Asset Category Hierarchy Support & General Clean-Up
Fixes & Minor refactor of validation/autoupdate to support native node classes properly post composition/preset support changes

diff --git a/Engine/Plugins/Runtime/Metasound/Source/MetasoundEditor/Private/MetasoundDetailCustomization.cpp b/Engine/Plugins/Runtime/Metasound/Source/MetasoundEditor/Private/MetasoundDetailCustomization.cpp
--- a/Engine/Plugins/Runtime/Metasound/Source/MetasoundEditor/Private/MetasoundDetailCustomization.cpp
+++ b/Engine/Plugins/Runtime/Metasound/Source/MetasoundEditor/Private/MetasoundDetailCustomization.cpp
@@ -79,12 +79,19 @@ namespace Metasound
 				{
 					IDetailCategoryBuilder& GeneralCategoryBuilder = DetailLayout.EditCategory("MetaSound");
 					const FName AuthorPropertyPath = BuildChildPath(GetMetadataPropertyPath(), FMetasoundFrontendClassMetadata::GetAuthorPropertyName());
+					const FName CategoryHierarchyPropertyPath = BuildChildPath(GetMetadataPropertyPath(), FMetasoundFrontendClassMetadata::GetCategoryHierarchyPropertyName());
 					const FName DescPropertyPath = BuildChildPath(GetMetadataPropertyPath(), FMetasoundFrontendClassMetadata::GetDescriptionPropertyName());
+
+					const FName ClassNamePropertyPath = BuildChildPath(GetMetadataPropertyPath(), FMetasoundFrontendClassMetadata::GetClassNamePropertyName());
+					const FName ClassNameNamePropertyPath = BuildChildPath(ClassNamePropertyPath, GET_MEMBER_NAME_CHECKED(FMetasoundFrontendClassName, Name));
+
 					const FName VersionPropertyPath = BuildChildPath(GetMetadataPropertyPath(), FMetasoundFrontendClassMetadata::GetVersionPropertyName());
 					const FName MajorVersionPropertyPath = BuildChildPath(VersionPropertyPath, GET_MEMBER_NAME_CHECKED(FMetasoundFrontendVersionNumber, Major));
 					const FName MinorVersionPropertyPath = BuildChildPath(VersionPropertyPath, GET_MEMBER_NAME_CHECKED(FMetasoundFrontendVersionNumber, Minor));
 
 					TSharedPtr<IPropertyHandle> AuthorHandle = DetailLayout.GetProperty(AuthorPropertyPath);
+					TSharedPtr<IPropertyHandle> CategoryHierarchyHandle = DetailLayout.GetProperty(CategoryHierarchyPropertyPath);
+					TSharedPtr<IPropertyHandle> ClassNameHandle = DetailLayout.GetProperty(ClassNameNamePropertyPath);
 					TSharedPtr<IPropertyHandle> DescHandle = DetailLayout.GetProperty(DescPropertyPath);
 					TSharedPtr<IPropertyHandle> MajorVersionHandle = DetailLayout.GetProperty(MajorVersionPropertyPath);
 					TSharedPtr<IPropertyHandle> MinorVersionHandle = DetailLayout.GetProperty(MinorVersionPropertyPath);
@@ -94,6 +101,24 @@ namespace Metasound
 					GeneralCategoryBuilder.AddProperty(MajorVersionHandle);
 					GeneralCategoryBuilder.AddProperty(MinorVersionHandle);
 
+					static const FText ClassGuidName = LOCTEXT("MetasoundClassGuidPropertyName", "Class Guid");
+					GeneralCategoryBuilder.AddCustomRow(ClassGuidName).NameContent()
+					[
+						SNew(SHorizontalBox)
+						+ SHorizontalBox::Slot()
+						.VAlign(VAlign_Center)
+						[
+							SNew(STextBlock)
+							.Text(ClassGuidName)
+							.Font(FEditorStyle::GetFontStyle(TEXT("PropertyWindow.NormalFont")))
+						]
+					]
+					.ValueContent()
+					[
+						ClassNameHandle->CreatePropertyValueWidget()
+					];
+					GeneralCategoryBuilder.AddProperty(CategoryHierarchyHandle);
+
 					DetailLayout.HideCategory("Advanced");
 					DetailLayout.HideCategory("Analysis");
 					DetailLayout.HideCategory("Attenuation");
diff --git a/Engine/Plugins/Runtime/Metasound/Source/MetasoundEditor/Private/MetasoundEditor.cpp b/Engine/Plugins/Runtime/Metasound/Source/MetasoundEditor/Private/MetasoundEditor.cpp
--- a/Engine/Plugins/Runtime/Metasound/Source/MetasoundEditor/Private/MetasoundEditor.cpp
+++ b/Engine/Plugins/Runtime/Metasound/Source/MetasoundEditor/Private/MetasoundEditor.cpp
@@ -50,8 +50,8 @@
 #include "Widgets/Images/SImage.h"
 #include "Widgets/Input/SButton.h"
 #include "Widgets/Notifications/SNotificationList.h"
-#include "Widgets/SOverlay.h"
 #include "Widgets/SBoxPanel.h"
+#include "Widgets/SOverlay.h"
 
 // This needs to be moved to public directory
 #include "../../GraphEditor/Private/GraphActionNode.h"
@@ -324,7 +324,8 @@ namespace Metasound
 			Metasound = ObjectToEdit;
 
 			const bool bAutoUpdate = true;
-			FGraphBuilder::ValidateGraph(*Metasound, bAutoUpdate);
+			const bool bClearUpdateNotes = false;
+			FGraphBuilder::ValidateGraph(*Metasound, bAutoUpdate, bClearUpdateNotes);
 
 			GEditor->RegisterForUndo(this);
 
@@ -387,6 +388,8 @@ namespace Metasound
 
 			ExtendToolbar();
 			RegenerateMenusAndToolbars();
+
+			FGraphBuilder::ValidateGraph(*Metasound, false /* bAutoUpdate */);
 		}
 
 		UObject* FEditor::GetMetasoundObject() const
@@ -1181,6 +1184,7 @@ namespace Metasound
 								if (MajorUpdateVersion.IsValid() && MajorUpdateVersion > NodeHandle->GetClassMetadata().GetVersion())
 								{
 									ExternalNode->UpdateToVersion(MajorUpdateVersion, false /* bInPropagateErrorMessages */);
+									FGraphBuilder::SynchronizeGraph(*Metasound);
 								}
 							}
 						}
@@ -1746,8 +1750,7 @@ namespace Metasound
 				NotifyUserModifiedBySync();
 			}
 
-			const bool bAutoUpdate = true;
-			FGraphBuilder::ValidateGraph(*Metasound, bAutoUpdate);
+			FGraphBuilder::ValidateGraph(*Metasound);
 
 			if (bNotifyReferenceLoop)
 			{
diff --git a/Engine/Plugins/Runtime/Metasound/Source/MetasoundEditor/Private/MetasoundEditorGraph.cpp b/Engine/Plugins/Runtime/Metasound/Source/MetasoundEditor/Private/MetasoundEditorGraph.cpp
--- a/Engine/Plugins/Runtime/Metasound/Source/MetasoundEditor/Private/MetasoundEditorGraph.cpp
+++ b/Engine/Plugins/Runtime/Metasound/Source/MetasoundEditor/Private/MetasoundEditorGraph.cpp
@@ -446,13 +446,13 @@ bool UMetasoundEditorGraph::Synchronize()
 	return false;
 }
 
-bool UMetasoundEditorGraph::Validate(bool bInAutoUpdate)
+bool UMetasoundEditorGraph::Validate(bool bInAutoUpdate, bool bInClearUpdateNotes)
 {
 	using namespace Metasound::Editor;
 
 	if (UObject* ParentMetasound = GetOuter())
 	{
-		return FGraphBuilder::ValidateGraph(*ParentMetasound, bInAutoUpdate);
+		return FGraphBuilder::ValidateGraph(*ParentMetasound, bInAutoUpdate, bInClearUpdateNotes);
 	}
 
 	return false;
diff --git a/Engine/Plugins/Runtime/Metasound/Source/MetasoundEditor/Private/MetasoundEditorGraphBuilder.cpp b/Engine/Plugins/Runtime/Metasound/Source/MetasoundEditor/Private/MetasoundEditorGraphBuilder.cpp
--- a/Engine/Plugins/Runtime/Metasound/Source/MetasoundEditor/Private/MetasoundEditorGraphBuilder.cpp
+++ b/Engine/Plugins/Runtime/Metasound/Source/MetasoundEditor/Private/MetasoundEditorGraphBuilder.cpp
@@ -50,9 +50,6 @@ namespace Metasound
 
 		const FName FGraphBuilder::PinSubCategoryTime = "time";
 
-		const FText FGraphBuilder::ConvertMenuName = LOCTEXT("MetasoundConversionsMenu", "Conversions");
-		const FText FGraphBuilder::FunctionMenuName = LOCTEXT("MetasoundFunctionsMenu", "Functions");
-
 		namespace GraphBuilderPrivate
 		{
 			void DeleteNode(UObject& InMetaSound, Frontend::FNodeHandle InNodeHandle)
@@ -149,50 +146,58 @@ namespace Metasound
 			return NewGraphNode;
 		}
 
-		bool FGraphBuilder::ValidateGraph(UObject& InMetaSoundToUpdate, bool bInAutoUpdate)
+		bool FGraphBuilder::ValidateGraph(UObject& InMetaSound, bool bInAutoUpdate, bool bInClearUpdateNotes)
 		{
 			using namespace Frontend;
 			using namespace GraphBuilderPrivate;
 
-			SynchronizeGraph(InMetaSoundToUpdate);
+			SynchronizeGraph(InMetaSound);
 
 			TSharedPtr<SGraphEditor> GraphEditor;
-			TSharedPtr<FEditor> MetaSoundEditor = GetEditorForMetasound(InMetaSoundToUpdate);
+			TSharedPtr<FEditor> MetaSoundEditor = GetEditorForMetasound(InMetaSound);
 			if (MetaSoundEditor.IsValid())
 			{
 				GraphEditor = MetaSoundEditor->GetGraphEditor();
 			}
 
-			FMetasoundAssetBase* MetaSoundAsset = IMetasoundUObjectRegistry::Get().GetObjectAsAssetBase(&InMetaSoundToUpdate);
+			FMetasoundAssetBase* MetaSoundAsset = IMetasoundUObjectRegistry::Get().GetObjectAsAssetBase(&InMetaSound);
 			check(MetaSoundAsset);
-
-			// Run initial validation pass before updating to capture upgrade message state prior to auto-update
 			UMetasoundEditorGraph& Graph = *CastChecked<UMetasoundEditorGraph>(&MetaSoundAsset->GetGraphChecked());
+
 			FGraphValidationResults Results;
-			Graph.ValidateInternal(Results);
 
-			bool bAutoUpdated = false;
-			if (Results.IsValid() && bInAutoUpdate)
+			auto ValidateGraphInteral = [&]()
 			{
-				check(MetaSoundAsset);
-				if (MetaSoundAsset->AutoUpdate(UMetaSoundAssetSubsystem::Get(), true /* bMarkDirty*/, true /* bUpdateReferencedAssets*/))
+				Graph.ValidateInternal(Results, bInClearUpdateNotes);
+				for (const FGraphNodeValidationResult& Result : Results.GetResults())
 				{
-					SynchronizeGraph(InMetaSoundToUpdate);
-
-					Graph.ValidateInternal(Results, false /* bClearUpgradeMessaging */);
-					for (const FGraphNodeValidationResult& Result : Results.GetResults())
+					if (GraphEditor.IsValid())
 					{
-						if (GraphEditor.IsValid())
-						{
-							check(Result.Node);
-							GraphEditor->RefreshNode(*Result.Node);
-							MetaSoundEditor->RefreshInterface();
-						}
+						check(Result.Node);
+						GraphEditor->RefreshNode(*Result.Node);
 					}
+				}
 
-					InMetaSoundToUpdate.MarkPackageDirty();
+				InMetaSound.MarkPackageDirty();
+			};
+
+			if (bInAutoUpdate)
+			{
+				if (MetaSoundAsset->AutoUpdate(UMetaSoundAssetSubsystem::Get(), true /* bMarkDirty*/, true /* bUpdateReferencedAssets*/))
+				{
+					SynchronizeGraph(InMetaSound);
+					ValidateGraphInteral();
 				}
 			}
+			else
+			{
+				ValidateGraphInteral();
+			}
+
+			if (MetaSoundEditor.IsValid())
+			{
+				MetaSoundEditor->RefreshInterface();
+			}
 
 			return Results.IsValid();
 		}
@@ -904,10 +909,6 @@ namespace Metasound
 			}
 
 			FNodeHandle NodeHandle = Node->GetNodeHandle();
-			if (!ensure(NodeHandle->IsValid()))
-			{
-				return false;
-			}
 
 			// Remove connects only to pins associated with this EdGraph node
 			// only (Iterate pins and not Frontend representation to preserve
diff --git a/Engine/Plugins/Runtime/Metasound/Source/MetasoundEditor/Private/MetasoundEditorGraphBuilder.h b/Engine/Plugins/Runtime/Metasound/Source/MetasoundEditor/Private/MetasoundEditorGraphBuilder.h
--- a/Engine/Plugins/Runtime/Metasound/Source/MetasoundEditor/Private/MetasoundEditorGraphBuilder.h
+++ b/Engine/Plugins/Runtime/Metasound/Source/MetasoundEditor/Private/MetasoundEditorGraphBuilder.h
@@ -50,8 +50,8 @@ namespace Metasound
 			// Custom pin-related styles for non-literal types (ex. wire color, pin heads, etc.)
 			static const FName PinSubCategoryTime; // Time type
 
-			static const FText ConvertMenuName;
 			static const FText FunctionMenuName;
+			static const FText GraphMenuName;
 
 			// Adds an EdGraph node to mirror the provided FNodeHandle.
 			static UMetasoundEditorGraphNode* AddNode(UObject& InMetaSound, Frontend::FNodeHandle InNodeHandle, FVector2D InLocation, bool bInSelectNewNode = true);
@@ -61,7 +61,7 @@ namespace Metasound
 			static TSharedPtr<FEditor> GetEditorForGraph(const UEdGraph& EdGraph);
 
 			// Validates MetaSound and optionally performs AutoUpdate logic in Frontend with any additional logic required to refresh editor & respective editor object state.
-			static bool ValidateGraph(UObject& InMetaSoundToUpdate, bool bAutoUpdate = true);
+			static bool ValidateGraph(UObject& InMetaSound, bool bAutoUpdate = true, bool bInClearUpdateNotes = false);
 
 			// Wraps RegisterGraphWithFrontend logic in Frontend with any additional logic required to refresh editor & respective editor object state.
 			static void RegisterGraphWithFrontend(UObject& InMetaSound);
diff --git a/Engine/Plugins/Runtime/Metasound/Source/MetasoundEditor/Private/MetasoundEditorGraphNode.cpp b/Engine/Plugins/Runtime/Metasound/Source/MetasoundEditor/Private/MetasoundEditorGraphNode.cpp
--- a/Engine/Plugins/Runtime/Metasound/Source/MetasoundEditor/Private/MetasoundEditorGraphNode.cpp
+++ b/Engine/Plugins/Runtime/Metasound/Source/MetasoundEditor/Private/MetasoundEditorGraphNode.cpp
@@ -38,6 +38,11 @@ namespace Metasound
 				InNode.bHasCompilerMessage = true;
 				InNode.ErrorMsg = InMessage;
 				InNode.ErrorType = InSeverity;
+
+				if (InSeverity == EMessageSeverity::Error)
+				{
+					UE_LOG(LogMetasoundEditor, Error, TEXT("%s"), *InMessage);
+				}
 			}
 		} // namespace GraphNodePrivate
 	} // namespace Editor
@@ -573,7 +578,7 @@ bool UMetasoundEditorGraphExternalNode::CanAutoUpdate() const
 {
 	using namespace Metasound::Frontend;
 
-	return GetConstNodeHandle()->CanAutoUpdate();
+	return GetConstNodeHandle()->CanAutoUpdate(UMetaSoundAssetSubsystem::Get());
 }
 
 UMetasoundEditorGraphExternalNode* UMetasoundEditorGraphExternalNode::UpdateToVersion(const FMetasoundFrontendVersionNumber& InNewVersion, bool bInPropagateErrorMessages)
@@ -612,23 +617,6 @@ UMetasoundEditorGraphExternalNode* UMetasoundEditorGraphExternalNode::UpdateToVe
 	GetGraph()->RemoveNode(this);
 	MetaSound.Modify();
 
-	// TODO: Pare down to SynchronizeNode/NodeConnections call to avoid having to synchronize the entire graph for all version calls
-	FGraphBuilder::SynchronizeGraph(MetaSound);
-
-	if (InitVersion != InNewVersion)
-	{
-		NodeUpgradeMessage = FText::Format(LOCTEXT("NodeVersionUpgradeMessageFormat", "Class interface updated: {0} to {1}"),
-			FText::FromString(*InitVersion.ToString()),
-			FText::FromString(*InNewVersion.ToString())
-		);
-	}
-	else
-	{
-		NodeUpgradeMessage = FText::Format(LOCTEXT("NodeVersionRefreshMessageFormat", "Class interface updated: detected interface change for {0}"),
-			FText::FromString(*InitVersion.ToString())
-		);
-	}
-
 	if (InitNodeID != ReplacementNodeID)
 	{
 		if (bInPropagateErrorMessages)
@@ -649,8 +637,7 @@ bool UMetasoundEditorGraphExternalNode::Validate(Metasound::Editor::FGraphNodeVa
 
 	OutResult = FGraphNodeValidationResult(*this);
 
-	FNodeHandle NodeHandle = GetNodeHandle();
-
+	// 1. Reset ed node validation state
 	if (ErrorType != EMessageSeverity::Info)
 	{
 		ErrorType = EMessageSeverity::Info;
@@ -664,80 +651,166 @@ bool UMetasoundEditorGraphExternalNode::Validate(Metasound::Editor::FGraphNodeVa
 		OutResult.bIsDirty = true;
 	}
 
-	if (bClearUpgradeMessage && !NodeUpgradeMessage.IsEmpty())
+	if (bClearUpgradeMessage)
 	{
-		NodeUpgradeMessage = FText::GetEmpty();
-		OutResult.bIsDirty = true;
-	}
-
-	if (!NodeHandle->IsValid())
-	{
-		GraphNodePrivate::SetGraphNodeMessage(*this, EMessageSeverity::Error, TEXT("Class definition missing (If defined in asset, corrisponding graph definition may be invalid)"));
-		OutResult.bIsInvalid = true;
-		OutResult.bIsDirty = true;
-		return !OutResult.bIsInvalid;
+		if (!NodeUpgradeMessage.IsEmpty())
+		{
+			NodeUpgradeMessage = FText::GetEmpty();
+			OutResult.bIsDirty = true;
+		}
 	}
 
+	// 2. Check if node is invalid, version is missing and cache if interface changes exist between the document's records and the registry
+	FNodeHandle NodeHandle = GetNodeHandle();
 	const FMetasoundFrontendClassMetadata& Metadata = NodeHandle->GetClassMetadata();
-	const Metasound::FNodeClassName NodeClassName = Metadata.GetClassName().ToNodeClassName();
-	const TArray<FMetasoundFrontendClass> SortedClasses = ISearchEngine::Get().FindClassesWithName(NodeClassName, true /* bSortByVersion */);
 
-	if (SortedClasses.IsEmpty())
+	FClassInterfaceUpdates InterfaceUpdates;
+	if (!NodeHandle->DiffAgainstRegistryInterface(InterfaceUpdates, false /* bUseHighestMinorVersion */))
 	{
-		OutResult.MissingClass = NodeClassName.GetFullName();
-		const FString NewErrorMsg = FString::Format(TEXT("Class definition '{0}' missing: {1}"),
+		
+		if (NodeHandle->IsValid())
+		{
+			const FText* PromptIfMissing = nullptr;
+			FString FormattedClassName;
+			if (bIsClassNative)
 			{
-				*Metadata.GetClassName().ToString(),
-				*Metadata.GetPromptIfMissing().ToString()
+				PromptIfMissing = &Metadata.GetPromptIfMissing();
+				FormattedClassName = FString::Format(TEXT("{0} {1} ({2})"), { *Metadata.GetDisplayName().ToString(), *Metadata.GetVersion().ToString(), *Metadata.GetClassName().ToString() });
+			}
+			else
+			{
+				static const FText AssetPromptIfMissing = LOCTEXT("PromptIfAssetMissing", "Asset may have not been saved, deleted or is not loaded (ex. in an unloaded plugin).");
+				PromptIfMissing = &AssetPromptIfMissing;
+				FormattedClassName = FString::Format(TEXT("{0} {1} ({2})"), { *Metadata.GetDisplayName().ToString(), *Metadata.GetVersion().ToString(), *Metadata.GetClassName().Name.ToString() });
+			}
+
+			const FString NewErrorMsg = FString::Format(TEXT("Class definition '{0}' not found: {1}"),
+			{
+				*FormattedClassName,
+				*PromptIfMissing->ToString()
 			});
 
-		for (UEdGraphPin* Pin : Pins)
+			GraphNodePrivate::SetGraphNodeMessage(*this, EMessageSeverity::Error, NewErrorMsg);
+		}
+		else
 		{
-			Pin->bOrphanedPin = true;
+			if (bIsClassNative)
+			{
+				GraphNodePrivate::SetGraphNodeMessage(*this, EMessageSeverity::Error, FString::Format(
+					TEXT("Class '{0}' definition missing for last known natively defined node."),
+					{ *ClassName.ToString() }));
+			}
+			else
+			{
+				GraphNodePrivate::SetGraphNodeMessage(*this, EMessageSeverity::Error,
+					FString::Format(TEXT("Class definition missing for asset with guid '{0}': Asset is either missing or invalid"),
+					{ *ClassName.Name.ToString() }));
+			}
 		}
 
-		GraphNodePrivate::SetGraphNodeMessage(*this, EMessageSeverity::Error, NewErrorMsg);
 		OutResult.bIsDirty = true;
 		OutResult.bIsInvalid = true;
-		return !OutResult.bIsInvalid;
 	}
 
+	// 3. Report if node was nativized
+	const FNodeRegistryKey RegistryKey = NodeRegistryKey::CreateKey(Metadata);
+	bool bNewIsClassNative = FMetasoundFrontendRegistryContainer::Get()->IsNodeNative(RegistryKey);
+	if (bIsClassNative != bNewIsClassNative)
+	{
+		if (bNewIsClassNative)
+		{
+			GraphNodePrivate::SetGraphNodeMessage(*this, EMessageSeverity::Info, FString::Format(TEXT("Class '{0}' has been nativized."),
+			{
+				*Metadata.GetDisplayName().ToString()
+			}));
+		}
+
+		OutResult.bIsDirty = true;
+		bIsClassNative = bNewIsClassNative;
+	}
+
+	// 4. Report if node was auto-updated
+	FMetasoundFrontendNodeStyle Style = NodeHandle->GetNodeStyle();
+	if (Style.bMessageNodeUpdated)
+	{
+		if (bClearUpgradeMessage)
+		{
+			Style.bMessageNodeUpdated = false;
+			NodeHandle->SetNodeStyle(Style);
+		}
+		else
+		{
+			GraphNodePrivate::SetGraphNodeMessage(*this, EMessageSeverity::Info, FString::Format(TEXT("Node class '{0}' updated to version {1}"),
+			{
+				*Metadata.GetDisplayName().ToString(),
+				*Metadata.GetVersion().ToString()
+			}));
+		}
+
+		OutResult.bIsDirty = true;
+	}
+
+	// 5. Reset pin state (if pin was orphaned or clear if no longer orphaned)
 	for (UEdGraphPin* Pin : Pins)
 	{
-		if (Pin->bOrphanedPin)
+		bool bWasRemoved = false;
+		if (Pin->Direction == EGPD_Input)
+		{
+			FInputHandle Input = FGraphBuilder::GetInputHandleFromPin(Pin);
+			bWasRemoved |= InterfaceUpdates.RemovedInputs.ContainsByPredicate([&](const FMetasoundFrontendClassInput& ClassInput)
+			{
+				return Input->GetName() == ClassInput.Name && Input->GetDataType() == ClassInput.TypeName;
+			});
+		}
+
+		if (Pin->Direction == EGPD_Output)
+		{
+			FOutputHandle Output = FGraphBuilder::GetOutputHandleFromPin(Pin);
+			bWasRemoved |= InterfaceUpdates.RemovedOutputs.ContainsByPredicate([&](const FMetasoundFrontendClassOutput& ClassOutput)
+				{
+					return Output->GetName() == ClassOutput.Name && Output->GetDataType() == ClassOutput.TypeName;
+				});
+		}
+
+		if (Pin->bOrphanedPin != bWasRemoved)
 		{
-			Pin->bOrphanedPin = false;
+			Pin->bOrphanedPin = bWasRemoved;
 			OutResult.bIsDirty = true;
 		}
 	}
 
+	// 6. Find all available versions & report if upgrade available
+	const Metasound::FNodeClassName NodeClassName = Metadata.GetClassName().ToNodeClassName();
+	const TArray<FMetasoundFrontendClass> SortedClasses = ISearchEngine::Get().FindClassesWithName(NodeClassName, true /* bSortByVersion */);
 	const FMetasoundFrontendVersionNumber& CurrentVersion = Metadata.GetVersion();
-	if (SortedClasses[0].Metadata.GetVersion() > CurrentVersion)
+	const FMetasoundFrontendClass& HighestRegistryClass = SortedClasses[0];
+	if (HighestRegistryClass.Metadata.GetVersion() > CurrentVersion)
 	{
 		FMetasoundFrontendClass HighestMinorVersionClass;
 		FString NodeMsg;
 		EMessageSeverity::Type Severity;
 
-		auto ClassVersionExists = [InCurrentVersion = &CurrentVersion](const FMetasoundFrontendClass& AvailableClass)
+		const bool bClassVersionExists = SortedClasses.ContainsByPredicate([InCurrentVersion = &CurrentVersion](const FMetasoundFrontendClass& AvailableClass)
 		{
 			return AvailableClass.Metadata.GetVersion() == *InCurrentVersion;
-		};
-		if (SortedClasses.ContainsByPredicate(ClassVersionExists))
+		});
+		if (bClassVersionExists)
 		{
-			NodeMsg = FString::Format(TEXT("Node class '{0}' out-of-date: Eligible for upgrade to {1}"),
+			NodeMsg = FString::Format(TEXT("Node class '{0} {1}' is prior version: Eligible for upgrade to {2}"),
 			{
 				*Metadata.GetClassName().ToString(),
-				*SortedClasses[0].Metadata.GetVersion().ToString()
+				*Metadata.GetVersion().ToString(),
+				*HighestRegistryClass.Metadata.GetVersion().ToString()
 			});
 			Severity = EMessageSeverity::Warning;
 		}
 		else
 		{
-			NodeMsg = FString::Format(TEXT("Node class '{0} ({1})' is missing.  Highest version ({2}) found."),
+			NodeMsg = FString::Format(TEXT("Node class '{0} {1}' is missing and ineligible for auto-update.  Highest version '{2}' found."),
 			{
 				*Metadata.GetClassName().ToString(),
 				*Metadata.GetVersion().ToString(),
-				*SortedClasses[0].Metadata.GetVersion().ToString()
+				*HighestRegistryClass.Metadata.GetVersion().ToString()
 			});
 			Severity = EMessageSeverity::Error;
 			OutResult.bIsInvalid = true;
@@ -745,13 +818,33 @@ bool UMetasoundEditorGraphExternalNode::Validate(Metasound::Editor::FGraphNodeVa
 
 		GraphNodePrivate::SetGraphNodeMessage(*this, Severity, NodeMsg);
 		OutResult.bIsDirty = true;
-		return !OutResult.bIsInvalid;
 	}
-
-	OutResult.bIsAutoUpdateEligible = CanAutoUpdate();
-
-	const FNodeRegistryKey RegistryKey = NodeRegistryKey::CreateKey(Metadata);
-	bIsClassNative = FMetasoundFrontendRegistryContainer::Get()->IsNodeNative(RegistryKey);
+	else if (HighestRegistryClass.Metadata.GetVersion() == CurrentVersion)
+	{
+		if (InterfaceUpdates.ContainsChanges())
+		{
+			GraphNodePrivate::SetGraphNodeMessage(*this, EMessageSeverity::Error,
+			FString::Format(TEXT("Node & registered class interface mismatch: '{0} {1}'. Class either versioned improperly, class key collision exists, or AutoUpdate disabled in 'MetaSound' Developer Settings."),
+			{
+				*Metadata.GetClassName().ToString(),
+				*Metadata.GetVersion().ToString()
+			}));
+			OutResult.bIsDirty = true;
+			OutResult.bIsInvalid = true;
+		}
+	}
+	else
+	{
+		GraphNodePrivate::SetGraphNodeMessage(*this, EMessageSeverity::Error,
+			FString::Format(TEXT("Node with class '{0} {1}' interface version higher than that of highest minor revision ({2}) in class registry."),
+			{
+				*Metadata.GetClassName().ToString(),
+				*Metadata.GetVersion().ToString(),
+				*HighestRegistryClass.Metadata.GetVersion().ToString()
+			}));
+		OutResult.bIsDirty = true;
+		OutResult.bIsInvalid = true;
+	}
 
 	return !OutResult.bIsInvalid;
 }
diff --git a/Engine/Plugins/Runtime/Metasound/Source/MetasoundEditor/Private/MetasoundEditorGraphSchema.cpp b/Engine/Plugins/Runtime/Metasound/Source/MetasoundEditor/Private/MetasoundEditorGraphSchema.cpp
--- a/Engine/Plugins/Runtime/Metasound/Source/MetasoundEditor/Private/MetasoundEditorGraphSchema.cpp
+++ b/Engine/Plugins/Runtime/Metasound/Source/MetasoundEditor/Private/MetasoundEditorGraphSchema.cpp
@@ -24,9 +24,12 @@
 #include "MetasoundEditorModule.h"
 #include "MetasoundEditorSettings.h"
 #include "MetasoundFrontend.h"
+#include "MetasoundFrontendNodesCategories.h"
 #include "MetasoundFrontendRegistries.h"
 #include "MetasoundFrontendSearchEngine.h"
+#include "MetasoundStandardNodesCategories.h"
 #include "ScopedTransaction.h"
+#include "Styling/SlateStyleRegistry.h"
 #include "Toolkits/ToolkitManager.h"
 #include "ToolMenus.h"
 #include "UObject/NoExportTypes.h"
@@ -49,6 +52,9 @@ namespace Metasound
 	{
 		namespace SchemaPrivate
 		{
+			static const FText CategoryDelim = LOCTEXT("MetaSoundActionsCategoryDelim", "|");
+			static const FText KeywordDelim = LOCTEXT("MetaSoundKeywordDelim", " ");
+
 			static const FText InputDisplayNameFormat = LOCTEXT("DisplayNameAddInputFormat", "Get {0}");
 			static const FText InputTooltipFormat = LOCTEXT("TooltipAddInputFormat", "Adds a getter for the input '{0}' to the graph.");
 
@@ -60,7 +66,9 @@ namespace Metasound
 				Default = 0,
 				Inputs,
 				Outputs,
-				External
+				Graphs,
+				Functions,
+				Conversions
 			};
 
 			const FText& GetContextGroupDisplayName(EPrimaryContextGroup InContextGroup)
@@ -69,18 +77,34 @@ namespace Metasound
 				{
 					case EPrimaryContextGroup::Inputs:
 					{
-						static const FText InputGroupName = LOCTEXT("InputActions", "Input Actions");
+						static const FText InputGroupName = LOCTEXT("InputActions", "Inputs");
 						return InputGroupName;
 					}
 					case EPrimaryContextGroup::Outputs:
 					{
-						static const FText OutputGroupName = LOCTEXT("OutputActions", "Output Actions");
+						static const FText OutputGroupName = LOCTEXT("OutputActions", "Outputs");
 						return OutputGroupName;
 					}
 
-					// External nodes & any other group other than those above should not use a primary context group display name
+					case EPrimaryContextGroup::Graphs:
+					{
+						static const FText GraphGroupName = LOCTEXT("GraphActions", "Graphs");
+						return GraphGroupName;
+					}
+
+					case EPrimaryContextGroup::Functions:
+					{
+						static const FText FunctionGroupName = LOCTEXT("FunctionActions", "Functions");
+						return FunctionGroupName;
+					}
+
+					case EPrimaryContextGroup::Conversions:
+					{
+						static const FText ConversionGroupName = LOCTEXT("ConversionActions", "Conversions");
+						return ConversionGroupName;
+					}
+
 					case EPrimaryContextGroup::Default:
-					case EPrimaryContextGroup::External:
 					default:
 					{
 						checkNoEntry();
@@ -173,6 +197,45 @@ namespace Metasound
 	} // namespace Editor
 } // namespace Metasound
 
+
+const FSlateBrush* FMetasoundGraphSchemaAction_NewNode::GetIconBrush() const
+{
+	using namespace Metasound::Frontend;
+
+	if (const ISlateStyle* MetasoundStyle = FSlateStyleRegistry::FindSlateStyle("MetaSoundStyle"))
+	{
+		const FNodeRegistryKey RegistryKey = NodeRegistryKey::CreateKey(ClassMetadata);
+		const bool bIsClassNative = FMetasoundFrontendRegistryContainer::Get()->IsNodeNative(RegistryKey);
+		if (bIsClassNative)
+		{
+			return MetasoundStyle->GetBrush(TEXT("MetasoundEditor.Graph.Node.Class.Native"));
+		}
+
+		return MetasoundStyle->GetBrush(TEXT("MetasoundEditor.Graph.Node.Class.Graph"));
+	}
+
+	return Super::GetIconBrush();
+}
+
+const FLinearColor& FMetasoundGraphSchemaAction_NewNode::GetIconColor() const
+{
+	using namespace Metasound::Frontend;
+
+	if (const UMetasoundEditorSettings* EditorSettings = GetDefault<UMetasoundEditorSettings>())
+	{
+		const FNodeRegistryKey RegistryKey = NodeRegistryKey::CreateKey(ClassMetadata);
+		const bool bIsClassNative = FMetasoundFrontendRegistryContainer::Get()->IsNodeNative(RegistryKey);
+		if (bIsClassNative)
+		{
+			return EditorSettings->NativeNodeTitleColor;
+		}
+
+		return EditorSettings->AssetReferenceNodeTitleColor;
+	}
+
+	return Super::GetIconColor();
+}
+
 UEdGraphNode* FMetasoundGraphSchemaAction_NewNode::PerformAction(UEdGraph* ParentGraph, UEdGraphPin* FromPin, const FVector2D Location, bool bSelectNewNode /* = true */)
 {
 	using namespace Metasound::Editor;
@@ -194,11 +257,30 @@ UEdGraphNode* FMetasoundGraphSchemaAction_NewNode::PerformAction(UEdGraph* Paren
 }
 
 FMetasoundGraphSchemaAction_NewInput::FMetasoundGraphSchemaAction_NewInput(FText InNodeCategory, FText InDisplayName, FGuid InNodeID, FText InToolTip, const int32 InGrouping)
-	: FEdGraphSchemaAction(MoveTemp(InNodeCategory), MoveTemp(InDisplayName), MoveTemp(InToolTip), InGrouping)
+	: FMetasoundGraphSchemaAction(MoveTemp(InNodeCategory), MoveTemp(InDisplayName), MoveTemp(InToolTip), InGrouping)
 	, NodeID(InNodeID)
 {
 }
 
+const FSlateBrush* FMetasoundGraphSchemaAction_NewInput::GetIconBrush() const
+{
+	if (const ISlateStyle* MetasoundStyle = FSlateStyleRegistry::FindSlateStyle("MetaSoundStyle"))
+	{
+		return MetasoundStyle->GetBrush("MetasoundEditor.Graph.Node.Class.Input");
+	}
+	return Super::GetIconBrush();
+}
+
+const FLinearColor& FMetasoundGraphSchemaAction_NewInput::GetIconColor() const
+{
+	if (const UMetasoundEditorSettings* EditorSettings = GetDefault<UMetasoundEditorSettings>())
+	{
+		return EditorSettings->InputNodeTitleColor;
+	}
+
+	return Super::GetIconColor();
+}
+
 UEdGraphNode* FMetasoundGraphSchemaAction_NewInput::PerformAction(UEdGraph* ParentGraph, UEdGraphPin* FromPin, const FVector2D InLocation, bool bSelectNewNode /* = true */)
 {
 	using namespace Metasound::Editor;
@@ -239,7 +321,7 @@ UEdGraphNode* FMetasoundGraphSchemaAction_NewInput::PerformAction(UEdGraph* Pare
 }
 
 FMetasoundGraphSchemaAction_PromoteToInput::FMetasoundGraphSchemaAction_PromoteToInput(FText InNodeCategory, FText InDisplayName, FText InToolTip, const int32 InGrouping)
-	: FEdGraphSchemaAction(MoveTemp(InNodeCategory), MoveTemp(InDisplayName), MoveTemp(InToolTip), InGrouping)
+	: FMetasoundGraphSchemaAction_NewInput(MoveTemp(InNodeCategory), MoveTemp(InDisplayName), FGuid(), MoveTemp(InToolTip), InGrouping)
 {
 }
 
@@ -294,11 +376,31 @@ UEdGraphNode* FMetasoundGraphSchemaAction_PromoteToInput::PerformAction(UEdGraph
 }
 
 FMetasoundGraphSchemaAction_NewOutput::FMetasoundGraphSchemaAction_NewOutput(FText InNodeCategory, FText InDisplayName, FGuid InOutputNodeID, FText InToolTip, const int32 InGrouping)
-	: FEdGraphSchemaAction(MoveTemp(InNodeCategory), MoveTemp(InDisplayName), MoveTemp(InToolTip), InGrouping)
+	: FMetasoundGraphSchemaAction(MoveTemp(InNodeCategory), MoveTemp(InDisplayName), MoveTemp(InToolTip), InGrouping)
 	, NodeID(InOutputNodeID)
 {
 }
 
+const FSlateBrush* FMetasoundGraphSchemaAction_NewOutput::GetIconBrush() const
+{
+	if (const ISlateStyle* MetasoundStyle = FSlateStyleRegistry::FindSlateStyle("MetaSoundStyle"))
+	{
+		return MetasoundStyle->GetBrush("MetasoundEditor.Graph.Node.Class.Output");
+	}
+
+	return Super::GetIconBrush();
+}
+
+const FLinearColor& FMetasoundGraphSchemaAction_NewOutput::GetIconColor() const
+{
+	if (const UMetasoundEditorSettings* EditorSettings = GetDefault<UMetasoundEditorSettings>())
+	{
+		return EditorSettings->OutputNodeTitleColor;
+	}
+
+	return Super::GetIconColor();
+}
+
 UEdGraphNode* FMetasoundGraphSchemaAction_NewOutput::PerformAction(UEdGraph* ParentGraph, UEdGraphPin* FromPin, const FVector2D Location, bool bSelectNewNode /* = true */)
 {
 	using namespace Metasound::Editor;
@@ -334,7 +436,7 @@ UEdGraphNode* FMetasoundGraphSchemaAction_NewOutput::PerformAction(UEdGraph* Par
 }
 
 FMetasoundGraphSchemaAction_PromoteToOutput::FMetasoundGraphSchemaAction_PromoteToOutput(FText InNodeCategory, FText InDisplayName, FText InToolTip, const int32 InGrouping)
-	: FEdGraphSchemaAction(MoveTemp(InNodeCategory), MoveTemp(InDisplayName), MoveTemp(InToolTip), InGrouping)
+	: FMetasoundGraphSchemaAction_NewOutput(MoveTemp(InNodeCategory), MoveTemp(InDisplayName), FGuid(), MoveTemp(InToolTip), InGrouping)
 {
 }
 
@@ -414,6 +516,18 @@ UEdGraphNode* FMetasoundGraphSchemaAction_NewComment::PerformAction(UEdGraph* Pa
 	return FEdGraphSchemaAction_NewNode::SpawnNodeFromTemplate<UEdGraphNode_Comment>(ParentGraph, CommentTemplate, SpawnLocation);
 }
 
+const FSlateBrush* FMetasoundGraphSchemaAction_NewComment::GetIconBrush() const
+{
+	// TODO: Implement (Find icon & rig up)
+	return Super::GetIconBrush();
+}
+
+const FLinearColor& FMetasoundGraphSchemaAction_NewComment::GetIconColor() const
+{
+	// TODO: Implement (Set to white when icon found)
+	return Super::GetIconColor();
+}
+
 UEdGraphNode* FMetasoundGraphSchemaAction_Paste::PerformAction(UEdGraph* ParentGraph, UEdGraphPin* FromPin, const FVector2D Location, bool bSelectNewNode/* = true*/)
 {
 	using namespace Metasound::Editor;
@@ -543,10 +657,10 @@ void UMetasoundEditorGraphSchema::GetGraphContextActions(FGraphContextMenuBuilde
 			});
 
 			TSharedPtr<FMetasoundGraphSchemaAction_PromoteToOutput> NewNodeAction = MakeShared<FMetasoundGraphSchemaAction_PromoteToOutput>(
-				SchemaPrivate::GetContextGroupDisplayName(SchemaPrivate::EPrimaryContextGroup::Inputs),
+				SchemaPrivate::GetContextGroupDisplayName(SchemaPrivate::EPrimaryContextGroup::Outputs),
 				LOCTEXT("PromoteToOutputName", "Promote To Graph Output"),
 				LOCTEXT("PromoteToOutputTooltip", "Promotes node output to graph output"),
-				static_cast<int32>(SchemaPrivate::EPrimaryContextGroup::Inputs));
+				static_cast<int32>(SchemaPrivate::EPrimaryContextGroup::Outputs));
 
 			static_cast<FGraphActionMenuBuilder&>(ContextMenuBuilder).AddAction(NewNodeAction);
 		}
@@ -961,8 +1075,8 @@ void UMetasoundEditorGraphSchema::DroppedAssetsOnNode(const TArray<FAssetData>&
 void UMetasoundEditorGraphSchema::GetConversionActions(FGraphActionMenuBuilder& ActionMenuBuilder, Metasound::Editor::FActionClassFilters InFilters, bool bShowSelectedActions) const
 {
 	using namespace Metasound;
+	using namespace Metasound::Editor;
 
-	const FText MenuJoinFormat = LOCTEXT("MetasoundActionsFormatSubCategory", "{0}|{1}");
 	const bool bIncludeDeprecated = static_cast<bool>(EnableDeprecatedMetaSoundNodeClassCreationCVar);
 	const TArray<FMetasoundFrontendClass> FrontendClasses = Frontend::ISearchEngine::Get().FindAllClasses(bIncludeDeprecated);
 	for (const FMetasoundFrontendClass& FrontendClass : FrontendClasses)
@@ -982,14 +1096,16 @@ void UMetasoundEditorGraphSchema::GetConversionActions(FGraphActionMenuBuilder&
 			? Metadata.GetDescription()
 			: FText::Format(LOCTEXT("MetasoundTooltipAuthorFormat", "{0}\nAuthor: {1}"), Metadata.GetDescription(), Metadata.GetAuthor());
 
-		if (!Metadata.GetCategoryHierarchy().IsEmpty() && !Metadata.GetCategoryHierarchy()[0].CompareTo(Editor::FGraphBuilder::ConvertMenuName))
+		if (!Metadata.GetCategoryHierarchy().IsEmpty() && !Metadata.GetCategoryHierarchy()[0].CompareTo(NodeCategories::Conversions))
 		{
+			FText KeywordsText = FText::Join(SchemaPrivate::KeywordDelim, Metadata.GetKeywords());
 			TSharedPtr<FMetasoundGraphSchemaAction_NewNode> NewNodeAction = MakeShared<FMetasoundGraphSchemaAction_NewNode>
 			(
-				Editor::FGraphBuilder::ConvertMenuName,
+				NodeCategories::Conversions,
 				Metadata.GetDisplayName(),
 				Tooltip,
-				0
+				static_cast<int32>(SchemaPrivate::EPrimaryContextGroup::Conversions),
+				KeywordsText
 			);
 
 			Metadata.SetType(EMetasoundFrontendClassType::External);
@@ -1059,10 +1175,9 @@ void UMetasoundEditorGraphSchema::GetDataTypeOutputNodeActions(FGraphContextMenu
 void UMetasoundEditorGraphSchema::GetFunctionActions(FGraphActionMenuBuilder& ActionMenuBuilder, Metasound::Editor::FActionClassFilters InFilters, bool bShowSelectedActions, Metasound::Frontend::FConstGraphHandle InGraphHandle) const
 {
 	using namespace Metasound::Editor;
+	using namespace Metasound::Editor::SchemaPrivate;
 	using namespace Metasound::Frontend;
 
-	const FText MenuJoinFormat = LOCTEXT("MetasoundActionsFormatSubCategory", "{0}|{1}");
-
 	const UMetaSoundAssetSubsystem& AssetSubsystem = UMetaSoundAssetSubsystem::Get();
 	FMetasoundAssetBase* ParentAsset = nullptr;
 	if (InGraphHandle->IsValid())
@@ -1087,10 +1202,17 @@ void UMetasoundEditorGraphSchema::GetFunctionActions(FGraphActionMenuBuilder& Ac
 			continue;
 		}
 
-		if (ParentAsset && FrontendClass.Metadata.GetType() == EMetasoundFrontendClassType::External)
+		if (FrontendClass.Metadata.GetType() != EMetasoundFrontendClassType::External)
+		{
+			continue;
+		}
+
+		const FSoftObjectPath* Path = nullptr;
+		if (ParentAsset)
 		{
 			const FNodeRegistryKey RegistryKey = NodeRegistryKey::CreateKey(FrontendClass.Metadata);
-			if (const FSoftObjectPath* Path = AssetSubsystem.FindObjectPathFromKey(RegistryKey))
+			Path = AssetSubsystem.FindObjectPathFromKey(RegistryKey);
+			if (Path)
 			{
 				if (ParentAsset->AddingReferenceCausesLoop(*Path, AssetSubsystem))
 				{
@@ -1104,15 +1226,21 @@ void UMetasoundEditorGraphSchema::GetFunctionActions(FGraphActionMenuBuilder& Ac
 			? Metadata.GetDescription()
 			: FText::Format(LOCTEXT("MetasoundTooltipAuthorFormat", "{0}\nAuthor: {1}"), Metadata.GetDescription(), Metadata.GetAuthor());
 
-		if (Metadata.GetCategoryHierarchy().IsEmpty() || Metadata.GetCategoryHierarchy()[0].CompareTo(FGraphBuilder::ConvertMenuName))
+		if (Metadata.GetCategoryHierarchy().IsEmpty() || Metadata.GetCategoryHierarchy()[0].CompareTo(Metasound::NodeCategories::Conversions))
 		{
-			const FText CategoriesText = FText::Join(LOCTEXT("MetasoundActionsCategoryDelim", "|"), Metadata.GetCategoryHierarchy());
+			TArray<FText> CategoryHeirarchy { GetContextGroupDisplayName(Path ? EPrimaryContextGroup::Graphs : EPrimaryContextGroup::Functions) };
+			CategoryHeirarchy.Append(Metadata.GetCategoryHierarchy());
+			const FText KeywordsText = FText::Join(KeywordDelim, Metadata.GetKeywords());
+			const FText CategoryText = FText::Join(CategoryDelim, CategoryHeirarchy);
+			const int32 GroupID = static_cast<int32>(Path ? EPrimaryContextGroup::Graphs : EPrimaryContextGroup::Functions);
+
 			TSharedPtr<FMetasoundGraphSchemaAction_NewNode> NewNodeAction = MakeShared<FMetasoundGraphSchemaAction_NewNode>
 			(
-				FText::Format(MenuJoinFormat, FGraphBuilder::FunctionMenuName, CategoriesText),
+				CategoryText,
 				Metadata.GetDisplayName(),
 				Tooltip,
-				0
+				GroupID,
+				KeywordsText
 			);
 
 			Metadata.SetType(EMetasoundFrontendClassType::External);
diff --git a/Engine/Plugins/Runtime/Metasound/Source/MetasoundEditor/Public/MetasoundEditorGraph.h b/Engine/Plugins/Runtime/Metasound/Source/MetasoundEditor/Public/MetasoundEditorGraph.h
--- a/Engine/Plugins/Runtime/Metasound/Source/MetasoundEditor/Public/MetasoundEditorGraph.h
+++ b/Engine/Plugins/Runtime/Metasound/Source/MetasoundEditor/Public/MetasoundEditorGraph.h
@@ -159,7 +159,7 @@ public:
 	bool IsPreviewing() const;
 
 	virtual bool Synchronize() override;
-	virtual bool Validate(bool bInAutoUpdate) override;
+	virtual bool Validate(bool bInAutoUpdate, bool bInClearUpdateNotes) override;
 
 	virtual void RegisterGraphWithFrontend() override;
 
diff --git a/Engine/Plugins/Runtime/Metasound/Source/MetasoundEditor/Public/MetasoundEditorGraphSchema.h b/Engine/Plugins/Runtime/Metasound/Source/MetasoundEditor/Public/MetasoundEditorGraphSchema.h
--- a/Engine/Plugins/Runtime/Metasound/Source/MetasoundEditor/Public/MetasoundEditorGraphSchema.h
+++ b/Engine/Plugins/Runtime/Metasound/Source/MetasoundEditor/Public/MetasoundEditorGraphSchema.h
@@ -38,9 +38,35 @@ namespace Metasound
 	} // namespace Editor
 } // namespace Metasound
 
+USTRUCT()
+struct METASOUNDEDITOR_API FMetasoundGraphSchemaAction : public FEdGraphSchemaAction
+{
+	GENERATED_USTRUCT_BODY();
+
+	FMetasoundGraphSchemaAction()
+		: FEdGraphSchemaAction()
+	{}
+
+	FMetasoundGraphSchemaAction(FText InNodeCategory, FText InMenuDesc, FText InToolTip, const int32 InGrouping, FText InKeywords = FText::GetEmpty())
+		: FEdGraphSchemaAction(MoveTemp(InNodeCategory), MoveTemp(InMenuDesc), MoveTemp(InToolTip), InGrouping, MoveTemp(InKeywords))
+	{}
+
+	virtual const FSlateBrush* GetIconBrush() const
+	{
+		return FEditorStyle::GetBrush("NoBrush");
+	}
+
+	virtual const FLinearColor& GetIconColor() const
+	{
+		static const FLinearColor DefaultColor;
+		return DefaultColor;
+	}
+};
+
+
 /** Action to add an input reference to the graph */
 USTRUCT()
-struct METASOUNDEDITOR_API FMetasoundGraphSchemaAction_NewInput : public FEdGraphSchemaAction
+struct METASOUNDEDITOR_API FMetasoundGraphSchemaAction_NewInput : public FMetasoundGraphSchemaAction
 {
 	GENERATED_USTRUCT_BODY();
 
@@ -48,11 +74,17 @@ struct METASOUNDEDITOR_API FMetasoundGraphSchemaAction_NewInput : public FEdGrap
 	FGuid NodeID;
 
 	FMetasoundGraphSchemaAction_NewInput()
-		: FEdGraphSchemaAction()
+		: FMetasoundGraphSchemaAction()
 	{}
 
 	FMetasoundGraphSchemaAction_NewInput(FText InNodeCategory, FText InDisplayName, FGuid InInputNodeID, FText InToolTip, const int32 InGrouping);
 
+	//~ Begin FMetasoundGraphSchemaAction Interface
+	virtual const FSlateBrush* GetIconBrush() const override;
+
+	virtual const FLinearColor& GetIconColor() const override;
+	//~ End FMetasoundGraphSchemaAction Interface
+
 	//~ Begin FEdGraphSchemaAction Interface
 	virtual UEdGraphNode* PerformAction(UEdGraph* ParentGraph, UEdGraphPin* FromPin, const FVector2D Location, bool bSelectNewNode = true) override;
 	//~ End FEdGraphSchemaAction Interface
@@ -60,12 +92,12 @@ struct METASOUNDEDITOR_API FMetasoundGraphSchemaAction_NewInput : public FEdGrap
 
 /** Promotes an input to a graph input, using its respective literal value as the default value */
 USTRUCT()
-struct METASOUNDEDITOR_API FMetasoundGraphSchemaAction_PromoteToInput : public FEdGraphSchemaAction
+struct METASOUNDEDITOR_API FMetasoundGraphSchemaAction_PromoteToInput : public FMetasoundGraphSchemaAction_NewInput
 {
 	GENERATED_USTRUCT_BODY();
 
 	FMetasoundGraphSchemaAction_PromoteToInput()
-		: FEdGraphSchemaAction()
+		: FMetasoundGraphSchemaAction_NewInput()
 	{}
 
 	FMetasoundGraphSchemaAction_PromoteToInput(FText InNodeCategory, FText InDisplayName, FText InToolTip, const int32 InGrouping);
@@ -77,7 +109,7 @@ struct METASOUNDEDITOR_API FMetasoundGraphSchemaAction_PromoteToInput : public F
 
 /** Adds an output to the graph */
 USTRUCT()
-struct METASOUNDEDITOR_API FMetasoundGraphSchemaAction_NewOutput : public FEdGraphSchemaAction
+struct METASOUNDEDITOR_API FMetasoundGraphSchemaAction_NewOutput : public FMetasoundGraphSchemaAction
 {
 	GENERATED_USTRUCT_BODY();
 
@@ -85,11 +117,17 @@ struct METASOUNDEDITOR_API FMetasoundGraphSchemaAction_NewOutput : public FEdGra
 	FGuid NodeID;
 
 	FMetasoundGraphSchemaAction_NewOutput()
-		: FEdGraphSchemaAction()
+		: FMetasoundGraphSchemaAction()
 	{}
 
 	FMetasoundGraphSchemaAction_NewOutput(FText InNodeCategory, FText InDisplayName, FGuid InOutputNodeID, FText InToolTip, const int32 InGrouping);
 
+	//~ Begin FMetasoundGraphSchemaAction Interface
+	virtual const FSlateBrush* GetIconBrush() const override;
+
+	virtual const FLinearColor& GetIconColor() const override;
+	//~ End FMetasoundGraphSchemaAction Interface
+
 	//~ Begin FEdGraphSchemaAction Interface
 	virtual UEdGraphNode* PerformAction(UEdGraph* ParentGraph, UEdGraphPin* FromPin, const FVector2D Location, bool bSelectNewNode = true) override;
 	//~ End FEdGraphSchemaAction Interface
@@ -97,12 +135,12 @@ struct METASOUNDEDITOR_API FMetasoundGraphSchemaAction_NewOutput : public FEdGra
 
 /** Promotes a node output to a graph output */
 USTRUCT()
-struct METASOUNDEDITOR_API FMetasoundGraphSchemaAction_PromoteToOutput : public FEdGraphSchemaAction
+struct METASOUNDEDITOR_API FMetasoundGraphSchemaAction_PromoteToOutput : public FMetasoundGraphSchemaAction_NewOutput
 {
 	GENERATED_USTRUCT_BODY();
 
 	FMetasoundGraphSchemaAction_PromoteToOutput()
-		: FEdGraphSchemaAction()
+		: FMetasoundGraphSchemaAction_NewOutput()
 	{}
 
 	FMetasoundGraphSchemaAction_PromoteToOutput(FText InNodeCategory, FText InDisplayName, FText InToolTip, const int32 InGrouping);
@@ -114,7 +152,7 @@ struct METASOUNDEDITOR_API FMetasoundGraphSchemaAction_PromoteToOutput : public
 
 /** Action to add a node to the graph */
 USTRUCT()
-struct METASOUNDEDITOR_API FMetasoundGraphSchemaAction_NewNode : public FEdGraphSchemaAction
+struct METASOUNDEDITOR_API FMetasoundGraphSchemaAction_NewNode : public FMetasoundGraphSchemaAction
 {
 	GENERATED_USTRUCT_BODY();
 
@@ -122,13 +160,19 @@ struct METASOUNDEDITOR_API FMetasoundGraphSchemaAction_NewNode : public FEdGraph
 	FMetasoundFrontendClassMetadata ClassMetadata;
 
 	FMetasoundGraphSchemaAction_NewNode() 
-		: FEdGraphSchemaAction()
+		: FMetasoundGraphSchemaAction()
 	{}
 
-	FMetasoundGraphSchemaAction_NewNode(const FText& InNodeCategory, const FText& InMenuDesc, const FText& InToolTip, const int32 InGrouping)
-		: FEdGraphSchemaAction(InNodeCategory, InMenuDesc, InToolTip, InGrouping)
+	FMetasoundGraphSchemaAction_NewNode(const FText& InNodeCategory, const FText& InMenuDesc, const FText& InToolTip, const int32 InGrouping, FText InKeywords = FText::GetEmpty())
+		: FMetasoundGraphSchemaAction(InNodeCategory, InMenuDesc, InToolTip, InGrouping, InKeywords)
 	{}
 
+	//~ Begin FMetasoundGraphSchemaAction Interface
+	virtual const FSlateBrush* GetIconBrush() const override;
+
+	virtual const FLinearColor& GetIconColor() const override;
+	//~ End FMetasoundGraphSchemaAction Interface
+
 	//~ Begin FEdGraphSchemaAction Interface
 	virtual UEdGraphNode* PerformAction(UEdGraph* ParentGraph, UEdGraphPin* FromPin, const FVector2D Location, bool bSelectNewNode = true) override;
 	//~ End FEdGraphSchemaAction Interface
@@ -155,18 +199,24 @@ struct METASOUNDEDITOR_API FMetasoundGraphSchemaAction_NewFromSelected : public
 
 /** Action to create new comment */
 USTRUCT()
-struct METASOUNDEDITOR_API FMetasoundGraphSchemaAction_NewComment : public FEdGraphSchemaAction
+struct METASOUNDEDITOR_API FMetasoundGraphSchemaAction_NewComment : public FMetasoundGraphSchemaAction
 {
 	GENERATED_USTRUCT_BODY();
 
-	FMetasoundGraphSchemaAction_NewComment() 
-		: FEdGraphSchemaAction()
+	FMetasoundGraphSchemaAction_NewComment()
+		: FMetasoundGraphSchemaAction()
 	{}
 
-	FMetasoundGraphSchemaAction_NewComment(FText InNodeCategory, FText InMenuDesc, FText InToolTip, const int32 InGrouping)
-		: FEdGraphSchemaAction(MoveTemp(InNodeCategory), MoveTemp(InMenuDesc), MoveTemp(InToolTip), InGrouping)
+	FMetasoundGraphSchemaAction_NewComment(const FText& InNodeCategory, const FText& InMenuDesc, const FText& InToolTip, const int32 InGrouping)
+		: FMetasoundGraphSchemaAction(InNodeCategory, InMenuDesc, InToolTip, InGrouping)
 	{}
 
+	//~ Begin FMetasoundGraphSchemaAction Interface
+	virtual const FSlateBrush* GetIconBrush() const override;
+
+	virtual const FLinearColor& GetIconColor() const override;
+	//~ End FMetasoundGraphSchemaAction Interface
+
 	//~ Begin FEdGraphSchemaAction Interface
 	virtual UEdGraphNode* PerformAction(class UEdGraph* ParentGraph, UEdGraphPin* FromPin, const FVector2D Location, bool bSelectNewNode = true) override;
 	//~ End FEdGraphSchemaAction Interface
@@ -174,16 +224,16 @@ struct METASOUNDEDITOR_API FMetasoundGraphSchemaAction_NewComment : public FEdGr
 
 /** Action to paste clipboard contents into the graph */
 USTRUCT()
-struct METASOUNDEDITOR_API FMetasoundGraphSchemaAction_Paste : public FEdGraphSchemaAction
+struct METASOUNDEDITOR_API FMetasoundGraphSchemaAction_Paste : public FMetasoundGraphSchemaAction
 {
 	GENERATED_USTRUCT_BODY();
 
 	FMetasoundGraphSchemaAction_Paste() 
-		: FEdGraphSchemaAction()
+		: FMetasoundGraphSchemaAction()
 	{}
 
-	FMetasoundGraphSchemaAction_Paste(FText InNodeCategory, FText InMenuDesc, FText InToolTip, const int32 InGrouping)
-		: FEdGraphSchemaAction(MoveTemp(InNodeCategory), MoveTemp(InMenuDesc), MoveTemp(InToolTip), InGrouping)
+	FMetasoundGraphSchemaAction_Paste(const FText& InNodeCategory, const FText& InMenuDesc, const FText& InToolTip, const int32 InGrouping)
+		: FMetasoundGraphSchemaAction(InNodeCategory, InMenuDesc, InToolTip, InGrouping)
 	{}
 
 	//~ Begin FEdGraphSchemaAction Interface
diff --git a/Engine/Plugins/Runtime/Metasound/Source/MetasoundEditor/Public/MetasoundEditorGraphValidation.h b/Engine/Plugins/Runtime/Metasound/Source/MetasoundEditor/Public/MetasoundEditorGraphValidation.h
--- a/Engine/Plugins/Runtime/Metasound/Source/MetasoundEditor/Public/MetasoundEditorGraphValidation.h
+++ b/Engine/Plugins/Runtime/Metasound/Source/MetasoundEditor/Public/MetasoundEditorGraphValidation.h
@@ -16,9 +16,6 @@ namespace Metasound
 			// Node associated with validation result
 			UMetasoundEditorGraphExternalNode* Node = nullptr;
 
-			// Whether or not the node can be auto-updated
-			bool bIsAutoUpdateEligible = false;
-
 			// Whether associated node is in invalid state (document is corrupt and no Frontend representation could be
 			// found for the node)
 			bool bIsInvalid = false;
@@ -26,10 +23,6 @@ namespace Metasound
 			// Whether validation made changes to the node and is now in a dirty state
 			bool bIsDirty = false;
 
-			// If not none, name of GraphNode's missing (not found in Frontend Registry) class as described in the owning
-			// Frontend document model.
-			FName MissingClass;
-
 			FGraphNodeValidationResult(UMetasoundEditorGraphExternalNode& InNode)
 				: Node(&InNode)
 			{
@@ -51,25 +44,10 @@ namespace Metasound
 			{
 				auto NodeIsInvalid = [](const FGraphNodeValidationResult& Result)
 				{
-					return Result.bIsInvalid || !Result.MissingClass.IsNone();
+					return Result.bIsInvalid;
 				};
 				return !NodeResults.ContainsByPredicate(NodeIsInvalid);
 			}
-
-			// Returns set of missing class names from the validation results
-			TSet<FName> FindMissingClasses() const
-			{
-				TSet<FName> MissingClasses;
-				for (const FGraphNodeValidationResult& Result : NodeResults)
-				{
-					if (!Result.MissingClass.IsNone())
-					{
-						MissingClasses.Add(Result.MissingClass);
-					}
-				}
-
-				return MissingClasses;
-			}
 		};
 	} // namespace Editor
 } // namespace Metasound
diff --git a/Engine/Plugins/Runtime/Metasound/Source/MetasoundEngine/MetasoundEngine.Build.cs b/Engine/Plugins/Runtime/Metasound/Source/MetasoundEngine/MetasoundEngine.Build.cs
--- a/Engine/Plugins/Runtime/Metasound/Source/MetasoundEngine/MetasoundEngine.Build.cs
+++ b/Engine/Plugins/Runtime/Metasound/Source/MetasoundEngine/MetasoundEngine.Build.cs
@@ -14,10 +14,10 @@ namespace UnrealBuildTool.Rules
 				new string[]
 				{
 					"Core",
-					"MetasoundGraphCore",
-					"MetasoundGenerator",
 					"AudioExtensions",
 					"AudioCodecEngine",
+					"MetasoundGraphCore",
+					"MetasoundGenerator",
 					"SignalProcessing"
 				}
 			);
@@ -27,6 +27,7 @@ namespace UnrealBuildTool.Rules
 				new string[]
 				{
 					"CoreUObject",
+					"DeveloperSettings",
 					"Engine",
 					"MetasoundFrontend",
 					"MetasoundStandardNodes",
diff --git a/Engine/Plugins/Runtime/Metasound/Source/MetasoundEngine/Private/MetasoundEngineArchetypes.cpp b/Engine/Plugins/Runtime/Metasound/Source/MetasoundEngine/Private/MetasoundEngineArchetypes.cpp
--- a/Engine/Plugins/Runtime/Metasound/Source/MetasoundEngine/Private/MetasoundEngineArchetypes.cpp
+++ b/Engine/Plugins/Runtime/Metasound/Source/MetasoundEngine/Private/MetasoundEngineArchetypes.cpp
@@ -204,6 +204,12 @@ namespace Metasound
 				return SoundIsPreviewSoundName;
 			}
 
+			const FString& GetGraphName()
+			{
+				static const FString SoundGraphName = TEXT("GraphName");
+				return SoundGraphName;
+			}
+
 			FMetasoundFrontendArchetype GetArchetype()
 			{
 				FMetasoundFrontendArchetype Archetype;
@@ -296,6 +302,12 @@ namespace Metasound
 				return SoundIsPreviewSoundName;
 			}
 
+			const FString& GetGraphName()
+			{
+				static const FString SoundGraphName = TEXT("GraphName");
+				return SoundGraphName;
+			}
+
 			const FString& GetInstanceIDName()
 			{
 				static const FString& Name = FMetasoundInstanceTransmitter::GetInstanceIDEnvironmentVariableName();
@@ -468,6 +480,12 @@ namespace Metasound
 				return SoundIsPreviewSoundName;
 			}
 
+			const FString& GetGraphName()
+			{
+				static const FString SoundGraphName = TEXT("GraphName");
+				return SoundGraphName;
+			}
+
 			const FString& GetInstanceIDName()
 			{
 				static const FString& Name = FMetasoundInstanceTransmitter::GetInstanceIDEnvironmentVariableName();
@@ -664,6 +682,12 @@ namespace Metasound
 				check(CurrentMonoVersion::GetInstanceIDName() == CurrentStereoVersion::GetInstanceIDName());
 				return CurrentMonoVersion::GetInstanceIDName();
 			}
+
+			const FString& GetGraphName()
+			{
+				check(CurrentMonoVersion::GetGraphName() == CurrentStereoVersion::GetGraphName());
+				return CurrentMonoVersion::GetGraphName();
+			}
 		}
 
 		namespace MetasoundSourceMono
diff --git a/Engine/Plugins/Runtime/Metasound/Source/MetasoundEngine/Private/MetasoundEngineArchetypes.h b/Engine/Plugins/Runtime/Metasound/Source/MetasoundEngine/Private/MetasoundEngineArchetypes.h
--- a/Engine/Plugins/Runtime/Metasound/Source/MetasoundEngine/Private/MetasoundEngineArchetypes.h
+++ b/Engine/Plugins/Runtime/Metasound/Source/MetasoundEngine/Private/MetasoundEngineArchetypes.h
@@ -37,6 +37,7 @@ namespace Metasound
 			const FString& GetAudioDeviceHandleVariableName();
 			const FString& GetSoundUniqueIdName();
 			const FString& GetIsPreviewSoundName();
+			const FString& GetGraphName();
 		}
 
 		// V1.1 of a Metasound Mono source. Uses FAudioBuffer as output.
@@ -50,6 +51,7 @@ namespace Metasound
 			const FString& GetSoundUniqueIdName();
 			const FString& GetIsPreviewSoundName();
 			const FString& GetInstanceIDName();
+			const FString& GetGraphName();
 		}
 
 		// V1.1 of a Metasound Stereo source. Uses two FAudioBuffers as outputs.
@@ -64,6 +66,7 @@ namespace Metasound
 			const FString& GetSoundUniqueIdName();
 			const FString& GetIsPreviewSoundName();
 			const FString& GetInstanceIDName();
+			const FString& GetGraphName();
 		}
 
 		// Current version of Metasound Source.
@@ -75,6 +78,7 @@ namespace Metasound
 			const FString& GetSoundUniqueIdName();
 			const FString& GetIsPreviewSoundName();
 			const FString& GetInstanceIDName();
+			const FString& GetGraphName();
 		}
 
 		// Current version of MetasoundSourceMono
diff --git a/Engine/Plugins/Runtime/Metasound/Source/MetasoundEngine/Private/MetasoundSource.cpp b/Engine/Plugins/Runtime/Metasound/Source/MetasoundEngine/Private/MetasoundSource.cpp
--- a/Engine/Plugins/Runtime/Metasound/Source/MetasoundEngine/Private/MetasoundSource.cpp
+++ b/Engine/Plugins/Runtime/Metasound/Source/MetasoundEngine/Private/MetasoundSource.cpp
@@ -320,6 +320,10 @@ Metasound::FMetasoundEnvironment UMetaSoundSource::CreateEnvironment(const FSoun
 	FMetasoundEnvironment Environment = CreateEnvironment();
 	Environment.SetValue<bool>(MetasoundSource::GetIsPreviewSoundName(), InParams.bIsPreviewSound);
 	Environment.SetValue<uint64>(MetasoundSource::GetInstanceIDName(), InParams.InstanceID);
+	
+#if WITH_METASOUND_DEBUG_ENVIRONMENT
+	Environment.SetValue<FString>(MetasoundSource::GetGraphName(), GetFullName());
+#endif // WITH_METASOUND_DEBUG_ENVIRONMENT
 
 	return Environment;
 }
diff --git a/Engine/Plugins/Runtime/Metasound/Source/MetasoundEngine/Private/MetasoundUObjectRegistry.cpp b/Engine/Plugins/Runtime/Metasound/Source/MetasoundEngine/Private/MetasoundUObjectRegistry.cpp
--- a/Engine/Plugins/Runtime/Metasound/Source/MetasoundEngine/Private/MetasoundUObjectRegistry.cpp
+++ b/Engine/Plugins/Runtime/Metasound/Source/MetasoundEngine/Private/MetasoundUObjectRegistry.cpp
@@ -14,6 +14,7 @@
 #include "MetasoundFrontendArchetypeRegistry.h"
 #include "MetasoundFrontendRegistries.h"
 #include "MetasoundFrontendTransform.h"
+#include "MetasoundSettings.h"
 #include "MetasoundSource.h"
 #include "UObject/Object.h"
 
@@ -82,6 +83,7 @@ void UMetaSoundAssetSubsystem::PostEngineInit()
 	if (UAssetManager* AssetManager = UAssetManager::GetIfValid())
 	{
 		AssetManager->CallOrRegister_OnCompletedInitialScan(FSimpleMulticastDelegate::FDelegate::CreateUObject(this, &UMetaSoundAssetSubsystem::PostInitAssetScan));
+		RebuildBlacklistCache(*AssetManager);
 	}
 	else
 	{
@@ -179,6 +181,60 @@ void UMetaSoundAssetSubsystem::AddOrUpdateAsset(const FAssetData& InAssetData, b
 	PathMap.FindOrAdd(RegistryKey) = InAssetData.ObjectPath;
 }
 
+bool UMetaSoundAssetSubsystem::CanAutoUpdate(const FMetasoundFrontendClassName& InClassName) const
+{
+	const UMetaSoundSettings* Settings = GetDefault<UMetaSoundSettings>();
+	if (!Settings->bAutoUpdateEnabled)
+	{
+		return false;
+	}
+
+	return !AutoUpdateBlacklistCache.Contains(InClassName.GetFullName());
+}
+
+void UMetaSoundAssetSubsystem::RebuildBlacklistCache(const UAssetManager& InAssetManager)
+{
+	const UMetaSoundSettings* Settings = GetDefault<UMetaSoundSettings>();
+	if (Settings->BlacklistCacheChangeID == AutoUpdateBlacklistChangeID)
+	{
+		return;
+	}
+
+	AutoUpdateBlacklistCache.Reset();
+
+	for (const FMetasoundFrontendClassName& ClassName : Settings->AutoUpdateBlacklist)
+	{
+		AutoUpdateBlacklistCache.Add(ClassName.GetFullName());
+	}
+
+	for (const FDefaultMetaSoundAssetAutoUpdateSettings& UpdateSettings : Settings->AutoUpdateAssetBlacklist)
+	{
+		if (UAssetManager* AssetManager = UAssetManager::GetIfValid())
+		{
+			FAssetData AssetData;
+			if (AssetManager->GetAssetDataForPath(UpdateSettings.MetaSound, AssetData))
+			{
+				FString AssetClassID;
+				if (AssetData.GetTagValue(Metasound::AssetTags::AssetClassID, AssetClassID))
+				{
+					const FMetasoundFrontendClassName ClassName = { FName(), *AssetClassID, FName() };
+					AutoUpdateBlacklistCache.Add(ClassName.GetFullName());
+				}
+			}
+		}
+	}
+
+	AutoUpdateBlacklistChangeID = Settings->BlacklistCacheChangeID;
+}
+
+void UMetaSoundAssetSubsystem::RescanAutoUpdateBlacklist()
+{
+	if (const UAssetManager* AssetManager = UAssetManager::GetIfValid())
+	{
+		RebuildBlacklistCache(*AssetManager);
+	}
+}
+
 FMetasoundAssetBase* UMetaSoundAssetSubsystem::FindAssetFromKey(const Metasound::Frontend::FNodeRegistryKey& RegistryKey) const
 {
 	if (const FSoftObjectPath* ObjectPath = FindObjectPathFromKey(RegistryKey))
diff --git a/Engine/Plugins/Runtime/Metasound/Source/MetasoundEngine/Private/MetasoundWavePlayerNode.cpp b/Engine/Plugins/Runtime/Metasound/Source/MetasoundEngine/Private/MetasoundWavePlayerNode.cpp
--- a/Engine/Plugins/Runtime/Metasound/Source/MetasoundEngine/Private/MetasoundWavePlayerNode.cpp
+++ b/Engine/Plugins/Runtime/Metasound/Source/MetasoundEngine/Private/MetasoundWavePlayerNode.cpp
@@ -300,7 +300,7 @@ namespace Metasound
 
 			if (IsWaveValid())
 			{
-				
+				UE_LOG(LogMetaSound, Verbose, TEXT("Starting Sound: '%s'"), *CurrentWaveAsset->GetFullName());
 				SoundAssetSampleRate = CurrentWaveAsset->GetSampleRate();
 				SoundAssetDurationSeconds = CurrentWaveAsset->GetDuration();
 				SoundAssetNumFrames = CurrentWaveAsset->GetNumFrames();
diff --git a/Engine/Plugins/Runtime/Metasound/Source/MetasoundEngine/Public/Metasound.h b/Engine/Plugins/Runtime/Metasound/Source/MetasoundEngine/Public/Metasound.h
--- a/Engine/Plugins/Runtime/Metasound/Source/MetasoundEngine/Public/Metasound.h
+++ b/Engine/Plugins/Runtime/Metasound/Source/MetasoundEngine/Public/Metasound.h
@@ -30,7 +30,7 @@ public:
 	virtual bool NeedsLoadForEditorGame() const override { return false; }
 
 	virtual bool Synchronize() { return false; }
-	virtual bool Validate(bool bInAutoUpdate) { return false; }
+	virtual bool Validate(bool bInAutoUpdate, bool bInClearUpdateNotes) { return false; }
 	virtual void RegisterGraphWithFrontend() { }
 };
 
@@ -53,7 +53,7 @@ namespace Metasound
 	{
 		if (UMetasoundEditorGraphBase* MetaSoundGraph = Cast<UMetasoundEditorGraphBase>(InMetaSound.GetGraph()))
 		{
-			if (MetaSoundGraph->Validate(false /* bAutoUpdate */))
+			if (MetaSoundGraph->Validate(false /* bAutoUpdate */, false /* bClearUpdateNotes */))
 			{
 				MetaSoundGraph->RegisterGraphWithFrontend();
 			}
@@ -83,7 +83,8 @@ namespace Metasound
 		if (UMetasoundEditorGraphBase* MetaSoundGraph = Cast<UMetasoundEditorGraphBase>(InMetaSound.GetGraph()))
 		{
 			const bool bAutoUpdate = false;
-			const bool bIsValid = MetaSoundGraph->Validate(bAutoUpdate /* bAutoUpdate */);
+			const bool bClearUpdateNotes = true;
+			const bool bIsValid = MetaSoundGraph->Validate(bAutoUpdate, bClearUpdateNotes);
 			if (bIsValid)
 			{
 				MetaSoundGraph->RegisterGraphWithFrontend();
diff --git a/Engine/Plugins/Runtime/Metasound/Source/MetasoundEngine/Public/MetasoundSettings.h b/Engine/Plugins/Runtime/Metasound/Source/MetasoundEngine/Public/MetasoundSettings.h
new file mode 100644
--- /dev/null
+++ b/Engine/Plugins/Runtime/Metasound/Source/MetasoundEngine/Public/MetasoundSettings.h
@@ -0,0 +1,65 @@
+// Copyright Epic Games, Inc. All Rights Reserved.
+#pragma once
+
+#include "CoreMinimal.h"
+#include "Engine/DeveloperSettings.h"
+#include "MetasoundFrontendDocument.h"
+#include "UObject/NoExportTypes.h"
+#include "UObject/Object.h"
+#include "UObject/ObjectMacros.h"
+
+#include "MetasoundSettings.generated.h"
+
+
+// Forward Declarations
+struct FMetasoundFrontendClassName;
+class IMetaSoundAssetInterface;
+
+
+UENUM()
+enum class EMetaSoundMessageLevel : uint8
+{
+	Error,
+	Warning,
+	Info
+};
+
+USTRUCT()
+struct METASOUNDENGINE_API FDefaultMetaSoundAssetAutoUpdateSettings
+{
+	GENERATED_BODY()
+
+	/** MetaSound to blacklist from AutoUpdate. */
+	UPROPERTY(EditAnywhere, Category = "AutoUpdate", meta = (AllowedClasses = "MetaSound, MetaSoundSource"))
+	FSoftObjectPath MetaSound;
+};
+
+UCLASS(config = MetaSound, defaultconfig, meta = (DisplayName = "MetaSounds"))
+class METASOUNDENGINE_API UMetaSoundSettings : public UDeveloperSettings
+{
+	GENERATED_BODY()
+
+public:
+	/** If true, AutoUpdate is enabled, increasing load times.  If false, skips AutoUpdate on load, but can result in MetaSounds failing to load, 
+	  * register, and execute if interface differences are present. */
+	UPROPERTY(EditAnywhere, config, Category = AutoUpdate)
+	bool bAutoUpdateEnabled = true;
+
+	/** List of native MetaSound classes whose node references should not be AutoUpdated. */
+	UPROPERTY(EditAnywhere, config, Category = AutoUpdate, meta = (DisplayName = "Blacklist", EditCondition = "bAutoUpdateEnabled"))
+	TArray<FMetasoundFrontendClassName> AutoUpdateBlacklist;
+
+	/** List of MetaSound assets whose node references should not be AutoUpdated. */
+	UPROPERTY(EditAnywhere, config, Category = AutoUpdate, meta = (DisplayName = "Asset Blacklist", EditCondition = "bAutoUpdateEnabled"))
+	TArray<FDefaultMetaSoundAssetAutoUpdateSettings> AutoUpdateAssetBlacklist;
+
+	UPROPERTY(Transient)
+	int32 BlacklistCacheChangeID = 0;
+
+#if WITH_EDITOR
+	virtual void PostEditChangeProperty(FPropertyChangedEvent& PropertyChangedEvent) override
+	{
+		BlacklistCacheChangeID++;
+	}
+#endif // WITH_EDITOR
+};
diff --git a/Engine/Plugins/Runtime/Metasound/Source/MetasoundEngine/Public/MetasoundUObjectRegistry.h b/Engine/Plugins/Runtime/Metasound/Source/MetasoundEngine/Public/MetasoundUObjectRegistry.h
--- a/Engine/Plugins/Runtime/Metasound/Source/MetasoundEngine/Public/MetasoundUObjectRegistry.h
+++ b/Engine/Plugins/Runtime/Metasound/Source/MetasoundEngine/Public/MetasoundUObjectRegistry.h
@@ -10,6 +10,9 @@
 
 #include "MetasoundUObjectRegistry.generated.h"
 
+// Forward Declarations
+class UAssetManager;
+
 
 /** The subsystem in charge of the MetaSound asset registry */
 UCLASS()
@@ -34,6 +37,8 @@ public:
 	void RenameAsset(const FAssetData& InAssetData, bool bInReregisterWithFrontend = true);
 	void SynchronizeAssetClassDisplayName(const FAssetData& InAssetData);
 
+	virtual bool CanAutoUpdate(const FMetasoundFrontendClassName& InClassName) const override;
+	virtual void RescanAutoUpdateBlacklist() override;
 	virtual FMetasoundAssetBase* FindAssetFromKey(const Metasound::Frontend::FNodeRegistryKey& RegistryKey) const override;
 	virtual const FSoftObjectPath* FindObjectPathFromKey(const Metasound::Frontend::FNodeRegistryKey& RegistryKey) const override;
 	virtual FMetasoundAssetBase* TryLoadAsset(const FSoftObjectPath& InObjectPath) const override;
@@ -43,6 +48,10 @@ protected:
 	void PostInitAssetScan();
 
 private:
+	void RebuildBlacklistCache(const UAssetManager& InAssetManager);
+
+	int32 AutoUpdateBlacklistChangeID = INDEX_NONE;
+	TSet<FName> AutoUpdateBlacklistCache;
 	TMap<Metasound::Frontend::FNodeRegistryKey, FSoftObjectPath> PathMap;
 };
 
diff --git a/Engine/Plugins/Runtime/Metasound/Source/MetasoundFrontend/Private/MetasoundArrayNodes.cpp b/Engine/Plugins/Runtime/Metasound/Source/MetasoundFrontend/Private/MetasoundArrayNodes.cpp
--- a/Engine/Plugins/Runtime/Metasound/Source/MetasoundFrontend/Private/MetasoundArrayNodes.cpp
+++ b/Engine/Plugins/Runtime/Metasound/Source/MetasoundFrontend/Private/MetasoundArrayNodes.cpp
@@ -20,8 +20,8 @@ namespace Metasound
 				PluginAuthor,
 				PluginNodeMissingPrompt,
 				InDefaultInterface,
-				{LOCTEXT("ArrayCategory", "Array")},
-				{TEXT("Array")},
+				{ LOCTEXT("ArrayCategory", "Array") },
+				{ LOCTEXT("MetasoundArrayKeyword", "Array") },
 				FNodeDisplayStyle{}
 			};
 
diff --git a/Engine/Plugins/Runtime/Metasound/Source/MetasoundFrontend/Private/MetasoundFrontendInjectReceiveNodes.cpp b/Engine/Plugins/Runtime/Metasound/Source/MetasoundFrontend/Private/MetasoundFrontendInjectReceiveNodes.cpp
--- a/Engine/Plugins/Runtime/Metasound/Source/MetasoundFrontend/Private/MetasoundFrontendInjectReceiveNodes.cpp
+++ b/Engine/Plugins/Runtime/Metasound/Source/MetasoundFrontend/Private/MetasoundFrontendInjectReceiveNodes.cpp
@@ -93,7 +93,7 @@ namespace Metasound
 				{
 					FNodeClassMetadata Metadata
 					{
-						FNodeClassName{"MetasoundFrontendInjectReceiveNodes", "ReceiveNodeAddress", ""},
+						FNodeClassName { "MetasoundFrontendInjectReceiveNodes", "ReceiveNodeAddress", "" },
 						1, // MajorVersion
 						0, // MinorVersion
 						FText::GetEmpty(), // DisplayName
@@ -102,8 +102,8 @@ namespace Metasound
 						FText::GetEmpty(), // Prompt If Missing
 						CreateVertexInterface(), // DefaultInterface
 						TArray<FText>(), // CategoryHierachy
-						TArray<FName>(), // Keywoards
-						FNodeDisplayStyle {} // NodeDisplayStyle
+						TArray<FText>(), // Keywords
+						FNodeDisplayStyle() // NodeDisplayStyle
 					}; 
 
 					return Metadata;
diff --git a/Engine/Plugins/Runtime/Metasound/Source/MetasoundFrontend/Private/MetasoundFrontendInvalidController.h b/Engine/Plugins/Runtime/Metasound/Source/MetasoundFrontend/Private/MetasoundFrontendInvalidController.h
--- a/Engine/Plugins/Runtime/Metasound/Source/MetasoundFrontend/Private/MetasoundFrontendInvalidController.h
+++ b/Engine/Plugins/Runtime/Metasound/Source/MetasoundFrontend/Private/MetasoundFrontendInvalidController.h
@@ -4,11 +4,13 @@
 
 
 #include "CoreMinimal.h"
+#include "MetasoundAssetBase.h"
 #include "MetasoundBuilderInterface.h"
 #include "MetasoundFrontendController.h"
 #include "MetasoundFrontendDocument.h"
 #include "MetasoundGraph.h"
 
+
 namespace Metasound
 {
 	namespace Frontend
@@ -173,7 +175,8 @@ namespace Metasound
 
 			bool IsRequired() const override { return false; }
 
-			bool CanAutoUpdate() const override { return false; }
+			bool DiffAgainstRegistryInterface(FClassInterfaceUpdates& OutInterfaceUpdates, bool bInUseHighestMinorVersion) const override { return false; }
+			bool CanAutoUpdate(const IMetaSoundAssetInterface& AssetInterface, FClassInterfaceUpdates* OutInterfaceUpdates = nullptr) const override { return false; }
 			FMetasoundFrontendVersionNumber FindHighestVersionInRegistry() const override { return FMetasoundFrontendVersionNumber::GetInvalid(); }
 			FMetasoundFrontendVersionNumber FindHighestMinorVersionInRegistry() const override { return FMetasoundFrontendVersionNumber::GetInvalid(); }
 
diff --git a/Engine/Plugins/Runtime/Metasound/Source/MetasoundFrontend/Private/MetasoundFrontendNodesCategories.cpp b/Engine/Plugins/Runtime/Metasound/Source/MetasoundFrontend/Private/MetasoundFrontendNodesCategories.cpp
new file mode 100644
--- /dev/null
+++ b/Engine/Plugins/Runtime/Metasound/Source/MetasoundFrontend/Private/MetasoundFrontendNodesCategories.cpp
@@ -0,0 +1,15 @@
+// Copyright Epic Games, Inc. All Rights Reserved.
+
+#include "MetasoundFrontendNodesCategories.h"
+
+#define LOCTEXT_NAMESPACE "MetasoundNodeCategories"
+
+namespace Metasound
+{
+	namespace NodeCategories
+	{
+		const FText Conversions  = { LOCTEXT("Metasound_ConversionsCategory", "Conversions") };
+	}
+}
+
+#undef LOCTEXT_NAMESPACE
\ No newline at end of file
diff --git a/Engine/Plugins/Runtime/Metasound/Source/MetasoundFrontend/Private/MetasoundFrontendStandardController.cpp b/Engine/Plugins/Runtime/Metasound/Source/MetasoundFrontend/Private/MetasoundFrontendStandardController.cpp
--- a/Engine/Plugins/Runtime/Metasound/Source/MetasoundFrontend/Private/MetasoundFrontendStandardController.cpp
+++ b/Engine/Plugins/Runtime/Metasound/Source/MetasoundFrontend/Private/MetasoundFrontendStandardController.cpp
@@ -4,9 +4,11 @@
 
 #include "Algo/ForEach.h"
 #include "Algo/NoneOf.h"
+#include "Algo/Transform.h"
 #include "CoreMinimal.h"
 #include "HAL/FileManager.h"
 #include "MetasoundAccessPtr.h"
+#include "MetasoundAssetBase.h"
 #include "MetasoundFrontendArchetypeRegistry.h"
 #include "MetasoundFrontendController.h"
 #include "MetasoundFrontendDocument.h"
@@ -1637,7 +1639,7 @@ namespace Metasound
 				return this->AsShared();
 			}
 
-			const FMetasoundFrontendNodeStyle Style = GetNodeStyle();
+			FMetasoundFrontendNodeStyle Style = GetNodeStyle();
 
 			using FConnectionKey = TPair<FString, FName>;
 
@@ -1740,40 +1742,144 @@ namespace Metasound
 			return ReplacementNode;
 		}
 
-		bool FBaseNodeController::CanAutoUpdate() const
+		bool FBaseNodeController::DiffAgainstRegistryInterface(FClassInterfaceUpdates& OutInterfaceUpdates, bool bInUseHighestMinorVersion) const
 		{
-			const FMetasoundFrontendClassMetadata& Metadata = GetClassMetadata();
+			OutInterfaceUpdates = FClassInterfaceUpdates();
 
-			const FMetasoundFrontendVersionNumber& CurrentVersion = Metadata.GetVersion();
-			Metasound::FNodeClassName NodeClassName = Metadata.GetClassName().ToNodeClassName();
+			const FMetasoundFrontendClassMetadata& NodeClassMetadata = GetClassMetadata();
+			const FMetasoundFrontendClassInterface& NodeClassInterface = GetClassInterface();
+
+			Metasound::FNodeClassName NodeClassName = NodeClassMetadata.GetClassName().ToNodeClassName();
+
+			const bool bSortByVersion = true;
+
+			FMetasoundFrontendClass RegistryClass;
+			if (bInUseHighestMinorVersion)
+			{
+				if (!ISearchEngine::Get().FindClassWithMajorVersion(NodeClassName, NodeClassMetadata.GetVersion().Major, RegistryClass))
+				{
+					OutInterfaceUpdates.RemovedInputs = NodeClassInterface.Inputs;
+					OutInterfaceUpdates.RemovedOutputs = NodeClassInterface.Outputs;
+					return false;
+				}
+			}
+			else
+			{
+				const TArray<FMetasoundFrontendClass> Classes = ISearchEngine::Get().FindClassesWithName(NodeClassName, bSortByVersion);
+				const FMetasoundFrontendClass* ExactClass = Classes.FindByPredicate([CurrentVersion = &NodeClassMetadata.GetVersion()](const FMetasoundFrontendClass& AvailableClass)
+				{
+					return AvailableClass.Metadata.GetVersion() == *CurrentVersion;
+				});
+
+				if (!ExactClass)
+				{
+					OutInterfaceUpdates.RemovedInputs = NodeClassInterface.Inputs;
+					OutInterfaceUpdates.RemovedOutputs = NodeClassInterface.Outputs;
+					return false;
+				}
+				RegistryClass = *ExactClass;
+			}
+
+			for (const FMetasoundFrontendClassInput& Input : NodeClassInterface.Inputs)
+			{
+				auto IsFunctionalEquivalent = [NodeClassInput = &Input](const FMetasoundFrontendClassInput& Iter)
+				{
+					return FMetasoundFrontendClassVertex::IsFunctionalEquivalent(*NodeClassInput, Iter);
+				};
+
+				const int32 Index = RegistryClass.Interface.Inputs.FindLastByPredicate(IsFunctionalEquivalent);
+				if (Index == INDEX_NONE)
+				{
+					OutInterfaceUpdates.RemovedInputs.Add(Input);
+				}
+				else
+				{
+					RegistryClass.Interface.Inputs.RemoveAtSwap(Index, 1, false /* bAllowShrinking */);
+				}
+			}
+			OutInterfaceUpdates.AddedInputs = MoveTemp(RegistryClass.Interface.Inputs);
 
-			// TODO: Natively defined member should probably be added to the FMetasoundFrontendClass
-			// instead of being an independent member on the INodeRegistryEntry.
-			const FNodeRegistryKey RegistryKey = FMetasoundFrontendRegistryContainer::Get()->GetRegistryKey(Metadata);
-			if (FMetasoundFrontendRegistryContainer::Get()->IsNodeNative(RegistryKey))
+			for (const FMetasoundFrontendClassOutput& Output : NodeClassInterface.Outputs)
+			{
+				auto IsFunctionalEquivalent = [NodeClassOutput = &Output](const FMetasoundFrontendClassOutput& Iter)
+				{
+					return FMetasoundFrontendClassVertex::IsFunctionalEquivalent(*NodeClassOutput, Iter);
+				};
+
+				const int32 Index = RegistryClass.Interface.Outputs.FindLastByPredicate(IsFunctionalEquivalent);
+				if (Index == INDEX_NONE)
+				{
+					OutInterfaceUpdates.RemovedOutputs.Add(Output);
+				}
+				else
+				{
+					RegistryClass.Interface.Outputs.RemoveAtSwap(Index, 1, false /* bAllowShrinking */);
+				}
+			}
+			OutInterfaceUpdates.AddedOutputs = MoveTemp(RegistryClass.Interface.Outputs);
+
+			return true;
+		}
+
+		bool FBaseNodeController::CanAutoUpdate(const IMetaSoundAssetInterface& AssetInterface, FClassInterfaceUpdates* OutInterfaceUpdates) const
+		{
+			const FMetasoundFrontendClassMetadata& NodeClassMetadata = GetClassMetadata();
+			if (!AssetInterface.CanAutoUpdate(NodeClassMetadata.GetClassName()))
 			{
 				return false;
 			}
 
-			FMetasoundFrontendClass ClassWithMajorVersion;
-			if (!ISearchEngine::Get().FindClassWithMajorVersion(NodeClassName, CurrentVersion.Major, ClassWithMajorVersion))
+			FMetasoundFrontendClass RegistryClass;
+			if (!ISearchEngine::Get().FindClassWithMajorVersion(
+				NodeClassMetadata.GetClassName().ToNodeClassName(),
+				NodeClassMetadata.GetVersion().Major,
+				RegistryClass))
 			{
 				return false;
 			}
 
-			if (ClassWithMajorVersion.Metadata.GetChangeID() == Metadata.GetChangeID())
+			if (RegistryClass.Metadata.GetVersion() < NodeClassMetadata.GetVersion())
 			{
-				const FGuid& ChangeID = GetClassInterface().GetChangeID();
-				if (ClassWithMajorVersion.Interface.GetChangeID() == ChangeID)
+				return false;
+			}
+
+			if (RegistryClass.Metadata.GetVersion() == NodeClassMetadata.GetVersion())
+			{
+				// TODO: Merge these paths.  Shouldn't use different logic to
+				// define changes in native vs asset class definitions.
+				const FNodeRegistryKey RegistryKey = NodeRegistryKey::CreateKey(RegistryClass.Metadata);
+				const bool bIsClassNative = FMetasoundFrontendRegistryContainer::Get()->IsNodeNative(RegistryKey);
+				if (bIsClassNative)
 				{
-					return false;
+					FClassInterfaceUpdates InterfaceUpdates;
+					DiffAgainstRegistryInterface(InterfaceUpdates, true /* bUseHighestMinorVersion */);
+					if (OutInterfaceUpdates)
+					{
+						*OutInterfaceUpdates = InterfaceUpdates;
+					}
+
+					if (!InterfaceUpdates.ContainsChanges())
+					{
+						return false;
+					}
+				}
+				else
+				{
+					if (RegistryClass.Metadata.GetChangeID() == NodeClassMetadata.GetChangeID())
+					{
+						const FGuid& NodeClassInterfaceChangeID = GetClassInterface().GetChangeID();
+						if (RegistryClass.Interface.GetChangeID() == NodeClassInterfaceChangeID)
+						{
+							return false;
+						}
+					}
 				}
 			}
 
 			return true;
 		}
 
-		FDocumentAccess FBaseNodeController::ShareAccess() 
+		FDocumentAccess FBaseNodeController::ShareAccess()
 		{
 			FDocumentAccess Access;
 
@@ -4082,25 +4188,6 @@ namespace Metasound
 					NumDependenciesRemovedThisItr = Document->Dependencies.RemoveAllSwap(IsDependencyUnreferenced);
 				}
 				while (NumDependenciesRemovedThisItr > 0);
-
-				/*
-				FRegistry* Registry = FRegistry::Get();
-				if (ensure(Registry))
-				{
-					// Make sure classes are up-to-date with registered versions of class
-					for (FMetasoundFrontendClass& Class : DocumentPtr->Dependencies)
-					{
-						FMetasoundFrontendClass RegisteredClass;
-						FNodeRegistryKey Key = FRegistry::GetRegistryKey(Class.Metadata);
-						if (Registry->FindFrontendClassFromRegistered(Key, RegisteredClass))
-						{
-							Class = RegisteredClass;
-							// TODO: update nodes for class.
-							// TODO: perform minor version updating.
-						}
-					}
-				}
-				*/
 			}
 		}
 
diff --git a/Engine/Plugins/Runtime/Metasound/Source/MetasoundFrontend/Private/MetasoundFrontendStandardController.h b/Engine/Plugins/Runtime/Metasound/Source/MetasoundFrontend/Private/MetasoundFrontendStandardController.h
--- a/Engine/Plugins/Runtime/Metasound/Source/MetasoundFrontend/Private/MetasoundFrontendStandardController.h
+++ b/Engine/Plugins/Runtime/Metasound/Source/MetasoundFrontend/Private/MetasoundFrontendStandardController.h
@@ -11,6 +11,8 @@
 #include "UObject/Object.h"
 
 
+class IMetaSoundAssetInterface;
+
 namespace Metasound
 {
 	namespace Frontend
@@ -375,7 +377,9 @@ namespace Metasound
 			const FMetasoundFrontendNodeStyle& GetNodeStyle() const override;
 			void SetNodeStyle(const FMetasoundFrontendNodeStyle& InStyle) override;
 
-			bool CanAutoUpdate() const override;
+			bool DiffAgainstRegistryInterface(FClassInterfaceUpdates& OutInterfaceUpdates, bool bInUseHighestMinorVersion) const override;
+
+			bool CanAutoUpdate(const IMetaSoundAssetInterface& AssetInterface, FClassInterfaceUpdates* OutInterfaceUpdates = nullptr) const override;
 			FNodeHandle ReplaceWithVersion(const FMetasoundFrontendVersionNumber& InNewVersion) override;
 			FMetasoundFrontendVersionNumber FindHighestVersionInRegistry() const override;
 			FMetasoundFrontendVersionNumber FindHighestMinorVersionInRegistry() const override;
diff --git a/Engine/Plugins/Runtime/Metasound/Source/MetasoundFrontend/Private/MetasoundFrontendTransform.cpp b/Engine/Plugins/Runtime/Metasound/Source/MetasoundFrontend/Private/MetasoundFrontendTransform.cpp
--- a/Engine/Plugins/Runtime/Metasound/Source/MetasoundFrontend/Private/MetasoundFrontendTransform.cpp
+++ b/Engine/Plugins/Runtime/Metasound/Source/MetasoundFrontend/Private/MetasoundFrontendTransform.cpp
@@ -396,20 +396,19 @@ namespace Metasound
 			RootGraph->IterateNodes([&](FNodeHandle NodeHandle)
 			{
 				using namespace Metasound::Frontend;
-
-				const FMetasoundFrontendClassMetadata& ClassMetadata = NodeHandle->GetClassMetadata();
-
-				if (!NodeHandle->CanAutoUpdate())
+				FClassInterfaceUpdates InterfaceUpdates;
+				if (!NodeHandle->CanAutoUpdate(*AssetInterface, &InterfaceUpdates))
 				{
 					return;
 				}
 
+				const FMetasoundFrontendClassMetadata ClassMetadata = NodeHandle->GetClassMetadata();
 				FMetasoundFrontendVersionNumber UpdateVersion = NodeHandle->FindHighestMinorVersionInRegistry();
 				if (UpdateVersion.IsValid() && UpdateVersion > ClassMetadata.GetVersion())
 				{
 					UE_LOG(LogMetaSound, Display, TEXT("Auto-Updating node class '%s': Newer minor version '%s' found."), *ClassMetadata.GetDisplayName().ToString(), *UpdateVersion.ToString());
 				}
-				else
+				else if (InterfaceUpdates.ContainsChanges())
 				{
 					UpdateVersion = ClassMetadata.GetVersion();
 					UE_LOG(LogMetaSound, Display, TEXT("Auto-Updating node with class '%s (%s)': Interface change detected."), *ClassMetadata.GetDisplayName().ToString(), *UpdateVersion.ToString());
@@ -442,7 +441,12 @@ namespace Metasound
 				for (const TPair<FNodeHandle, FMetasoundFrontendVersionNumber>& Pair : NodesToUpdate)
 				{
 					FNodeHandle ExistingNode = Pair.Key;
+					FMetasoundFrontendVersionNumber InitialVersion = ExistingNode->GetClassMetadata().GetVersion();
 					FNodeHandle NewNode = ExistingNode->ReplaceWithVersion(Pair.Value);
+					FMetasoundFrontendNodeStyle Style = NewNode->GetNodeStyle();
+					Style.bMessageNodeUpdated = NewNode->GetClassMetadata().GetVersion() > InitialVersion;
+					NewNode->SetNodeStyle(Style);
+
 					bDidEdit |= NewNode->GetID() != ExistingNode->GetID();
 				}
 			}
diff --git a/Engine/Plugins/Runtime/Metasound/Source/MetasoundFrontend/Public/MetasoundArrayRandomNode.h b/Engine/Plugins/Runtime/Metasound/Source/MetasoundFrontend/Public/MetasoundArrayRandomNode.h
--- a/Engine/Plugins/Runtime/Metasound/Source/MetasoundFrontend/Public/MetasoundArrayRandomNode.h
+++ b/Engine/Plugins/Runtime/Metasound/Source/MetasoundFrontend/Public/MetasoundArrayRandomNode.h
@@ -168,7 +168,7 @@ namespace Metasound
 			FInt32ReadRef InNoRepeatOrder = InParams.InputDataReferences.GetDataReadReferenceOrConstructWithVertexDefault<int32>(Inputs, InputNoRepeatOrderName, InParams.OperatorSettings);
 			FBoolReadRef bInEnableSharedState = InParams.InputDataReferences.GetDataReadReferenceOrConstructWithVertexDefault<bool>(Inputs, InputEnableSharedStateName, InParams.OperatorSettings);
 
-			return MakeUnique<TArrayRandomGetOperator>(InParams, InTriggerNext, InTriggerReset, InInputArray, InInputWeightsArray, InSeedValue, InNoRepeatOrder, bInEnableSharedState);
+			return MakeUnique<TArrayRandomGetOperator<ArrayType>>(InParams, InTriggerNext, InTriggerReset, InInputArray, InInputWeightsArray, InSeedValue, InNoRepeatOrder, bInEnableSharedState);
 		}
 
 		TArrayRandomGetOperator(
@@ -190,7 +190,11 @@ namespace Metasound
 			, TriggerOnNext(FTriggerWriteRef::CreateNew(InParams.OperatorSettings))
 			, TriggerOnReset(FTriggerWriteRef::CreateNew(InParams.OperatorSettings))
 			, OutValue(TDataWriteReferenceFactory<ElementType>::CreateAny(InParams.OperatorSettings))
+#if WITH_METASOUND_DEBUG_ENVIRONMENT
+			, GraphName(*InParams.Environment.GetValue<FString>("GraphName"))
+#endif // WITH_METASOUND_DEBUG_ENVIRONMENT
 		{
+
 			// Check to see if this is a global shuffler or a local one. 
 			// Global shuffler will use a namespace to opt into it.
 			PrevSeedValue = *SeedValue;
@@ -228,10 +232,12 @@ namespace Metasound
 					ArrayRandomGet = MakeUnique<FArrayRandomGet>(PrevSeedValue, PrevArraySize, WeightsArray, PrevNoRepeatOrder);
 				}
 			}
+#if WITH_METASOUND_DEBUG_ENVIRONMENT
 			else
 			{
-				UE_LOG(LogMetaSound, Error, TEXT("Array Random Get: Can't retrieve random elements from an empty array"));
+ 				UE_LOG(LogMetaSound, Verbose, TEXT("Array Random Get: Can't retrieve random elements from an empty array in graph '%s'"), *GraphName);
 			}
+#endif // WITH_METASOUND_DEBUG_ENVIRONMENT
 		}
 
 		virtual ~TArrayRandomGetOperator() = default;
@@ -289,7 +295,9 @@ namespace Metasound
 			{
 				if (!bHasLoggedEmptyArrayWarning)
 				{
-					UE_LOG(LogMetaSound, Warning, TEXT("Array Random Get has an empty array input."));
+#if WITH_METASOUND_DEBUG_ENVIRONMENT
+					UE_LOG(LogMetaSound, Verbose, TEXT("Array Random Get: empty array input (Graph '%s')"), *GraphName);
+#endif // WITH_METASOUND_DEBUG_ENVIRONMENT
 					bHasLoggedEmptyArrayWarning = true;
 				}
 				return;
@@ -411,6 +419,10 @@ namespace Metasound
 		FTriggerWriteRef TriggerOnReset;
 		TDataWriteReference<ElementType> OutValue;
 
+#if WITH_METASOUND_DEBUG_ENVIRONMENT
+		FString GraphName;
+#endif // WITH_METASOUND_DEBUG_ENVIRONMENT
+
 		// Data
 		TUniquePtr<FArrayRandomGet> ArrayRandomGet;
 		TArray<float> WeightsArray;
diff --git a/Engine/Plugins/Runtime/Metasound/Source/MetasoundFrontend/Public/MetasoundAssetBase.h b/Engine/Plugins/Runtime/Metasound/Source/MetasoundFrontend/Public/MetasoundAssetBase.h
--- a/Engine/Plugins/Runtime/Metasound/Source/MetasoundFrontend/Public/MetasoundAssetBase.h
+++ b/Engine/Plugins/Runtime/Metasound/Source/MetasoundFrontend/Public/MetasoundAssetBase.h
@@ -40,9 +40,13 @@ namespace Metasound
 class METASOUNDFRONTEND_API IMetaSoundAssetInterface
 {
 public:
+	virtual bool CanAutoUpdate(const FMetasoundFrontendClassName& InClassName) const = 0;
 	virtual FMetasoundAssetBase* FindAssetFromKey(const Metasound::Frontend::FNodeRegistryKey& InRegistryKey) const = 0;
 	virtual const FSoftObjectPath* FindObjectPathFromKey(const Metasound::Frontend::FNodeRegistryKey& InRegistryKey) const = 0;
 
+	// Rescans settings for blacklisted assets not to run reference auto-update against.
+	virtual void RescanAutoUpdateBlacklist() = 0;
+
 	// Attempts to load an FMetasoundAssetBase from the given path, or returns it if its already loaded
 	virtual FMetasoundAssetBase* TryLoadAsset(const FSoftObjectPath& InObjectPath) const = 0;
 };
diff --git a/Engine/Plugins/Runtime/Metasound/Source/MetasoundFrontend/Public/MetasoundAutoConverterNode.h b/Engine/Plugins/Runtime/Metasound/Source/MetasoundFrontend/Public/MetasoundAutoConverterNode.h
--- a/Engine/Plugins/Runtime/Metasound/Source/MetasoundFrontend/Public/MetasoundAutoConverterNode.h
+++ b/Engine/Plugins/Runtime/Metasound/Source/MetasoundFrontend/Public/MetasoundAutoConverterNode.h
@@ -11,6 +11,7 @@
 #include "MetasoundDataReference.h"
 #include "MetasoundExecutableOperator.h"
 #include "MetasoundFrontend.h"
+#include "MetasoundFrontendNodesCategories.h"
 
 #include <type_traits>
 
@@ -72,8 +73,13 @@ namespace Metasound
 				Info.PromptIfMissing = PluginNodeMissingPrompt;
 				Info.DefaultInterface = DeclareVertexInterface();
 
-				Info.CategoryHierarchy = { LOCTEXT("Metasound_ConvertNodeCategory", "Conversions") };
-				Info.Keywords = { "Convert", GetMetasoundDataTypeName<FromDataType>(), GetMetasoundDataTypeName<ToDataType>()};
+				Info.CategoryHierarchy.Emplace(NodeCategories::Conversions);
+				Info.Keywords =
+				{
+					LOCTEXT("MetasoundConvertKeyword", "Convert"),
+					FText::FromName(GetMetasoundDataTypeName<FromDataType>()),
+					FText::FromName(GetMetasoundDataTypeName<ToDataType>())
+				};
 
 				return Info;
 			};
diff --git a/Engine/Plugins/Runtime/Metasound/Source/MetasoundFrontend/Public/MetasoundFrontendController.h b/Engine/Plugins/Runtime/Metasound/Source/MetasoundFrontend/Public/MetasoundFrontendController.h
--- a/Engine/Plugins/Runtime/Metasound/Source/MetasoundFrontend/Public/MetasoundFrontendController.h
+++ b/Engine/Plugins/Runtime/Metasound/Source/MetasoundFrontend/Public/MetasoundFrontendController.h
@@ -11,6 +11,8 @@
 #include "MetasoundFrontendRegistries.h"
 #include "MetasoundGraph.h"
 
+class IMetaSoundAssetInterface;
+
 /* Metasound Controllers and Handles provide a object oriented interface for  manipulating Metasound Documents. 
  *
  * Each controller interface is associated with a single Metasound entity such as 
@@ -84,6 +86,32 @@ namespace Metasound
 		};
 
 
+		/** Provides list of interface members that have been added or removed
+		  * when querying if a node's class has been updated */
+		struct FClassInterfaceUpdates
+		{
+			TArray<FMetasoundFrontendClassInput> AddedInputs;
+			TArray<FMetasoundFrontendClassOutput> AddedOutputs;
+			TArray<FMetasoundFrontendClassInput> RemovedInputs;
+			TArray<FMetasoundFrontendClassOutput> RemovedOutputs;
+
+			bool ContainsRemovedMembers() const
+			{
+				return !RemovedInputs.IsEmpty() || !RemovedOutputs.IsEmpty();
+			}
+
+			bool ContainsAddedMembers() const
+			{
+				return !AddedInputs.IsEmpty() || !AddedOutputs.IsEmpty();
+			}
+
+			bool ContainsChanges() const
+			{
+				return ContainsRemovedMembers() || ContainsAddedMembers();
+			}
+		};
+
+
 		/** IDocumentAccessor describes an interface for various I*Controllers to interact with
 		 * each other without exposing that interface publicly or requiring friendship 
 		 * between various controller implementation classes. 
@@ -397,10 +425,18 @@ namespace Metasound
 			virtual const FMetasoundFrontendInterfaceStyle& GetInputStyle() const = 0;
 			virtual const FMetasoundFrontendClassStyle& GetClassStyle() const = 0;
 
+			/**
+			  * Fills out the provided ClassInterfaceUpdate struct with the differences between
+			  * the registry's version of the class interface and that of the node.
+			  * @return Whether or not the interface was found in the registry.
+			  */
+			virtual bool DiffAgainstRegistryInterface(FClassInterfaceUpdates& OutInterfaceUpdates, bool bInUseHighestMinorVersion) const = 0;
+
 			/** Returns whether the node is eligible for auto-updating (i.e.
 			  * has undergone minor revision or the interface has changed, but
-			  * no higher major revision is available.) */
-			virtual bool CanAutoUpdate() const = 0;
+			  * no higher major revision is available. Optionally provide
+			  * interface updates to be populate with any information regarding interface updates. */
+			virtual bool CanAutoUpdate(const IMetaSoundAssetInterface& AssetInterface, FClassInterfaceUpdates* OutInterfaceUpdates = nullptr) const = 0;
 
 			/** Description of the given node. */
 			virtual const FText& GetDescription() const = 0;
diff --git a/Engine/Plugins/Runtime/Metasound/Source/MetasoundFrontend/Public/MetasoundFrontendDocument.h b/Engine/Plugins/Runtime/Metasound/Source/MetasoundFrontend/Public/MetasoundFrontendDocument.h
--- a/Engine/Plugins/Runtime/Metasound/Source/MetasoundFrontend/Public/MetasoundFrontendDocument.h
+++ b/Engine/Plugins/Runtime/Metasound/Source/MetasoundFrontend/Public/MetasoundFrontendDocument.h
@@ -229,7 +229,7 @@ struct METASOUNDFRONTEND_API FMetasoundFrontendVertex
 	UPROPERTY()	
 	FGuid VertexID;
 
-	// Returns true if vertices have equal name, type and number of IDs.
+	// Returns true if vertices have equal name & type.
 	static bool IsFunctionalEquivalent(const FMetasoundFrontendVertex& InLHS, const FMetasoundFrontendVertex& InRHS);
 };
 
@@ -304,6 +304,11 @@ struct METASOUNDFRONTEND_API FMetasoundFrontendNodeStyle
 	UPROPERTY()
 	FMetasoundFrontendNodeStyleDisplay Display;
 
+	// Whether or not to display if
+	// the node's version has been updated
+	UPROPERTY()
+	bool bMessageNodeUpdated = false;
+
 	UPROPERTY()
 	bool bIsPrivate = false;
 };
@@ -746,15 +751,15 @@ struct METASOUNDFRONTEND_API FMetasoundFrontendClassName
 	FMetasoundFrontendClassName(const Metasound::FNodeClassName& InName);
 
 	// Namespace of class.
-	UPROPERTY(EditAnywhere, Category = General)
+	UPROPERTY(VisibleAnywhere, Category = General)
 	FName Namespace;
 
 	// Name of class.
-	UPROPERTY(EditAnywhere, Category = General)
+	UPROPERTY(VisibleAnywhere, Category = General)
 	FName Name;
 
 	// Variant of class. The Variant is used to describe an equivalent class which performs the same operation but on differing types.
-	UPROPERTY(EditAnywhere, Category = General)
+	UPROPERTY(VisibleAnywhere, Category = General)
 	FName Variant;
 
 	// Returns a full name of the class.
@@ -799,19 +804,19 @@ struct METASOUNDFRONTEND_API FMetasoundFrontendClassMetadata
 		UPROPERTY(VisibleAnywhere, Category = Metasound)
 		FText DisplayName;
 
-		UPROPERTY(VisibleAnywhere, Category = Metasound)
+		UPROPERTY(EditAnywhere, Category = Metasound)
 		FText Description;
 
 		UPROPERTY(VisibleAnywhere, Category = Metasound)
 		FText PromptIfMissing;
 
-		UPROPERTY(VisibleAnywhere, Category = Metasound)
+		UPROPERTY(EditAnywhere, Category = Metasound)
 		FText Author;
 
 		UPROPERTY(VisibleAnywhere, Category = Metasound)
-		TArray<FName> Keywords;
+		TArray<FText> Keywords;
 
-		UPROPERTY(VisibleAnywhere, Category = Metasound)
+		UPROPERTY(EditAnywhere, Category = Metasound)
 		TArray<FText> CategoryHierarchy;
 
 		// If true, auto-update will manage (add and remove)
@@ -831,6 +836,11 @@ struct METASOUNDFRONTEND_API FMetasoundFrontendClassMetadata
 			return GET_MEMBER_NAME_CHECKED(FMetasoundFrontendClassMetadata, Author);
 		}
 
+		static FName GetCategoryHierarchyPropertyName()
+		{
+			return GET_MEMBER_NAME_CHECKED(FMetasoundFrontendClassMetadata, CategoryHierarchy);
+		}
+
 		static FName GetClassNamePropertyName()
 		{
 			return GET_MEMBER_NAME_CHECKED(FMetasoundFrontendClassMetadata, ClassName);
@@ -937,12 +947,12 @@ struct METASOUNDFRONTEND_API FMetasoundFrontendClassMetadata
 			ChangeID = FGuid::NewGuid();
 		}
 
-		const TArray<FName>& GetKeywords() const
+		const TArray<FText>& GetKeywords() const
 		{
 			return Keywords;
 		}
 
-		void SetKeywords(const TArray<FName>& InKeywords)
+		void SetKeywords(const TArray<FText>& InKeywords)
 		{
 			Keywords = InKeywords;
 			ChangeID = FGuid::NewGuid();
diff --git a/Engine/Plugins/Runtime/Metasound/Source/MetasoundFrontend/Public/MetasoundFrontendNodesCategories.h b/Engine/Plugins/Runtime/Metasound/Source/MetasoundFrontend/Public/MetasoundFrontendNodesCategories.h
new file mode 100644
--- /dev/null
+++ b/Engine/Plugins/Runtime/Metasound/Source/MetasoundFrontend/Public/MetasoundFrontendNodesCategories.h
@@ -0,0 +1,13 @@
+// Copyright Epic Games, Inc. All Rights Reserved.
+
+#pragma once
+
+#include "CoreMinimal.h"
+
+namespace Metasound
+{
+	namespace NodeCategories
+	{
+		extern const FText METASOUNDFRONTEND_API Conversions;
+	}
+}
diff --git a/Engine/Plugins/Runtime/Metasound/Source/MetasoundFrontend/Public/MetasoundReceiveNode.h b/Engine/Plugins/Runtime/Metasound/Source/MetasoundFrontend/Public/MetasoundReceiveNode.h
--- a/Engine/Plugins/Runtime/Metasound/Source/MetasoundFrontend/Public/MetasoundReceiveNode.h
+++ b/Engine/Plugins/Runtime/Metasound/Source/MetasoundFrontend/Public/MetasoundReceiveNode.h
@@ -59,7 +59,7 @@ namespace Metasound
 				Info.PromptIfMissing = PluginNodeMissingPrompt;
 				Info.DefaultInterface = DeclareVertexInterface();
 				Info.CategoryHierarchy = { LOCTEXT("Metasound_TransmissionNodeCategory", "Transmission") };
-				Info.Keywords = {TEXT("receive"), GetMetasoundDataTypeName<TDataType>()};
+				Info.Keywords = { LOCTEXT("Metasound_ReceiveKeyword", "Receive"), FText::FromName(GetMetasoundDataTypeName<TDataType>()) };
 
 				return Info;
 			};
diff --git a/Engine/Plugins/Runtime/Metasound/Source/MetasoundFrontend/Public/MetasoundSendNode.h b/Engine/Plugins/Runtime/Metasound/Source/MetasoundFrontend/Public/MetasoundSendNode.h
--- a/Engine/Plugins/Runtime/Metasound/Source/MetasoundFrontend/Public/MetasoundSendNode.h
+++ b/Engine/Plugins/Runtime/Metasound/Source/MetasoundFrontend/Public/MetasoundSendNode.h
@@ -61,7 +61,7 @@ namespace Metasound
 				Info.PromptIfMissing = PluginNodeMissingPrompt;
 				Info.DefaultInterface = DeclareVertexInterface();
 				Info.CategoryHierarchy = { LOCTEXT("Metasound_TransmissionNodeCategory", "Transmission") };
-				Info.Keywords = { "Send", GetMetasoundDataTypeName<TDataType>()};
+				Info.Keywords = { LOCTEXT("Metasound_SendKeyword", "Send"), FText::FromName(GetMetasoundDataTypeName<TDataType>()) };
 
 				return Info;
 			};
diff --git a/Engine/Plugins/Runtime/Metasound/Source/MetasoundGraphCore/Public/MetasoundEnvironment.h b/Engine/Plugins/Runtime/Metasound/Source/MetasoundGraphCore/Public/MetasoundEnvironment.h
--- a/Engine/Plugins/Runtime/Metasound/Source/MetasoundGraphCore/Public/MetasoundEnvironment.h
+++ b/Engine/Plugins/Runtime/Metasound/Source/MetasoundGraphCore/Public/MetasoundEnvironment.h
@@ -4,8 +4,13 @@
 
 #include "Containers/SortedMap.h"
 #include "CoreMinimal.h"
+#include "Misc/Build.h"
 #include <type_traits>
 
+#ifndef WITH_METASOUND_DEBUG_ENVIRONMENT
+#define WITH_METASOUND_DEBUG_ENVIRONMENT !UE_BUILD_SHIPPING
+#endif
+
 #define DECLARE_METASOUND_ENVIRONMENT_VARIABLE_TYPE(DLL_API, VarType) \
 	template<> \
 	struct DLL_API ::Metasound::TEnvironmentVariableTypeInfo<VarType> \
diff --git a/Engine/Plugins/Runtime/Metasound/Source/MetasoundGraphCore/Public/MetasoundNodeInterface.h b/Engine/Plugins/Runtime/Metasound/Source/MetasoundGraphCore/Public/MetasoundNodeInterface.h
--- a/Engine/Plugins/Runtime/Metasound/Source/MetasoundGraphCore/Public/MetasoundNodeInterface.h
+++ b/Engine/Plugins/Runtime/Metasound/Source/MetasoundGraphCore/Public/MetasoundNodeInterface.h
@@ -110,7 +110,7 @@ namespace Metasound
 		TArray<FText> CategoryHierarchy;
 
 		/** List of keywords for contextual node searching. */
-		TArray<FName> Keywords;
+		TArray<FText> Keywords;
 
 		/** Display style for node when visualized. */
 		FNodeDisplayStyle DisplayStyle;
diff --git a/Engine/Plugins/Runtime/Metasound/Source/MetasoundStandardNodes/Private/MetasoundADEnvelopeNode.cpp b/Engine/Plugins/Runtime/Metasound/Source/MetasoundStandardNodes/Private/MetasoundADEnvelopeNode.cpp
--- a/Engine/Plugins/Runtime/Metasound/Source/MetasoundStandardNodes/Private/MetasoundADEnvelopeNode.cpp
+++ b/Engine/Plugins/Runtime/Metasound/Source/MetasoundStandardNodes/Private/MetasoundADEnvelopeNode.cpp
@@ -6,6 +6,7 @@
 #include "MetasoundNodeRegistrationMacro.h"
 #include "MetasoundDataTypeRegistrationMacro.h"
 #include "MetasoundPrimitives.h"
+#include "MetasoundStandardNodesCategories.h"
 #include "MetasoundStandardNodesNames.h"
 #include "MetasoundTrigger.h"
 #include "MetasoundTime.h"
@@ -225,7 +226,7 @@ namespace Metasound
 
 				const FNodeClassMetadata Metadata
 				{
-					FNodeClassName{FName("AD Envelope"), OperatorName, DataTypeName},
+					FNodeClassName { "AD Envelope", OperatorName, DataTypeName },
 					1, // Major Version
 					0, // Minor Version
 					NodeDisplayName,
@@ -233,9 +234,9 @@ namespace Metasound
 					PluginAuthor,
 					PluginNodeMissingPrompt,
 					NodeInterface,
-					{LOCTEXT("EnvelopeCat", "Envelopes")},
-					{TEXT("Envelope")},
-					FNodeDisplayStyle{}
+					{ NodeCategories::Envelopes },
+					{ },
+					FNodeDisplayStyle()
 				};
 
 				return Metadata;
diff --git a/Engine/Plugins/Runtime/Metasound/Source/MetasoundStandardNodes/Private/MetasoundADSREnvelopeNode.cpp b/Engine/Plugins/Runtime/Metasound/Source/MetasoundStandardNodes/Private/MetasoundADSREnvelopeNode.cpp
--- a/Engine/Plugins/Runtime/Metasound/Source/MetasoundStandardNodes/Private/MetasoundADSREnvelopeNode.cpp
+++ b/Engine/Plugins/Runtime/Metasound/Source/MetasoundStandardNodes/Private/MetasoundADSREnvelopeNode.cpp
@@ -6,6 +6,7 @@
 #include "MetasoundNodeRegistrationMacro.h"
 #include "MetasoundDataTypeRegistrationMacro.h"
 #include "MetasoundPrimitives.h"
+#include "MetasoundStandardNodesCategories.h"
 #include "MetasoundStandardNodesNames.h"
 #include "MetasoundTrigger.h"
 #include "MetasoundTime.h"
@@ -364,7 +365,7 @@ namespace Metasound
 
 				const FNodeClassMetadata Metadata
 				{
-					FNodeClassName{FName("ADSR Envelope"), OperatorName, DataTypeName},
+					FNodeClassName { "ADSR Envelope", OperatorName, DataTypeName },
 					1, // Major Version
 					0, // Minor Version
 					NodeDisplayName,
@@ -372,9 +373,9 @@ namespace Metasound
 					PluginAuthor,
 					PluginNodeMissingPrompt,
 					NodeInterface,
-					{LOCTEXT("EnvelopeCat", "Envelopes")},
-					{TEXT("Envelope")},
-					FNodeDisplayStyle{}
+					{ NodeCategories::Envelopes },
+					{ },
+					FNodeDisplayStyle()
 				};
 
 				return Metadata;
diff --git a/Engine/Plugins/Runtime/Metasound/Source/MetasoundStandardNodes/Private/MetasoundBPMToSecondsNode.cpp b/Engine/Plugins/Runtime/Metasound/Source/MetasoundStandardNodes/Private/MetasoundBPMToSecondsNode.cpp
--- a/Engine/Plugins/Runtime/Metasound/Source/MetasoundStandardNodes/Private/MetasoundBPMToSecondsNode.cpp
+++ b/Engine/Plugins/Runtime/Metasound/Source/MetasoundStandardNodes/Private/MetasoundBPMToSecondsNode.cpp
@@ -151,7 +151,7 @@ namespace Metasound
 		auto InitNodeInfo = []() -> FNodeClassMetadata
 		{
 			FNodeClassMetadata Info;
-			Info.ClassName = { Metasound::StandardNodes::Namespace, TEXT("BPMToSeconds"), TEXT("") };
+			Info.ClassName = { StandardNodes::Namespace, TEXT("BPMToSeconds"), TEXT("") };
 			Info.MajorVersion = 1;
 			Info.MinorVersion = 0;
 			Info.DisplayName = LOCTEXT("BPMToSecondsDisplayName", "BPM To Seconds");
diff --git a/Engine/Plugins/Runtime/Metasound/Source/MetasoundStandardNodes/Private/MetasoundBasicFilters.cpp b/Engine/Plugins/Runtime/Metasound/Source/MetasoundStandardNodes/Private/MetasoundBasicFilters.cpp
--- a/Engine/Plugins/Runtime/Metasound/Source/MetasoundStandardNodes/Private/MetasoundBasicFilters.cpp
+++ b/Engine/Plugins/Runtime/Metasound/Source/MetasoundStandardNodes/Private/MetasoundBasicFilters.cpp
@@ -160,7 +160,7 @@ namespace Metasound
 		auto InitNodeInfo = []() -> FNodeClassMetadata
 		{
 			FNodeClassMetadata Info;
-			Info.ClassName = { Metasound::StandardNodes::Namespace, TEXT("Ladder Filter"), Metasound::StandardNodes::AudioVariant };
+			Info.ClassName = { StandardNodes::Namespace, TEXT("Ladder Filter"), StandardNodes::AudioVariant };
 			Info.MajorVersion = 1;
 			Info.MinorVersion = 0;
 			Info.DisplayName = LOCTEXT("Metasound_LadderFilterNodeDisplayName", "Ladder Filter");
@@ -168,7 +168,7 @@ namespace Metasound
 			Info.Author = PluginAuthor;
 			Info.PromptIfMissing = PluginNodeMissingPrompt;
 			Info.DefaultInterface = DeclareVertexInterface();
-			Info.CategoryHierarchy.Emplace(StandardNodes::Filters);
+			Info.CategoryHierarchy.Emplace(NodeCategories::Filters);
 
 			return Info;
 		};
@@ -342,7 +342,7 @@ namespace Metasound
 		auto InitNodeInfo = []() -> FNodeClassMetadata
 		{
 			FNodeClassMetadata Info;
-			Info.ClassName = { Metasound::StandardNodes::Namespace, TEXT("State Variable Filter"), Metasound::StandardNodes::AudioVariant };
+			Info.ClassName = { StandardNodes::Namespace, TEXT("State Variable Filter"), StandardNodes::AudioVariant };
 			Info.MajorVersion = 1;
 			Info.MinorVersion = 0;
 			Info.DisplayName = LOCTEXT("Metasound_StateVariableFilterNodeDisplayName", "State Variable Filter");
@@ -350,7 +350,7 @@ namespace Metasound
 				Info.Author = PluginAuthor;
 			Info.PromptIfMissing = PluginNodeMissingPrompt;
 			Info.DefaultInterface = DeclareVertexInterface();
-			Info.CategoryHierarchy.Emplace(StandardNodes::Filters);
+			Info.CategoryHierarchy.Emplace(NodeCategories::Filters);
 
 			return Info;
 		};
@@ -512,7 +512,7 @@ namespace Metasound
 		auto InitNodeInfo = []() -> FNodeClassMetadata
 		{
 			FNodeClassMetadata Info;
-			Info.ClassName = { Metasound::StandardNodes::Namespace, TEXT("One-Pole Low Pass Filter"), Metasound::StandardNodes::AudioVariant };
+			Info.ClassName = { StandardNodes::Namespace, TEXT("One-Pole Low Pass Filter"), StandardNodes::AudioVariant };
 			Info.MajorVersion = 1;
 			Info.MinorVersion = 0;
 			Info.DisplayName = LOCTEXT("Metasound_OnePoleLpfNodeDisplayName", "One-Pole Low Pass Filter");
@@ -520,7 +520,7 @@ namespace Metasound
 				Info.Author = PluginAuthor;
 			Info.PromptIfMissing = PluginNodeMissingPrompt;
 			Info.DefaultInterface = DeclareVertexInterface();
-			Info.CategoryHierarchy.Emplace(StandardNodes::Filters);
+			Info.CategoryHierarchy.Emplace(NodeCategories::Filters);
 
 			return Info;
 		};
@@ -647,7 +647,7 @@ namespace Metasound
 		auto InitNodeInfo = []() -> FNodeClassMetadata
 		{
 			FNodeClassMetadata Info;
-			Info.ClassName = { Metasound::StandardNodes::Namespace, TEXT("One-Pole High Pass Filter"), Metasound::StandardNodes::AudioVariant };
+			Info.ClassName = { StandardNodes::Namespace, TEXT("One-Pole High Pass Filter"), StandardNodes::AudioVariant };
 			Info.MajorVersion = 1;
 			Info.MinorVersion = 0;
 			Info.DisplayName = LOCTEXT("Metasound_OnePoleHpfNodeDisplayName", "One-Pole High Pass Filter");
@@ -655,7 +655,7 @@ namespace Metasound
 				Info.Author = PluginAuthor;
 			Info.PromptIfMissing = PluginNodeMissingPrompt;
 			Info.DefaultInterface = DeclareVertexInterface();
-			Info.CategoryHierarchy.Emplace(StandardNodes::Filters);
+			Info.CategoryHierarchy.Emplace(NodeCategories::Filters);
 
 			return Info;
 		};
@@ -801,7 +801,7 @@ namespace Metasound
 		auto InitNodeInfo = []() -> FNodeClassMetadata
 		{
 			FNodeClassMetadata Info;
-			Info.ClassName = { Metasound::StandardNodes::Namespace, TEXT("Biquad Filter"), Metasound::StandardNodes::AudioVariant };
+			Info.ClassName = { StandardNodes::Namespace, TEXT("Biquad Filter"), StandardNodes::AudioVariant };
 			Info.MajorVersion = 1;
 			Info.MinorVersion = 0;
 			Info.DisplayName = LOCTEXT("Metasound_BiquadFilterNodeDisplayName", "Biquad Filter");
@@ -809,7 +809,7 @@ namespace Metasound
 				Info.Author = PluginAuthor;
 			Info.PromptIfMissing = PluginNodeMissingPrompt;
 			Info.DefaultInterface = DeclareVertexInterface();
-			Info.CategoryHierarchy.Emplace(StandardNodes::Filters);
+			Info.CategoryHierarchy.Emplace(NodeCategories::Filters);
 
 			return Info;
 		};
diff --git a/Engine/Plugins/Runtime/Metasound/Source/MetasoundStandardNodes/Private/MetasoundBitcrusherNode.cpp b/Engine/Plugins/Runtime/Metasound/Source/MetasoundStandardNodes/Private/MetasoundBitcrusherNode.cpp
--- a/Engine/Plugins/Runtime/Metasound/Source/MetasoundStandardNodes/Private/MetasoundBitcrusherNode.cpp
+++ b/Engine/Plugins/Runtime/Metasound/Source/MetasoundStandardNodes/Private/MetasoundBitcrusherNode.cpp
@@ -56,7 +56,7 @@ namespace Metasound
 
 				FNodeClassMetadata Metadata
 				{
-					FNodeClassName{StandardNodes::Namespace, TEXT("Bitcrusher"), StandardNodes::AudioVariant},
+					FNodeClassName { StandardNodes::Namespace, "Bitcrusher", StandardNodes::AudioVariant },
 					1, // Major Version
 					0, // Minor Version
 					LOCTEXT("BitcrusherDisplayName", "Bitcrusher"),
@@ -64,9 +64,9 @@ namespace Metasound
 					PluginAuthor,
 					PluginNodeMissingPrompt,
 					NodeInterface,
-					{StandardNodes::Filters},
-					{TEXT("Bitcrusher")},
-					FNodeDisplayStyle{}
+					{ NodeCategories::Filters },
+					{ },
+					FNodeDisplayStyle()
 				};
 
 				return Metadata;
diff --git a/Engine/Plugins/Runtime/Metasound/Source/MetasoundStandardNodes/Private/MetasoundClampNode.cpp b/Engine/Plugins/Runtime/Metasound/Source/MetasoundStandardNodes/Private/MetasoundClampNode.cpp
--- a/Engine/Plugins/Runtime/Metasound/Source/MetasoundStandardNodes/Private/MetasoundClampNode.cpp
+++ b/Engine/Plugins/Runtime/Metasound/Source/MetasoundStandardNodes/Private/MetasoundClampNode.cpp
@@ -27,7 +27,7 @@ namespace Metasound
 		{
 			FNodeClassMetadata Metadata
 			{
-				FNodeClassName{FName("Clamp"), InOperatorName, InDataTypeName},
+				FNodeClassName { "Clamp", InOperatorName, InDataTypeName },
 				1, // Major Version
 				0, // Minor Version
 				InDisplayName,
@@ -35,9 +35,9 @@ namespace Metasound
 				PluginAuthor,
 				PluginNodeMissingPrompt,
 				InDefaultInterface,
-				{LOCTEXT("ClampCategory", "Math")},
-				{TEXT("Clamp")},
-				FNodeDisplayStyle{}
+				{ LOCTEXT("ClampCategory", "Math") },
+				{ LOCTEXT("ClampKeyword", "Clamp") },
+				FNodeDisplayStyle()
 			};
 
 			return Metadata;
diff --git a/Engine/Plugins/Runtime/Metasound/Source/MetasoundStandardNodes/Private/MetasoundConversionNodes.cpp b/Engine/Plugins/Runtime/Metasound/Source/MetasoundStandardNodes/Private/MetasoundConversionNodes.cpp
--- a/Engine/Plugins/Runtime/Metasound/Source/MetasoundStandardNodes/Private/MetasoundConversionNodes.cpp
+++ b/Engine/Plugins/Runtime/Metasound/Source/MetasoundStandardNodes/Private/MetasoundConversionNodes.cpp
@@ -182,7 +182,7 @@ namespace Metasound
 				const FText NodeDescription = FText::Format(LOCTEXT("ConverterNodeDesc", "Converts from {0} to {1}."), FromTypeText, ToTypeText);
 
 				FNodeClassMetadata Info;
-				Info.ClassName = { Metasound::StandardNodes::Namespace, ClassName, "" };
+				Info.ClassName = { StandardNodes::Namespace, ClassName, "" };
 				Info.MajorVersion = 1;
 				Info.MinorVersion = 0;
 				Info.DisplayName = NodeDisplayName;
@@ -190,7 +190,7 @@ namespace Metasound
 				Info.Author = PluginAuthor;
 				Info.PromptIfMissing = PluginNodeMissingPrompt;
 				Info.DefaultInterface = GetDefaultInterface();
-				Info.CategoryHierarchy.Emplace(StandardNodes::Conversions);
+				Info.CategoryHierarchy.Emplace(NodeCategories::Conversions);
 
 				return Info;
 			};
diff --git a/Engine/Plugins/Runtime/Metasound/Source/MetasoundStandardNodes/Private/MetasoundCrossfadeNode.cpp b/Engine/Plugins/Runtime/Metasound/Source/MetasoundStandardNodes/Private/MetasoundCrossfadeNode.cpp
--- a/Engine/Plugins/Runtime/Metasound/Source/MetasoundStandardNodes/Private/MetasoundCrossfadeNode.cpp
+++ b/Engine/Plugins/Runtime/Metasound/Source/MetasoundStandardNodes/Private/MetasoundCrossfadeNode.cpp
@@ -6,6 +6,7 @@
 #include "MetasoundFacade.h"
 #include "MetasoundExecutableOperator.h"
 #include "MetasoundPrimitives.h"
+#include "MetasoundStandardNodesCategories.h"
 #include "MetasoundStandardNodesNames.h"
 #include "MetasoundTrigger.h"
 #include "Internationalization/Text.h"
@@ -189,7 +190,7 @@ namespace Metasound
 
 				FNodeClassMetadata Metadata
 				{
-					FNodeClassName{FName("Crossfade"), OperatorName, DataTypeName},
+					FNodeClassName { "Crossfade", OperatorName, DataTypeName },
 					1, // Major Version
 					0, // Minor Version
 					NodeDisplayName,
@@ -197,9 +198,9 @@ namespace Metasound
 					PluginAuthor,
 					PluginNodeMissingPrompt,
 					NodeInterface,
-					{LOCTEXT("CrossfadeCategory", "Math")},
-					{TEXT("Crossfade")},
-					FNodeDisplayStyle{}
+					{ NodeCategories::Envelopes },
+					{ },
+					FNodeDisplayStyle()
 				};
 				return Metadata;
 			};
diff --git a/Engine/Plugins/Runtime/Metasound/Source/MetasoundStandardNodes/Private/MetasoundDebugLogNode.cpp b/Engine/Plugins/Runtime/Metasound/Source/MetasoundStandardNodes/Private/MetasoundDebugLogNode.cpp
--- a/Engine/Plugins/Runtime/Metasound/Source/MetasoundStandardNodes/Private/MetasoundDebugLogNode.cpp
+++ b/Engine/Plugins/Runtime/Metasound/Source/MetasoundStandardNodes/Private/MetasoundDebugLogNode.cpp
@@ -18,7 +18,7 @@ namespace Metasound
 		{
 			FNodeClassMetadata Metadata
 			{
-				FNodeClassName{FName("Print Log"), InOperatorName, InDataTypeName},
+				FNodeClassName { "Print Log", InOperatorName, InDataTypeName },
 				1, // Major Version
 				0, // Minor Version
 				InDisplayName,
@@ -26,9 +26,9 @@ namespace Metasound
 				PluginAuthor,
 				PluginNodeMissingPrompt,
 				InDefaultInterface,
-				{ StandardNodes::DebugUtils },
-				{TEXT("Print Log")},
-				FNodeDisplayStyle{}
+				{ NodeCategories::Debug },
+				{ },
+				FNodeDisplayStyle()
 			};
 
 			return Metadata;
diff --git a/Engine/Plugins/Runtime/Metasound/Source/MetasoundStandardNodes/Private/MetasoundDecibelsToLinearGain.cpp b/Engine/Plugins/Runtime/Metasound/Source/MetasoundStandardNodes/Private/MetasoundDecibelsToLinearGain.cpp
--- a/Engine/Plugins/Runtime/Metasound/Source/MetasoundStandardNodes/Private/MetasoundDecibelsToLinearGain.cpp
+++ b/Engine/Plugins/Runtime/Metasound/Source/MetasoundStandardNodes/Private/MetasoundDecibelsToLinearGain.cpp
@@ -102,7 +102,7 @@ namespace Metasound
 			const FText NodeDescription = LOCTEXT("Metasound_DecibelsToLinearGainDescription", "Converts a logarithmic (dB) gain value to a linear gain value.");
 
 			FNodeClassMetadata Info;
-			Info.ClassName = { Metasound::StandardNodes::Namespace, OperatorName, DataTypeName };
+			Info.ClassName = { StandardNodes::Namespace, OperatorName, DataTypeName };
 			Info.MajorVersion = 1;
 			Info.MinorVersion = 0;
 			Info.DisplayName = NodeDisplayName;
@@ -110,7 +110,7 @@ namespace Metasound
 			Info.Author = PluginAuthor;
 			Info.PromptIfMissing = PluginNodeMissingPrompt;
 			Info.DefaultInterface = GetVertexInterface();
-			Info.CategoryHierarchy.Emplace(StandardNodes::Dynamics);
+			Info.CategoryHierarchy.Emplace(NodeCategories::Dynamics);
 
 			return Info;
 		};
diff --git a/Engine/Plugins/Runtime/Metasound/Source/MetasoundStandardNodes/Private/MetasoundDelayNode.cpp b/Engine/Plugins/Runtime/Metasound/Source/MetasoundStandardNodes/Private/MetasoundDelayNode.cpp
--- a/Engine/Plugins/Runtime/Metasound/Source/MetasoundStandardNodes/Private/MetasoundDelayNode.cpp
+++ b/Engine/Plugins/Runtime/Metasound/Source/MetasoundStandardNodes/Private/MetasoundDelayNode.cpp
@@ -189,7 +189,7 @@ namespace Metasound
 		auto InitNodeInfo = []() -> FNodeClassMetadata
 		{
 			FNodeClassMetadata Info;
-			Info.ClassName = { Metasound::StandardNodes::Namespace, TEXT("Delay"), Metasound::StandardNodes::AudioVariant };
+			Info.ClassName = { StandardNodes::Namespace, TEXT("Delay"), StandardNodes::AudioVariant };
 			Info.MajorVersion = 1;
 			Info.MinorVersion = 0;
 			Info.DisplayName = LOCTEXT("Metasound_DelayDisplayName", "Delay");
@@ -197,7 +197,7 @@ namespace Metasound
 			Info.Author = PluginAuthor;
 			Info.PromptIfMissing = PluginNodeMissingPrompt;
 			Info.DefaultInterface = GetVertexInterface();
-			Info.CategoryHierarchy.Emplace(StandardNodes::Delays);
+			Info.CategoryHierarchy.Emplace(NodeCategories::Delays);
 			return Info;
 		};
 
diff --git a/Engine/Plugins/Runtime/Metasound/Source/MetasoundStandardNodes/Private/MetasoundEnvelopeFollowerNode.cpp b/Engine/Plugins/Runtime/Metasound/Source/MetasoundStandardNodes/Private/MetasoundEnvelopeFollowerNode.cpp
--- a/Engine/Plugins/Runtime/Metasound/Source/MetasoundStandardNodes/Private/MetasoundEnvelopeFollowerNode.cpp
+++ b/Engine/Plugins/Runtime/Metasound/Source/MetasoundStandardNodes/Private/MetasoundEnvelopeFollowerNode.cpp
@@ -24,6 +24,7 @@ namespace Metasound
 		static const TCHAR* InParamNameAudioInput = TEXT("In");
 		static const TCHAR* InParamNameAttackTime = TEXT("Attack Time");
 		static const TCHAR* InParamNameReleaseTime = TEXT("Release Time");
+		static const TCHAR* PoopTest = TEXT("Poop Test");
 		static const TCHAR* OutParamNameEnvelope = TEXT("Envelope");
 	}
 
@@ -85,7 +86,7 @@ namespace Metasound
 		InputDataReferences.AddDataReadReference(EnvelopeFollower::InParamNameAudioInput, AudioInput);
 		InputDataReferences.AddDataReadReference(EnvelopeFollower::InParamNameAttackTime, AttackTime);
 		InputDataReferences.AddDataReadReference(EnvelopeFollower::InParamNameReleaseTime, ReleaseTime);
-
+		InputDataReferences.AddDataReadReference(EnvelopeFollower::PoopTest, ReleaseTime);
 		return InputDataReferences;
 	}
 
@@ -128,7 +129,8 @@ namespace Metasound
 			FInputVertexInterface(
 				TInputDataVertexModel<FAudioBuffer>(EnvelopeFollower::InParamNameAudioInput, LOCTEXT("AudioInputToolTT", "Audio input.")),
 				TInputDataVertexModel<FTime>(EnvelopeFollower::InParamNameAttackTime, LOCTEXT("AttackTimeTT", "The attack time of the envelope follower."), 0.01f),
-				TInputDataVertexModel<FTime>(EnvelopeFollower::InParamNameReleaseTime, LOCTEXT("ReleaseTimeTT", "The release time of the envelope follower."), 0.1f)
+				TInputDataVertexModel<FTime>(EnvelopeFollower::InParamNameReleaseTime, LOCTEXT("ReleaseTimeTT", "The release time of the envelope follower."), 0.1f),
+				TInputDataVertexModel<FTime>(EnvelopeFollower::PoopTest, LOCTEXT("PoopTest", "Just your average poop test."), 0.1f)
 			),
 			FOutputVertexInterface(
 				TOutputDataVertexModel<float>(EnvelopeFollower::OutParamNameEnvelope, LOCTEXT("EnvelopeFollowerOutputTT", "The output envelope value of the audio signal."))
@@ -143,9 +145,9 @@ namespace Metasound
 		auto InitNodeInfo = []() -> FNodeClassMetadata
 		{
 			FNodeClassMetadata Info;
-			Info.ClassName = { Metasound::StandardNodes::Namespace, TEXT("Envelope Follower"), TEXT("") };
+			Info.ClassName = { StandardNodes::Namespace, TEXT("Envelope Follower"), TEXT("") };
 			Info.MajorVersion = 1;
-			Info.MinorVersion = 0;
+			Info.MinorVersion = 1;
 			Info.DisplayName = LOCTEXT("Metasound_EnvelopeFollowerDisplayName", "Envelope Follower");
 			Info.Description = LOCTEXT("Metasound_EnvelopeFollowerDescription", "Outputs an envelope from an input audio signal.");
 			Info.Author = PluginAuthor;
diff --git a/Engine/Plugins/Runtime/Metasound/Source/MetasoundStandardNodes/Private/MetasoundFlangerNode.cpp b/Engine/Plugins/Runtime/Metasound/Source/MetasoundStandardNodes/Private/MetasoundFlangerNode.cpp
--- a/Engine/Plugins/Runtime/Metasound/Source/MetasoundStandardNodes/Private/MetasoundFlangerNode.cpp
+++ b/Engine/Plugins/Runtime/Metasound/Source/MetasoundStandardNodes/Private/MetasoundFlangerNode.cpp
@@ -39,7 +39,7 @@ namespace Metasound
 				const FText NodeDescription = LOCTEXT("Metasound_FlangerNodeDescription", "Applies a flanger effect to input audio.");
 
 				FNodeClassMetadata Info;
-				Info.ClassName = { Metasound::StandardNodes::Namespace, OperatorName, TEXT("") };
+				Info.ClassName = { StandardNodes::Namespace, OperatorName, TEXT("") };
 				Info.MajorVersion = 1;
 				Info.MinorVersion = 0;
 				Info.DisplayName = NodeDisplayName;
diff --git a/Engine/Plugins/Runtime/Metasound/Source/MetasoundStandardNodes/Private/MetasoundFreqMultiplierToSemitoneNode.cpp b/Engine/Plugins/Runtime/Metasound/Source/MetasoundStandardNodes/Private/MetasoundFreqMultiplierToSemitoneNode.cpp
--- a/Engine/Plugins/Runtime/Metasound/Source/MetasoundStandardNodes/Private/MetasoundFreqMultiplierToSemitoneNode.cpp
+++ b/Engine/Plugins/Runtime/Metasound/Source/MetasoundStandardNodes/Private/MetasoundFreqMultiplierToSemitoneNode.cpp
@@ -111,7 +111,7 @@ namespace Metasound
 			const FText NodeDescription = LOCTEXT("Metasound_FrequencyMultiplierToSemitoneDescription", "Converts a frequency multiplier to the corresponding difference in semitones.");
 
 			FNodeClassMetadata Info;
-			Info.ClassName = { Metasound::StandardNodes::Namespace, OperatorName, DataTypeName };
+			Info.ClassName = { StandardNodes::Namespace, OperatorName, DataTypeName };
 			Info.MajorVersion = 1;
 			Info.MinorVersion = 0;
 			Info.DisplayName = NodeDisplayName;
@@ -119,7 +119,7 @@ namespace Metasound
 			Info.Author = PluginAuthor;
 			Info.PromptIfMissing = PluginNodeMissingPrompt;
 			Info.DefaultInterface = GetVertexInterface();
-			Info.CategoryHierarchy.Emplace(StandardNodes::Conversions);
+			Info.CategoryHierarchy.Emplace(NodeCategories::Conversions);
 
 			return Info;
 		};
diff --git a/Engine/Plugins/Runtime/Metasound/Source/MetasoundStandardNodes/Private/MetasoundFreqToMIDINode.cpp b/Engine/Plugins/Runtime/Metasound/Source/MetasoundStandardNodes/Private/MetasoundFreqToMIDINode.cpp
--- a/Engine/Plugins/Runtime/Metasound/Source/MetasoundStandardNodes/Private/MetasoundFreqToMIDINode.cpp
+++ b/Engine/Plugins/Runtime/Metasound/Source/MetasoundStandardNodes/Private/MetasoundFreqToMIDINode.cpp
@@ -112,7 +112,7 @@ namespace Metasound
 			const FText NodeDescription = LOCTEXT("Metasound_FreqToMidiNodeDescription", "Converts a frequency (Hz) value to a Midi Note value.");
 
 			FNodeClassMetadata Info;
-			Info.ClassName = { Metasound::StandardNodes::Namespace, OperatorName, DataTypeName };
+			Info.ClassName = { StandardNodes::Namespace, OperatorName, DataTypeName };
 			Info.MajorVersion = 1;
 			Info.MinorVersion = 0;
 			Info.DisplayName = NodeDisplayName;
@@ -120,7 +120,7 @@ namespace Metasound
 			Info.Author = PluginAuthor;
 			Info.PromptIfMissing = PluginNodeMissingPrompt;
 			Info.DefaultInterface = GetVertexInterface();
-			Info.CategoryHierarchy.Emplace(StandardNodes::Music);
+			Info.CategoryHierarchy.Emplace(NodeCategories::Music);
 
 			return Info;
 		};
diff --git a/Engine/Plugins/Runtime/Metasound/Source/MetasoundStandardNodes/Private/MetasoundITDPannerNode.cpp b/Engine/Plugins/Runtime/Metasound/Source/MetasoundStandardNodes/Private/MetasoundITDPannerNode.cpp
--- a/Engine/Plugins/Runtime/Metasound/Source/MetasoundStandardNodes/Private/MetasoundITDPannerNode.cpp
+++ b/Engine/Plugins/Runtime/Metasound/Source/MetasoundStandardNodes/Private/MetasoundITDPannerNode.cpp
@@ -310,7 +310,7 @@ namespace Metasound
 		auto InitNodeInfo = []() -> FNodeClassMetadata
 		{
 			FNodeClassMetadata Info;
-			Info.ClassName = { Metasound::StandardNodes::Namespace, TEXT("ITD Panner"), TEXT("") };
+			Info.ClassName = { StandardNodes::Namespace, TEXT("ITD Panner"), TEXT("") };
 			Info.MajorVersion = 1;
 			Info.MinorVersion = 0;
 			Info.DisplayName = LOCTEXT("Metasound_ITDPannerDisplayName", "ITD Panner");
@@ -318,7 +318,7 @@ namespace Metasound
 			Info.Author = PluginAuthor;
 			Info.PromptIfMissing = PluginNodeMissingPrompt;
 			Info.DefaultInterface = GetVertexInterface();
-			Info.CategoryHierarchy.Emplace(StandardNodes::Spatalization);
+			Info.CategoryHierarchy.Emplace(NodeCategories::Spatialization);
 			return Info;
 		};
 
diff --git a/Engine/Plugins/Runtime/Metasound/Source/MetasoundStandardNodes/Private/MetasoundInterpToNode.cpp b/Engine/Plugins/Runtime/Metasound/Source/MetasoundStandardNodes/Private/MetasoundInterpToNode.cpp
--- a/Engine/Plugins/Runtime/Metasound/Source/MetasoundStandardNodes/Private/MetasoundInterpToNode.cpp
+++ b/Engine/Plugins/Runtime/Metasound/Source/MetasoundStandardNodes/Private/MetasoundInterpToNode.cpp
@@ -140,7 +140,7 @@ namespace Metasound
 		auto InitNodeInfo = []() -> FNodeClassMetadata
 		{
 			FNodeClassMetadata Info;
-			Info.ClassName = { Metasound::StandardNodes::Namespace, TEXT("InterpTo"), Metasound::StandardNodes::AudioVariant };
+			Info.ClassName = { StandardNodes::Namespace, TEXT("InterpTo"), StandardNodes::AudioVariant };
 			Info.MajorVersion = 1;
 			Info.MinorVersion = 0;
 			Info.DisplayName = LOCTEXT("Metasound_InterpDisplayName", "InterpTo");
diff --git a/Engine/Plugins/Runtime/Metasound/Source/MetasoundStandardNodes/Private/MetasoundLinearGainToDecibels.cpp b/Engine/Plugins/Runtime/Metasound/Source/MetasoundStandardNodes/Private/MetasoundLinearGainToDecibels.cpp
--- a/Engine/Plugins/Runtime/Metasound/Source/MetasoundStandardNodes/Private/MetasoundLinearGainToDecibels.cpp
+++ b/Engine/Plugins/Runtime/Metasound/Source/MetasoundStandardNodes/Private/MetasoundLinearGainToDecibels.cpp
@@ -107,7 +107,7 @@ namespace Metasound
 			const FText NodeDescription = LOCTEXT("Metasound_LinearGainToDecibelsDescription", "Converts a linear gain value to a logarithmic (dB) gain value.");
 
 			FNodeClassMetadata Info;
-			Info.ClassName = { Metasound::StandardNodes::Namespace, OperatorName, DataTypeName };
+			Info.ClassName = { StandardNodes::Namespace, OperatorName, DataTypeName };
 			Info.MajorVersion = 1;
 			Info.MinorVersion = 0;
 			Info.DisplayName = NodeDisplayName;
@@ -115,7 +115,7 @@ namespace Metasound
 			Info.Author = PluginAuthor;
 			Info.PromptIfMissing = PluginNodeMissingPrompt;
 			Info.DefaultInterface = GetVertexInterface();
-			Info.CategoryHierarchy.Emplace(StandardNodes::Dynamics);
+			Info.CategoryHierarchy.Emplace(NodeCategories::Dynamics);
 
 			return Info;
 		};
diff --git a/Engine/Plugins/Runtime/Metasound/Source/MetasoundStandardNodes/Private/MetasoundLinearToLogFrequencyNode.cpp b/Engine/Plugins/Runtime/Metasound/Source/MetasoundStandardNodes/Private/MetasoundLinearToLogFrequencyNode.cpp
--- a/Engine/Plugins/Runtime/Metasound/Source/MetasoundStandardNodes/Private/MetasoundLinearToLogFrequencyNode.cpp
+++ b/Engine/Plugins/Runtime/Metasound/Source/MetasoundStandardNodes/Private/MetasoundLinearToLogFrequencyNode.cpp
@@ -37,7 +37,7 @@ namespace Metasound
 				FText NodeDescription = LOCTEXT("LinearToLogFrequencyDescription", "Converts a linear space input value to log-frequency space.");
 
 				FNodeClassMetadata Info;
-				Info.ClassName = { Metasound::StandardNodes::Namespace, OperatorName, TEXT("") };
+				Info.ClassName = { StandardNodes::Namespace, OperatorName, TEXT("") };
 				Info.MajorVersion = 1;
 				Info.MinorVersion = 0;
 				Info.DisplayName = NodeDisplayName;
@@ -45,7 +45,7 @@ namespace Metasound
 				Info.Author = PluginAuthor;
 				Info.PromptIfMissing = PluginNodeMissingPrompt;
 				Info.DefaultInterface = GetVertexInterface();
-				Info.CategoryHierarchy.Emplace(StandardNodes::Music);
+				Info.CategoryHierarchy.Emplace(NodeCategories::Music);
 
 				return Info;
 			};
diff --git a/Engine/Plugins/Runtime/Metasound/Source/MetasoundStandardNodes/Private/MetasoundMIDIToFreqNode.cpp b/Engine/Plugins/Runtime/Metasound/Source/MetasoundStandardNodes/Private/MetasoundMIDIToFreqNode.cpp
--- a/Engine/Plugins/Runtime/Metasound/Source/MetasoundStandardNodes/Private/MetasoundMIDIToFreqNode.cpp
+++ b/Engine/Plugins/Runtime/Metasound/Source/MetasoundStandardNodes/Private/MetasoundMIDIToFreqNode.cpp
@@ -154,7 +154,7 @@ namespace Metasound
 			const FText NodeDescription = LOCTEXT("Metasound_MidiToFreqNodeDescription", "Converts a Midi note value to a frequency (hz) value.");
 
 			FNodeClassMetadata Info;
-			Info.ClassName = { Metasound::StandardNodes::Namespace, OperatorName, DataTypeName };
+			Info.ClassName = { StandardNodes::Namespace, OperatorName, DataTypeName };
 			Info.MajorVersion = 1;
 			Info.MinorVersion = 0;
 			Info.DisplayName = NodeDisplayName;
@@ -162,7 +162,7 @@ namespace Metasound
 			Info.Author = PluginAuthor;
 			Info.PromptIfMissing = PluginNodeMissingPrompt;
 			Info.DefaultInterface = GetVertexInterface();
-			Info.CategoryHierarchy.Emplace(StandardNodes::Music);
+			Info.CategoryHierarchy.Emplace(NodeCategories::Music);
 
 			return Info;
 		};
diff --git a/Engine/Plugins/Runtime/Metasound/Source/MetasoundStandardNodes/Private/MetasoundMapRangeNode.cpp b/Engine/Plugins/Runtime/Metasound/Source/MetasoundStandardNodes/Private/MetasoundMapRangeNode.cpp
--- a/Engine/Plugins/Runtime/Metasound/Source/MetasoundStandardNodes/Private/MetasoundMapRangeNode.cpp
+++ b/Engine/Plugins/Runtime/Metasound/Source/MetasoundStandardNodes/Private/MetasoundMapRangeNode.cpp
@@ -6,6 +6,7 @@
 #include "MetasoundFacade.h"
 #include "MetasoundExecutableOperator.h"
 #include "MetasoundPrimitives.h"
+#include "MetasoundStandardNodesCategories.h"
 #include "MetasoundStandardNodesNames.h"
 #include "MetasoundTrigger.h"
 #include "Internationalization/Text.h"
@@ -305,7 +306,7 @@ namespace Metasound
 
 				FNodeClassMetadata Metadata
 				{
-					FNodeClassName{FName("MapRange"), OperatorName, FMappingClass::GetDataTypeName()},
+					FNodeClassName { "MapRange", OperatorName, FMappingClass::GetDataTypeName() },
 					1, // Major Version
 					0, // Minor Version
 					FMappingClass::GetNodeName(),
@@ -313,9 +314,9 @@ namespace Metasound
 					PluginAuthor,
 					PluginNodeMissingPrompt,
 					NodeInterface,
-					{LOCTEXT("MapRangeCategory", "Math")},
-					{TEXT("MapRange")},
-					FNodeDisplayStyle{}
+					{ NodeCategories::Math },
+					{ },
+					FNodeDisplayStyle()
 				};
 				return Metadata;
 			};
diff --git a/Engine/Plugins/Runtime/Metasound/Source/MetasoundStandardNodes/Private/MetasoundMathNodes.cpp b/Engine/Plugins/Runtime/Metasound/Source/MetasoundStandardNodes/Private/MetasoundMathNodes.cpp
--- a/Engine/Plugins/Runtime/Metasound/Source/MetasoundStandardNodes/Private/MetasoundMathNodes.cpp
+++ b/Engine/Plugins/Runtime/Metasound/Source/MetasoundStandardNodes/Private/MetasoundMathNodes.cpp
@@ -13,25 +13,27 @@
 #include "MetasoundNodeInterface.h"
 #include "MetasoundOperatorInterface.h"
 #include "MetasoundPrimitives.h"
+#include "MetasoundStandardNodesCategories.h"
 #include "MetasoundStandardNodesNames.h"
 #include "MetasoundTime.h"
 #include "MetasoundTrigger.h"
 
 #define LOCTEXT_NAMESPACE "MetasoundMathOpNode"
 
-#define DEFINE_METASOUND_MATHOP(OpName, DataTypeName, DataClass, Description) \
+#define DEFINE_METASOUND_MATHOP(OpName, DataTypeName, DataClass, Description, Keywords) \
 	class F##OpName##DataTypeName##Node \
 		: public ::Metasound::TMathOpNode<F##OpName##DataTypeName##Node, ::Metasound::TMathOp##OpName<DataClass, DataClass>, DataClass, DataClass> \
 	{ \
 	public: \
-		static FNodeClassName GetClassName() { return {::Metasound::StandardNodes::Namespace, TEXT(#OpName), TEXT(#DataTypeName)}; } \
+		static FNodeClassName GetClassName() { return { ::Metasound::StandardNodes::Namespace, TEXT(#OpName), TEXT(#DataTypeName)}; } \
 		static FText GetDisplayName() { return MathOpNames::Get##OpName##DisplayName<DataClass>(); } \
 		static FText GetDescription() { return Description; } \
 		static FName GetImageName() { return "MetasoundEditor.Graph.Node.Math." #OpName; } \
-		F##OpName##DataTypeName##Node(const FNodeInitData& InInitData) : TMathOpNode(InInitData) { } \
+		static TArray<FText> GetKeywords() { return Keywords; } \
+F##OpName##DataTypeName##Node(const FNodeInitData& InInitData) : TMathOpNode(InInitData) { } \
 	};
 
-#define DEFINE_METASOUND_OPERAND_TYPED_MATHOP(OpName, DataTypeName, DataClass, OperandTypeName, OperandDataClass, Description) \
+#define DEFINE_METASOUND_OPERAND_TYPED_MATHOP(OpName, DataTypeName, DataClass, OperandTypeName, OperandDataClass, Description, Keywords) \
 	class F##OpName##DataTypeName##OperandTypeName##Node \
 		: public ::Metasound::TMathOpNode<F##OpName##DataTypeName##OperandTypeName##Node, ::Metasound::TMathOp##OpName<DataClass, OperandDataClass>, DataClass, OperandDataClass> \
 	{ \
@@ -40,6 +42,7 @@
 		static FText GetDisplayName() { return MathOpNames::Get##OpName##DisplayName<DataClass, OperandDataClass>(); } \
 		static FText GetDescription() { return Description; } \
 		static FName GetImageName() { return "MetasoundEditor.Graph.Node.Math." #OpName; } \
+		static TArray<FText> GetKeywords() { return Keywords; } \
 		F##OpName##DataTypeName##OperandTypeName##Node(const FNodeInitData& InInitData) : TMathOpNode(InInitData) { } \
 	};
 
@@ -51,6 +54,13 @@ namespace Metasound
 
 		static const FString AdditionalOperandsName = TEXT("AdditionalOperands");
 
+		static const TArray<FText> AddKeywords = { LOCTEXT("AddMathKeyword", "+") };
+		static const TArray<FText> SubtractKeywords = { LOCTEXT("SubtractMathKeyword", "-") };
+		static const TArray<FText> MultiplyKeywords = { LOCTEXT("MultiplyMathKeyword", "*") };
+		static const TArray<FText> DivideKeywords = { LOCTEXT("DivideMathKeyword", "/") };
+		static const TArray<FText> PowerKeywords = { LOCTEXT("PowerMathKeyword", "^") };
+		static const TArray<FText> ModuloKeywords = { LOCTEXT("ModuloMathKeyword", "%") };
+
 		template<typename DataType>
 		const FText GetAddDisplayName()
 		{
@@ -209,7 +219,8 @@ namespace Metasound
 					Info.PromptIfMissing = PluginNodeMissingPrompt;
 					Info.DefaultInterface = TMathOpClass::GetVertexInterface();
 					Info.DisplayStyle = DisplayStyle;
-					Info.CategoryHierarchy = { LOCTEXT("Metasound_MathCategory", "Math") };
+					Info.CategoryHierarchy = { NodeCategories::Math };
+					Info.Keywords = TDerivedClass::GetKeywords();
 
 					return Info;
 				};
@@ -1435,40 +1446,40 @@ namespace Metasound
 #endif
 
 	// Definitions
-	DEFINE_METASOUND_MATHOP(Add, Float, float, LOCTEXT("Metasound_MathAddFloatNodeDescription", "Adds floats."))
-	DEFINE_METASOUND_MATHOP(Add, Int32, int32, LOCTEXT("Metasound_MathAddInt32NodeDescription", "Adds int32s."))
-	//DEFINE_METASOUND_MATHOP(Add, Int64, int64, LOCTEXT("Metasound_MathAddInt64NodeDescription", "Adds int64s."))
-	DEFINE_METASOUND_MATHOP(Add, Audio, FAudioBuffer, LOCTEXT("Metasound_MathAddBufferNodeDescription", "Adds buffers together by sample."))
-	DEFINE_METASOUND_MATHOP(Add, Time, FTime, LOCTEXT("Metasound_MathAddTimeNodeDescription", "Adds time values."))
+	DEFINE_METASOUND_MATHOP(Add, Float, float, LOCTEXT("Metasound_MathAddFloatNodeDescription", "Adds floats."), MathOpNames::AddKeywords)
+	DEFINE_METASOUND_MATHOP(Add, Int32, int32, LOCTEXT("Metasound_MathAddInt32NodeDescription", "Adds int32s."), MathOpNames::AddKeywords)
+	//DEFINE_METASOUND_MATHOP(Add, Int64, int64, LOCTEXT("Metasound_MathAddInt64NodeDescription", "Adds int64s."), MathOpNames::AddKeywords)
+	DEFINE_METASOUND_MATHOP(Add, Audio, FAudioBuffer, LOCTEXT("Metasound_MathAddBufferNodeDescription", "Adds buffers together by sample."), MathOpNames::AddKeywords)
+	DEFINE_METASOUND_MATHOP(Add, Time, FTime, LOCTEXT("Metasound_MathAddTimeNodeDescription", "Adds time values."), MathOpNames::AddKeywords)
 
-	DEFINE_METASOUND_OPERAND_TYPED_MATHOP(Add, Audio, FAudioBuffer, Float, float, LOCTEXT("Metasound_MathAddAudioFloatNodeDescription", "Add floats to buffer sample-by-sample."))
+	DEFINE_METASOUND_OPERAND_TYPED_MATHOP(Add, Audio, FAudioBuffer, Float, float, LOCTEXT("Metasound_MathAddAudioFloatNodeDescription", "Add floats to buffer sample-by-sample."), MathOpNames::AddKeywords)
 
-	DEFINE_METASOUND_MATHOP(Subtract, Float, float, LOCTEXT("Metasound_MathSubractFloatNodeDescription", "Subtracts floats."))
-	DEFINE_METASOUND_MATHOP(Subtract, Int32, int32, LOCTEXT("Metasound_MathSubractInt32NodeDescription", "Subtracts int32s."))
-	//DEFINE_METASOUND_MATHOP(Subtract, Int64, int64, LOCTEXT("Metasound_MathSubractInt64NodeDescription", "Subtracts int64s."))
-	DEFINE_METASOUND_MATHOP(Subtract, Audio, FAudioBuffer, LOCTEXT("Metasound_MathSubtractBufferNodeDescription", "Subtracts buffers sample-by-sample."))
-	DEFINE_METASOUND_MATHOP(Subtract, Time, FTime, LOCTEXT("Metasound_MathSubractTimeNodeDescription", "Subtracts time values."))
+	DEFINE_METASOUND_MATHOP(Subtract, Float, float, LOCTEXT("Metasound_MathSubractFloatNodeDescription", "Subtracts floats."), MathOpNames::SubtractKeywords)
+	DEFINE_METASOUND_MATHOP(Subtract, Int32, int32, LOCTEXT("Metasound_MathSubractInt32NodeDescription", "Subtracts int32s."), MathOpNames::SubtractKeywords)
+	//DEFINE_METASOUND_MATHOP(Subtract, Int64, int64, LOCTEXT("Metasound_MathSubractInt64NodeDescription", "Subtracts int64s."), MathOpNames::SubtractKeywords)
+	DEFINE_METASOUND_MATHOP(Subtract, Audio, FAudioBuffer, LOCTEXT("Metasound_MathSubtractBufferNodeDescription", "Subtracts buffers sample-by-sample."), MathOpNames::SubtractKeywords)
+	DEFINE_METASOUND_MATHOP(Subtract, Time, FTime, LOCTEXT("Metasound_MathSubractTimeNodeDescription", "Subtracts time values."), MathOpNames::SubtractKeywords)
 
-	DEFINE_METASOUND_MATHOP(Multiply, Float, float, LOCTEXT("Metasound_MathMultiplyFloatNodeDescription", "Multiplies floats."))
-	DEFINE_METASOUND_MATHOP(Multiply, Int32, int32, LOCTEXT("Metasound_MathMultiplyInt32NodeDescription", "Multiplies int32s."))
-	//DEFINE_METASOUND_MATHOP(Multiply, Int64, int64, LOCTEXT("Metasound_MathMultiplyInt64NodeDescription", "Multiplies int64s."))
-	DEFINE_METASOUND_MATHOP(Multiply, Audio, FAudioBuffer, LOCTEXT("Metasound_MathMultiplyBufferNodeDescription", "Multiplies buffers together sample-by-sample."))
+	DEFINE_METASOUND_MATHOP(Multiply, Float, float, LOCTEXT("Metasound_MathMultiplyFloatNodeDescription", "Multiplies floats."), MathOpNames::MultiplyKeywords)
+	DEFINE_METASOUND_MATHOP(Multiply, Int32, int32, LOCTEXT("Metasound_MathMultiplyInt32NodeDescription", "Multiplies int32s."), MathOpNames::MultiplyKeywords)
+	//DEFINE_METASOUND_MATHOP(Multiply, Int64, int64, LOCTEXT("Metasound_MathMultiplyInt64NodeDescription", "Multiplies int64s."), MathOpNames::MultiplyKeywords)
+	DEFINE_METASOUND_MATHOP(Multiply, Audio, FAudioBuffer, LOCTEXT("Metasound_MathMultiplyBufferNodeDescription", "Multiplies buffers together sample-by-sample."), MathOpNames::MultiplyKeywords)
 
-	DEFINE_METASOUND_OPERAND_TYPED_MATHOP(Multiply, Audio, FAudioBuffer, Float, float, LOCTEXT("Metasound_MathMultiplyAudioByFloatDescription", "Multiplies buffer by float scalars."))
-	DEFINE_METASOUND_OPERAND_TYPED_MATHOP(Multiply, Time, FTime, Float, float, LOCTEXT("Metasound_MathMultiplyTimeNodeDescription", "Scales time by floats."))
+	DEFINE_METASOUND_OPERAND_TYPED_MATHOP(Multiply, Audio, FAudioBuffer, Float, float, LOCTEXT("Metasound_MathMultiplyAudioByFloatDescription", "Multiplies buffer by float scalars."), MathOpNames::MultiplyKeywords)
+	DEFINE_METASOUND_OPERAND_TYPED_MATHOP(Multiply, Time, FTime, Float, float, LOCTEXT("Metasound_MathMultiplyTimeNodeDescription", "Scales time by floats."), MathOpNames::MultiplyKeywords)
 
-	DEFINE_METASOUND_MATHOP(Divide, Float, float, LOCTEXT("Metasound_MathDivideFloatNodeDescription", "Divide float by another float."))
-	DEFINE_METASOUND_MATHOP(Divide, Int32, int32, LOCTEXT("Metasound_MathDivideInt32NodeDescription", "Divide int32 by another int32."))
-	//DEFINE_METASOUND_MATHOP(Divide, Int64, int64, LOCTEXT("Metasound_MathDivideInt64NodeDescription", "Divide int64 by another int64."))
+	DEFINE_METASOUND_MATHOP(Divide, Float, float, LOCTEXT("Metasound_MathDivideFloatNodeDescription", "Divide float by another float."), MathOpNames::DivideKeywords)
+	DEFINE_METASOUND_MATHOP(Divide, Int32, int32, LOCTEXT("Metasound_MathDivideInt32NodeDescription", "Divide int32 by another int32."), MathOpNames::DivideKeywords)
+	//DEFINE_METASOUND_MATHOP(Divide, Int64, int64, LOCTEXT("Metasound_MathDivideInt64NodeDescription", "Divide int64 by another int64."), MathOpNames::DivideKeywords)
 
-	DEFINE_METASOUND_OPERAND_TYPED_MATHOP(Divide, Time, FTime, Float, float, LOCTEXT("Metasound_MathDivideTimeNodeDescription", "Divides time by floats."))
+	DEFINE_METASOUND_OPERAND_TYPED_MATHOP(Divide, Time, FTime, Float, float, LOCTEXT("Metasound_MathDivideTimeNodeDescription", "Divides time by floats."), MathOpNames::DivideKeywords)
 
-	DEFINE_METASOUND_MATHOP(Modulo, Int32, int32, LOCTEXT("Metasound_MathModulusInt32NodeDescription", "Modulo int32 by another int32."))
-	//DEFINE_METASOUND_MATHOP(Modulo, Int64, int64, LOCTEXT("Metasound_MathModulusInt64NodeDescription", "Modulo int64 by another int64."))
+	DEFINE_METASOUND_MATHOP(Modulo, Int32, int32, LOCTEXT("Metasound_MathModulusInt32NodeDescription", "Modulo int32 by another int32."), MathOpNames::ModuloKeywords)
+	//DEFINE_METASOUND_MATHOP(Modulo, Int64, int64, LOCTEXT("Metasound_MathModulusInt64NodeDescription", "Modulo int64 by another int64."), MathOpNames::ModuloKeywords)
 
-	DEFINE_METASOUND_MATHOP(Power, Float, float, LOCTEXT("Metasound_MathPowerFloatNodeDescription", "Raise float to the power of another float."))
+	DEFINE_METASOUND_MATHOP(Power, Float, float, LOCTEXT("Metasound_MathPowerFloatNodeDescription", "Raise float to the power of another float."), MathOpNames::PowerKeywords)
 
-	DEFINE_METASOUND_MATHOP(Logarithm, Float, float, LOCTEXT("Metasound_MathLogarithmFloatNodeDescription", "Calculate float-base logarithm of another float."))
+	DEFINE_METASOUND_MATHOP(Logarithm, Float, float, LOCTEXT("Metasound_MathLogarithmFloatNodeDescription", "Calculate float-base logarithm of another float."), TArray<FText>())
 
 	METASOUND_REGISTER_NODE(FAddFloatNode)
 	METASOUND_REGISTER_NODE(FAddInt32Node)
diff --git a/Engine/Plugins/Runtime/Metasound/Source/MetasoundStandardNodes/Private/MetasoundMaxNode.cpp b/Engine/Plugins/Runtime/Metasound/Source/MetasoundStandardNodes/Private/MetasoundMaxNode.cpp
--- a/Engine/Plugins/Runtime/Metasound/Source/MetasoundStandardNodes/Private/MetasoundMaxNode.cpp
+++ b/Engine/Plugins/Runtime/Metasound/Source/MetasoundStandardNodes/Private/MetasoundMaxNode.cpp
@@ -5,6 +5,7 @@
 #include "MetasoundFacade.h"
 #include "MetasoundExecutableOperator.h"
 #include "MetasoundPrimitives.h"
+#include "MetasoundStandardNodesCategories.h"
 #include "MetasoundStandardNodesNames.h"
 #include "MetasoundAudioBuffer.h"
 #include "Internationalization/Text.h"
@@ -27,7 +28,7 @@ namespace Metasound
 		{
 			FNodeClassMetadata Metadata
 			{
-				FNodeClassName{FName("Max"), InOperatorName, InDataTypeName},
+				FNodeClassName { "Max", InOperatorName, InDataTypeName },
 				1, // Major Version
 				0, // Minor Version
 				InDisplayName,
@@ -35,8 +36,8 @@ namespace Metasound
 				PluginAuthor,
 				PluginNodeMissingPrompt,
 				InDefaultInterface,
-				{LOCTEXT("MaxCategory", "Math")},
-				{TEXT("Max")},
+				{ NodeCategories::Math },
+				{ },
 				FNodeDisplayStyle{}
 			};
 
diff --git a/Engine/Plugins/Runtime/Metasound/Source/MetasoundStandardNodes/Private/MetasoundMidSideNodes.cpp b/Engine/Plugins/Runtime/Metasound/Source/MetasoundStandardNodes/Private/MetasoundMidSideNodes.cpp
--- a/Engine/Plugins/Runtime/Metasound/Source/MetasoundStandardNodes/Private/MetasoundMidSideNodes.cpp
+++ b/Engine/Plugins/Runtime/Metasound/Source/MetasoundStandardNodes/Private/MetasoundMidSideNodes.cpp
@@ -61,7 +61,7 @@ namespace Metasound
 
 				FNodeClassMetadata Metadata
 				{
-					FNodeClassName{StandardNodes::Namespace, TEXT("Mid-Side Decode"), StandardNodes::AudioVariant},
+					FNodeClassName { StandardNodes::Namespace, "Mid-Side Decode", StandardNodes::AudioVariant },
 					1, // Major Version
 					0, // Minor Version
 					LOCTEXT("MidSideEncodeDisplayName", "Mid-Side Encode"),
@@ -69,8 +69,8 @@ namespace Metasound
 					PluginAuthor,
 					PluginNodeMissingPrompt,
 					NodeInterface,
-					{StandardNodes::Spatalization},
-					{TEXT("MidSideEncode")},
+					{ NodeCategories::Spatialization},
+					{ },
 					FNodeDisplayStyle{}
 				};
 
@@ -285,7 +285,7 @@ namespace Metasound
 
 				FNodeClassMetadata Metadata
 				{
-					FNodeClassName{StandardNodes::Namespace, TEXT("Mid-Side Encode"), StandardNodes::AudioVariant},
+					FNodeClassName { StandardNodes::Namespace, "Mid-Side Encode", StandardNodes::AudioVariant },
 					1, // Major Version
 					0, // Minor Version
 					LOCTEXT("MidSideDecodeDisplayName", "Mid-Side Decode"),
@@ -293,9 +293,9 @@ namespace Metasound
 					PluginAuthor,
 					PluginNodeMissingPrompt,
 					NodeInterface,
-					{StandardNodes::Spatalization},
-					{TEXT("MidSideDecode")},
-					FNodeDisplayStyle{}
+					{ NodeCategories::Spatialization },
+					{ },
+					FNodeDisplayStyle()
 				};
 
 				return Metadata;
diff --git a/Engine/Plugins/Runtime/Metasound/Source/MetasoundStandardNodes/Private/MetasoundMidiNoteQuantizerNode.cpp b/Engine/Plugins/Runtime/Metasound/Source/MetasoundStandardNodes/Private/MetasoundMidiNoteQuantizerNode.cpp
--- a/Engine/Plugins/Runtime/Metasound/Source/MetasoundStandardNodes/Private/MetasoundMidiNoteQuantizerNode.cpp
+++ b/Engine/Plugins/Runtime/Metasound/Source/MetasoundStandardNodes/Private/MetasoundMidiNoteQuantizerNode.cpp
@@ -110,7 +110,7 @@ namespace Metasound
 		auto InitNodeInfo = []() -> FNodeClassMetadata
 		{
 			FNodeClassMetadata Info;
-			Info.ClassName = { Metasound::StandardNodes::Namespace, TEXT("MIDI Note Quantizer"), Metasound::StandardNodes::AudioVariant };
+			Info.ClassName = { StandardNodes::Namespace, TEXT("MIDI Note Quantizer"), StandardNodes::AudioVariant };
 			Info.MajorVersion = 1;
 			Info.MinorVersion = 0;
 			Info.DisplayName = LOCTEXT("Metasound_MIDI_Note_Quantizer_NodeDisplayName", "MIDI Note Quantizer");
@@ -118,7 +118,7 @@ namespace Metasound
 			Info.Author = PluginAuthor;
 			Info.PromptIfMissing = PluginNodeMissingPrompt;
 			Info.DefaultInterface = DeclareVertexInterface();
-			Info.CategoryHierarchy.Emplace(StandardNodes::Music);
+			Info.CategoryHierarchy.Emplace(NodeCategories::Music);
 
 			return Info;
 		};
diff --git a/Engine/Plugins/Runtime/Metasound/Source/MetasoundStandardNodes/Private/MetasoundMidiScaleToArrayNode.cpp b/Engine/Plugins/Runtime/Metasound/Source/MetasoundStandardNodes/Private/MetasoundMidiScaleToArrayNode.cpp
--- a/Engine/Plugins/Runtime/Metasound/Source/MetasoundStandardNodes/Private/MetasoundMidiScaleToArrayNode.cpp
+++ b/Engine/Plugins/Runtime/Metasound/Source/MetasoundStandardNodes/Private/MetasoundMidiScaleToArrayNode.cpp
@@ -104,7 +104,7 @@ namespace Metasound
 		auto InitNodeInfo = []() -> FNodeClassMetadata
 		{
 			FNodeClassMetadata Info;
-			Info.ClassName = { Metasound::StandardNodes::Namespace, TEXT("Musical Scale To Note Array"), Metasound::StandardNodes::AudioVariant };
+			Info.ClassName = { StandardNodes::Namespace, TEXT("Musical Scale To Note Array"), StandardNodes::AudioVariant };
 			Info.MajorVersion = 1;
 			Info.MinorVersion = 0;
 			Info.DisplayName = LOCTEXT("Metasound_ScaleToNoteArray_NodeDisplayName", "Scale to Note Array");
@@ -112,7 +112,7 @@ namespace Metasound
 			Info.Author = PluginAuthor;
 			Info.PromptIfMissing = PluginNodeMissingPrompt;
 			Info.DefaultInterface = DeclareVertexInterface();
-			Info.CategoryHierarchy.Emplace(StandardNodes::Music);
+			Info.CategoryHierarchy.Emplace(NodeCategories::Music);
 
 			return Info;
 		};
diff --git a/Engine/Plugins/Runtime/Metasound/Source/MetasoundStandardNodes/Private/MetasoundMinNode.cpp b/Engine/Plugins/Runtime/Metasound/Source/MetasoundStandardNodes/Private/MetasoundMinNode.cpp
--- a/Engine/Plugins/Runtime/Metasound/Source/MetasoundStandardNodes/Private/MetasoundMinNode.cpp
+++ b/Engine/Plugins/Runtime/Metasound/Source/MetasoundStandardNodes/Private/MetasoundMinNode.cpp
@@ -4,6 +4,7 @@
 #include "MetasoundFacade.h"
 #include "MetasoundExecutableOperator.h"
 #include "MetasoundPrimitives.h"
+#include "MetasoundStandardNodesCategories.h"
 #include "MetasoundStandardNodesNames.h"
 #include "MetasoundAudioBuffer.h"
 #include "Internationalization/Text.h"
@@ -27,7 +28,7 @@ namespace Metasound
 		{
 			FNodeClassMetadata Metadata
 			{
-				FNodeClassName{FName("Min"), InOperatorName, InDataTypeName},
+				FNodeClassName { "Min", InOperatorName, InDataTypeName },
 				1, // Major Version
 				0, // Minor Version
 				InDisplayName,
@@ -35,9 +36,9 @@ namespace Metasound
 				PluginAuthor,
 				PluginNodeMissingPrompt,
 				InDefaultInterface,
-				{LOCTEXT("MaxCategory", "Math")},
-				{TEXT("Min")},
-				FNodeDisplayStyle{}
+				{ NodeCategories::Math },
+				{ },
+				FNodeDisplayStyle()
 			};
 
 			return Metadata;
diff --git a/Engine/Plugins/Runtime/Metasound/Source/MetasoundStandardNodes/Private/MetasoundMixerNode.cpp b/Engine/Plugins/Runtime/Metasound/Source/MetasoundStandardNodes/Private/MetasoundMixerNode.cpp
--- a/Engine/Plugins/Runtime/Metasound/Source/MetasoundStandardNodes/Private/MetasoundMixerNode.cpp
+++ b/Engine/Plugins/Runtime/Metasound/Source/MetasoundStandardNodes/Private/MetasoundMixerNode.cpp
@@ -209,7 +209,7 @@ namespace Metasound
 		{
 			FNodeClassMetadata Metadata
 			{
-				FNodeClassName{FName("AudioMixer"), InOperatorName, TEXT("")},
+				FNodeClassName { "AudioMixer", InOperatorName, FName() },
 				1, // Major Version
 				0, // Minor Version
 				InDisplayName,
@@ -217,8 +217,8 @@ namespace Metasound
 				PluginAuthor,
 				PluginNodeMissingPrompt,
 				InDefaultInterface,
-				{ StandardNodes::Audio },
-				{TEXT("AudioMixer")},
+				{ NodeCategories::Audio },
+				{ LOCTEXT("Metasound_AudioMixerKeyword", "Mixer")},
 				FNodeDisplayStyle{}
 			};
 
diff --git a/Engine/Plugins/Runtime/Metasound/Source/MetasoundStandardNodes/Private/MetasoundNoiseGenerator.cpp b/Engine/Plugins/Runtime/Metasound/Source/MetasoundStandardNodes/Private/MetasoundNoiseGenerator.cpp
--- a/Engine/Plugins/Runtime/Metasound/Source/MetasoundStandardNodes/Private/MetasoundNoiseGenerator.cpp
+++ b/Engine/Plugins/Runtime/Metasound/Source/MetasoundStandardNodes/Private/MetasoundNoiseGenerator.cpp
@@ -177,7 +177,7 @@ namespace Metasound
 		auto InitNodeInfo = []() -> FNodeClassMetadata
 		{
 			FNodeClassMetadata Info;
-			Info.ClassName = { Metasound::StandardNodes::Namespace, TEXT("Noise"), Metasound::StandardNodes::AudioVariant };
+			Info.ClassName = { StandardNodes::Namespace, TEXT("Noise"), StandardNodes::AudioVariant };
 			Info.MajorVersion = 1;
 			Info.MinorVersion = 0;
 			Info.DisplayName = LOCTEXT("Metasound_NoiseNodeDisplayNameX", "Noise");
@@ -185,7 +185,7 @@ namespace Metasound
 			Info.Author = PluginAuthor;
 			Info.PromptIfMissing = PluginNodeMissingPrompt;
 			Info.DefaultInterface = DeclareVertexInterface();
-			Info.CategoryHierarchy.Emplace(StandardNodes::Generators);
+			Info.CategoryHierarchy.Emplace(NodeCategories::Generators);
 			return Info;
 		};
 		static const FNodeClassMetadata Info = InitNodeInfo();
diff --git a/Engine/Plugins/Runtime/Metasound/Source/MetasoundStandardNodes/Private/MetasoundOscillatorNodes.cpp b/Engine/Plugins/Runtime/Metasound/Source/MetasoundStandardNodes/Private/MetasoundOscillatorNodes.cpp
--- a/Engine/Plugins/Runtime/Metasound/Source/MetasoundStandardNodes/Private/MetasoundOscillatorNodes.cpp
+++ b/Engine/Plugins/Runtime/Metasound/Source/MetasoundStandardNodes/Private/MetasoundOscillatorNodes.cpp
@@ -669,7 +669,7 @@ namespace Metasound
 			auto InitNodeInfo = []() -> FNodeClassMetadata
 			{
 				FNodeClassMetadata Info;
-				Info.ClassName = { Metasound::StandardNodes::Namespace, TEXT("Sine"), Metasound::StandardNodes::AudioVariant };
+				Info.ClassName = { StandardNodes::Namespace, TEXT("Sine"), StandardNodes::AudioVariant };
 				Info.MajorVersion = 1;
 				Info.MinorVersion = 0;
 				Info.DisplayName = LOCTEXT("Metasound_SineNodeDisplayName", "Sine");
@@ -677,7 +677,7 @@ namespace Metasound
 				Info.Author = PluginAuthor;
 				Info.PromptIfMissing = PluginNodeMissingPrompt;
 				Info.DefaultInterface = GetVertexInterface();
-				Info.CategoryHierarchy.Emplace(StandardNodes::Generators);
+				Info.CategoryHierarchy.Emplace(NodeCategories::Generators);
 				return Info;
 			};
 			static const FNodeClassMetadata Info = InitNodeInfo();
@@ -792,7 +792,7 @@ namespace Metasound
 			auto InitNodeInfo = []() -> FNodeClassMetadata
 			{
 				FNodeClassMetadata Info;
-				Info.ClassName = { Metasound::StandardNodes::Namespace, TEXT("Saw"), Metasound::StandardNodes::AudioVariant };
+				Info.ClassName = { StandardNodes::Namespace, TEXT("Saw"), StandardNodes::AudioVariant };
 				Info.MajorVersion = 1;
 				Info.MinorVersion = 0;
 				Info.DisplayName = LOCTEXT("Metasound_SawNodeDisplayName", "Saw");
@@ -800,7 +800,7 @@ namespace Metasound
 				Info.Author = PluginAuthor;
 				Info.PromptIfMissing = PluginNodeMissingPrompt;
 				Info.DefaultInterface = GetVertexInterface();
-				Info.CategoryHierarchy.Emplace(StandardNodes::Generators);
+				Info.CategoryHierarchy.Emplace(NodeCategories::Generators);
 				return Info;
 			};
 			static const FNodeClassMetadata Info = InitNodeInfo();
@@ -1021,7 +1021,7 @@ namespace Metasound
 			auto InitNodeInfo = []() -> FNodeClassMetadata
 			{
 				FNodeClassMetadata Info;
-				Info.ClassName = { Metasound::StandardNodes::Namespace, TEXT("Square"), Metasound::StandardNodes::AudioVariant };
+				Info.ClassName = { StandardNodes::Namespace, TEXT("Square"), StandardNodes::AudioVariant };
 				Info.MajorVersion = 1;
 				Info.MinorVersion = 0;
 				Info.DisplayName = LOCTEXT("Metasound_SquareNodeDisplayName", "Square");
@@ -1029,7 +1029,7 @@ namespace Metasound
 				Info.Author = PluginAuthor;
 				Info.PromptIfMissing = PluginNodeMissingPrompt;
 				Info.DefaultInterface = GetVertexInterface();
-				Info.CategoryHierarchy.Emplace(StandardNodes::Generators);
+				Info.CategoryHierarchy.Emplace(NodeCategories::Generators);
 				return Info;
 			};
 			static const FNodeClassMetadata Info = InitNodeInfo();
@@ -1135,7 +1135,7 @@ namespace Metasound
 			auto InitNodeInfo = []() -> FNodeClassMetadata
 			{
 				FNodeClassMetadata Info;
-				Info.ClassName = { Metasound::StandardNodes::Namespace, TEXT("Triangle"), Metasound::StandardNodes::AudioVariant };
+				Info.ClassName = { StandardNodes::Namespace, TEXT("Triangle"), StandardNodes::AudioVariant };
 				Info.MajorVersion = 1;
 				Info.MinorVersion = 0;
 				Info.DisplayName = LOCTEXT("Metasound_TriangleNodeDisplayName", "Triangle");
@@ -1143,7 +1143,7 @@ namespace Metasound
 				Info.Author = PluginAuthor;
 				Info.PromptIfMissing = PluginNodeMissingPrompt;
 				Info.DefaultInterface = GetVertexInterface();
-				Info.CategoryHierarchy.Emplace(StandardNodes::Generators);
+				Info.CategoryHierarchy.Emplace(NodeCategories::Generators);
 				return Info;
 			};
 			static const FNodeClassMetadata Info = InitNodeInfo();
@@ -1205,7 +1205,8 @@ namespace Metasound
 		// Common pins
 		static constexpr const TCHAR* WaveshapePinName = TEXT("Shape");
 		static constexpr const TCHAR* LfoOutPinName = TEXT("Out");
-	
+		static constexpr const TCHAR* TestOutName = TEXT("TestOut");
+
 		static const FVertexInterface& GetVertexInterface()
 		{
 			static const FVertexInterface Interface
@@ -1220,7 +1221,8 @@ namespace Metasound
 					TInputDataVertexModel<float>(FSquareOscilatorNode::FFactory::PulseWidthPinName, LOCTEXT("LfoPulseWidthDescription", "Pulse Width (0..1)"), 0.5f)
 				},
 				FOutputVertexInterface{
-					TOutputDataVertexModel<float>(LfoOutPinName, LOCTEXT("LfoOutputDescription", "Output of the LFO (blockrate)"))
+					TOutputDataVertexModel<float>(LfoOutPinName, LOCTEXT("LfoOutputDescription", "Output of the LFO (blockrate)")),
+// 					TOutputDataVertexModel<FAudioBuffer>(TestOutName, LOCTEXT("LfoOutputDescription2", "Test Output"))
 				}
 			};
 			return Interface;
@@ -1231,15 +1233,15 @@ namespace Metasound
 			auto InitNodeInfo = []() -> FNodeClassMetadata
 			{
 				FNodeClassMetadata Info;
-				Info.ClassName = { Metasound::StandardNodes::Namespace, TEXT("LFO"), Metasound::StandardNodes::AudioVariant };
+				Info.ClassName = { StandardNodes::Namespace, TEXT("LFO"), StandardNodes::AudioVariant };
 				Info.MajorVersion = 1;
-				Info.MinorVersion = 0;
+				Info.MinorVersion = 6;
 				Info.DisplayName = LOCTEXT("Metasound_LfoNodeDisplayName", "LFO");
 				Info.Description = LOCTEXT("Metasound_LfoNodeDescription", "Low frequency oscillator < blockrate");
 				Info.Author = PluginAuthor;
 				Info.PromptIfMissing = PluginNodeMissingPrompt;
 				Info.DefaultInterface = GetVertexInterface();
-				Info.CategoryHierarchy.Emplace(StandardNodes::Generators);
+				Info.CategoryHierarchy.Emplace(NodeCategories::Generators);
 				return Info;
 			};
 			static const FNodeClassMetadata Info = InitNodeInfo();
@@ -1277,6 +1279,7 @@ namespace Metasound
 			, PhaseOffset{MoveTemp(InPhaseOffset)}
 			, PulseWidth{MoveTemp(InPulseWidth)}
 			, Output{FFloatWriteRef::CreateNew(0.f)}
+			, TestOut(FAudioBufferWriteRef::CreateNew(InSettings))
 		{
 			ResetPhase();
 		}
@@ -1298,6 +1301,7 @@ namespace Metasound
 		{
 			FDataReferenceCollection OutputDataReferences;
 			OutputDataReferences.AddDataReadReference(LfoOutPinName, Output);
+// 			OutputDataReferences.AddDataReadReference(TestOutName, TestOut);
 			return OutputDataReferences;
 		}
 
@@ -1377,7 +1381,7 @@ namespace Metasound
 		FFloatReadRef PhaseOffset;
 		FFloatReadRef PulseWidth;
 		FFloatWriteRef Output;
-
+		FAudioBufferWriteRef TestOut;
 		Generators::FSawPolySmoothGenerator SawGenerator;
 		Generators::FSineWaveTableGenerator SineGenerator;
 		Generators::FTriangleGenerator TriangleGenerator;
diff --git a/Engine/Plugins/Runtime/Metasound/Source/MetasoundStandardNodes/Private/MetasoundSampleAndHoldNode.cpp b/Engine/Plugins/Runtime/Metasound/Source/MetasoundStandardNodes/Private/MetasoundSampleAndHoldNode.cpp
--- a/Engine/Plugins/Runtime/Metasound/Source/MetasoundStandardNodes/Private/MetasoundSampleAndHoldNode.cpp
+++ b/Engine/Plugins/Runtime/Metasound/Source/MetasoundStandardNodes/Private/MetasoundSampleAndHoldNode.cpp
@@ -38,7 +38,7 @@ namespace Metasound
 
 				FNodeClassMetadata Metadata
 				{
-					FNodeClassName{FName("SampleAndHold"), TEXT("Sample And Hold"), Metasound::StandardNodes::AudioVariant},
+					FNodeClassName { "SampleAndHold", "Sample And Hold", StandardNodes::AudioVariant },
 					1, // Major Version
 					0, // Minor Version
 					LOCTEXT("SampleAndHoldDisplayName", "Sample And Hold"),
@@ -46,9 +46,9 @@ namespace Metasound
 					PluginAuthor,
 					PluginNodeMissingPrompt,
 					NodeInterface,
-					{StandardNodes::Filters},
-					{TEXT("SampleAndHold")},
-					FNodeDisplayStyle{}
+					{ NodeCategories::Filters },
+					{ },
+					FNodeDisplayStyle()
 				};
 
 				return Metadata;
diff --git a/Engine/Plugins/Runtime/Metasound/Source/MetasoundStandardNodes/Private/MetasoundSemitoneToFreqMultiplierNode.cpp b/Engine/Plugins/Runtime/Metasound/Source/MetasoundStandardNodes/Private/MetasoundSemitoneToFreqMultiplierNode.cpp
--- a/Engine/Plugins/Runtime/Metasound/Source/MetasoundStandardNodes/Private/MetasoundSemitoneToFreqMultiplierNode.cpp
+++ b/Engine/Plugins/Runtime/Metasound/Source/MetasoundStandardNodes/Private/MetasoundSemitoneToFreqMultiplierNode.cpp
@@ -113,7 +113,7 @@ namespace Metasound
 			const FText NodeDescription = LOCTEXT("Metasound_SemitoneToFrequencyMultiplierDescription", "Converts a number of semitones to the corresponding frequency multiplier.");
 
 			FNodeClassMetadata Info;
-			Info.ClassName = { Metasound::StandardNodes::Namespace, OperatorName, DataTypeName };
+			Info.ClassName = { StandardNodes::Namespace, OperatorName, DataTypeName };
 			Info.MajorVersion = 1;
 			Info.MinorVersion = 0;
 			Info.DisplayName = NodeDisplayName;
@@ -121,7 +121,7 @@ namespace Metasound
 			Info.Author = PluginAuthor;
 			Info.PromptIfMissing = PluginNodeMissingPrompt;
 			Info.DefaultInterface = GetVertexInterface();
-			Info.CategoryHierarchy.Emplace(StandardNodes::Conversions);
+			Info.CategoryHierarchy.Emplace(NodeCategories::Conversions);
 
 			return Info;
 		};
diff --git a/Engine/Plugins/Runtime/Metasound/Source/MetasoundStandardNodes/Private/MetasoundStandardNodesCategories.cpp b/Engine/Plugins/Runtime/Metasound/Source/MetasoundStandardNodes/Private/MetasoundStandardNodesCategories.cpp
--- a/Engine/Plugins/Runtime/Metasound/Source/MetasoundStandardNodes/Private/MetasoundStandardNodesCategories.cpp
+++ b/Engine/Plugins/Runtime/Metasound/Source/MetasoundStandardNodes/Private/MetasoundStandardNodesCategories.cpp
@@ -6,20 +6,21 @@
 
 namespace Metasound
 {
-	namespace StandardNodes
+	namespace NodeCategories
 	{
-		const FText Generators = { LOCTEXT("Metasound_GeneratorsCategory", "Generators") };
-		const FText Filters = { LOCTEXT("Metasound_FiltersCategory", "Filters") };
-		const FText TriggerUtils = { LOCTEXT("Metasound_TriggerUtilsCategory", "Triggers") };
-		const FText Music = { LOCTEXT("Metasound_MusicCategory", "Music") };
-		const FText RandomUtils = { LOCTEXT("Metasound_RandomCategory", "Random") };
 		const FText Audio = { LOCTEXT("Metasound_AudioCategory", "Audio") };
-		const FText Conversions  = { LOCTEXT("Metasound_ConversionsCategory", "Conversions") };
-		const FText DebugUtils = { LOCTEXT("Metasound_DebugUtilsCategory", "Debug") };
-		const FText Io = { LOCTEXT("Metasound_IoCategory", "External IO") };
+		const FText Debug = { LOCTEXT("Metasound_DebugCategory", "Debug") };
 		const FText Delays = { LOCTEXT("Metasound_EffectsCategory", "Delays") };
-		const FText Spatalization = { LOCTEXT("Metasound_SpatializationCategory", "Spatialization") };
 		const FText Dynamics = { LOCTEXT("Metasound_DynamicsCategory", "Dynamics") };
+		const FText Envelopes = { LOCTEXT("Metasound_EnvelopesCategory", "Envelopes") };
+		const FText Filters = { LOCTEXT("Metasound_FiltersCategory", "Filters") };
+		const FText Generators = { LOCTEXT("Metasound_GeneratorsCategory", "Generators") };
+		const FText Io = { LOCTEXT("Metasound_IoCategory", "External IO") };
+		const FText Math = { LOCTEXT("Metasound_MathCategory", "Math") };
+		const FText Music = { LOCTEXT("Metasound_MusicCategory", "Music") };
+		const FText RandomUtils = { LOCTEXT("Metasound_RandomCategory", "Random") };
+		const FText Spatialization = { LOCTEXT("Metasound_SpatializationCategory", "Spatialization") };
+		const FText Trigger = { LOCTEXT("Metasound_TriggerCategory", "Triggers") };
 	}
 }
 
diff --git a/Engine/Plugins/Runtime/Metasound/Source/MetasoundStandardNodes/Private/MetasoundStereoDelayNode.cpp b/Engine/Plugins/Runtime/Metasound/Source/MetasoundStandardNodes/Private/MetasoundStereoDelayNode.cpp
--- a/Engine/Plugins/Runtime/Metasound/Source/MetasoundStandardNodes/Private/MetasoundStereoDelayNode.cpp
+++ b/Engine/Plugins/Runtime/Metasound/Source/MetasoundStandardNodes/Private/MetasoundStereoDelayNode.cpp
@@ -307,7 +307,7 @@ namespace Metasound
 		auto InitNodeInfo = []() -> FNodeClassMetadata
 		{
 			FNodeClassMetadata Info;
-			Info.ClassName = { Metasound::StandardNodes::Namespace, TEXT("Stereo Delay"), Metasound::StandardNodes::AudioVariant };
+			Info.ClassName = { StandardNodes::Namespace, TEXT("Stereo Delay"), StandardNodes::AudioVariant };
 			Info.MajorVersion = 1;
 			Info.MinorVersion = 0;
 			Info.DisplayName = LOCTEXT("Metasound_StereoDelayDisplayName", "Stereo Delay");
@@ -315,7 +315,7 @@ namespace Metasound
 			Info.Author = PluginAuthor;
 			Info.PromptIfMissing = PluginNodeMissingPrompt;
 			Info.DefaultInterface = GetVertexInterface();
-			Info.CategoryHierarchy.Emplace(StandardNodes::Delays);
+			Info.CategoryHierarchy.Emplace(NodeCategories::Delays);
 			return Info;
 		};
 
diff --git a/Engine/Plugins/Runtime/Metasound/Source/MetasoundStandardNodes/Private/MetasoundStereoPannerNode.cpp b/Engine/Plugins/Runtime/Metasound/Source/MetasoundStandardNodes/Private/MetasoundStereoPannerNode.cpp
--- a/Engine/Plugins/Runtime/Metasound/Source/MetasoundStandardNodes/Private/MetasoundStereoPannerNode.cpp
+++ b/Engine/Plugins/Runtime/Metasound/Source/MetasoundStandardNodes/Private/MetasoundStereoPannerNode.cpp
@@ -250,7 +250,7 @@ namespace Metasound
 		auto InitNodeInfo = []() -> FNodeClassMetadata
 		{
 			FNodeClassMetadata Info;
-			Info.ClassName = { Metasound::StandardNodes::Namespace, TEXT("Stereo Panner"), TEXT("") };
+			Info.ClassName = { StandardNodes::Namespace, TEXT("Stereo Panner"), TEXT("") };
 			Info.MajorVersion = 1;
 			Info.MinorVersion = 0;
 			Info.DisplayName = LOCTEXT("Metasound_StereoPannerDisplayName", "Stereo Panner");
@@ -258,7 +258,7 @@ namespace Metasound
 			Info.Author = PluginAuthor;
 			Info.PromptIfMissing = PluginNodeMissingPrompt;
 			Info.DefaultInterface = GetVertexInterface();
-			Info.CategoryHierarchy.Emplace(StandardNodes::Spatalization);
+			Info.CategoryHierarchy.Emplace(NodeCategories::Spatialization);
 			return Info;
 		};
 
diff --git a/Engine/Plugins/Runtime/Metasound/Source/MetasoundStandardNodes/Private/MetasoundSwitchGateNode.cpp b/Engine/Plugins/Runtime/Metasound/Source/MetasoundStandardNodes/Private/MetasoundSwitchGateNode.cpp
--- a/Engine/Plugins/Runtime/Metasound/Source/MetasoundStandardNodes/Private/MetasoundSwitchGateNode.cpp
+++ b/Engine/Plugins/Runtime/Metasound/Source/MetasoundStandardNodes/Private/MetasoundSwitchGateNode.cpp
@@ -196,7 +196,7 @@ namespace Metasound
 		auto InitNodeInfo = []() -> FNodeClassMetadata
 		{
 			FNodeClassMetadata Info;
-			Info.ClassName = { Metasound::StandardNodes::Namespace, TEXT("Switch"), Metasound::StandardNodes::AudioVariant };
+			Info.ClassName = { StandardNodes::Namespace, TEXT("Switch"), StandardNodes::AudioVariant };
 			Info.MajorVersion = 1;
 			Info.MinorVersion = 0;
 			Info.DisplayName = LOCTEXT("Metasound_SwitchNodeDisplayName", "Switch");
@@ -415,7 +415,7 @@ namespace Metasound
 		auto InitNodeInfo = []() -> FNodeClassMetadata
 		{
 			FNodeClassMetadata Info;
-			Info.ClassName = {Metasound::StandardNodes::Namespace, TEXT("Gate"), Metasound::StandardNodes::AudioVariant};
+			Info.ClassName = {StandardNodes::Namespace, TEXT("Gate"), StandardNodes::AudioVariant};
 			Info.MajorVersion = 1;
 			Info.MinorVersion = 0;
 			Info.DisplayName = LOCTEXT("Metasound_GateNodeDisplayName", "Gate");
diff --git a/Engine/Plugins/Runtime/Metasound/Source/MetasoundStandardNodes/Private/MetasoundTriggerAccumulatorNode.cpp b/Engine/Plugins/Runtime/Metasound/Source/MetasoundStandardNodes/Private/MetasoundTriggerAccumulatorNode.cpp
--- a/Engine/Plugins/Runtime/Metasound/Source/MetasoundStandardNodes/Private/MetasoundTriggerAccumulatorNode.cpp
+++ b/Engine/Plugins/Runtime/Metasound/Source/MetasoundStandardNodes/Private/MetasoundTriggerAccumulatorNode.cpp
@@ -21,7 +21,7 @@ namespace Metasound
 		{
 			FNodeClassMetadata Metadata
 			{
-				FNodeClassName{FName("TriggerAccumulator"), InOperatorName, TEXT("")},
+				FNodeClassName { "TriggerAccumulator", InOperatorName, FName() },
 				1, // Major Version
 				0, // Minor Version
 				InDisplayName,
@@ -29,9 +29,9 @@ namespace Metasound
 				PluginAuthor,
 				PluginNodeMissingPrompt,
 				InDefaultInterface,
-				{StandardNodes::TriggerUtils},
-				{TEXT("TriggerAccumulator")},
-				FNodeDisplayStyle{}
+				{ NodeCategories::Trigger },
+				{ },
+				FNodeDisplayStyle()
 			};
 
 			return Metadata;
diff --git a/Engine/Plugins/Runtime/Metasound/Source/MetasoundStandardNodes/Private/MetasoundTriggerAnyNode.cpp b/Engine/Plugins/Runtime/Metasound/Source/MetasoundStandardNodes/Private/MetasoundTriggerAnyNode.cpp
--- a/Engine/Plugins/Runtime/Metasound/Source/MetasoundStandardNodes/Private/MetasoundTriggerAnyNode.cpp
+++ b/Engine/Plugins/Runtime/Metasound/Source/MetasoundStandardNodes/Private/MetasoundTriggerAnyNode.cpp
@@ -22,7 +22,7 @@ namespace Metasound
 		{
 			FNodeClassMetadata Metadata
 			{
-				FNodeClassName{FName("TriggerAny"), InOperatorName, TEXT("")},
+				FNodeClassName { "TriggerAny", InOperatorName, FName() },
 				1, // Major Version
 				0, // Minor Version
 				InDisplayName,
@@ -30,8 +30,8 @@ namespace Metasound
 				PluginAuthor,
 				PluginNodeMissingPrompt,
 				InDefaultInterface,
-				{StandardNodes::TriggerUtils},
-				{TEXT("TriggerAny")},
+				{ NodeCategories::Trigger },
+				{ },
 				FNodeDisplayStyle{}
 			};
 
diff --git a/Engine/Plugins/Runtime/Metasound/Source/MetasoundStandardNodes/Private/MetasoundTriggerCoinNode.cpp b/Engine/Plugins/Runtime/Metasound/Source/MetasoundStandardNodes/Private/MetasoundTriggerCoinNode.cpp
--- a/Engine/Plugins/Runtime/Metasound/Source/MetasoundStandardNodes/Private/MetasoundTriggerCoinNode.cpp
+++ b/Engine/Plugins/Runtime/Metasound/Source/MetasoundStandardNodes/Private/MetasoundTriggerCoinNode.cpp
@@ -196,7 +196,7 @@ namespace Metasound
 		auto InitNodeInfo = []() -> FNodeClassMetadata
 		{
 			FNodeClassMetadata Info;
-			Info.ClassName = { Metasound::StandardNodes::Namespace, TEXT("Trigger Filter"), TEXT("") };
+			Info.ClassName = { StandardNodes::Namespace, TEXT("Trigger Filter"), TEXT("") };
 			Info.MajorVersion = 1;
 			Info.MinorVersion = 0;
 			Info.DisplayName = LOCTEXT("Metasound_TriggerCoinNodeDisplayName", "Trigger Filter");
@@ -204,7 +204,7 @@ namespace Metasound
 			Info.Author = PluginAuthor;
 			Info.PromptIfMissing = PluginNodeMissingPrompt;
 			Info.DefaultInterface = GetVertexInterface();
-			Info.CategoryHierarchy.Emplace(StandardNodes::TriggerUtils);
+			Info.CategoryHierarchy.Emplace(NodeCategories::Trigger);
 
 			return Info;
 		};
diff --git a/Engine/Plugins/Runtime/Metasound/Source/MetasoundStandardNodes/Private/MetasoundTriggerCompareNode.cpp b/Engine/Plugins/Runtime/Metasound/Source/MetasoundStandardNodes/Private/MetasoundTriggerCompareNode.cpp
--- a/Engine/Plugins/Runtime/Metasound/Source/MetasoundStandardNodes/Private/MetasoundTriggerCompareNode.cpp
+++ b/Engine/Plugins/Runtime/Metasound/Source/MetasoundStandardNodes/Private/MetasoundTriggerCompareNode.cpp
@@ -12,7 +12,7 @@ namespace Metasound
 		{
 			FNodeClassMetadata Metadata
 			{
-				FNodeClassName{FName("TriggerCompare"), InOperatorName, InDataTypeName},
+				FNodeClassName { "TriggerCompare", InOperatorName, InDataTypeName },
 				1, // Major Version
 				0, // Minor Version
 				InDisplayName,
@@ -20,9 +20,9 @@ namespace Metasound
 				PluginAuthor,
 				PluginNodeMissingPrompt,
 				InDefaultInterface,
-				{StandardNodes::TriggerUtils},
-				{TEXT("TriggerCompare")},
-				FNodeDisplayStyle{}
+				{ NodeCategories::Trigger },
+				{ },
+				FNodeDisplayStyle()
 			};
 
 			return Metadata;
diff --git a/Engine/Plugins/Runtime/Metasound/Source/MetasoundStandardNodes/Private/MetasoundTriggerCounterNode.cpp b/Engine/Plugins/Runtime/Metasound/Source/MetasoundStandardNodes/Private/MetasoundTriggerCounterNode.cpp
--- a/Engine/Plugins/Runtime/Metasound/Source/MetasoundStandardNodes/Private/MetasoundTriggerCounterNode.cpp
+++ b/Engine/Plugins/Runtime/Metasound/Source/MetasoundStandardNodes/Private/MetasoundTriggerCounterNode.cpp
@@ -204,7 +204,7 @@ namespace Metasound
 		auto InitNodeInfo = []() -> FNodeClassMetadata
 		{
 			FNodeClassMetadata Info;
-			Info.ClassName = {Metasound::StandardNodes::Namespace, TEXT("Trigger Counter"), TEXT("")};
+			Info.ClassName = {StandardNodes::Namespace, TEXT("Trigger Counter"), TEXT("")};
 			Info.MajorVersion = 1;
 			Info.MinorVersion = 0;
 			Info.DisplayName = LOCTEXT("Metasound_TriggerCounterNodeDisplayName", "Trigger Counter");
@@ -212,7 +212,7 @@ namespace Metasound
 			Info.Author = PluginAuthor;
 			Info.PromptIfMissing = PluginNodeMissingPrompt;
 			Info.DefaultInterface = GetVertexInterface();
-			Info.CategoryHierarchy.Emplace(StandardNodes::TriggerUtils);
+			Info.CategoryHierarchy.Emplace(NodeCategories::Trigger);
 
 			return Info;
 		};
diff --git a/Engine/Plugins/Runtime/Metasound/Source/MetasoundStandardNodes/Private/MetasoundTriggerDelayNode.cpp b/Engine/Plugins/Runtime/Metasound/Source/MetasoundStandardNodes/Private/MetasoundTriggerDelayNode.cpp
--- a/Engine/Plugins/Runtime/Metasound/Source/MetasoundStandardNodes/Private/MetasoundTriggerDelayNode.cpp
+++ b/Engine/Plugins/Runtime/Metasound/Source/MetasoundStandardNodes/Private/MetasoundTriggerDelayNode.cpp
@@ -172,7 +172,7 @@ namespace Metasound
 		auto InitNodeInfo = []() -> FNodeClassMetadata
 		{
 			FNodeClassMetadata Info;
-			Info.ClassName = {Metasound::StandardNodes::Namespace, TEXT("Trigger Delay"), TEXT("")};
+			Info.ClassName = {StandardNodes::Namespace, TEXT("Trigger Delay"), TEXT("")};
 			Info.MajorVersion = 1;
 			Info.MinorVersion = 0;
 			Info.DisplayName = LOCTEXT("Metasound_TriggerDelayNodeDisplayName", "Trigger Delay");
@@ -180,7 +180,7 @@ namespace Metasound
 			Info.Author = PluginAuthor;
 			Info.PromptIfMissing = PluginNodeMissingPrompt;
 			Info.DefaultInterface = GetVertexInterface();
-			Info.CategoryHierarchy.Emplace(StandardNodes::TriggerUtils);
+			Info.CategoryHierarchy.Emplace(NodeCategories::Trigger);
 
 			return Info;
 		};
diff --git a/Engine/Plugins/Runtime/Metasound/Source/MetasoundStandardNodes/Private/MetasoundTriggerOnThresholdNode.cpp b/Engine/Plugins/Runtime/Metasound/Source/MetasoundStandardNodes/Private/MetasoundTriggerOnThresholdNode.cpp
--- a/Engine/Plugins/Runtime/Metasound/Source/MetasoundStandardNodes/Private/MetasoundTriggerOnThresholdNode.cpp
+++ b/Engine/Plugins/Runtime/Metasound/Source/MetasoundStandardNodes/Private/MetasoundTriggerOnThresholdNode.cpp
@@ -199,7 +199,7 @@ namespace Metasound
 		auto InitNodeInfo = []() -> FNodeClassMetadata
 		{
 			FNodeClassMetadata Info;
-			Info.ClassName = {Metasound::StandardNodes::Namespace, TEXT("TriggerOnThreshold"), Metasound::StandardNodes::AudioVariant};
+			Info.ClassName = {StandardNodes::Namespace, TEXT("TriggerOnThreshold"), StandardNodes::AudioVariant};
 			Info.MajorVersion = 1;
 			Info.MinorVersion = 0;
 			Info.DisplayName = LOCTEXT("Metasound_TriggerOnThresholdNodeDisplayName", "Trigger On Threshold");
@@ -207,7 +207,7 @@ namespace Metasound
 			Info.Author = PluginAuthor;
 			Info.PromptIfMissing = PluginNodeMissingPrompt;
 			Info.DefaultInterface = DeclareVertexInterface(); 
-			Info.CategoryHierarchy.Emplace(StandardNodes::TriggerUtils);
+			Info.CategoryHierarchy.Emplace(NodeCategories::Trigger);
 
 			return Info;
 		};
diff --git a/Engine/Plugins/Runtime/Metasound/Source/MetasoundStandardNodes/Private/MetasoundTriggerPipeNode.cpp b/Engine/Plugins/Runtime/Metasound/Source/MetasoundStandardNodes/Private/MetasoundTriggerPipeNode.cpp
--- a/Engine/Plugins/Runtime/Metasound/Source/MetasoundStandardNodes/Private/MetasoundTriggerPipeNode.cpp
+++ b/Engine/Plugins/Runtime/Metasound/Source/MetasoundStandardNodes/Private/MetasoundTriggerPipeNode.cpp
@@ -180,7 +180,7 @@ namespace Metasound
 			Info.Author = PluginAuthor;
 			Info.PromptIfMissing = PluginNodeMissingPrompt;
 			Info.DefaultInterface = GetVertexInterface();
-			Info.CategoryHierarchy.Emplace(StandardNodes::TriggerUtils);
+			Info.CategoryHierarchy.Emplace(NodeCategories::Trigger);
 
 			return Info;
 		};
diff --git a/Engine/Plugins/Runtime/Metasound/Source/MetasoundStandardNodes/Private/MetasoundTriggerRepeatNode.cpp b/Engine/Plugins/Runtime/Metasound/Source/MetasoundStandardNodes/Private/MetasoundTriggerRepeatNode.cpp
--- a/Engine/Plugins/Runtime/Metasound/Source/MetasoundStandardNodes/Private/MetasoundTriggerRepeatNode.cpp
+++ b/Engine/Plugins/Runtime/Metasound/Source/MetasoundStandardNodes/Private/MetasoundTriggerRepeatNode.cpp
@@ -169,7 +169,7 @@ namespace Metasound
 		auto InitNodeInfo = []() -> FNodeClassMetadata
 		{
 			FNodeClassMetadata Info;
-			Info.ClassName = { Metasound::StandardNodes::Namespace, TEXT("TriggerRepeat"), TEXT("") };
+			Info.ClassName = { StandardNodes::Namespace, TEXT("TriggerRepeat"), TEXT("") };
 			Info.MajorVersion = 1;
 			Info.MinorVersion = 0;
 			Info.DisplayName = LOCTEXT("Metasound_PeriodicTriggerNodeDisplayName", "Trigger Repeat");
@@ -177,7 +177,7 @@ namespace Metasound
 			Info.Author = PluginAuthor;
 			Info.PromptIfMissing = PluginNodeMissingPrompt;
 			Info.DefaultInterface = GetVertexInterface();
-			Info.CategoryHierarchy.Emplace(StandardNodes::TriggerUtils);
+			Info.CategoryHierarchy.Emplace(NodeCategories::Trigger);
 
 			return Info;
 		};
diff --git a/Engine/Plugins/Runtime/Metasound/Source/MetasoundStandardNodes/Private/MetasoundTriggerRouteNode.cpp b/Engine/Plugins/Runtime/Metasound/Source/MetasoundStandardNodes/Private/MetasoundTriggerRouteNode.cpp
--- a/Engine/Plugins/Runtime/Metasound/Source/MetasoundStandardNodes/Private/MetasoundTriggerRouteNode.cpp
+++ b/Engine/Plugins/Runtime/Metasound/Source/MetasoundStandardNodes/Private/MetasoundTriggerRouteNode.cpp
@@ -81,7 +81,7 @@ namespace Metasound
 		{
 			FNodeClassMetadata Metadata
 			{
-				FNodeClassName{FName("TriggerRoute"), InOperatorName, InDataTypeName},
+				FNodeClassName{ "TriggerRoute", InOperatorName, InDataTypeName},
 				1, // Major Version
 				0, // Minor Version
 				InDisplayName,
@@ -89,9 +89,9 @@ namespace Metasound
 				PluginAuthor,
 				PluginNodeMissingPrompt,
 				InDefaultInterface,
-				{StandardNodes::TriggerUtils},
-				{TEXT("TriggerRoute")},
-				FNodeDisplayStyle{}
+				{ NodeCategories::Trigger },
+				{ },
+				FNodeDisplayStyle()
 			};
 
 			return Metadata;
diff --git a/Engine/Plugins/Runtime/Metasound/Source/MetasoundStandardNodes/Private/MetasoundTriggerToggleNode.cpp b/Engine/Plugins/Runtime/Metasound/Source/MetasoundStandardNodes/Private/MetasoundTriggerToggleNode.cpp
--- a/Engine/Plugins/Runtime/Metasound/Source/MetasoundStandardNodes/Private/MetasoundTriggerToggleNode.cpp
+++ b/Engine/Plugins/Runtime/Metasound/Source/MetasoundStandardNodes/Private/MetasoundTriggerToggleNode.cpp
@@ -145,7 +145,7 @@ namespace Metasound
 		auto InitNodeInfo = []() -> FNodeClassMetadata
 		{
 			FNodeClassMetadata Info;
-			Info.ClassName = {Metasound::StandardNodes::Namespace, TEXT("Trigger Toggle"), TEXT("")};
+			Info.ClassName = {StandardNodes::Namespace, TEXT("Trigger Toggle"), TEXT("")};
 			Info.MajorVersion = 1;
 			Info.MinorVersion = 0;
 			Info.DisplayName = LOCTEXT("Metasound_TriggerToggleNodeDisplayName", "Trigger Toggle");
@@ -153,7 +153,7 @@ namespace Metasound
 			Info.Author = PluginAuthor;
 			Info.PromptIfMissing = PluginNodeMissingPrompt;
 			Info.DefaultInterface = GetVertexInterface();
-			Info.CategoryHierarchy.Emplace(StandardNodes::TriggerUtils);
+			Info.CategoryHierarchy.Emplace(NodeCategories::Trigger);
 
 			return Info;
 		};
diff --git a/Engine/Plugins/Runtime/Metasound/Source/MetasoundStandardNodes/Private/MetasoundValueNode.cpp b/Engine/Plugins/Runtime/Metasound/Source/MetasoundStandardNodes/Private/MetasoundValueNode.cpp
--- a/Engine/Plugins/Runtime/Metasound/Source/MetasoundStandardNodes/Private/MetasoundValueNode.cpp
+++ b/Engine/Plugins/Runtime/Metasound/Source/MetasoundStandardNodes/Private/MetasoundValueNode.cpp
@@ -14,7 +14,7 @@ namespace Metasound
 		{
 			FNodeClassMetadata Metadata
 			{
-				FNodeClassName{FName("Value"), InOperatorName, InDataTypeName},
+				FNodeClassName{ "Value", InOperatorName, InDataTypeName },
 				1, // Major Version
 				0, // Minor Version
 				InDisplayName,
@@ -22,8 +22,8 @@ namespace Metasound
 				PluginAuthor,
 				PluginNodeMissingPrompt,
 				InDefaultInterface,
-				{LOCTEXT("ValueCategory", "Value")},
-				{TEXT("Value")},
+				{ LOCTEXT("ValueCategory", "Value") },
+				{ },
 				FNodeDisplayStyle{}
 			};
 
diff --git a/Engine/Plugins/Runtime/Metasound/Source/MetasoundStandardNodes/Private/MetasoundWaveShaperNode.cpp b/Engine/Plugins/Runtime/Metasound/Source/MetasoundStandardNodes/Private/MetasoundWaveShaperNode.cpp
--- a/Engine/Plugins/Runtime/Metasound/Source/MetasoundStandardNodes/Private/MetasoundWaveShaperNode.cpp
+++ b/Engine/Plugins/Runtime/Metasound/Source/MetasoundStandardNodes/Private/MetasoundWaveShaperNode.cpp
@@ -151,7 +151,7 @@ namespace Metasound
 		auto InitNodeInfo = []() -> FNodeClassMetadata
 		{
 			FNodeClassMetadata Info;
-			Info.ClassName = { Metasound::StandardNodes::Namespace, TEXT("WaveShaper"), Metasound::StandardNodes::AudioVariant };
+			Info.ClassName = { StandardNodes::Namespace, TEXT("WaveShaper"), StandardNodes::AudioVariant };
 			Info.MajorVersion = 1;
 			Info.MinorVersion = 0;
 			Info.DisplayName = LOCTEXT("Metasound_WaveShaperDisplayName", "WaveShaper");
diff --git a/Engine/Plugins/Runtime/Metasound/Source/MetasoundStandardNodes/Private/MetasoundWaveWriterNode.cpp b/Engine/Plugins/Runtime/Metasound/Source/MetasoundStandardNodes/Private/MetasoundWaveWriterNode.cpp
--- a/Engine/Plugins/Runtime/Metasound/Source/MetasoundStandardNodes/Private/MetasoundWaveWriterNode.cpp
+++ b/Engine/Plugins/Runtime/Metasound/Source/MetasoundStandardNodes/Private/MetasoundWaveWriterNode.cpp
@@ -228,7 +228,7 @@ namespace Metasound
 			auto InitNodeInfo = []() -> FNodeClassMetadata
 			{
 				FNodeClassMetadata Info;
-				Info.ClassName = { Metasound::StandardNodes::Namespace, TEXT("WaveWriter"), Metasound::StandardNodes::AudioVariant };
+				Info.ClassName = { StandardNodes::Namespace, TEXT("WaveWriter"), StandardNodes::AudioVariant };
 				Info.MajorVersion = 1;
 				Info.MinorVersion = 0;
 				Info.DisplayName = LOCTEXT("Metasound_WaveWriterNodeDisplayName", "Wave Writer");
@@ -236,7 +236,7 @@ namespace Metasound
 				Info.Author = PluginAuthor;
 				Info.PromptIfMissing = PluginNodeMissingPrompt;
 				Info.DefaultInterface = DeclareVertexInterface();
-				Info.CategoryHierarchy.Emplace(StandardNodes::Io);
+				Info.CategoryHierarchy.Emplace(NodeCategories::Io);
 				return Info;
 			};
 			static const FNodeClassMetadata Info = InitNodeInfo();
diff --git a/Engine/Plugins/Runtime/Metasound/Source/MetasoundStandardNodes/Public/MetasoundRandomNode.h b/Engine/Plugins/Runtime/Metasound/Source/MetasoundStandardNodes/Public/MetasoundRandomNode.h
--- a/Engine/Plugins/Runtime/Metasound/Source/MetasoundStandardNodes/Public/MetasoundRandomNode.h
+++ b/Engine/Plugins/Runtime/Metasound/Source/MetasoundStandardNodes/Public/MetasoundRandomNode.h
@@ -300,7 +300,7 @@ namespace Metasound
 			auto InitNodeInfo = []() -> FNodeClassMetadata
 			{
 				FNodeClassMetadata Info;
-				Info.ClassName = { Metasound::StandardNodes::Namespace, TRandomNodeSpecialization<ValueType>::GetClassName(), TEXT("") };
+				Info.ClassName = { StandardNodes::Namespace, TRandomNodeSpecialization<ValueType>::GetClassName(), TEXT("") };
 				Info.MajorVersion = 1;
 				Info.MinorVersion = 0;
 				Info.DisplayName = TRandomNodeSpecialization<ValueType>::GetDisplayName();
@@ -308,7 +308,7 @@ namespace Metasound
 				Info.Author = PluginAuthor;
 				Info.PromptIfMissing = PluginNodeMissingPrompt;
 				Info.DefaultInterface = GetDefaultInterface();
-				Info.CategoryHierarchy.Emplace(StandardNodes::RandomUtils);
+				Info.CategoryHierarchy.Emplace(NodeCategories::RandomUtils);
 
 				return Info;
 			};
diff --git a/Engine/Plugins/Runtime/Metasound/Source/MetasoundStandardNodes/Public/MetasoundStandardNodesCategories.h b/Engine/Plugins/Runtime/Metasound/Source/MetasoundStandardNodes/Public/MetasoundStandardNodesCategories.h
--- a/Engine/Plugins/Runtime/Metasound/Source/MetasoundStandardNodes/Public/MetasoundStandardNodesCategories.h
+++ b/Engine/Plugins/Runtime/Metasound/Source/MetasoundStandardNodes/Public/MetasoundStandardNodesCategories.h
@@ -3,22 +3,25 @@
 #pragma once
 
 #include "CoreMinimal.h"
+#include "MetasoundAutoConverterNode.h"
+
 
 namespace Metasound
 {
-	namespace StandardNodes
+	namespace NodeCategories
 	{
-		extern const FText Generators;
-		extern const FText Filters;
-		extern const FText TriggerUtils;
-		extern const FText Music;
-		extern const FText Audio;
-		extern const FText RandomUtils;
-		extern const FText Conversions;
-		extern const FText DebugUtils;
-		extern const FText Io;
-		extern const FText Delays;
-		extern const FText Spatalization;
-		extern const FText Dynamics;
+		extern const FText METASOUNDSTANDARDNODES_API Audio;
+		extern const FText METASOUNDSTANDARDNODES_API Debug;
+		extern const FText METASOUNDSTANDARDNODES_API Delays;
+		extern const FText METASOUNDSTANDARDNODES_API Dynamics;
+		extern const FText METASOUNDSTANDARDNODES_API Envelopes;
+		extern const FText METASOUNDSTANDARDNODES_API Filters;
+		extern const FText METASOUNDSTANDARDNODES_API Generators;
+		extern const FText METASOUNDSTANDARDNODES_API Io;
+		extern const FText METASOUNDSTANDARDNODES_API Math;
+		extern const FText METASOUNDSTANDARDNODES_API Music;
+		extern const FText METASOUNDSTANDARDNODES_API RandomUtils;
+		extern const FText METASOUNDSTANDARDNODES_API Spatialization;
+		extern const FText METASOUNDSTANDARDNODES_API Trigger;
 	}
 }
