Build Fix: Missing category updates (Earlier PFaS was using earlier CL)

diff --git a/Engine/Plugins/Runtime/Metasound/Source/MetasoundStandardNodes/Private/MetasoundAdditiveSynthNode.cpp b/Engine/Plugins/Runtime/Metasound/Source/MetasoundStandardNodes/Private/MetasoundAdditiveSynthNode.cpp
--- a/Engine/Plugins/Runtime/Metasound/Source/MetasoundStandardNodes/Private/MetasoundAdditiveSynthNode.cpp
+++ b/Engine/Plugins/Runtime/Metasound/Source/MetasoundStandardNodes/Private/MetasoundAdditiveSynthNode.cpp
@@ -1,15 +1,15 @@
 // Copyright Epic Games, Inc. All Rights Reserved.
 
-#include "Internationalization/Text.h"
+#include "DSP/BufferVectorOperations.h"
+#include "DSP/SineWaveTableOsc.h"
 #include "MetasoundFacade.h"
+#include "Internationalization/Text.h"
 #include "MetasoundExecutableOperator.h"
 #include "MetasoundNodeRegistrationMacro.h"
+#include "MetasoundParamHelper.h"
 #include "MetasoundPrimitives.h"
-#include "MetasoundStandardNodesNames.h"
 #include "MetasoundStandardNodesCategories.h"
-#include "MetasoundParamHelper.h"
-#include "DSP/BufferVectorOperations.h"
-#include "DSP/SineWaveTableOsc.h"
+#include "MetasoundStandardNodesNames.h"
 
 #define LOCTEXT_NAMESPACE "MetasoundAdditiveSynthNode"
 
@@ -37,12 +37,12 @@ namespace Metasound
 		{
 			auto InitNodeInfo = []() -> FNodeClassMetadata
 			{
-				const FName OperatorName = TEXT("Additive Synth");
+				const FName OperatorName = "Additive Synth";
 				const FText NodeDisplayName = LOCTEXT("Metasound_AdditiveSynthNodeDisplayName", "Additive Synth");
 				const FText NodeDescription = LOCTEXT("Metasound_AdditiveSynthNodeDescription", "Synthesizes audio output given input array of sinusoids.");
 
 				FNodeClassMetadata Info;
-				Info.ClassName = { Metasound::StandardNodes::Namespace, OperatorName, TEXT("") };
+				Info.ClassName = { Metasound::StandardNodes::Namespace, OperatorName, FName() };
 				Info.MajorVersion = 1;
 				Info.MinorVersion = 0;
 				Info.DisplayName = NodeDisplayName;
diff --git a/Engine/Plugins/Runtime/Metasound/Source/MetasoundStandardNodes/Private/MetasoundCompressorNode.cpp b/Engine/Plugins/Runtime/Metasound/Source/MetasoundStandardNodes/Private/MetasoundCompressorNode.cpp
--- a/Engine/Plugins/Runtime/Metasound/Source/MetasoundStandardNodes/Private/MetasoundCompressorNode.cpp
+++ b/Engine/Plugins/Runtime/Metasound/Source/MetasoundStandardNodes/Private/MetasoundCompressorNode.cpp
@@ -1,17 +1,17 @@
 // Copyright Epic Games, Inc. All Rights Reserved.
 
+#include "DSP/Dsp.h"
+#include "DSP/DynamicsProcessor.h"
 #include "Internationalization/Text.h"
-#include "MetasoundExecutableOperator.h"
+#include "MetasoundAudioBuffer.h"
 #include "MetasoundEnumRegistrationMacro.h"
+#include "MetasoundExecutableOperator.h"
+#include "MetasoundFacade.h"
 #include "MetasoundNodeRegistrationMacro.h"
+#include "MetasoundParamHelper.h"
 #include "MetasoundPrimitives.h"
-#include "MetasoundStandardNodesNames.h"
-#include "MetasoundAudioBuffer.h"
 #include "MetasoundStandardNodesCategories.h"
-#include "MetasoundFacade.h"
-#include "MetasoundParamHelper.h"
-#include "DSP/Dsp.h"
-#include "DSP/DynamicsProcessor.h"
+#include "MetasoundStandardNodesNames.h"
 
 #define LOCTEXT_NAMESPACE "MetasoundStandardNodes_CompressorNode"
 
@@ -122,7 +122,7 @@ namespace Metasound
 
 				FNodeClassMetadata Metadata
 				{
-					FNodeClassName{StandardNodes::Namespace, TEXT("Compressor"), StandardNodes::AudioVariant},
+					FNodeClassName { StandardNodes::Namespace, "Compressor", StandardNodes::AudioVariant },
 					1, // Major Version
 					0, // Minor Version
 					LOCTEXT("CompressorDisplayName", "Compressor"),
@@ -130,9 +130,9 @@ namespace Metasound
 					PluginAuthor,
 					PluginNodeMissingPrompt,
 					NodeInterface,
-					{StandardNodes::Dynamics},
-					{TEXT("Compressor")},
-					FNodeDisplayStyle{}
+					{ NodeCategories::Dynamics },
+					{ },
+					FNodeDisplayStyle()
 				};
 
 				return Metadata;
diff --git a/Engine/Plugins/Runtime/Metasound/Source/MetasoundStandardNodes/Private/MetasoundTriggerControlNode.cpp b/Engine/Plugins/Runtime/Metasound/Source/MetasoundStandardNodes/Private/MetasoundTriggerControlNode.cpp
--- a/Engine/Plugins/Runtime/Metasound/Source/MetasoundStandardNodes/Private/MetasoundTriggerControlNode.cpp
+++ b/Engine/Plugins/Runtime/Metasound/Source/MetasoundStandardNodes/Private/MetasoundTriggerControlNode.cpp
@@ -201,7 +201,7 @@ namespace Metasound
 		auto InitNodeInfo = []() -> FNodeClassMetadata
 		{
 			FNodeClassMetadata Info;
-			Info.ClassName = { Metasound::StandardNodes::Namespace, TEXT("Trigger Control"), TEXT("") };
+			Info.ClassName = { Metasound::StandardNodes::Namespace, "Trigger Control", FName() };
 			Info.MajorVersion = 1;
 			Info.MinorVersion = 0;
 			Info.DisplayName = LOCTEXT("Metasound_TriggerControlNodeDisplayName", "Trigger Control");
@@ -209,7 +209,7 @@ namespace Metasound
 			Info.Author = PluginAuthor;
 			Info.PromptIfMissing = PluginNodeMissingPrompt;
 			Info.DefaultInterface = GetVertexInterface();
-			Info.CategoryHierarchy.Emplace(StandardNodes::TriggerUtils);
+			Info.CategoryHierarchy.Emplace(NodeCategories::Trigger);
 
 			return Info;
 		};
diff --git a/Engine/Plugins/Runtime/Metasound/Source/MetasoundStandardNodes/Private/MetasoundTriggerOnValueChangeNode.cpp b/Engine/Plugins/Runtime/Metasound/Source/MetasoundStandardNodes/Private/MetasoundTriggerOnValueChangeNode.cpp
--- a/Engine/Plugins/Runtime/Metasound/Source/MetasoundStandardNodes/Private/MetasoundTriggerOnValueChangeNode.cpp
+++ b/Engine/Plugins/Runtime/Metasound/Source/MetasoundStandardNodes/Private/MetasoundTriggerOnValueChangeNode.cpp
@@ -161,7 +161,7 @@ namespace Metasound
 			Info.Author = PluginAuthor;
 			Info.PromptIfMissing = PluginNodeMissingPrompt;
 			Info.DefaultInterface = GetVertexInterface();
-			Info.CategoryHierarchy.Emplace(StandardNodes::TriggerUtils);
+			Info.CategoryHierarchy.Emplace(NodeCategories::Trigger);
 
 			return Info;
 		};
diff --git a/Engine/Plugins/Runtime/Metasound/Source/MetasoundStandardNodes/Private/MetasoundTriggerOnceNode.cpp b/Engine/Plugins/Runtime/Metasound/Source/MetasoundStandardNodes/Private/MetasoundTriggerOnceNode.cpp
--- a/Engine/Plugins/Runtime/Metasound/Source/MetasoundStandardNodes/Private/MetasoundTriggerOnceNode.cpp
+++ b/Engine/Plugins/Runtime/Metasound/Source/MetasoundStandardNodes/Private/MetasoundTriggerOnceNode.cpp
@@ -163,7 +163,7 @@ namespace Metasound
 		auto InitNodeInfo = []() -> FNodeClassMetadata
 		{
 			FNodeClassMetadata Info;
-			Info.ClassName = { Metasound::StandardNodes::Namespace, TEXT("Trigger Once"), TEXT("") };
+			Info.ClassName = { Metasound::StandardNodes::Namespace, "Trigger Once", FName() };
 			Info.MajorVersion = 1;
 			Info.MinorVersion = 0;
 			Info.DisplayName = LOCTEXT("Metasound_TriggerOnceNodeDisplayName", "Trigger Once");
@@ -171,7 +171,7 @@ namespace Metasound
 			Info.Author = PluginAuthor;
 			Info.PromptIfMissing = PluginNodeMissingPrompt;
 			Info.DefaultInterface = GetVertexInterface();
-			Info.CategoryHierarchy.Emplace(StandardNodes::TriggerUtils);
+			Info.CategoryHierarchy.Emplace(NodeCategories::Trigger);
 
 			return Info;
 		};
