Remove test comments/code
#rnx

diff --git a/Engine/Plugins/Runtime/Metasound/Source/MetasoundStandardNodes/Private/MetasoundOscillatorNodes.cpp b/Engine/Plugins/Runtime/Metasound/Source/MetasoundStandardNodes/Private/MetasoundOscillatorNodes.cpp
--- a/Engine/Plugins/Runtime/Metasound/Source/MetasoundStandardNodes/Private/MetasoundOscillatorNodes.cpp
+++ b/Engine/Plugins/Runtime/Metasound/Source/MetasoundStandardNodes/Private/MetasoundOscillatorNodes.cpp
@@ -1205,8 +1205,7 @@ namespace Metasound
 		// Common pins
 		static constexpr const TCHAR* WaveshapePinName = TEXT("Shape");
 		static constexpr const TCHAR* LfoOutPinName = TEXT("Out");
-		static constexpr const TCHAR* TestOutName = TEXT("TestOut");
-
+	
 		static const FVertexInterface& GetVertexInterface()
 		{
 			static const FVertexInterface Interface
@@ -1221,8 +1220,7 @@ namespace Metasound
 					TInputDataVertexModel<float>(FSquareOscilatorNode::FFactory::PulseWidthPinName, LOCTEXT("LfoPulseWidthDescription", "Pulse Width (0..1)"), 0.5f)
 				},
 				FOutputVertexInterface{
-					TOutputDataVertexModel<float>(LfoOutPinName, LOCTEXT("LfoOutputDescription", "Output of the LFO (blockrate)")),
-// 					TOutputDataVertexModel<FAudioBuffer>(TestOutName, LOCTEXT("LfoOutputDescription2", "Test Output"))
+					TOutputDataVertexModel<float>(LfoOutPinName, LOCTEXT("LfoOutputDescription", "Output of the LFO (blockrate)"))
 				}
 			};
 			return Interface;
@@ -1235,7 +1233,7 @@ namespace Metasound
 				FNodeClassMetadata Info;
 				Info.ClassName = { StandardNodes::Namespace, TEXT("LFO"), StandardNodes::AudioVariant };
 				Info.MajorVersion = 1;
-				Info.MinorVersion = 6;
+				Info.MinorVersion = 0;
 				Info.DisplayName = LOCTEXT("Metasound_LfoNodeDisplayName", "LFO");
 				Info.Description = LOCTEXT("Metasound_LfoNodeDescription", "Low frequency oscillator < blockrate");
 				Info.Author = PluginAuthor;
@@ -1279,7 +1277,6 @@ namespace Metasound
 			, PhaseOffset{MoveTemp(InPhaseOffset)}
 			, PulseWidth{MoveTemp(InPulseWidth)}
 			, Output{FFloatWriteRef::CreateNew(0.f)}
-			, TestOut(FAudioBufferWriteRef::CreateNew(InSettings))
 		{
 			ResetPhase();
 		}
@@ -1301,7 +1298,6 @@ namespace Metasound
 		{
 			FDataReferenceCollection OutputDataReferences;
 			OutputDataReferences.AddDataReadReference(LfoOutPinName, Output);
-// 			OutputDataReferences.AddDataReadReference(TestOutName, TestOut);
 			return OutputDataReferences;
 		}
 
@@ -1381,7 +1377,7 @@ namespace Metasound
 		FFloatReadRef PhaseOffset;
 		FFloatReadRef PulseWidth;
 		FFloatWriteRef Output;
-		FAudioBufferWriteRef TestOut;
+
 		Generators::FSawPolySmoothGenerator SawGenerator;
 		Generators::FSineWaveTableGenerator SineGenerator;
 		Generators::FTriangleGenerator TriangleGenerator;
