Add safeguard against missing class in MetaSound validation

diff --git a/Engine/Plugins/Runtime/Metasound/Source/MetasoundEditor/Private/MetasoundEditorGraphNode.cpp b/Engine/Plugins/Runtime/Metasound/Source/MetasoundEditor/Private/MetasoundEditorGraphNode.cpp
--- a/Engine/Plugins/Runtime/Metasound/Source/MetasoundEditor/Private/MetasoundEditorGraphNode.cpp
+++ b/Engine/Plugins/Runtime/Metasound/Source/MetasoundEditor/Private/MetasoundEditorGraphNode.cpp
@@ -782,69 +782,83 @@ bool UMetasoundEditorGraphExternalNode::Validate(Metasound::Editor::FGraphNodeVa
 	// 6. Find all available versions & report if upgrade available
 	const Metasound::FNodeClassName NodeClassName = Metadata.GetClassName().ToNodeClassName();
 	const TArray<FMetasoundFrontendClass> SortedClasses = ISearchEngine::Get().FindClassesWithName(NodeClassName, true /* bSortByVersion */);
-	const FMetasoundFrontendVersionNumber& CurrentVersion = Metadata.GetVersion();
-	const FMetasoundFrontendClass& HighestRegistryClass = SortedClasses[0];
-	if (HighestRegistryClass.Metadata.GetVersion() > CurrentVersion)
+	if (SortedClasses.IsEmpty())
 	{
-		FMetasoundFrontendClass HighestMinorVersionClass;
-		FString NodeMsg;
-		EMessageSeverity::Type Severity;
-
-		const bool bClassVersionExists = SortedClasses.ContainsByPredicate([InCurrentVersion = &CurrentVersion](const FMetasoundFrontendClass& AvailableClass)
-		{
-			return AvailableClass.Metadata.GetVersion() == *InCurrentVersion;
-		});
-		if (bClassVersionExists)
-		{
-			NodeMsg = FString::Format(TEXT("Node class '{0} {1}' is prior version: Eligible for upgrade to {2}"),
+		GraphNodePrivate::SetGraphNodeMessage(*this, EMessageSeverity::Error,
+			FString::Format(TEXT("Class '{0} {1}' not registered."),
 			{
 				*Metadata.GetClassName().ToString(),
-				*Metadata.GetVersion().ToString(),
-				*HighestRegistryClass.Metadata.GetVersion().ToString()
-			});
-			Severity = EMessageSeverity::Warning;
+				*Metadata.GetVersion().ToString()
+			}));
+		OutResult.bIsDirty = true;
+		OutResult.bIsInvalid = true;
+	}
+	else
+	{
+		const FMetasoundFrontendVersionNumber& CurrentVersion = Metadata.GetVersion();
+		const FMetasoundFrontendClass& HighestRegistryClass = SortedClasses[0];
+		if (HighestRegistryClass.Metadata.GetVersion() > CurrentVersion)
+		{
+			FMetasoundFrontendClass HighestMinorVersionClass;
+			FString NodeMsg;
+			EMessageSeverity::Type Severity;
+
+			const bool bClassVersionExists = SortedClasses.ContainsByPredicate([InCurrentVersion = &CurrentVersion](const FMetasoundFrontendClass& AvailableClass)
+				{
+					return AvailableClass.Metadata.GetVersion() == *InCurrentVersion;
+				});
+			if (bClassVersionExists)
+			{
+				NodeMsg = FString::Format(TEXT("Node class '{0} {1}' is prior version: Eligible for upgrade to {2}"),
+					{
+						*Metadata.GetClassName().ToString(),
+						*Metadata.GetVersion().ToString(),
+						*HighestRegistryClass.Metadata.GetVersion().ToString()
+					});
+				Severity = EMessageSeverity::Warning;
+			}
+			else
+			{
+				NodeMsg = FString::Format(TEXT("Node class '{0} {1}' is missing and ineligible for auto-update.  Highest version '{2}' found."),
+					{
+						*Metadata.GetClassName().ToString(),
+						*Metadata.GetVersion().ToString(),
+						*HighestRegistryClass.Metadata.GetVersion().ToString()
+					});
+				Severity = EMessageSeverity::Error;
+				OutResult.bIsInvalid = true;
+			}
+
+			GraphNodePrivate::SetGraphNodeMessage(*this, Severity, NodeMsg);
+			OutResult.bIsDirty = true;
 		}
-		else
+		else if (HighestRegistryClass.Metadata.GetVersion() == CurrentVersion)
 		{
-			NodeMsg = FString::Format(TEXT("Node class '{0} {1}' is missing and ineligible for auto-update.  Highest version '{2}' found."),
+			if (InterfaceUpdates.ContainsChanges())
 			{
-				*Metadata.GetClassName().ToString(),
-				*Metadata.GetVersion().ToString(),
-				*HighestRegistryClass.Metadata.GetVersion().ToString()
-			});
-			Severity = EMessageSeverity::Error;
-			OutResult.bIsInvalid = true;
+				GraphNodePrivate::SetGraphNodeMessage(*this, EMessageSeverity::Error,
+					FString::Format(TEXT("Node & registered class interface mismatch: '{0} {1}'. Class either versioned improperly, class key collision exists, or AutoUpdate disabled in 'MetaSound' Developer Settings."),
+						{
+							*Metadata.GetClassName().ToString(),
+							*Metadata.GetVersion().ToString()
+						}));
+				OutResult.bIsDirty = true;
+				OutResult.bIsInvalid = true;
+			}
 		}
-
-		GraphNodePrivate::SetGraphNodeMessage(*this, Severity, NodeMsg);
-		OutResult.bIsDirty = true;
-	}
-	else if (HighestRegistryClass.Metadata.GetVersion() == CurrentVersion)
-	{
-		if (InterfaceUpdates.ContainsChanges())
+		else
 		{
 			GraphNodePrivate::SetGraphNodeMessage(*this, EMessageSeverity::Error,
-			FString::Format(TEXT("Node & registered class interface mismatch: '{0} {1}'. Class either versioned improperly, class key collision exists, or AutoUpdate disabled in 'MetaSound' Developer Settings."),
-			{
-				*Metadata.GetClassName().ToString(),
-				*Metadata.GetVersion().ToString()
-			}));
+				FString::Format(TEXT("Node with class '{0} {1}' interface version higher than that of highest minor revision ({2}) in class registry."),
+					{
+						*Metadata.GetClassName().ToString(),
+						*Metadata.GetVersion().ToString(),
+						*HighestRegistryClass.Metadata.GetVersion().ToString()
+					}));
 			OutResult.bIsDirty = true;
 			OutResult.bIsInvalid = true;
 		}
 	}
-	else
-	{
-		GraphNodePrivate::SetGraphNodeMessage(*this, EMessageSeverity::Error,
-			FString::Format(TEXT("Node with class '{0} {1}' interface version higher than that of highest minor revision ({2}) in class registry."),
-			{
-				*Metadata.GetClassName().ToString(),
-				*Metadata.GetVersion().ToString(),
-				*HighestRegistryClass.Metadata.GetVersion().ToString()
-			}));
-		OutResult.bIsDirty = true;
-		OutResult.bIsInvalid = true;
-	}
 
 	return !OutResult.bIsInvalid;
 }
