NPC Vehicle Audio Controller Component and Preset

(phil only reviewed the metasoundfrontend change)

diff --git a/Engine/Plugins/Runtime/Metasound/Source/MetasoundFrontend/MetasoundFrontend.Build.cs b/Engine/Plugins/Runtime/Metasound/Source/MetasoundFrontend/MetasoundFrontend.Build.cs
--- a/Engine/Plugins/Runtime/Metasound/Source/MetasoundFrontend/MetasoundFrontend.Build.cs
+++ b/Engine/Plugins/Runtime/Metasound/Source/MetasoundFrontend/MetasoundFrontend.Build.cs
@@ -16,7 +16,9 @@ namespace UnrealBuildTool.Rules
 					"Core",
 					"CoreUObject",
 					"Serialization",
-					"SignalProcessing"
+					"SignalProcessing",
+					"MetasoundGraphCore"
+
 				}
 			);
 
@@ -24,7 +26,6 @@ namespace UnrealBuildTool.Rules
 			(
 				new string[]
 				{
-					"MetasoundGraphCore"
 				}
 			);
 
