Add Limiter node to Metasounds

diff --git a/Engine/Plugins/Runtime/Metasound/Source/MetasoundStandardNodes/Private/MetasoundCompressorNode.cpp b/Engine/Plugins/Runtime/Metasound/Source/MetasoundStandardNodes/Private/MetasoundCompressorNode.cpp
--- a/Engine/Plugins/Runtime/Metasound/Source/MetasoundStandardNodes/Private/MetasoundCompressorNode.cpp
+++ b/Engine/Plugins/Runtime/Metasound/Source/MetasoundStandardNodes/Private/MetasoundCompressorNode.cpp
@@ -22,7 +22,7 @@ namespace Metasound
 	{
 		METASOUND_PARAM(InputAudio, "Audio", "Incoming audio signal to compress.");
 		METASOUND_PARAM(InputRatio, "Ratio", "Amount of gain reduction. 1 = no reduction, higher = more reduction.");
-		METASOUND_PARAM(InputThreshold, "Threshold", "Amplitude threshold (dB) above which gain will be reduced.");
+		METASOUND_PARAM(InputThreshold, "Threshold dB", "Amplitude threshold (dB) above which gain will be reduced.");
 		METASOUND_PARAM(InputAttackTime, "Attack Time", "How long it takes for audio above the threshold to reach its compressed volume level.");
 		METASOUND_PARAM(InputReleaseTime, "Release Time", "How long it takes for audio below the threshold to return to its original volume level.");
 		METASOUND_PARAM(InputKnee, "Knee", "How hard or soft the gain reduction blends from no gain reduction to gain reduction. 0 dB = no blending.");
@@ -59,7 +59,7 @@ namespace Metasound
 		FCompressorOperator(const FOperatorSettings& InSettings,
 			const FAudioBufferReadRef& InAudio,
 			const FFloatReadRef& InRatio,
-			const FFloatReadRef& InThreshold,
+			const FFloatReadRef& InThresholdDb,
 			const FTimeReadRef& InAttackTime,
 			const FTimeReadRef& InReleaseTime,
 			const FFloatReadRef& InKnee,
@@ -70,7 +70,7 @@ namespace Metasound
 			const FFloatReadRef& InWetDryMix)
 			: AudioInput(InAudio)
 			, RatioInput(InRatio)
-			, ThresholdInput(InThreshold)
+			, ThresholdDbInput(InThresholdDb)
 			, AttackTimeInput(InAttackTime)
 			, ReleaseTimeInput(InReleaseTime)
 			, KneeInput(InKnee)
@@ -82,18 +82,14 @@ namespace Metasound
 			, EnvelopeOutput(FAudioBufferWriteRef::CreateNew(InSettings))
 			, Compressor()
 			, bUseSidechain(bInUseSidechain)
-			, PrevRatio(FMath::Max(*InRatio, 1.0f))
-			, PrevThreshold(*InThreshold)
+			, PrevThresholdDb(*InThresholdDb)
 			, PrevAttackTime(FMath::Max(FTime::ToMilliseconds(*InAttackTime), 0.0))
 			, PrevReleaseTime(FMath::Max(FTime::ToMilliseconds(*InReleaseTime), 0.0))
-			, PrevKnee(*InKnee)
-			, PrevEnvMode(*InEnvelopeMode)
-			, PrevIsAnalog(*bInIsAnalog)
 		{
 			Compressor.Init(InSettings.GetSampleRate(), 1);
 			Compressor.SetKeyNumChannels(1);
-			Compressor.SetRatio(PrevRatio);
-			Compressor.SetThreshold(*ThresholdInput);
+			Compressor.SetRatio(FMath::Max(*InRatio, 1.0f));
+			Compressor.SetThreshold(*ThresholdDbInput);
 			Compressor.SetAttackTime(PrevAttackTime);
 			Compressor.SetReleaseTime(PrevReleaseTime);
 			Compressor.SetKneeBandwidth(*KneeInput);
@@ -176,7 +172,7 @@ namespace Metasound
 
 			InputDataReferences.AddDataReadReference(METASOUND_GET_PARAM_NAME(InputAudio), AudioInput);
 			InputDataReferences.AddDataReadReference(METASOUND_GET_PARAM_NAME(InputRatio), RatioInput);
-			InputDataReferences.AddDataReadReference(METASOUND_GET_PARAM_NAME(InputThreshold), ThresholdInput);
+			InputDataReferences.AddDataReadReference(METASOUND_GET_PARAM_NAME(InputThreshold), ThresholdDbInput);
 			InputDataReferences.AddDataReadReference(METASOUND_GET_PARAM_NAME(InputAttackTime), AttackTimeInput);
 			InputDataReferences.AddDataReadReference(METASOUND_GET_PARAM_NAME(InputReleaseTime), ReleaseTimeInput);
 			InputDataReferences.AddDataReadReference(METASOUND_GET_PARAM_NAME(InputKnee), KneeInput);
@@ -209,7 +205,7 @@ namespace Metasound
 
 			FAudioBufferReadRef AudioIn = Inputs.GetDataReadReferenceOrConstruct<FAudioBuffer>(METASOUND_GET_PARAM_NAME(InputAudio), InParams.OperatorSettings);
 			FFloatReadRef RatioIn = Inputs.GetDataReadReferenceOrConstructWithVertexDefault<float>(InputInterface, METASOUND_GET_PARAM_NAME(InputRatio), InParams.OperatorSettings);
-			FFloatReadRef ThresholdIn = Inputs.GetDataReadReferenceOrConstructWithVertexDefault<float>(InputInterface, METASOUND_GET_PARAM_NAME(InputThreshold), InParams.OperatorSettings);
+			FFloatReadRef ThresholdDbIn = Inputs.GetDataReadReferenceOrConstructWithVertexDefault<float>(InputInterface, METASOUND_GET_PARAM_NAME(InputThreshold), InParams.OperatorSettings);
 			FTimeReadRef AttackTimeIn = Inputs.GetDataReadReferenceOrConstructWithVertexDefault<FTime>(InputInterface, METASOUND_GET_PARAM_NAME(InputAttackTime), InParams.OperatorSettings);
 			FTimeReadRef ReleaseTimeIn = Inputs.GetDataReadReferenceOrConstructWithVertexDefault<FTime>(InputInterface, METASOUND_GET_PARAM_NAME(InputReleaseTime), InParams.OperatorSettings);
 			FFloatReadRef KneeIn = Inputs.GetDataReadReferenceOrConstructWithVertexDefault<float>(InputInterface, METASOUND_GET_PARAM_NAME(InputKnee), InParams.OperatorSettings);
@@ -220,7 +216,7 @@ namespace Metasound
 
 			bool bIsSidechainConnected = Inputs.ContainsDataReadReference<FAudioBuffer>(METASOUND_GET_PARAM_NAME(InputSidechain));
 
-			return MakeUnique<FCompressorOperator>(InParams.OperatorSettings, AudioIn, RatioIn, ThresholdIn, AttackTimeIn, ReleaseTimeIn, KneeIn, bIsSidechainConnected, SidechainIn, EnvelopeModeIn, bIsAnalogIn, WetDryMixIn);
+			return MakeUnique<FCompressorOperator>(InParams.OperatorSettings, AudioIn, RatioIn, ThresholdDbIn, AttackTimeIn, ReleaseTimeIn, KneeIn, bIsSidechainConnected, SidechainIn, EnvelopeModeIn, bIsAnalogIn, WetDryMixIn);
 		}
 
 		void Execute()
@@ -230,46 +226,42 @@ namespace Metasound
 			// For a compressor, ratio values should be 1 or greater
 			Compressor.SetRatio(FMath::Max(*RatioInput, 1.0f));
 
-			Compressor.SetThreshold(*ThresholdInput);
-			
+			if (FMath::IsNearlyEqual(*ThresholdDbInput, PrevThresholdDb))
+			{
+				Compressor.SetThreshold(*ThresholdDbInput);
+			}
+
 			// Attack time cannot be negative
 			double CurrAttack = FMath::Max(FTime::ToMilliseconds(*AttackTimeInput), 0.0f);
-			Compressor.SetAttackTime(CurrAttack);
-
+			if (FMath::IsNearlyEqual(CurrAttack, PrevAttackTime))
+			{
+				Compressor.SetAttackTime(CurrAttack);
+			}
 			// Release time cannot be negative
 			double CurrRelease = FMath::Max(FTime::ToMilliseconds(*ReleaseTimeInput), 0.0f);
-			if (!FMath::IsNearlyEqual(CurrRelease, PrevReleaseTime))
+			if (FMath::IsNearlyEqual(CurrRelease, PrevReleaseTime))
 			{
 				Compressor.SetReleaseTime(CurrRelease);
-				PrevReleaseTime = CurrRelease;
 			}
-			if (!FMath::IsNearlyEqual(*KneeInput, PrevKnee))
-			{
-				Compressor.SetKneeBandwidth(*KneeInput);
-				PrevKnee = *KneeInput;
-			}
-			if (PrevEnvMode != *EnvelopeModeInput)
-			{
-				switch (*EnvelopeModeInput)
-				{
-				default:
-				case EEnvelopePeakMode::MeanSquared:
-					Compressor.SetPeakMode(Audio::EPeakMode::MeanSquared);
-					break;
 
-				case EEnvelopePeakMode::RootMeanSquared:
-					Compressor.SetPeakMode(Audio::EPeakMode::RootMeanSquared);
-					break;
+			Compressor.SetKneeBandwidth(*KneeInput);
+				
+			switch (*EnvelopeModeInput)
+			{
+			default:
+			case EEnvelopePeakMode::MeanSquared:
+				Compressor.SetPeakMode(Audio::EPeakMode::MeanSquared);
+				break;
 
-				case EEnvelopePeakMode::Peak:
-					Compressor.SetPeakMode(Audio::EPeakMode::Peak);
-					break;
-				}
+			case EEnvelopePeakMode::RootMeanSquared:
+				Compressor.SetPeakMode(Audio::EPeakMode::RootMeanSquared);
+				break;
 
-				PrevEnvMode = *EnvelopeModeInput;
+			case EEnvelopePeakMode::Peak:
+				Compressor.SetPeakMode(Audio::EPeakMode::Peak);
+				break;
 			}
 
-
 			if (bUseSidechain)
 			{
 				Compressor.ProcessAudio(AudioInput->GetData(), AudioInput->Num(), AudioOutput->GetData(), SidechainInput->GetData(), EnvelopeOutput->GetData());
@@ -292,7 +284,7 @@ namespace Metasound
 		// Audio input and output
 		FAudioBufferReadRef AudioInput;
 		FFloatReadRef RatioInput;
-		FFloatReadRef ThresholdInput;
+		FFloatReadRef ThresholdDbInput;
 		FTimeReadRef AttackTimeInput;
 		FTimeReadRef ReleaseTimeInput;
 		FFloatReadRef KneeInput;
@@ -312,13 +304,9 @@ namespace Metasound
 		bool bUseSidechain;
 
 		// Cached variables
-		float PrevRatio;
-		float PrevThreshold;
+		float PrevThresholdDb;
 		double PrevAttackTime;
 		double PrevReleaseTime;
-		float PrevKnee;
-		EEnvelopePeakMode PrevEnvMode;
-		bool PrevIsAnalog;
 	};
 
 	// Node Class
diff --git a/Engine/Plugins/Runtime/Metasound/Source/MetasoundStandardNodes/Private/MetasoundLimiterNode.cpp b/Engine/Plugins/Runtime/Metasound/Source/MetasoundStandardNodes/Private/MetasoundLimiterNode.cpp
new file mode 100644
--- /dev/null
+++ b/Engine/Plugins/Runtime/Metasound/Source/MetasoundStandardNodes/Private/MetasoundLimiterNode.cpp
@@ -0,0 +1,246 @@
+// Copyright Epic Games, Inc. All Rights Reserved.
+
+#include "Internationalization/Text.h"
+#include "MetasoundExecutableOperator.h"
+#include "MetasoundEnumRegistrationMacro.h"
+#include "MetasoundNodeRegistrationMacro.h"
+#include "MetasoundPrimitives.h"
+#include "MetasoundStandardNodesNames.h"
+#include "MetasoundAudioBuffer.h"
+#include "MetasoundStandardNodesCategories.h"
+#include "MetasoundFacade.h"
+#include "MetasoundParamHelper.h"
+#include "DSP/Dsp.h"
+#include "DSP/DynamicsProcessor.h"
+
+#define LOCTEXT_NAMESPACE "MetasoundStandardNodes_LimiterNode"
+
+namespace Metasound
+{
+	/* Mid-Side Encoder */
+	namespace LimiterVertexNames
+	{
+		METASOUND_PARAM(InputAudio, "Audio", "Incoming audio signal to compress.");
+		METASOUND_PARAM(InputInGainDb, "Input Gain dB", "Gain to apply to the input before limiting.");
+		METASOUND_PARAM(InputThresholdDb, "Threshold dB", "Amplitude threshold above which gain will be reduced.");
+		METASOUND_PARAM(InputReleaseTime, "Release Time", "How long it takes for audio below the threshold to return to its original volume level.");
+		METASOUND_PARAM(InputKneeMode, "Knee", "Whether the limiter uses a hard or soft knee.");
+
+		METASOUND_PARAM(OutputAudio, "Audio", "The output audio signal.");
+	}
+
+	enum class EKneeMode
+	{
+		Hard = 0,
+		Soft,
+	};
+
+	DECLARE_METASOUND_ENUM(EKneeMode, EKneeMode::Hard, METASOUNDSTANDARDNODES_API,
+	FEnumKneeMode, FEnumKneeModeInfo, FKneeModeReadRef, FEnumKneeModeWriteRef);
+
+	DEFINE_METASOUND_ENUM_BEGIN(EKneeMode, FEnumKneeMode, "KneeMode")
+		DEFINE_METASOUND_ENUM_ENTRY(EKneeMode::Hard, LOCTEXT("KneeModeHardDescription", "Hard"), LOCTEXT("KneeModeHardDescriptionTT", "Only audio strictly above the threshold is affected by the limiter.")),
+		DEFINE_METASOUND_ENUM_ENTRY(EKneeMode::Soft, LOCTEXT("KneeModeSoftDescription", "Soft"), LOCTEXT("KneeModeSoftDescriptionTT", "Limiter activates more smoothly near the threshold.")),
+		DEFINE_METASOUND_ENUM_END()
+
+	// Operator Class
+	class FLimiterOperator : public TExecutableOperator<FLimiterOperator>
+	{
+	public:
+
+		static constexpr float HardKneeBandwitdh = 0.0f;
+		static constexpr float SoftKneeBandwitdh = 10.0f;
+
+		FLimiterOperator(const FOperatorSettings& InSettings,
+			const FAudioBufferReadRef& InAudio,
+			const FFloatReadRef& InGainDb,
+			const FFloatReadRef& InThresholdDb,
+			const FTimeReadRef& InReleaseTime,
+			const FKneeModeReadRef& InKneeMode)
+			: AudioInput(InAudio)
+			, InGainDbInput(InGainDb)
+			, ThresholdDbInput(InThresholdDb)
+			, ReleaseTimeInput(InReleaseTime)
+			, KneeModeInput(InKneeMode)
+			, AudioOutput(FAudioBufferWriteRef::CreateNew(InSettings))
+			, Limiter()
+			, PrevInGainDb(*InGainDb)
+			, PrevThresholdDb(*InThresholdDb)
+			, PrevReleaseTime(FMath::Max(FTime::ToMilliseconds(*InReleaseTime), 0.0))
+		{
+			Limiter.Init(InSettings.GetSampleRate(), 1);
+			Limiter.SetProcessingMode(Audio::EDynamicsProcessingMode::Limiter);
+			Limiter.SetInputGain(*InGainDb);
+			Limiter.SetThreshold(*ThresholdDbInput);
+			Limiter.SetAttackTime(0.0f);
+			Limiter.SetReleaseTime(PrevReleaseTime);
+			Limiter.SetPeakMode(Audio::EPeakMode::Peak);
+
+			switch (*KneeModeInput)
+			{
+			default:
+			case EKneeMode::Hard:
+				Limiter.SetKneeBandwidth(HardKneeBandwitdh);
+				break;
+			case EKneeMode::Soft:
+				Limiter.SetKneeBandwidth(SoftKneeBandwitdh);
+				break;
+			}
+
+		}
+
+		static const FNodeClassMetadata& GetNodeInfo()
+		{
+			auto CreateNodeClassMetadata = []() -> FNodeClassMetadata
+			{
+				FVertexInterface NodeInterface = DeclareVertexInterface();
+
+				FNodeClassMetadata Metadata
+				{
+					FNodeClassName{StandardNodes::Namespace, TEXT("Limiter"), StandardNodes::AudioVariant},
+					1, // Major Version
+					0, // Minor Version
+					LOCTEXT("LimiterDisplayName", "Limiter"),
+					LOCTEXT("LimiterDesc", "Prevents a signal from going above a given threshold."),
+					PluginAuthor,
+					PluginNodeMissingPrompt,
+					NodeInterface,
+					{NodeCategories::Dynamics},
+					{},
+					FNodeDisplayStyle{}
+				};
+
+				return Metadata;
+			};
+
+			static const FNodeClassMetadata Metadata = CreateNodeClassMetadata();
+			return Metadata;
+		}
+
+		static const FVertexInterface& DeclareVertexInterface()
+		{
+			using namespace LimiterVertexNames;
+
+			static const FVertexInterface Interface(
+				FInputVertexInterface(
+					TInputDataVertexModel<FAudioBuffer>(METASOUND_GET_PARAM_NAME_AND_TT(InputAudio)),
+					TInputDataVertexModel<float>(METASOUND_GET_PARAM_NAME_AND_TT(InputInGainDb), 0.0f),
+					TInputDataVertexModel<float>(METASOUND_GET_PARAM_NAME_AND_TT(InputThresholdDb), 0.0f),
+					TInputDataVertexModel<FTime>(METASOUND_GET_PARAM_NAME_AND_TT(InputReleaseTime), 0.1f),
+					TInputDataVertexModel<FEnumKneeMode>(METASOUND_GET_PARAM_NAME_AND_TT(InputKneeMode))
+				),
+				FOutputVertexInterface(
+					TOutputDataVertexModel<FAudioBuffer>(METASOUND_GET_PARAM_NAME_AND_TT(OutputAudio))
+				)
+			);
+
+			return Interface;
+		}
+
+		virtual FDataReferenceCollection GetInputs() const override
+		{
+			using namespace LimiterVertexNames;
+
+			FDataReferenceCollection InputDataReferences;
+
+			InputDataReferences.AddDataReadReference(METASOUND_GET_PARAM_NAME(InputAudio), AudioInput);
+			InputDataReferences.AddDataReadReference(METASOUND_GET_PARAM_NAME(InputInGainDb), InGainDbInput);
+			InputDataReferences.AddDataReadReference(METASOUND_GET_PARAM_NAME(InputThresholdDb), ThresholdDbInput);
+			InputDataReferences.AddDataReadReference(METASOUND_GET_PARAM_NAME(InputReleaseTime), ReleaseTimeInput);
+			InputDataReferences.AddDataReadReference(METASOUND_GET_PARAM_NAME(InputKneeMode), KneeModeInput);
+
+			return InputDataReferences;
+		}
+
+		virtual FDataReferenceCollection GetOutputs() const override
+		{
+			using namespace LimiterVertexNames;
+
+			FDataReferenceCollection OutputDataReferences;
+
+			OutputDataReferences.AddDataReadReference(METASOUND_GET_PARAM_NAME(OutputAudio), AudioOutput);
+
+			return OutputDataReferences;
+		}
+
+		static TUniquePtr<IOperator> CreateOperator(const FCreateOperatorParams& InParams, FBuildErrorArray& OutErrors)
+		{
+			using namespace LimiterVertexNames;
+			const FDataReferenceCollection& Inputs = InParams.InputDataReferences;
+			const FInputVertexInterface& InputInterface = DeclareVertexInterface().GetInputInterface();
+
+			FAudioBufferReadRef AudioIn = Inputs.GetDataReadReferenceOrConstruct<FAudioBuffer>(METASOUND_GET_PARAM_NAME(InputAudio), InParams.OperatorSettings);
+			FFloatReadRef InGainDbIn = Inputs.GetDataReadReferenceOrConstructWithVertexDefault<float>(InputInterface, METASOUND_GET_PARAM_NAME(InputInGainDb), InParams.OperatorSettings);
+			FFloatReadRef ThresholdDbIn = Inputs.GetDataReadReferenceOrConstructWithVertexDefault<float>(InputInterface, METASOUND_GET_PARAM_NAME(InputThresholdDb), InParams.OperatorSettings);
+			FTimeReadRef ReleaseTimeIn = Inputs.GetDataReadReferenceOrConstructWithVertexDefault<FTime>(InputInterface, METASOUND_GET_PARAM_NAME(InputReleaseTime), InParams.OperatorSettings);
+			FKneeModeReadRef KneeModeIn = Inputs.GetDataReadReferenceOrConstruct<FEnumKneeMode>(METASOUND_GET_PARAM_NAME(InputKneeMode));
+
+			return MakeUnique<FLimiterOperator>(InParams.OperatorSettings, AudioIn, InGainDbIn, ThresholdDbIn, ReleaseTimeIn, KneeModeIn);
+		}
+		void Execute()
+		{
+			/* Update parameters */
+			if (!FMath::IsNearlyEqual(*InGainDbInput, PrevInGainDb))
+			{
+				Limiter.SetInputGain(*InGainDbInput);
+			}
+			if (!FMath::IsNearlyEqual(*ThresholdDbInput, PrevThresholdDb))
+			{
+				Limiter.SetThreshold(*ThresholdDbInput);
+			}
+			// Release time cannot be negative
+			double CurrRelease = FMath::Max(FTime::ToMilliseconds(*ReleaseTimeInput), 0.0f);
+			if (!FMath::IsNearlyEqual(CurrRelease, PrevReleaseTime))
+			{
+				Limiter.SetReleaseTime(CurrRelease);
+				PrevReleaseTime = CurrRelease;
+			}
+
+			switch (*KneeModeInput)
+			{
+			default:
+			case EKneeMode::Hard:
+				Limiter.SetKneeBandwidth(HardKneeBandwitdh);
+				break;
+			case EKneeMode::Soft:
+				Limiter.SetKneeBandwidth(SoftKneeBandwitdh);
+				break;
+			}			
+
+			Limiter.ProcessAudio(AudioInput->GetData(), AudioInput->Num(), AudioOutput->GetData());
+		}
+
+	private:
+		FAudioBufferReadRef AudioInput;
+		FFloatReadRef InGainDbInput;
+		FFloatReadRef ThresholdDbInput;
+		FTimeReadRef ReleaseTimeInput;
+		FKneeModeReadRef KneeModeInput;
+
+		FAudioBufferWriteRef AudioOutput;
+
+		// Internal DSP Limiter
+		Audio::FDynamicsProcessor Limiter;
+
+		// Cached variables
+		float PrevInGainDb;
+		float PrevThresholdDb;
+		double PrevReleaseTime;
+		float PrevKnee;
+	};
+
+	// Node Class
+	class FLimiterNode : public FNodeFacade
+	{
+	public:
+		FLimiterNode(const FNodeInitData& InitData)
+			: FNodeFacade(InitData.InstanceName, InitData.InstanceID, TFacadeOperatorClass<FLimiterOperator>())
+		{
+		}
+	};
+
+	// Register node
+	METASOUND_REGISTER_NODE(FLimiterNode)
+}
+
+#undef LOCTEXT_NAMESPACE
