Adding cpuprofile tracing to Metadsound and Audio Mixer

diff --git a/Engine/Plugins/Runtime/Metasound/Source/MetasoundEngine/Private/MetasoundEngineModule.cpp b/Engine/Plugins/Runtime/Metasound/Source/MetasoundEngine/Private/MetasoundEngineModule.cpp
--- a/Engine/Plugins/Runtime/Metasound/Source/MetasoundEngine/Private/MetasoundEngineModule.cpp
+++ b/Engine/Plugins/Runtime/Metasound/Source/MetasoundEngine/Private/MetasoundEngineModule.cpp
@@ -18,7 +18,10 @@ class FMetasoundEngineModule : public IMetasoundEngineModule
 {
 	virtual void StartupModule() override
 	{
+		FModuleManager::Get().LoadModuleChecked("MetasoundGraphCore");
 		FModuleManager::Get().LoadModuleChecked("MetasoundFrontend");
+		FModuleManager::Get().LoadModuleChecked("MetasoundStandardNodes");
+		FModuleManager::Get().LoadModuleChecked("MetasoundGenerator");
 		FModuleManager::Get().LoadModuleChecked("AudioCodecEngine");
 
 		Metasound::Engine::RegisterArchetypes();
diff --git a/Engine/Plugins/Runtime/Metasound/Source/MetasoundEngine/Private/MetasoundSource.cpp b/Engine/Plugins/Runtime/Metasound/Source/MetasoundEngine/Private/MetasoundSource.cpp
--- a/Engine/Plugins/Runtime/Metasound/Source/MetasoundEngine/Private/MetasoundSource.cpp
+++ b/Engine/Plugins/Runtime/Metasound/Source/MetasoundEngine/Private/MetasoundSource.cpp
@@ -20,6 +20,7 @@
 #include "MetasoundOperatorSettings.h"
 #include "MetasoundPrimitives.h"
 #include "MetasoundReceiveNode.h"
+#include "MetasoundTrace.h"
 #include "MetasoundTrigger.h"
 #include "MetasoundUObjectRegistry.h"
 #include "UObject/ObjectSaveContext.h"
@@ -227,6 +228,8 @@ ISoundGeneratorPtr UMetaSoundSource::CreateSoundGenerator(const FSoundGeneratorI
 	using namespace Metasound::Frontend;
 	using namespace Metasound::Engine;
 
+	METASOUND_TRACE_CPUPROFILER_EVENT_SCOPE(MetaSoundSource::CreateSoundGenerator);
+
 	Duration = INDEFINITELY_LOOPING_DURATION;
 	bLooping = true;
 
diff --git a/Engine/Plugins/Runtime/Metasound/Source/MetasoundEngine/Private/MetasoundUObjectRegistry.cpp b/Engine/Plugins/Runtime/Metasound/Source/MetasoundEngine/Private/MetasoundUObjectRegistry.cpp
--- a/Engine/Plugins/Runtime/Metasound/Source/MetasoundEngine/Private/MetasoundUObjectRegistry.cpp
+++ b/Engine/Plugins/Runtime/Metasound/Source/MetasoundEngine/Private/MetasoundUObjectRegistry.cpp
@@ -16,6 +16,7 @@
 #include "MetasoundFrontendTransform.h"
 #include "MetasoundSettings.h"
 #include "MetasoundSource.h"
+#include "MetasoundTrace.h"
 #include "UObject/Object.h"
 
 
@@ -93,6 +94,8 @@ void UMetaSoundAssetSubsystem::PostEngineInit()
 
 void UMetaSoundAssetSubsystem::PostInitAssetScan()
 {
+	METASOUND_TRACE_CPUPROFILER_EVENT_SCOPE(UMetaSoundAssetSubsystem::PostInitAssetScan);
+
 	UAssetManager& AssetManager = UAssetManager::Get();
 
 	FAssetManagerSearchRules Rules;
@@ -125,6 +128,8 @@ void UMetaSoundAssetSubsystem::AddOrUpdateAsset(UObject& InObject, bool bInRegis
 	using namespace Metasound::AssetSubsystemPrivate;
 	using namespace Metasound::Frontend;
 
+	METASOUND_TRACE_CPUPROFILER_EVENT_SCOPE(UMetaSoundAssetSubsystem::AddOrUpdateAsset);
+
 	FMetasoundAssetBase* MetaSoundAsset = Metasound::IMetasoundUObjectRegistry::Get().GetObjectAsAssetBase(&InObject);
 	check(MetaSoundAsset);
 
@@ -147,6 +152,8 @@ void UMetaSoundAssetSubsystem::AddOrUpdateAsset(const FAssetData& InAssetData, b
 	// TODO: Set to false for builds once registering without loading asset is supported.
 	static const bool bLoadRequiredToRegisterAssetClasses = true;
 
+	METASOUND_TRACE_CPUPROFILER_EVENT_SCOPE(UMetaSoundAssetSubsystem::AddOrUpdateAsset);
+
 	FNodeClassInfo ClassInfo;
 	bool bClassInfoFound = GetAssetClassInfo(InAssetData, ClassInfo);
 	if (!bClassInfoFound || bLoadRequiredToRegisterAssetClasses || bInRegisterWithFrontend)
diff --git a/Engine/Plugins/Runtime/Metasound/Source/MetasoundFrontend/Private/MetasoundAssetBase.cpp b/Engine/Plugins/Runtime/Metasound/Source/MetasoundFrontend/Private/MetasoundAssetBase.cpp
--- a/Engine/Plugins/Runtime/Metasound/Source/MetasoundFrontend/Private/MetasoundAssetBase.cpp
+++ b/Engine/Plugins/Runtime/Metasound/Source/MetasoundFrontend/Private/MetasoundAssetBase.cpp
@@ -20,6 +20,7 @@
 #include "MetasoundJsonBackend.h"
 #include "MetasoundLog.h"
 #include "MetasoundReceiveNode.h"
+#include "MetasoundTrace.h"
 #include "StructSerializer.h"
 #include "UObject/MetaData.h"
 
@@ -71,6 +72,8 @@ void FMetasoundAssetBase::RegisterGraphWithFrontend()
 	using namespace Metasound;
 	using namespace Metasound::Frontend;
 
+	METASOUND_TRACE_CPUPROFILER_EVENT_SCOPE(MetaSoundAssetBase::RegisterGraphWithFrontend);
+
 	FConstDocumentAccessPtr DocumentPtr = GetDocument();
 	FString AssetName;
 	FString AssetPath;
@@ -182,6 +185,7 @@ void FMetasoundAssetBase::RegisterGraphWithFrontend()
 void FMetasoundAssetBase::UnregisterGraphWithFrontend()
 {
 	using namespace Metasound::Frontend;
+	METASOUND_TRACE_CPUPROFILER_EVENT_SCOPE(MetaSoundAssetBase::UnregisterGraphWithFrontend);
 
 	if (!NodeRegistryKey::IsValid(RegistryKey))
 	{
@@ -292,6 +296,7 @@ void FMetasoundAssetBase::SetDocument(const FMetasoundFrontendDocument& InDocume
 
 void FMetasoundAssetBase::ConformDocumentToArchetype()
 {
+	METASOUND_TRACE_CPUPROFILER_EVENT_SCOPE(MetaSoundAssetBase::ConformDocumentToArchetype);
 	FMetasoundFrontendDocument& Doc = GetDocumentChecked();
 
 	FMetasoundFrontendVersion PreferredArchetypeVersion = GetPreferredArchetypeVersion(Doc);
@@ -309,6 +314,7 @@ void FMetasoundAssetBase::ConformDocumentToArchetype()
 bool FMetasoundAssetBase::AutoUpdate(const IMetaSoundAssetInterface& InAssetInterface, bool bInMarkDirty, bool bInUpdateReferencedAssets)
 {
 	using namespace Metasound::Frontend;
+	METASOUND_TRACE_CPUPROFILER_EVENT_SCOPE(MetaSoundAssetBase::AutoUpdate);
 
 	bool bDidEdit = false;
 
@@ -339,6 +345,7 @@ bool FMetasoundAssetBase::VersionAsset(const IMetaSoundAssetInterface& InAssetIn
 {
 	using namespace Metasound;
 	using namespace Metasound::Frontend;
+	METASOUND_TRACE_CPUPROFILER_EVENT_SCOPE(MetaSoundAssetBase::VersionAsset);
 
 	FName AssetName;
 	FString AssetPath;
@@ -386,6 +393,8 @@ TUniquePtr<Metasound::IGraph> FMetasoundAssetBase::BuildMetasoundDocument() cons
 	using namespace Metasound;
 	using namespace Metasound::Frontend;
 
+	METASOUND_TRACE_CPUPROFILER_EVENT_SCOPE(MetaSoundAssetBase::BuildMetasoundDocument);
+
 	const FMetasoundFrontendDocument* Doc = GetDocument().Get();
 	if (nullptr == Doc)
 	{
@@ -418,6 +427,8 @@ void FMetasoundAssetBase::RebuildReferencedAssets(const IMetaSoundAssetInterface
 {
 	using namespace Metasound::Frontend;
 
+	METASOUND_TRACE_CPUPROFILER_EVENT_SCOPE(MetaSoundAssetBase::RebuildReferencedAssets);
+
 	TSet<FSoftObjectPath>& ReferencedAssets = GetReferencedAssets();
 	ReferencedAssets.Reset();
 
@@ -653,6 +664,8 @@ Metasound::Frontend::FConstGraphHandle FMetasoundAssetBase::GetRootGraphHandle()
 
 bool FMetasoundAssetBase::ImportFromJSON(const FString& InJSON)
 {
+	METASOUND_TRACE_CPUPROFILER_EVENT_SCOPE(MetaSoundAssetBase::ImportFromJSON);
+
 	FMetasoundFrontendDocument* Document = GetDocument().Get();
 	if (ensure(nullptr != Document))
 	{
@@ -670,6 +683,8 @@ bool FMetasoundAssetBase::ImportFromJSON(const FString& InJSON)
 
 bool FMetasoundAssetBase::ImportFromJSONAsset(const FString& InAbsolutePath)
 {
+	METASOUND_TRACE_CPUPROFILER_EVENT_SCOPE(MetaSoundAssetBase::ImportFromJSONAsset);
+
 	Metasound::Frontend::FDocumentAccessPtr DocumentPtr = GetDocument();
 	if (FMetasoundFrontendDocument* Document = DocumentPtr.Get())
 	{
diff --git a/Engine/Plugins/Runtime/Metasound/Source/MetasoundGenerator/Private/MetasoundGenerator.cpp b/Engine/Plugins/Runtime/Metasound/Source/MetasoundGenerator/Private/MetasoundGenerator.cpp
--- a/Engine/Plugins/Runtime/Metasound/Source/MetasoundGenerator/Private/MetasoundGenerator.cpp
+++ b/Engine/Plugins/Runtime/Metasound/Source/MetasoundGenerator/Private/MetasoundGenerator.cpp
@@ -7,9 +7,9 @@
 #include "MetasoundOperatorBuilder.h"
 #include "MetasoundOperatorInterface.h"
 #include "MetasoundOutputNode.h"
+#include "MetasoundTrace.h"
 #include "MetasoundTrigger.h"
 
-
 namespace Metasound
 {
 
@@ -22,6 +22,7 @@ namespace Metasound
 
 	void FAsyncMetaSoundBuilder::DoWork()
 	{
+		METASOUND_TRACE_CPUPROFILER_EVENT_SCOPE(AsyncMetaSoundBuilder::DoWork);
 		// Create an instance of the new graph
 		FOperatorBuilder OperatorBuilder(FOperatorBuilderSettings::GetDefaultSettings());
 		FBuildGraphParams BuildParams{*InitParams.Graph, InitParams.OperatorSettings, FDataReferenceCollection{}, InitParams.Environment};
@@ -39,11 +40,7 @@ namespace Metasound
 			}
 		}
 
-		if (!GraphOperator.IsValid())
-		{
-			UE_LOG(LogMetaSound, Error, TEXT("Failed to build Metasound operator from graph in MetasoundSource [%s]"), *InitParams.MetaSoundName);
-		}
-		else
+		if (GraphOperator.IsValid())
 		{
 			TArray<FAudioBufferReadRef> OutputBuffers;
 			FDataReferenceCollection Outputs = GraphOperator->GetOutputs();
@@ -74,6 +71,13 @@ namespace Metasound
 
 			Generator->SetPendingGraph(MoveTemp(GeneratorData), bTriggerGenerator);
 		}
+		else 
+		{
+			UE_LOG(LogMetaSound, Error, TEXT("Failed to build Metasound operator from graph in MetasoundSource [%s]"), *InitParams.MetaSoundName);
+			// Set null generator data to inform that generator failed to build. 
+			// Otherwise, generator will continually wait for a new generator.
+			Generator->SetPendingGraphBuildFailed();
+		}
 	}
 
 	FMetasoundGenerator::FMetasoundGenerator(FMetasoundGeneratorInitParams&& InParams)
@@ -85,18 +89,20 @@ namespace Metasound
 		, OnPlayTriggerRef(FTriggerWriteRef::CreateNew(InParams.OperatorSettings))
 		, OnFinishedTriggerRef(FTriggerWriteRef::CreateNew(InParams.OperatorSettings))
 		, bPendingGraphTrigger(true)
-		, bIsNewGraphPending(true)
+		, bIsNewGraphPending(false)
+		, bIsWaitingForFirstGraph(true)
 	{
 		NumChannels = InParams.AudioOutputNames.Num();
 		NumFramesPerExecute = InParams.OperatorSettings.GetNumFramesPerBlock();
 		NumSamplesPerExecute = NumChannels * NumFramesPerExecute;
 
 		BuilderTask = MakeUnique<FBuilderTask>(this, MoveTemp(InParams), true /* bTriggerGenerator */);
-
+		
 		BuilderTask->StartSynchronousTask();
 		BuilderTask = nullptr;
 		UpdateGraphIfPending();
-
+		bIsWaitingForFirstGraph = false;
+		
 		//BuilderTask->StartBackgroundTask(GBackgroundPriorityThreadPool);
 	}
 
@@ -119,6 +125,16 @@ namespace Metasound
 		}
 	}
 
+	void FMetasoundGenerator::SetPendingGraphBuildFailed()
+	{
+		FScopeLock SetPendingGraphLock(&PendingGraphMutex);
+		{
+			PendingGraphData = TUniquePtr<FMetasoundGeneratorData>(nullptr);
+			bPendingGraphTrigger = false;
+			bIsNewGraphPending = true;
+		}
+	}
+
 	bool FMetasoundGenerator::UpdateGraphIfPending()
 	{
 		FScopeLock GraphLock(&PendingGraphMutex);
@@ -126,9 +142,10 @@ namespace Metasound
 		{
 			SetGraph(MoveTemp(PendingGraphData), bPendingGraphTrigger);
 			bIsNewGraphPending = false;
+			return true;
 		}
 
-		return bIsNewGraphPending;
+		return false;
 	}
 
 	void FMetasoundGenerator::SetGraph(TUniquePtr<FMetasoundGeneratorData>&& InData, bool bTriggerGraph)
@@ -198,6 +215,8 @@ namespace Metasound
 
 	int32 FMetasoundGenerator::OnGenerateAudio(float* OutAudio, int32 NumSamplesRemaining)
 	{
+		METASOUND_TRACE_CPUPROFILER_EVENT_SCOPE(MetasoundGenerator::OnGenerateAudio);
+
 		// Defer finishing the metasound generator one block
 		if (bIsFinishTriggered)
 		{
@@ -209,28 +228,23 @@ namespace Metasound
 			return 0;
 		}
 
-		const bool bPending = UpdateGraphIfPending();
+		const bool bDidUpdateGraph = UpdateGraphIfPending();
+		bIsWaitingForFirstGraph = bIsWaitingForFirstGraph && !bDidUpdateGraph;
 
 		// Output silent audio if we're still building a graph
-		if (bPending)
+		if (bIsWaitingForFirstGraph)
 		{
-			if (RootExecuter.IsNoOp() || NumSamplesPerExecute < 1)
-			{
-				FMemory::Memset(OutAudio, 0, sizeof(float)* NumSamplesRemaining);
-				return NumSamplesRemaining;
-			}
+			FMemory::Memset(OutAudio, 0, sizeof(float)* NumSamplesRemaining);
+			return NumSamplesRemaining;
 		}
 		// If no longer pending and executer is no-op, kill the MetaSound.
 		// Covers case where there was an error when building, resulting in
 		// Executer operator being assigned to NoOp.
-		else
+		else if (RootExecuter.IsNoOp() || NumSamplesPerExecute < 1)
 		{
-			if (RootExecuter.IsNoOp())
-			{
-				bIsFinished = true;
-				FMemory::Memset(OutAudio, 0, sizeof(float) * NumSamplesRemaining);
-				return NumSamplesRemaining;
-			}
+			bIsFinished = true;
+			FMemory::Memset(OutAudio, 0, sizeof(float) * NumSamplesRemaining);
+			return NumSamplesRemaining;
 		}
 
 		// If we have any audio left in the internal overflow buffer from 
diff --git a/Engine/Plugins/Runtime/Metasound/Source/MetasoundGenerator/Public/MetasoundGenerator.h b/Engine/Plugins/Runtime/Metasound/Source/MetasoundGenerator/Public/MetasoundGenerator.h
--- a/Engine/Plugins/Runtime/Metasound/Source/MetasoundGenerator/Public/MetasoundGenerator.h
+++ b/Engine/Plugins/Runtime/Metasound/Source/MetasoundGenerator/Public/MetasoundGenerator.h
@@ -141,9 +141,11 @@ namespace Metasound
 		 * @param bTriggerGraph - If true, "OnPlay" will be triggered on the new graph.
 		 */
 		void SetPendingGraph(FMetasoundGeneratorData&& InData, bool bTriggerGraph);
+		void SetPendingGraphBuildFailed();
 
 	private:
 		bool UpdateGraphIfPending();
+
 		// Internal set graph after checking compatibility.
 		void SetGraph(TUniquePtr<FMetasoundGeneratorData>&& InData, bool bTriggerGraph);
 
@@ -183,6 +185,7 @@ namespace Metasound
 		TUniquePtr<FMetasoundGeneratorData> PendingGraphData;
 		bool bPendingGraphTrigger;
 		bool bIsNewGraphPending;
+		bool bIsWaitingForFirstGraph;
 	};
 }
 
diff --git a/Engine/Plugins/Runtime/Metasound/Source/MetasoundGraphCore/Public/MetasoundTrace.h b/Engine/Plugins/Runtime/Metasound/Source/MetasoundGraphCore/Public/MetasoundTrace.h
new file mode 100644
--- /dev/null
+++ b/Engine/Plugins/Runtime/Metasound/Source/MetasoundGraphCore/Public/MetasoundTrace.h
@@ -0,0 +1,26 @@
+// Copyright Epic Games, Inc. All Rights Reserved.
+
+#pragma once
+
+#include "ProfilingDebugging/CpuProfilerTrace.h"
+
+#if !defined(METASOUND_CPUPROFILERTRACE_ENABLED)
+#if CPUPROFILERTRACE_ENABLED
+#define METASOUND_CPUPROFILERTRACE_ENABLED 1
+#else
+#define METASOUND_CPUPROFILERTRACE_ENABLED 0
+#endif
+#endif
+
+#if METASOUND_CPUPROFILERTRACE_ENABLED
+// Metasound CPU profiler trace enabled
+
+#define METASOUND_TRACE_CPUPROFILER_EVENT_SCOPE(Name) \
+	TRACE_CPUPROFILER_EVENT_SCOPE(Name)
+
+#else
+// Metasound CPU profiler trace *not* enabled
+
+#define METASOUND_TRACE_CPUPROFILER_EVENT_SCOPE(Name)
+
+#endif
diff --git a/Engine/Source/Runtime/AudioMixer/Private/AudioMixerSource.cpp b/Engine/Source/Runtime/AudioMixer/Private/AudioMixerSource.cpp
--- a/Engine/Source/Runtime/AudioMixer/Private/AudioMixerSource.cpp
+++ b/Engine/Source/Runtime/AudioMixer/Private/AudioMixerSource.cpp
@@ -6,6 +6,7 @@
 #include "AudioMixerSourceBuffer.h"
 #include "AudioMixerDevice.h"
 #include "AudioMixerSourceVoice.h"
+#include "AudioMixerTrace.h"
 #include "ContentStreaming.h"
 #include "IAudioExtensionPlugin.h"
 #include "ProfilingDebugging/CsvProfiler.h"
@@ -294,6 +295,7 @@ namespace Audio
 
 	bool FMixerSource::Init(FWaveInstance* InWaveInstance)
 	{
+		AUDIO_MIXER_TRACE_CPUPROFILER_EVENT_SCOPE(AudioMixerSource::Init);
 		AUDIO_MIXER_CHECK(MixerBuffer);
 		AUDIO_MIXER_CHECK(MixerBuffer->IsRealTimeSourceReady());
 
@@ -778,6 +780,8 @@ namespace Audio
 			return false;
 		}
 
+		AUDIO_MIXER_TRACE_CPUPROFILER_EVENT_SCOPE(AudioMixerSource::PrepareForInitialization);
+
 		// We are not initialized yet. We won't be until the sound file finishes loading and parsing the header.
 		InitializationState = EMixerSourceInitializationState::Initializing;
 
diff --git a/Engine/Source/Runtime/AudioMixerCore/Public/AudioMixerTrace.h b/Engine/Source/Runtime/AudioMixerCore/Public/AudioMixerTrace.h
new file mode 100644
--- /dev/null
+++ b/Engine/Source/Runtime/AudioMixerCore/Public/AudioMixerTrace.h
@@ -0,0 +1,27 @@
+// Copyright Epic Games, Inc. All Rights Reserved.
+
+#pragma once
+
+#include "ProfilingDebugging/CpuProfilerTrace.h"
+
+#if !defined(AUDIO_MIXER_CPUPROFILERTRACE_ENABLED)
+#if CPUPROFILERTRACE_ENABLED
+#define AUDIO_MIXER_CPUPROFILERTRACE_ENABLED 1
+#else
+#define AUDIO_MIXER_CPUPROFILERTRACE_ENABLED 0
+#endif
+#endif
+
+#if AUDIO_MIXER_CPUPROFILERTRACE_ENABLED
+// AudioMixer CPU profiler trace enabled
+
+#define AUDIO_MIXER_TRACE_CPUPROFILER_EVENT_SCOPE(Name) \
+    TRACE_CPUPROFILER_EVENT_SCOPE(Name)
+
+#else
+// AudioMixer CPU profiler trace *not* enabled
+
+#define AUDIO_MIXER_TRACE_CPUPROFILER_EVENT_SCOPE(Name)
+
+#endif
+
