Remove test code

diff --git a/Engine/Plugins/Runtime/Metasound/Source/MetasoundStandardNodes/Private/MetasoundEnvelopeFollowerNode.cpp b/Engine/Plugins/Runtime/Metasound/Source/MetasoundStandardNodes/Private/MetasoundEnvelopeFollowerNode.cpp
--- a/Engine/Plugins/Runtime/Metasound/Source/MetasoundStandardNodes/Private/MetasoundEnvelopeFollowerNode.cpp
+++ b/Engine/Plugins/Runtime/Metasound/Source/MetasoundStandardNodes/Private/MetasoundEnvelopeFollowerNode.cpp
@@ -24,7 +24,6 @@ namespace Metasound
 		static const TCHAR* InParamNameAudioInput = TEXT("In");
 		static const TCHAR* InParamNameAttackTime = TEXT("Attack Time");
 		static const TCHAR* InParamNameReleaseTime = TEXT("Release Time");
-		static const TCHAR* PoopTest = TEXT("Poop Test");
 		static const TCHAR* OutParamNameEnvelope = TEXT("Envelope");
 	}
 
@@ -86,7 +85,6 @@ namespace Metasound
 		InputDataReferences.AddDataReadReference(EnvelopeFollower::InParamNameAudioInput, AudioInput);
 		InputDataReferences.AddDataReadReference(EnvelopeFollower::InParamNameAttackTime, AttackTime);
 		InputDataReferences.AddDataReadReference(EnvelopeFollower::InParamNameReleaseTime, ReleaseTime);
-		InputDataReferences.AddDataReadReference(EnvelopeFollower::PoopTest, ReleaseTime);
 		return InputDataReferences;
 	}
 
@@ -129,8 +127,7 @@ namespace Metasound
 			FInputVertexInterface(
 				TInputDataVertexModel<FAudioBuffer>(EnvelopeFollower::InParamNameAudioInput, LOCTEXT("AudioInputToolTT", "Audio input.")),
 				TInputDataVertexModel<FTime>(EnvelopeFollower::InParamNameAttackTime, LOCTEXT("AttackTimeTT", "The attack time of the envelope follower."), 0.01f),
-				TInputDataVertexModel<FTime>(EnvelopeFollower::InParamNameReleaseTime, LOCTEXT("ReleaseTimeTT", "The release time of the envelope follower."), 0.1f),
-				TInputDataVertexModel<FTime>(EnvelopeFollower::PoopTest, LOCTEXT("PoopTest", "Just your average poop test."), 0.1f)
+				TInputDataVertexModel<FTime>(EnvelopeFollower::InParamNameReleaseTime, LOCTEXT("ReleaseTimeTT", "The release time of the envelope follower."), 0.1f)
 			),
 			FOutputVertexInterface(
 				TOutputDataVertexModel<float>(EnvelopeFollower::OutParamNameEnvelope, LOCTEXT("EnvelopeFollowerOutputTT", "The output envelope value of the audio signal."))
@@ -147,7 +144,7 @@ namespace Metasound
 			FNodeClassMetadata Info;
 			Info.ClassName = { StandardNodes::Namespace, TEXT("Envelope Follower"), TEXT("") };
 			Info.MajorVersion = 1;
-			Info.MinorVersion = 1;
+			Info.MinorVersion = 2;
 			Info.DisplayName = LOCTEXT("Metasound_EnvelopeFollowerDisplayName", "Envelope Follower");
 			Info.Description = LOCTEXT("Metasound_EnvelopeFollowerDescription", "Outputs an envelope from an input audio signal.");
 			Info.Author = PluginAuthor;
