Metasound additive synth node bugs
- Fixed default value for frequency not being used
- Clamped frequency and harmonic multiplier inputs to nonnegative/Nyquist
- Change phase range to 0-360 degrees for consistency

diff --git a/Engine/Plugins/Runtime/Metasound/Source/MetasoundStandardNodes/Private/MetasoundAdditiveSynthNode.cpp b/Engine/Plugins/Runtime/Metasound/Source/MetasoundStandardNodes/Private/MetasoundAdditiveSynthNode.cpp
--- a/Engine/Plugins/Runtime/Metasound/Source/MetasoundStandardNodes/Private/MetasoundAdditiveSynthNode.cpp
+++ b/Engine/Plugins/Runtime/Metasound/Source/MetasoundStandardNodes/Private/MetasoundAdditiveSynthNode.cpp
@@ -17,10 +17,10 @@ namespace Metasound
 {
 	namespace AdditiveSynthVertexNames
 	{
-		METASOUND_PARAM(BaseFrequency, "Base Frequency", "Sinusoid frequency that harmonics are based on.");
-		METASOUND_PARAM(HarmonicMultipliers, "HarmonicMultipliers", "Array of float harmonic multipliers that will be applied to the base frequency. Number of sinusoids used depends on the size of this array.");
+		METASOUND_PARAM(BaseFrequency, "Base Frequency", "Sinusoid frequency that harmonics are based on, clamped to 0.0 - Nyquist.");
+		METASOUND_PARAM(HarmonicMultipliers, "HarmonicMultipliers", "Array of float harmonic multipliers that will be applied to the base frequency. Number of sinusoids used depends on the size of this array. Clamped from 0.0 to a max such that the resulting frequency won't go above Nyquist.");
 		METASOUND_PARAM(Amplitudes, "Amplitudes", "Array of sinusoid amplitudes, clamped to 0.0 - 1.0");
-		METASOUND_PARAM(Phases, "Phases", "Array of sinusoid phases, clamped to 0.0 - 1.0.");
+		METASOUND_PARAM(Phases, "Phases", "Array of sinusoid phases in degrees, clamped to 0.0 - 360.0");
 		METASOUND_PARAM(PanAmounts, "Pan Amounts", "Array of pan amounts (using equal power law; -1.0 is full left, 1.0 is full right).");
 
 		METASOUND_PARAM(LeftAudioOut, "Out Left Audio", "Left output synthesized audio.");
@@ -84,9 +84,10 @@ namespace Metasound
 			using namespace AdditiveSynthVertexNames;
 
 			const FDataReferenceCollection& InputCollection = InParams.InputDataReferences;
-			const FInputVertexInterface& InputInterface = GetVertexInterface().GetInputInterface();
+			FVertexInterface Interface = GetVertexInterface();
+			const FInputVertexInterface& InputInterface = Interface.GetInputInterface();
 
-			FFloatReadRef InputBaseFrequency = InputCollection.GetDataReadReferenceOrConstruct<float>(METASOUND_GET_PARAM_NAME(BaseFrequency));
+			FFloatReadRef InputBaseFrequency = InputCollection.GetDataReadReferenceOrConstructWithVertexDefault<float>(InputInterface, METASOUND_GET_PARAM_NAME(BaseFrequency), InParams.OperatorSettings);
 			FInputFloatArrayReadRef InputHarmonicMultipliers = InputCollection.GetDataReadReferenceOrConstruct<FInputFloatArrayType>(METASOUND_GET_PARAM_NAME(HarmonicMultipliers));
 			FInputFloatArrayReadRef InputAmplitudes = InputCollection.GetDataReadReferenceOrConstruct<FInputFloatArrayType>(METASOUND_GET_PARAM_NAME(Amplitudes));
 			FInputFloatArrayReadRef InputPhases = InputCollection.GetDataReadReferenceOrConstruct<FInputFloatArrayType>(METASOUND_GET_PARAM_NAME(Phases));
@@ -138,7 +139,6 @@ namespace Metasound
 
 		void Execute()
 		{
-			float Frequency = *BaseFrequency;
 			int32 NumSamples = LeftAudioOut->Num();
 
 			// Zero out output buffers since we'll sum directly into them
@@ -159,11 +159,15 @@ namespace Metasound
 				CreateSinusoids();
 			}
 
+			// Flag to update frequency (which depends on node base frequency and individual harmonic multiplier) if necessary
 			bool UpdateFrequency = false;
-			// Flag to update frequency (which depends on node base frequency and indiviudal harmonic multiplier) if necessary
-			if (!FMath::IsNearlyEqual(*BaseFrequency, PreviousBaseFrequency))
+			float Frequency = *BaseFrequency;
+			// Clamp frequency to Nyquist 
+			const float Nyquist = SampleRate / 2.0f;
+			Frequency = FMath::Clamp(Frequency, 0.0f, Nyquist);
+			if (!FMath::IsNearlyEqual(Frequency, PreviousBaseFrequency))
 			{
-				PreviousBaseFrequency = *BaseFrequency;
+				PreviousBaseFrequency = Frequency;
 				UpdateFrequency = true;
 			}
 
@@ -173,7 +177,9 @@ namespace Metasound
 				Audio::FSineWaveTableOsc& WaveTableOsc = SinusoidData[SinusoidIndex];
 
 				// Update harmonic multiplier and resulting frequency if necessary
-				const float HarmonicMultiplier = (*HarmonicMultipliers)[SinusoidIndex];
+				float HarmonicMultiplier = (*HarmonicMultipliers)[SinusoidIndex];
+				// Maximum is Nyquist / Frequency because the actual frequency will be HarmonicMultiplier * Frequency
+				HarmonicMultiplier = FMath::Clamp(HarmonicMultiplier, 0.0f, Nyquist / Frequency);
 				if (!FMath::IsNearlyEqual(HarmonicMultiplier, PreviousHarmonicMultipliers[SinusoidIndex]))
 				{
 					PreviousHarmonicMultipliers[SinusoidIndex] = HarmonicMultiplier;
@@ -182,7 +188,7 @@ namespace Metasound
 				
 				if (UpdateFrequency)
 				{
-					WaveTableOsc.SetFrequencyHz((*BaseFrequency) * HarmonicMultiplier);
+					WaveTableOsc.SetFrequencyHz(Frequency * HarmonicMultiplier);
 				}
 
 				// Update phase if necessary, set to 0 if no phases are given
@@ -200,11 +206,11 @@ namespace Metasound
 				else
 				{
 					int32 PhaseModIndex = SinusoidIndex % Phases->Num();
-					Phase = FMath::Clamp((*Phases)[PhaseModIndex], 0.0f, 1.0f);
+					Phase = FMath::Clamp((*Phases)[PhaseModIndex], 0.0f, 360.0f);
 					if (!FMath::IsNearlyEqual(Phase, PreviousPhases[PhaseModIndex]))
 					{
 						PreviousPhases[PhaseModIndex] = Phase;
-						WaveTableOsc.SetPhase(Phase);
+						WaveTableOsc.SetPhase(Phase / 360.0f);
 					}
 				}
 
@@ -263,7 +269,7 @@ namespace Metasound
 				// Create wavetable oscillators for each sinusoid
 				Audio::FSineWaveTableOsc WaveTableOsc = Audio::FSineWaveTableOsc();
 				float Phase = Phases->Num() > 0 ? (*Phases)[SinusoidIndex % Phases->Num()] : 0.0f;
-				WaveTableOsc.Init(SampleRate, (*HarmonicMultipliers)[SinusoidIndex] * (*BaseFrequency), Phase);
+				WaveTableOsc.Init(SampleRate, (*HarmonicMultipliers)[SinusoidIndex] * (*BaseFrequency), Phase / 360.0f);
 				SinusoidData.Add(WaveTableOsc);
 			}
 		}
