Fix ensure on metasound enum promoted to input
Enabled metasound enum update via BP and editor

diff --git a/Engine/Plugins/Runtime/Metasound/Source/MetasoundFrontend/Private/MetasoundAssetBase.cpp b/Engine/Plugins/Runtime/Metasound/Source/MetasoundFrontend/Private/MetasoundAssetBase.cpp
--- a/Engine/Plugins/Runtime/Metasound/Source/MetasoundFrontend/Private/MetasoundAssetBase.cpp
+++ b/Engine/Plugins/Runtime/Metasound/Source/MetasoundFrontend/Private/MetasoundAssetBase.cpp
@@ -406,13 +406,7 @@ TUniquePtr<Metasound::IGraph> FMetasoundAssetBase::BuildMetasoundDocument() cons
 	TUniquePtr<FFrontendGraph> FrontendGraph = FFrontendGraphBuilder::CreateGraph(*Doc);
 	if (FrontendGraph.IsValid())
 	{
-		auto GetVertexName = [](const FMetasoundFrontendClassVertex& InVertex) { return InVertex.Name; };
-
-		TSet<FVertexKey> VerticesToSkip;
-		FMetasoundFrontendArchetype Archetype;
-		GetArchetype(Archetype);
-		Algo::Transform(Archetype.Interface.Inputs, VerticesToSkip, GetVertexName);
-
+		TSet<FVertexKey> VerticesToSkip = GetNonTransmittableInputVertices(*Doc);
 		bool bSuccessfullyInjectedReceiveNodes = InjectReceiveNodes(*FrontendGraph, FMetasoundInstanceTransmitter::CreateSendAddressFromEnvironment, VerticesToSkip);
 		if (!bSuccessfullyInjectedReceiveNodes)
 		{
@@ -546,6 +540,44 @@ TArray<FMetasoundAssetBase::FSendInfoAndVertexName> FMetasoundAssetBase::GetSend
 	return SendInfos;
 }
 
+TSet<Metasound::FVertexKey> FMetasoundAssetBase::GetNonTransmittableInputVertices(const FMetasoundFrontendDocument& InDoc) const
+{
+	using namespace Metasound;
+	using namespace Metasound::Frontend;
+
+	check(IsInGameThread());
+
+	TSet<FVertexKey> NonTransmittableInputVertices;
+
+	auto GetVertexKey = [](const FMetasoundFrontendClassVertex& InVertex) { return InVertex.Name; };
+
+	// Do not transmit vertices defined in archetype. Those are already accounted
+	// for by owning object.
+	FMetasoundFrontendArchetype Archetype;
+	GetArchetype(Archetype);
+	Algo::Transform(Archetype.Interface.Inputs, NonTransmittableInputVertices, GetVertexKey);
+
+
+	// Do not transmit vertices which are not transmittable. Async communication 
+	// is not supported without transmission.
+	FMetasoundFrontendRegistryContainer* Registry = FMetasoundFrontendRegistryContainer::Get();
+	check(nullptr != Registry);
+
+	auto IsNotTransmittable = [&Registry](const FMetasoundFrontendClassVertex& InVertex) -> bool
+	{	
+		FDataTypeRegistryInfo Info;
+		if (Registry->GetInfoForDataType(InVertex.TypeName, Info))
+		{
+			return !Info.bIsTransmittable;
+		}
+		return true;
+	};
+
+	Algo::TransformIf(InDoc.RootGraph.Interface.Inputs, NonTransmittableInputVertices, IsNotTransmittable, GetVertexKey);
+
+	return NonTransmittableInputVertices;
+}
+
 TArray<FString> FMetasoundAssetBase::GetTransmittableInputVertexNames() const
 {
 	using namespace Metasound;
diff --git a/Engine/Plugins/Runtime/Metasound/Source/MetasoundFrontend/Private/MetasoundFrontendInjectReceiveNodes.cpp b/Engine/Plugins/Runtime/Metasound/Source/MetasoundFrontend/Private/MetasoundFrontendInjectReceiveNodes.cpp
--- a/Engine/Plugins/Runtime/Metasound/Source/MetasoundFrontend/Private/MetasoundFrontendInjectReceiveNodes.cpp
+++ b/Engine/Plugins/Runtime/Metasound/Source/MetasoundFrontend/Private/MetasoundFrontendInjectReceiveNodes.cpp
@@ -115,18 +115,12 @@ namespace Metasound
 
 			TUniquePtr<INode> CreateReceiveNodeForDataType(const FGuid& InID, const FVertexKey& InVertexKey, const FName& InDataType)
 			{
-				const EMetasoundFrontendClassType Type = EMetasoundFrontendClassType::External;
-				const FNodeClassName& ClassName = ReceiveNodeInfo::GetClassNameForDataType(InDataType);
-				const int32 MajorVersion = ReceiveNodeInfo::GetCurrentMajorVersion();
-				const int32 MinorVersion = ReceiveNodeInfo::GetCurrentMinorVersion();
-
-				FNodeRegistryKey ReceiveNodeRegistryKey = NodeRegistryKey::CreateKey(Type, ClassName.GetFullName().ToString(), MajorVersion, MinorVersion);
 				FNodeInitData ReceiveNodeInitData 
 				{ 
 					FString::Format(TEXT("ReceiveInject_{0}"), { InVertexKey }),
 					InID
 				};
-				return FMetasoundFrontendRegistryContainer::Get()->CreateNode(ReceiveNodeRegistryKey, MoveTemp(ReceiveNodeInitData));
+				return FMetasoundFrontendRegistryContainer::Get()->CreateReceiveNode(InDataType, MoveTemp(ReceiveNodeInitData));
 			}
 		}
 
diff --git a/Engine/Plugins/Runtime/Metasound/Source/MetasoundFrontend/Private/MetasoundFrontendRegistries.cpp b/Engine/Plugins/Runtime/Metasound/Source/MetasoundFrontend/Private/MetasoundFrontendRegistries.cpp
--- a/Engine/Plugins/Runtime/Metasound/Source/MetasoundFrontend/Private/MetasoundFrontendRegistries.cpp
+++ b/Engine/Plugins/Runtime/Metasound/Source/MetasoundFrontend/Private/MetasoundFrontendRegistries.cpp
@@ -417,6 +417,8 @@ namespace Metasound
 				TUniquePtr<Metasound::INode> CreateInputNode(const FName& InInputType, Metasound::FInputNodeConstructorParams&& InParams) override;
 				TUniquePtr<Metasound::INode> CreateVariableNode(const FName& InVariableType, FVariableNodeConstructorParams&& InParams) override;
 				TUniquePtr<Metasound::INode> CreateOutputNode(const FName& InOutputType, Metasound::FOutputNodeConstructorParams&& InParams) override;
+				TUniquePtr<Metasound::INode> CreateReceiveNode(const FName& InVariableType, FReceiveNodeConstructorParams&& InParams) override;
+
 				virtual TUniquePtr<INode> CreateNode(const FNodeRegistryKey& InKey, FDefaultNodeConstructorParams&&) const override;
 				virtual TUniquePtr<INode> CreateNode(const FNodeRegistryKey& InKey, FDefaultLiteralNodeConstructorParams&&) const override;
 				virtual TUniquePtr<Metasound::INode> CreateNode(const FNodeRegistryKey& InKey, const Metasound::FNodeInitData& InInitData) const override;
@@ -547,6 +549,20 @@ namespace Metasound
 				}
 			}
 
+			TUniquePtr<Metasound::INode> FRegistryContainerImpl::CreateReceiveNode(const FName& InDataType, Metasound::FReceiveNodeConstructorParams&& InParams)
+			{
+				const IDataTypeRegistryEntry* Entry = FindDataTypeEntry(InDataType);
+
+				if (ensureAlwaysMsgf(nullptr != Entry, TEXT("Could not find data type [Name:%s]"), *InDataType.ToString()))
+				{
+					return Entry->CreateReceiveNode(MoveTemp(InParams));
+				}
+				else
+				{
+					return nullptr;
+				}
+			}
+
 			Metasound::FLiteral FRegistryContainerImpl::CreateLiteralFromUObject(const FName& InDataType, UObject* InObject)
 			{
 				const IDataTypeRegistryEntry* Entry = FindDataTypeEntry(InDataType);
diff --git a/Engine/Plugins/Runtime/Metasound/Source/MetasoundFrontend/Public/MetasoundAssetBase.h b/Engine/Plugins/Runtime/Metasound/Source/MetasoundFrontend/Public/MetasoundAssetBase.h
--- a/Engine/Plugins/Runtime/Metasound/Source/MetasoundFrontend/Public/MetasoundAssetBase.h
+++ b/Engine/Plugins/Runtime/Metasound/Source/MetasoundFrontend/Public/MetasoundAssetBase.h
@@ -191,6 +191,7 @@ protected:
 private:
 	Metasound::Frontend::FNodeRegistryKey RegistryKey;
 
+	TSet<Metasound::FVertexKey> GetNonTransmittableInputVertices(const FMetasoundFrontendDocument& InDoc) const;
 	TArray<FString> GetTransmittableInputVertexNames() const;
 	Metasound::FSendAddress CreateSendAddress(uint64 InInstanceID, const FString& InVertexName, const FName& InDataTypeName) const;
 	Metasound::Frontend::FNodeHandle AddInputPinForSendAddress(const Metasound::FMetasoundInstanceTransmitter::FSendInfo& InSendInfo, Metasound::Frontend::FGraphHandle InGraph) const;
diff --git a/Engine/Plugins/Runtime/Metasound/Source/MetasoundFrontend/Public/MetasoundDataTypeRegistrationMacro.h b/Engine/Plugins/Runtime/Metasound/Source/MetasoundFrontend/Public/MetasoundDataTypeRegistrationMacro.h
--- a/Engine/Plugins/Runtime/Metasound/Source/MetasoundFrontend/Public/MetasoundDataTypeRegistrationMacro.h
+++ b/Engine/Plugins/Runtime/Metasound/Source/MetasoundFrontend/Public/MetasoundDataTypeRegistrationMacro.h
@@ -423,6 +423,11 @@ namespace Metasound
 					return MakeUnique<TGetVariableNode<TDataType>>(InParams.NodeName, InParams.InstanceID, InParams.VertexName, MoveTemp(InParams.InitParam));
 				}
 
+				virtual TUniquePtr<INode> CreateReceiveNode(FReceiveNodeConstructorParams&& InParams) const override
+				{
+					return MakeUnique<TReceiveNode<TDataType>>(InParams);
+				}
+
 				virtual Audio::IProxyDataPtr CreateProxy(UObject* InObject) const override
 				{
 					// Only attempt to create proxy if the `UClassToUse` is not void.
diff --git a/Engine/Plugins/Runtime/Metasound/Source/MetasoundFrontend/Public/MetasoundFrontendRegistries.h b/Engine/Plugins/Runtime/Metasound/Source/MetasoundFrontend/Public/MetasoundFrontendRegistries.h
--- a/Engine/Plugins/Runtime/Metasound/Source/MetasoundFrontend/Public/MetasoundFrontendRegistries.h
+++ b/Engine/Plugins/Runtime/Metasound/Source/MetasoundFrontend/Public/MetasoundFrontendRegistries.h
@@ -36,6 +36,7 @@ namespace Metasound
 	using FOutputNodeConstructorParams = FDefaultNodeConstructorParams;
 	using FVariableNodeConstructorParams = FDefaultLiteralNodeConstructorParams;
 	using FInputNodeConstructorParams = FDefaultLiteralNodeConstructorParams;
+	using FReceiveNodeConstructorParams = FNodeInitData;
 
 	using FIterateMetasoundFrontendClassFunction = TFunctionRef<void(const FMetasoundFrontendClass&)>;
 
@@ -267,6 +268,9 @@ namespace Metasound
 			/** Create a variable node */
 			virtual TUniquePtr<INode> CreateVariableNode(FVariableNodeConstructorParams&&) const = 0;
 
+			/** Create a recevie node. */
+			virtual TUniquePtr<INode> CreateReceiveNode(FReceiveNodeConstructorParams&&) const = 0;
+
 			/** Create a proxy from a UObject. If this data type does not support
 			 * UObject proxies, return a nullptr. */
 			virtual Audio::IProxyDataPtr CreateProxy(UObject*) const = 0;
@@ -504,6 +508,7 @@ public:
 	virtual TUniquePtr<Metasound::INode> CreateInputNode(const FName& InInputType, Metasound::FInputNodeConstructorParams&& InParams) = 0;
 	virtual TUniquePtr<Metasound::INode> CreateVariableNode(const FName& InVariableType, Metasound::FVariableNodeConstructorParams&& InParams) = 0;
 	virtual TUniquePtr<Metasound::INode> CreateOutputNode(const FName& InOutputType, Metasound::FOutputNodeConstructorParams&& InParams) = 0;
+	virtual TUniquePtr<Metasound::INode> CreateReceiveNode(const FName& InVariableType, Metasound::FReceiveNodeConstructorParams&& InParams) = 0;
 
 	virtual TUniquePtr<Metasound::INode> CreateNode(const FNodeRegistryKey& InKey, Metasound::FDefaultNodeConstructorParams&&) const = 0;
 	virtual TUniquePtr<Metasound::INode> CreateNode(const FNodeRegistryKey& InKey, Metasound::FDefaultLiteralNodeConstructorParams&&) const = 0;
