Fixing for check hitting in metasound with -game

diff --git a/Engine/Plugins/Runtime/Metasound/Source/MetasoundFrontend/Private/MetasoundAssetBase.cpp b/Engine/Plugins/Runtime/Metasound/Source/MetasoundFrontend/Private/MetasoundAssetBase.cpp
--- a/Engine/Plugins/Runtime/Metasound/Source/MetasoundFrontend/Private/MetasoundAssetBase.cpp
+++ b/Engine/Plugins/Runtime/Metasound/Source/MetasoundFrontend/Private/MetasoundAssetBase.cpp
@@ -545,7 +545,7 @@ TSet<Metasound::FVertexKey> FMetasoundAssetBase::GetNonTransmittableInputVertice
 	using namespace Metasound;
 	using namespace Metasound::Frontend;
 
-	check(IsInGameThread());
+	check(IsInGameThread() || IsInAudioThread());
 
 	TSet<FVertexKey> NonTransmittableInputVertices;
 
