Build MetaSound instance on async task

diff --git a/Engine/Plugins/Runtime/Metasound/Source/MetasoundGenerator/Private/MetasoundGenerator.cpp b/Engine/Plugins/Runtime/Metasound/Source/MetasoundGenerator/Private/MetasoundGenerator.cpp
--- a/Engine/Plugins/Runtime/Metasound/Source/MetasoundGenerator/Private/MetasoundGenerator.cpp
+++ b/Engine/Plugins/Runtime/Metasound/Source/MetasoundGenerator/Private/MetasoundGenerator.cpp
@@ -10,6 +10,10 @@
 #include "MetasoundTrace.h"
 #include "MetasoundTrigger.h"
 
+#ifndef METASOUND_GENERATOR_BUILD_GRAPH_ASYNC
+#define METASOUND_GENERATOR_BUILD_GRAPH_ASYNC 1
+#endif
+
 namespace Metasound
 {
 
@@ -30,7 +34,6 @@ namespace Metasound
 
 		TUniquePtr<IOperator> GraphOperator = OperatorBuilder.BuildGraphOperator(BuildParams, BuildErrors);
 
-
 		// Log build errors
 		for (const IOperatorBuilder::FBuildErrorPtr& Error : BuildErrors)
 		{
@@ -98,12 +101,15 @@ namespace Metasound
 
 		BuilderTask = MakeUnique<FBuilderTask>(this, MoveTemp(InParams), true /* bTriggerGenerator */);
 		
+#if METASOUND_GENERATOR_BUILD_GRAPH_ASYNC
+		BuilderTask->StartBackgroundTask(GBackgroundPriorityThreadPool);
+		
+#else
 		BuilderTask->StartSynchronousTask();
 		BuilderTask = nullptr;
 		UpdateGraphIfPending();
 		bIsWaitingForFirstGraph = false;
-		
-		//BuilderTask->StartBackgroundTask(GBackgroundPriorityThreadPool);
+#endif
 	}
 
 	FMetasoundGenerator::~FMetasoundGenerator()
