perf improvements to metasound node searching

diff --git a/Engine/Plugins/Runtime/Metasound/Source/MetasoundFrontend/Private/MetasoundFrontendQuery.cpp b/Engine/Plugins/Runtime/Metasound/Source/MetasoundFrontend/Private/MetasoundFrontendQuery.cpp
--- a/Engine/Plugins/Runtime/Metasound/Source/MetasoundFrontend/Private/MetasoundFrontendQuery.cpp
+++ b/Engine/Plugins/Runtime/Metasound/Source/MetasoundFrontend/Private/MetasoundFrontendQuery.cpp
@@ -523,15 +523,15 @@ namespace Metasound
 
 				if (Step.IsValid())
 				{
-					auto FilterFunc = [&](const FFrontendQueryEntry* Entry)
-					{
-						return (nullptr != Entry) && Step->Filter(*Entry);
-					};
-
-					int32 NumRemoved = InOutResult.FilterSelection(FilterFunc);
-					if (NumRemoved > 0)
+					if (InOutResult.GetSelection().Num() > 0)
 					{
 						ResultState = EResultModificationState::Modified;
+						auto FilterFunc = [&](const FFrontendQueryEntry* Entry)
+						{
+							return (nullptr != Entry) && Step->Filter(*Entry);
+						};
+
+						InOutResult.FilterSelection(FilterFunc);
 					}
 				}
 
diff --git a/Engine/Plugins/Runtime/Metasound/Source/MetasoundFrontend/Private/MetasoundFrontendSearchEngine.cpp b/Engine/Plugins/Runtime/Metasound/Source/MetasoundFrontend/Private/MetasoundFrontendSearchEngine.cpp
--- a/Engine/Plugins/Runtime/Metasound/Source/MetasoundFrontend/Private/MetasoundFrontendSearchEngine.cpp
+++ b/Engine/Plugins/Runtime/Metasound/Source/MetasoundFrontend/Private/MetasoundFrontendSearchEngine.cpp
@@ -8,6 +8,7 @@
 #include "MetasoundFrontendQuery.h"
 #include "MetasoundFrontendQuerySteps.h"
 #include "MetasoundFrontendRegistryTransaction.h"
+#include "MetasoundTrace.h"
 
 static int32 ClearMetaSoundFrontendSearchEngineCacheCVar = 0;
 FAutoConsoleVariableRef CVarClearMetaSoundFrontendSearchEngineCache(
@@ -183,6 +184,27 @@ namespace Metasound
 				FName Name;
 			};
 
+
+			class FFilterNodeRegistrationEventsByClassName : public IFrontendQueryFilterStep
+			{
+			public:
+				FFilterNodeRegistrationEventsByClassName(const FMetasoundFrontendClassName& InName)
+				: Name(InName)
+				{
+				}
+
+				bool Filter(const FFrontendQueryEntry& InEntry) const override
+				{
+					if (ensure(InEntry.Value.IsType<FNodeRegistryTransaction>()))
+					{
+						return Name == InEntry.Value.Get<FNodeRegistryTransaction>().GetNodeClassInfo().ClassName;
+					}
+					return false;
+				}
+			private:
+				FMetasoundFrontendClassName Name;
+			};
+
 			class FFilterArchetypeRegistryTransactionsByNameAndMajorVersion : public IFrontendQueryFilterStep
 			{
 			public:
@@ -222,26 +244,6 @@ namespace Metasound
 				}
 			};
 
-			class FFilterRegistryTransactionsByName : public IFrontendQueryFilterStep
-			{
-			public:
-				FFilterRegistryTransactionsByName(const FName& InName)
-				: Name(InName)
-				{
-				}
-
-				bool Filter(const FFrontendQueryEntry& InEntry) const override
-				{
-					if (ensure(InEntry.Value.IsType<FArchetypeRegistryTransaction>()))
-					{
-						return Name == InEntry.Value.Get<FArchetypeRegistryTransaction>().GetArchetypeVersion().Name;
-					}
-					return false;
-				}
-			private:
-				FName Name;
-			};
-
 			class FTransformArchetypeRegistryTransactionToArchetypeVersion : public IFrontendQueryTransformStep
 			{
 			public:
@@ -298,6 +300,8 @@ namespace Metasound
 
 		TArray<FMetasoundFrontendClass> FSearchEngine::FindAllClasses(bool bInIncludeDeprecated)
 		{
+			METASOUND_TRACE_CPUPROFILER_EVENT_SCOPE(metasound::FSearchEngine::FindAllClasses);
+
 			FString QueryName = FString(TEXT("FindAllClasses"));
 			if (bInIncludeDeprecated)
 			{
@@ -344,6 +348,9 @@ namespace Metasound
 
 		TArray<FMetasoundFrontendClass> FSearchEngine::FindClassesWithName(const FNodeClassName& InName, bool bInSortByVersion)
 		{
+			METASOUND_TRACE_CPUPROFILER_EVENT_SCOPE(metasound::FSearchEngine::FindClassesWithName);
+			using namespace SearchEngineQuerySteps;
+
 			FString QueryName = FString(TEXT("FindClassesWithName"));
 			if (bInSortByVersion)
 			{
@@ -356,6 +363,7 @@ namespace Metasound
 			{
 				TUniquePtr<FFrontendQuery> NewQuery = MakeUnique<FFrontendQuery>();
 				NewQuery->AddStep<FNodeClassRegistrationEvents>()
+					.AddStep<FFilterNodeRegistrationEventsByClassName>(InName)
 					.AddStep<FMapRegistrationEventsToNodeRegistryKeys>()
 					.AddStep<FReduceRegistrationEventsToCurrentStatus>()
 					.AddStep<FTransformRegistrationEventsToClasses>()
@@ -387,6 +395,9 @@ namespace Metasound
 
 		bool FSearchEngine::FindClassWithHighestVersion(const FNodeClassName& InName, FMetasoundFrontendClass& OutClass)
 		{
+			METASOUND_TRACE_CPUPROFILER_EVENT_SCOPE(metasound::FSearchEngine::FindClassWithHighestVersion);
+			using namespace SearchEngineQuerySteps;
+
 			const FString QueryName = FString(TEXT("FindHighestVersion")) + InName.GetFullName().ToString();
 			FFrontendQuery* Query = FindQuery(QueryName);
 			if (!Query)
@@ -396,6 +407,7 @@ namespace Metasound
 				// TODO: Create index of class names and major version to avoid duplicating
 				// this query for each class name.
 				NewQuery->AddStep<FNodeClassRegistrationEvents>()
+					.AddStep<FFilterNodeRegistrationEventsByClassName>(InName)
 					.AddStep<FMapRegistrationEventsToNodeRegistryKeys>()
 					.AddStep<FReduceRegistrationEventsToCurrentStatus>()
 					.AddStep<FTransformRegistrationEventsToClasses>()
@@ -427,6 +439,9 @@ namespace Metasound
 
 		bool FSearchEngine::FindClassWithMajorVersion(const FNodeClassName& InName, int32 InMajorVersion, FMetasoundFrontendClass& OutClass)
 		{
+			METASOUND_TRACE_CPUPROFILER_EVENT_SCOPE(metasound::FSearchEngine::FindClassWithMajorVersion);
+			using namespace SearchEngineQuerySteps;
+
 			const FString QueryName = FString(TEXT("FindClassWithMajorVersion")) + InName.GetFullName().ToString() + FString::FromInt(InMajorVersion);
 			FFrontendQuery* Query = FindQuery(QueryName);
 
@@ -437,6 +452,7 @@ namespace Metasound
 				// TODO: Create index of class names to avoid duplicating
 				// this query for each class name.
 				NewQuery->AddStep<FNodeClassRegistrationEvents>()
+					.AddStep<FFilterNodeRegistrationEventsByClassName>(InName)
 					.AddStep<FMapRegistrationEventsToNodeRegistryKeys>()
 					.AddStep<FReduceRegistrationEventsToCurrentStatus>()
 					.AddStep<FTransformRegistrationEventsToClasses>()
@@ -522,7 +538,7 @@ namespace Metasound
 				TUniquePtr<FFrontendQuery> NewQuery = MakeUnique<FFrontendQuery>();
 
 				NewQuery->AddStep<FArchetypeRegistryTransactionSource>()
-					.AddStep<FFilterRegistryTransactionsByName>(InArchetypeName)
+					.AddStep<FFilterArchetypeRegistryTransactionsByName>(InArchetypeName)
 					.AddStep<FMapArchetypeRegistryTransactionsToArchetypeRegistryKeys>()
 					.AddStep<FReduceArchetypeRegistryTransactionsToCurrentStatus>()
 					.AddStep<FTransformArchetypeRegistryTransactionToArchetypeVersion>();
