Move MetaSound Registration Off of App Load
- Fix object picker getting into invalid state when using browsed object/clear not working
- Fix presets being pitched down an octave when set to stereo due to impromper mono down-mixing
- Fix for drag/drop not omitting assets that would cause reference loops
- Minor fix for marking referenced assets as dirtied if parent autoupdate call trickles down request

diff --git a/Engine/Plugins/Runtime/Metasound/Source/MetasoundEditor/Private/MetasoundAssetTypeActions.cpp b/Engine/Plugins/Runtime/Metasound/Source/MetasoundEditor/Private/MetasoundAssetTypeActions.cpp
--- a/Engine/Plugins/Runtime/Metasound/Source/MetasoundEditor/Private/MetasoundAssetTypeActions.cpp
+++ b/Engine/Plugins/Runtime/Metasound/Source/MetasoundEditor/Private/MetasoundAssetTypeActions.cpp
@@ -154,7 +154,13 @@ namespace Metasound
 		{
 			using namespace AssetTypeActionsPrivate;
 			RegisterPresetAction<UMetaSoundSource, UMetaSoundSourceFactory>(SourcePresetLabel, SourcePresetToolTip);
-			RegisterPresetAction<UMetaSound, UMetaSoundFactory, UMetaSoundSource>(PresetLabel, PresetToolTip);
+
+			// This is currently disabled because of the requirement that archetypes of a preset and its referenced
+			// asset must match.  TODO: A feature of archetypes/preset transform is required for this action to work
+			// that introduces the concept of an archetype "interface" (i.e. if the archetype one graph subscribes to
+			// or contains all of the necessary inputs/outputs of another), so that the preset can "implement" the
+			// referenced graphs "interface."
+// 			RegisterPresetAction<UMetaSound, UMetaSoundFactory, UMetaSoundSource>(PresetLabel, PresetToolTip);
 		}
 	} // namespace Editor
 } // namespace Metasound
diff --git a/Engine/Plugins/Runtime/Metasound/Source/MetasoundEditor/Private/MetasoundDetailCustomization.cpp b/Engine/Plugins/Runtime/Metasound/Source/MetasoundEditor/Private/MetasoundDetailCustomization.cpp
--- a/Engine/Plugins/Runtime/Metasound/Source/MetasoundEditor/Private/MetasoundDetailCustomization.cpp
+++ b/Engine/Plugins/Runtime/Metasound/Source/MetasoundEditor/Private/MetasoundDetailCustomization.cpp
@@ -14,6 +14,7 @@
 #include "MetasoundEditorSettings.h"
 #include "MetasoundFrontend.h"
 #include "MetasoundFrontendController.h"
+#include "MetasoundSource.h"
 #include "MetasoundUObjectRegistry.h"
 #include "PropertyEditorDelegates.h"
 #include "PropertyHandle.h"
@@ -65,7 +66,9 @@ namespace Metasound
 
 		void FMetasoundDetailCustomization::CustomizeDetails(IDetailLayoutBuilder& DetailLayout)
 		{
+			using namespace Metasound;
 			using namespace Metasound::Editor;
+			using namespace Metasound::Frontend;
 
 			EMetasoundActiveDetailView DetailsView = EMetasoundActiveDetailView::Metasound;
 			if (const UMetasoundEditorSettings* EditorSettings = GetDefault<UMetasoundEditorSettings>())
@@ -73,6 +76,16 @@ namespace Metasound
 				DetailsView = EditorSettings->DetailView;
 			}
 
+			// Currently only support modifying a single MetaSound at a time (Multiple MetaSound editting will be covered by separate tool)
+			TArray<TWeakObjectPtr<UObject>> Objects;
+			DetailLayout.GetObjectsBeingCustomized(Objects);
+			if (Objects.Num() > 1)
+			{
+				return;
+			}
+
+			TWeakObjectPtr<UMetaSoundSource> MetaSoundSource = Cast<UMetaSoundSource>(Objects[0].Get());
+
 			switch (DetailsView)
 			{
 				case EMetasoundActiveDetailView::Metasound:
@@ -96,6 +109,41 @@ namespace Metasound
 					TSharedPtr<IPropertyHandle> MajorVersionHandle = DetailLayout.GetProperty(MajorVersionPropertyPath);
 					TSharedPtr<IPropertyHandle> MinorVersionHandle = DetailLayout.GetProperty(MinorVersionPropertyPath);
 
+					// Invalid for UMetaSounds
+					TSharedPtr<IPropertyHandle> OutputFormat = DetailLayout.GetProperty(GET_MEMBER_NAME_CHECKED(UMetaSoundSource, OutputFormat));
+					if (OutputFormat.IsValid())
+					{
+						TAttribute<bool> IsEnabledAttribute = TAttribute<bool>::Create([InSource = MetaSoundSource]()
+						{
+							if (const FMetasoundAssetBase* MetaSoundAsset = IMetasoundUObjectRegistry::Get().GetObjectAsAssetBase(Cast<UObject>(InSource)))
+							{
+								FConstGraphHandle GraphHandle = MetaSoundAsset->GetRootGraphHandle();
+								return GraphHandle->GetGraphStyle().bIsGraphEditable;
+							}
+
+							return false;
+						});
+
+						TSharedRef<SWidget> OutputFormatValueWidget = OutputFormat->CreatePropertyValueWidget();
+						OutputFormatValueWidget->SetEnabled(IsEnabledAttribute);
+
+						static const FText OutputFormatName = LOCTEXT("MetasoundOutputFormatPropertyName", "Output Format");
+						GeneralCategoryBuilder.AddCustomRow(OutputFormatName)
+// 						.EditCondition(
+// 						{
+// 						}), nullptr)
+						.NameContent()
+						[
+							OutputFormat->CreatePropertyNameWidget()
+						]
+						.ValueContent()
+						[
+							OutputFormatValueWidget
+						];
+
+						OutputFormat->MarkHiddenByCustomization();
+					}
+
 					GeneralCategoryBuilder.AddProperty(AuthorHandle);
 					GeneralCategoryBuilder.AddProperty(DescHandle);
 					GeneralCategoryBuilder.AddProperty(MajorVersionHandle);
diff --git a/Engine/Plugins/Runtime/Metasound/Source/MetasoundEditor/Private/MetasoundEditor.cpp b/Engine/Plugins/Runtime/Metasound/Source/MetasoundEditor/Private/MetasoundEditor.cpp
--- a/Engine/Plugins/Runtime/Metasound/Source/MetasoundEditor/Private/MetasoundEditor.cpp
+++ b/Engine/Plugins/Runtime/Metasound/Source/MetasoundEditor/Private/MetasoundEditor.cpp
@@ -323,9 +323,8 @@ namespace Metasound
 
 			Metasound = ObjectToEdit;
 
-			const bool bAutoUpdate = true;
-			const bool bClearUpdateNotes = false;
-			FGraphBuilder::ValidateGraph(*Metasound, bAutoUpdate, bClearUpdateNotes);
+			FGraphBuilder::RegisterGraphWithFrontend(*Metasound);
+			FGraphBuilder::ValidateGraph(*Metasound);
 
 			GEditor->RegisterForUndo(this);
 
@@ -388,8 +387,6 @@ namespace Metasound
 
 			ExtendToolbar();
 			RegenerateMenusAndToolbars();
-
-			FGraphBuilder::ValidateGraph(*Metasound, false /* bAutoUpdate */);
 		}
 
 		UObject* FEditor::GetMetasoundObject() const
@@ -810,11 +807,16 @@ namespace Metasound
 		{
 			if (USoundBase* MetasoundToPlay = Cast<USoundBase>(Metasound))
 			{
-				if (!FGraphBuilder::ValidateGraph(*Metasound, false /* bAutoUpdate */))
+				if (!FGraphBuilder::ValidateGraph(*Metasound))
 				{
 					return;
 				}
 
+				// Even though the MetaSoundSource will attempt to register via InitResources
+				// later in this execution (and deeper in the stack), this call forces
+				// re-registering to make sure everything is up-to-date.
+				FGraphBuilder::RegisterGraphWithFrontend(*Metasound);
+
 				// Set the send to the audio bus that is used for analyzing the metasound output
 				check(GEditor);
 				if (UAudioComponent* PreviewComp = GEditor->PlayPreviewSound(MetasoundToPlay))
@@ -1336,7 +1338,7 @@ namespace Metasound
 				}
 			}
 
-			Metasound::Editor::FGraphBuilder::RegisterGraphWithFrontend(Graph->GetMetasoundChecked());
+			FGraphBuilder::RegisterGraphWithFrontend(Graph->GetMetasoundChecked());
 		}
 
 		void FEditor::DeleteSelected()
@@ -1606,9 +1608,9 @@ namespace Metasound
 					LookupMetadata.SetClassName(ExternalNode->GetClassName());
 					LookupMetadata.SetType(EMetasoundFrontendClassType::External);
 					const FNodeRegistryKey PastedRegistryKey = NodeRegistryKey::CreateKey(LookupMetadata);
-					if (const FSoftObjectPath* AssetPath = UMetaSoundAssetSubsystem::Get().FindObjectPathFromKey(PastedRegistryKey))
+					if (const FSoftObjectPath* AssetPath = IMetaSoundAssetManager::GetChecked().FindObjectPathFromKey(PastedRegistryKey))
 					{
-						if (MetasoundAsset->AddingReferenceCausesLoop(*AssetPath, UMetaSoundAssetSubsystem::Get()))
+						if (MetasoundAsset->AddingReferenceCausesLoop(*AssetPath))
 						{
 							FMetasoundFrontendClass MetaSoundClass;
 							FMetasoundFrontendRegistryContainer::Get()->FindFrontendClassFromRegistered(PastedRegistryKey, MetaSoundClass);
diff --git a/Engine/Plugins/Runtime/Metasound/Source/MetasoundEditor/Private/MetasoundEditorGraph.cpp b/Engine/Plugins/Runtime/Metasound/Source/MetasoundEditor/Private/MetasoundEditorGraph.cpp
--- a/Engine/Plugins/Runtime/Metasound/Source/MetasoundEditor/Private/MetasoundEditorGraph.cpp
+++ b/Engine/Plugins/Runtime/Metasound/Source/MetasoundEditor/Private/MetasoundEditorGraph.cpp
@@ -446,13 +446,13 @@ bool UMetasoundEditorGraph::Synchronize()
 	return false;
 }
 
-bool UMetasoundEditorGraph::Validate(bool bInAutoUpdate, bool bInClearUpdateNotes)
+bool UMetasoundEditorGraph::Validate(bool bInClearUpdateNotes)
 {
 	using namespace Metasound::Editor;
 
 	if (UObject* ParentMetasound = GetOuter())
 	{
-		return FGraphBuilder::ValidateGraph(*ParentMetasound, bInAutoUpdate, bInClearUpdateNotes);
+		return FGraphBuilder::ValidateGraph(*ParentMetasound, bInClearUpdateNotes);
 	}
 
 	return false;
diff --git a/Engine/Plugins/Runtime/Metasound/Source/MetasoundEditor/Private/MetasoundEditorGraphBuilder.cpp b/Engine/Plugins/Runtime/Metasound/Source/MetasoundEditor/Private/MetasoundEditorGraphBuilder.cpp
--- a/Engine/Plugins/Runtime/Metasound/Source/MetasoundEditor/Private/MetasoundEditorGraphBuilder.cpp
+++ b/Engine/Plugins/Runtime/Metasound/Source/MetasoundEditor/Private/MetasoundEditorGraphBuilder.cpp
@@ -114,8 +114,6 @@ namespace Metasound
 			check(MetaSoundAsset);
 
 			Frontend::FNodeHandle NodeHandle = MetaSoundAsset->GetRootGraphHandle()->AddNode(InMetadata);
-			MetaSoundAsset->RebuildReferencedAssets(UMetaSoundAssetSubsystem::Get());
-
 			return AddExternalNode(InMetaSound, NodeHandle, InLocation, bInSelectNewNode);
 		}
 
@@ -146,7 +144,7 @@ namespace Metasound
 			return NewGraphNode;
 		}
 
-		bool FGraphBuilder::ValidateGraph(UObject& InMetaSound, bool bInAutoUpdate, bool bInClearUpdateNotes)
+		bool FGraphBuilder::ValidateGraph(UObject& InMetaSound, bool bInClearUpdateNotes)
 		{
 			using namespace Frontend;
 			using namespace GraphBuilderPrivate;
@@ -166,39 +164,29 @@ namespace Metasound
 
 			FGraphValidationResults Results;
 
-			auto ValidateGraphInteral = [&]()
-			{
-				Graph.ValidateInternal(Results, bInClearUpdateNotes);
-				for (const FGraphNodeValidationResult& Result : Results.GetResults())
-				{
-					if (GraphEditor.IsValid())
-					{
-						check(Result.Node);
-						GraphEditor->RefreshNode(*Result.Node);
-					}
-				}
-
-				InMetaSound.MarkPackageDirty();
-			};
+			bool bMarkDirty = false;
 
-			if (bInAutoUpdate)
+			Graph.ValidateInternal(Results, bInClearUpdateNotes);
+			for (const FGraphNodeValidationResult& Result : Results.GetResults())
 			{
-				if (MetaSoundAsset->AutoUpdate(UMetaSoundAssetSubsystem::Get(), true /* bMarkDirty*/, true /* bUpdateReferencedAssets*/))
+				bMarkDirty |= Result.bIsDirty;
+				if (GraphEditor.IsValid())
 				{
-					SynchronizeGraph(InMetaSound);
-					ValidateGraphInteral();
+					check(Result.Node);
+					GraphEditor->RefreshNode(*Result.Node);
 				}
 			}
-			else
-			{
-				ValidateGraphInteral();
-			}
 
 			if (MetaSoundEditor.IsValid())
 			{
 				MetaSoundEditor->RefreshInterface();
 			}
 
+			if (bMarkDirty)
+			{
+				InMetaSound.MarkPackageDirty();
+			}
+
 			return Results.IsValid();
 		}
 
@@ -633,8 +621,6 @@ namespace Metasound
 					ExternalNode->SetNodeID(NewNode->GetID());
 
 					NodeHandle = NewNode;
-
-					MetaSoundAsset->RebuildReferencedAssets(UMetaSoundAssetSubsystem::Get());
 				}
 			}
 
@@ -858,9 +844,13 @@ namespace Metasound
 			Metadata.SetAutoUpdateManagesInterface(true);
 			PresetGraphHandle->SetGraphMetadata(Metadata);
 
+			FGraphBuilder::RegisterGraphWithFrontend(InMetaSoundReferenced);
+
 			const FMetasoundAssetBase* ReferencedAsset = IMetasoundUObjectRegistry::Get().GetObjectAsAssetBase(&InMetaSoundReferenced);
 			check(ReferencedAsset);
+
 			FRebuildPresetRootGraph(ReferencedAsset->GetDocumentHandle()).Transform(PresetAsset->GetDocumentHandle());
+			PresetAsset->ConformObjectDataToArchetype();
 		}
 
 		void FGraphBuilder::DeleteVariableNodeHandle(UMetasoundEditorGraphVariable& InVariable)
@@ -962,15 +952,6 @@ namespace Metasound
 			{
 				SynchronizeGraph(Graph->GetMetasoundChecked());
 			}
-			// SynchronizeGraph calls rebuild references, but it needs to be
-			// called to remove the given reference irrespective of error state.
-			else
-			{
-				FMetasoundAssetBase* MetaSoundAsset = IMetasoundUObjectRegistry::Get().GetObjectAsAssetBase(&Graph->GetMetasoundChecked());
-				check(MetaSoundAsset);
-
-				MetaSoundAsset->RebuildReferencedAssets(UMetaSoundAssetSubsystem::Get());
-			}
 
 			return bSuccess;
 		}
@@ -1039,16 +1020,23 @@ namespace Metasound
 			FMetasoundAssetBase* MetaSoundAsset = IMetasoundUObjectRegistry::Get().GetObjectAsAssetBase(&InMetaSound);
 			check(MetaSoundAsset);
 
-			MetaSoundAsset->RegisterGraphWithFrontend();
+			FMetaSoundAssetRegistrationOptions RegOptions;
+			RegOptions.bForceReregister = true;
+			RegOptions.bRegisterDependencies = true;
+			MetaSoundAsset->RegisterGraphWithFrontend(RegOptions);
 
 			if (GEditor)
 			{
 				TArray<UObject*> EditedAssets = GEditor->GetEditorSubsystem<UAssetEditorSubsystem>()->GetAllEditedAssets();
 				for (UObject* Asset : EditedAssets)
 				{
-					if (Asset != &InMetaSound && IMetasoundUObjectRegistry::Get().IsRegisteredClass(Asset))
+					if (Asset != &InMetaSound)
 					{
-						ValidateGraph(*Asset, true /* bAutoUpdate */);
+						if (FMetasoundAssetBase* EditedMetasound = IMetasoundUObjectRegistry::Get().GetObjectAsAssetBase(Asset))
+						{
+							EditedMetasound->RegisterGraphWithFrontend(RegOptions);
+							ValidateGraph(*Asset);
+						}
 					}
 				}
 			}
@@ -1071,7 +1059,7 @@ namespace Metasound
 				{
 					if (Asset != &InMetaSound && IMetasoundUObjectRegistry::Get().IsRegisteredClass(Asset))
 					{
-						ValidateGraph(*Asset, true /* bAutoUpdate */);
+						ValidateGraph(*Asset);
 					}
 				}
 			}
@@ -1431,8 +1419,6 @@ namespace Metasound
 
 			// Synchronize connections.
 			bIsEditorGraphDirty |= SynchronizeConnections(InMetaSound);
-
-			MetaSoundAsset->RebuildReferencedAssets(UMetaSoundAssetSubsystem::Get());
 			return bIsEditorGraphDirty;
 		}
 
diff --git a/Engine/Plugins/Runtime/Metasound/Source/MetasoundEditor/Private/MetasoundEditorGraphBuilder.h b/Engine/Plugins/Runtime/Metasound/Source/MetasoundEditor/Private/MetasoundEditorGraphBuilder.h
--- a/Engine/Plugins/Runtime/Metasound/Source/MetasoundEditor/Private/MetasoundEditorGraphBuilder.h
+++ b/Engine/Plugins/Runtime/Metasound/Source/MetasoundEditor/Private/MetasoundEditorGraphBuilder.h
@@ -60,8 +60,8 @@ namespace Metasound
 			static TSharedPtr<FEditor> GetEditorForMetasound(const UObject& MetaSound);
 			static TSharedPtr<FEditor> GetEditorForGraph(const UEdGraph& EdGraph);
 
-			// Validates MetaSound and optionally performs AutoUpdate logic in Frontend with any additional logic required to refresh editor & respective editor object state.
-			static bool ValidateGraph(UObject& InMetaSound, bool bAutoUpdate = true, bool bInClearUpdateNotes = false);
+			// Validates MetaSound graph.
+			static bool ValidateGraph(UObject& InMetaSound, bool bInClearUpdateNotes = false);
 
 			// Wraps RegisterGraphWithFrontend logic in Frontend with any additional logic required to refresh editor & respective editor object state.
 			static void RegisterGraphWithFrontend(UObject& InMetaSound);
diff --git a/Engine/Plugins/Runtime/Metasound/Source/MetasoundEditor/Private/MetasoundEditorGraphNode.cpp b/Engine/Plugins/Runtime/Metasound/Source/MetasoundEditor/Private/MetasoundEditorGraphNode.cpp
--- a/Engine/Plugins/Runtime/Metasound/Source/MetasoundEditor/Private/MetasoundEditorGraphNode.cpp
+++ b/Engine/Plugins/Runtime/Metasound/Source/MetasoundEditor/Private/MetasoundEditorGraphNode.cpp
@@ -19,6 +19,7 @@
 #include "MetasoundFrontend.h"
 #include "MetasoundFrontendRegistries.h"
 #include "MetasoundFrontendSearchEngine.h"
+#include "MetasoundLiteral.h"
 #include "MetasoundNodeInterface.h"
 #include "MetasoundUObjectRegistry.h"
 #include "ScopedTransaction.h"
@@ -271,72 +272,81 @@ void UMetasoundEditorGraphNode::PinDefaultValueChanged(UEdGraphPin* Pin)
 	}
 }
 
-FString UMetasoundEditorGraphNode::GetPinMetaData(FName InPinName, FName InKey)
+Metasound::Frontend::FDataTypeRegistryInfo UMetasoundEditorGraphNode::GetPinDataTypeInfo(const UEdGraphPin& InPin) const
 {
 	using namespace Metasound::Editor;
 	using namespace Metasound::Frontend;
 
-	if (InKey == "DisallowedClasses")
+	FMetasoundFrontendRegistryContainer* Registry = FMetasoundFrontendRegistryContainer::Get();
+	if (!ensure(Registry))
 	{
-		if (UEdGraphPin* Pin = FindPin(InPinName, EGPD_Input))
-		{
-			FInputHandle Handle = FGraphBuilder::GetInputHandleFromPin(Pin);
+		return FDataTypeRegistryInfo();
+	}
 
-			FMetasoundFrontendRegistryContainer* Registry = FMetasoundFrontendRegistryContainer::Get();
-			if (!ensure(Registry))
-			{
-				return FString();
-			}
+	Metasound::Frontend::FDataTypeRegistryInfo DataTypeInfo;
+	FConstInputHandle Handle = FGraphBuilder::GetConstInputHandleFromPin(&InPin);
+	if (!ensure(Registry->GetInfoForDataType(Handle->GetDataType(), DataTypeInfo)))
+	{
+		return FDataTypeRegistryInfo();
+	}
 
-			Metasound::Frontend::FDataTypeRegistryInfo DataTypeInfo;
-			if (!ensure(Registry->GetInfoForDataType(Handle->GetDataType(), DataTypeInfo)))
-			{
-				return FString();
-			}
+	return DataTypeInfo;
+}
 
-			const EMetasoundFrontendLiteralType LiteralType = GetMetasoundFrontendLiteralType(DataTypeInfo.PreferredLiteralType);
-			if (LiteralType != EMetasoundFrontendLiteralType::UObject && LiteralType != EMetasoundFrontendLiteralType::UObjectArray)
-			{
-				return FString();
-			}
+TSet<FString> UMetasoundEditorGraphNode::GetDisallowedPinClassNames(const UEdGraphPin& InPin) const
+{
+	using namespace Metasound::Editor;
+	using namespace Metasound::Frontend;
 
-			UClass* ProxyGenClass = DataTypeInfo.ProxyGeneratorClass;
-			if (!ProxyGenClass)
-			{
-				return FString();
-			}
+	const FDataTypeRegistryInfo DataTypeInfo = GetPinDataTypeInfo(InPin);
+	if (DataTypeInfo.PreferredLiteralType != Metasound::ELiteralType::UObjectProxy)
+	{
+		return { };
+	}
 
-			FString DisallowedClasses;
-			const FName ClassName = ProxyGenClass->GetFName();
-			for (TObjectIterator<UClass> ClassIt; ClassIt; ++ClassIt)
-			{
-				UClass* Class = *ClassIt;
-				if (!Class->IsNative())
-				{
-					continue;
-				}
+	UClass* ProxyGenClass = DataTypeInfo.ProxyGeneratorClass;
+	if (!ProxyGenClass)
+	{
+		return { };
+	}
 
-				if (Class->HasAnyClassFlags(CLASS_Abstract | CLASS_Deprecated | CLASS_NewerVersionExists))
-				{
-					continue;
-				}
+	TSet<FString> DisallowedClasses;
+	const FName ClassName = ProxyGenClass->GetFName();
+	for (TObjectIterator<UClass> ClassIt; ClassIt; ++ClassIt)
+	{
+		UClass* Class = *ClassIt;
+		if (!Class->IsNative())
+		{
+			continue;
+		}
 
-				if (ClassIt->GetFName() == ClassName)
-				{
-					continue;
-				}
+		if (Class->HasAnyClassFlags(CLASS_Abstract | CLASS_Deprecated | CLASS_NewerVersionExists))
+		{
+			continue;
+		}
 
-				if (Class->IsChildOf(ProxyGenClass))
-				{
-					if (!DisallowedClasses.IsEmpty())
-					{
-						DisallowedClasses += TEXT(",");
-					}
-					DisallowedClasses += *ClassIt->GetName();
-				}
-			}
+		if (ClassIt->GetFName() == ClassName)
+		{
+			continue;
+		}
+
+		if (Class->IsChildOf(ProxyGenClass))
+		{
+			DisallowedClasses.Add(ClassIt->GetName());
+		}
+	}
+
+	return DisallowedClasses;
+}
 
-			return DisallowedClasses;
+FString UMetasoundEditorGraphNode::GetPinMetaData(FName InPinName, FName InKey)
+{
+	if (InKey == "DisallowedClasses")
+	{
+		if (UEdGraphPin* Pin = FindPin(InPinName, EGPD_Input))
+		{
+			TSet<FString> DisallowedClasses = GetDisallowedPinClassNames(*Pin);
+			return FString::Join(DisallowedClasses.Array(), TEXT(","));
 		}
 
 		return FString();
@@ -578,7 +588,7 @@ bool UMetasoundEditorGraphExternalNode::CanAutoUpdate() const
 {
 	using namespace Metasound::Frontend;
 
-	return GetConstNodeHandle()->CanAutoUpdate(UMetaSoundAssetSubsystem::Get());
+	return GetConstNodeHandle()->CanAutoUpdate();
 }
 
 UMetasoundEditorGraphExternalNode* UMetasoundEditorGraphExternalNode::UpdateToVersion(const FMetasoundFrontendVersionNumber& InNewVersion, bool bInPropagateErrorMessages)
@@ -757,19 +767,19 @@ bool UMetasoundEditorGraphExternalNode::Validate(Metasound::Editor::FGraphNodeVa
 		if (Pin->Direction == EGPD_Input)
 		{
 			FInputHandle Input = FGraphBuilder::GetInputHandleFromPin(Pin);
-			bWasRemoved |= InterfaceUpdates.RemovedInputs.ContainsByPredicate([&](const FMetasoundFrontendClassInput& ClassInput)
+			bWasRemoved |= InterfaceUpdates.RemovedInputs.ContainsByPredicate([&](const FMetasoundFrontendClassInput* ClassInput)
 			{
-				return Input->GetName() == ClassInput.Name && Input->GetDataType() == ClassInput.TypeName;
+				return Input->GetName() == ClassInput->Name && Input->GetDataType() == ClassInput->TypeName;
 			});
 		}
 
 		if (Pin->Direction == EGPD_Output)
 		{
 			FOutputHandle Output = FGraphBuilder::GetOutputHandleFromPin(Pin);
-			bWasRemoved |= InterfaceUpdates.RemovedOutputs.ContainsByPredicate([&](const FMetasoundFrontendClassOutput& ClassOutput)
-				{
-					return Output->GetName() == ClassOutput.Name && Output->GetDataType() == ClassOutput.TypeName;
-				});
+			bWasRemoved |= InterfaceUpdates.RemovedOutputs.ContainsByPredicate([&](const FMetasoundFrontendClassOutput* ClassOutput)
+			{
+				return Output->GetName() == ClassOutput->Name && Output->GetDataType() == ClassOutput->TypeName;
+			});
 		}
 
 		if (Pin->bOrphanedPin != bWasRemoved)
diff --git a/Engine/Plugins/Runtime/Metasound/Source/MetasoundEditor/Private/MetasoundEditorGraphSchema.cpp b/Engine/Plugins/Runtime/Metasound/Source/MetasoundEditor/Private/MetasoundEditorGraphSchema.cpp
--- a/Engine/Plugins/Runtime/Metasound/Source/MetasoundEditor/Private/MetasoundEditorGraphSchema.cpp
+++ b/Engine/Plugins/Runtime/Metasound/Source/MetasoundEditor/Private/MetasoundEditorGraphSchema.cpp
@@ -27,6 +27,7 @@
 #include "MetasoundFrontendNodesCategories.h"
 #include "MetasoundFrontendRegistries.h"
 #include "MetasoundFrontendSearchEngine.h"
+#include "MetasoundLiteral.h"
 #include "MetasoundStandardNodesCategories.h"
 #include "ScopedTransaction.h"
 #include "Styling/SlateStyleRegistry.h"
@@ -876,6 +877,46 @@ bool UMetasoundEditorGraphSchema::TryCreateConnection(UEdGraphPin* PinA, UEdGrap
 	return true;
 }
 
+void UMetasoundEditorGraphSchema::TrySetDefaultObject(UEdGraphPin& Pin, UObject* NewDefaultObject, bool bInMarkAsModified) const
+{
+	using namespace Metasound::Editor;
+
+	if (UMetasoundEditorGraphNode* Node = Cast<UMetasoundEditorGraphNode>(Pin.GetOwningNode()))
+	{
+		if (Node->GetPinDataTypeInfo(Pin).PreferredLiteralType == Metasound::ELiteralType::UObjectProxy)
+		{
+			TrySetDefaultValue(Pin, NewDefaultObject ? NewDefaultObject->GetPathName() : FString(), bInMarkAsModified);
+			return;
+		}
+	}
+
+	Super::TrySetDefaultObject(Pin, NewDefaultObject, bInMarkAsModified);
+}
+
+void UMetasoundEditorGraphSchema::TrySetDefaultValue(UEdGraphPin& Pin, const FString& InNewDefaultValue, bool bInMarkAsModified) const
+{
+	if (UMetasoundEditorGraphNode* Node = Cast<UMetasoundEditorGraphNode>(Pin.GetOwningNode()))
+	{
+		if (Node->GetPinDataTypeInfo(Pin).PreferredLiteralType == Metasound::ELiteralType::UObjectProxy)
+		{
+			FSoftObjectPath Path = InNewDefaultValue;
+			const TSet<FString> DisallowedClassNames = Node->GetDisallowedPinClassNames(Pin);
+			if (UObject* Object = Path.TryLoad())
+			{
+				if (UClass* Class = Object->GetClass())
+				{
+					if (DisallowedClassNames.Contains(Class->GetName()))
+					{
+						return;
+					}
+				}
+			}
+		}
+	}
+
+	return Super::TrySetDefaultValue(Pin, InNewDefaultValue, bInMarkAsModified);
+}
+
 bool UMetasoundEditorGraphSchema::ShouldHidePinDefaultValue(UEdGraphPin* Pin) const
 {
 	// TODO: Determine if should be hidden from doc data
@@ -1035,6 +1076,9 @@ void UMetasoundEditorGraphSchema::DroppedAssetsOnGraph(const TArray<FAssetData>&
 	MetaSound.Modify();
 	Graph->Modify();
 
+	FMetasoundAssetBase* MetaSoundAsset = IMetasoundUObjectRegistry::Get().GetObjectAsAssetBase(&MetaSound);
+	check(MetaSoundAsset);
+
 	for (const FAssetData& DroppedAsset : Assets)
 	{
 		if (UObject* DroppedObject = DroppedAsset.GetAsset())
@@ -1045,15 +1089,19 @@ void UMetasoundEditorGraphSchema::DroppedAssetsOnGraph(const TArray<FAssetData>&
 				continue;
 			}
 
-			// This may not be necessary as dropping an asset on the graph may load it, thus triggering the registration from the MetaSoundAssetSubsystem.
+			if (MetaSoundAsset->AddingReferenceCausesLoop(FSoftObjectPath(DroppedAsset.ObjectPath)))
+			{
+				continue;
+			}
+
+			// This may not be necessary as dropping an asset on the graph may load it, thus triggering the registration from the MetaSoundAssetManager.
 			const FNodeClassInfo ClassInfo = DroppedMetaSoundAsset->GetAssetClassInfo();
 			const FNodeRegistryKey RegistryKey = NodeRegistryKey::CreateKey(ClassInfo);
 			if (ensure(IsValidNodeRegistryKey(RegistryKey)))
 			{
-				if (!FMetasoundFrontendRegistryContainer::Get()->IsNodeRegistered(RegistryKey))
-				{
-					DroppedMetaSoundAsset->RegisterGraphWithFrontend();
-				}
+				FMetaSoundAssetRegistrationOptions RegOptions;
+				RegOptions.bForceReregister = false;
+				DroppedMetaSoundAsset->RegisterGraphWithFrontend(RegOptions);
 			}
 
 			const FMetasoundFrontendClassName ClassName = DroppedMetaSoundAsset->GetAssetClassInfo().ClassName;
@@ -1178,14 +1226,14 @@ void UMetasoundEditorGraphSchema::GetFunctionActions(FGraphActionMenuBuilder& Ac
 	using namespace Metasound::Editor::SchemaPrivate;
 	using namespace Metasound::Frontend;
 
-	const UMetaSoundAssetSubsystem& AssetSubsystem = UMetaSoundAssetSubsystem::Get();
+	const IMetaSoundAssetManager& AssetManager = IMetaSoundAssetManager::GetChecked();
 	FMetasoundAssetBase* ParentAsset = nullptr;
 	if (InGraphHandle->IsValid())
 	{
 		FMetasoundFrontendClassMetadata AssetMetadata = InGraphHandle->GetGraphMetadata();
 		AssetMetadata.SetType(EMetasoundFrontendClassType::External);
 		const FNodeRegistryKey RegistryKey = NodeRegistryKey::CreateKey(AssetMetadata);
-		ParentAsset = AssetSubsystem.FindAssetFromKey(RegistryKey);
+		ParentAsset = AssetManager.FindAssetFromKey(RegistryKey);
 	}
 
 	const bool bIncludeDeprecated = static_cast<bool>(EnableDeprecatedMetaSoundNodeClassCreationCVar);
@@ -1211,10 +1259,10 @@ void UMetasoundEditorGraphSchema::GetFunctionActions(FGraphActionMenuBuilder& Ac
 		if (ParentAsset)
 		{
 			const FNodeRegistryKey RegistryKey = NodeRegistryKey::CreateKey(FrontendClass.Metadata);
-			Path = AssetSubsystem.FindObjectPathFromKey(RegistryKey);
+			Path = AssetManager.FindObjectPathFromKey(RegistryKey);
 			if (Path)
 			{
-				if (ParentAsset->AddingReferenceCausesLoop(*Path, AssetSubsystem))
+				if (ParentAsset->AddingReferenceCausesLoop(*Path))
 				{
 					continue;
 				}
diff --git a/Engine/Plugins/Runtime/Metasound/Source/MetasoundEditor/Private/MetasoundEditorModule.cpp b/Engine/Plugins/Runtime/Metasound/Source/MetasoundEditor/Private/MetasoundEditorModule.cpp
--- a/Engine/Plugins/Runtime/Metasound/Source/MetasoundEditor/Private/MetasoundEditorModule.cpp
+++ b/Engine/Plugins/Runtime/Metasound/Source/MetasoundEditor/Private/MetasoundEditorModule.cpp
@@ -154,11 +154,21 @@ namespace Metasound
 
 				if (IsMetaSoundAssetClass(InAssetData.AssetClass))
 				{
-					// Use the editor version of UnregisterWithFrontend so it refreshes any open MetaSound editors
-					UMetaSoundAssetSubsystem::Get().AddOrUpdateAsset(InAssetData, false /* bRegisterWithFrontend */);
-					if (UObject* AssetObject = InAssetData.GetAsset())
+					// in case they need to be 
+					check(GEngine);
+					UMetaSoundAssetSubsystem* AssetSubsystem = GEngine->GetEngineSubsystem<UMetaSoundAssetSubsystem>();
+					check(AssetSubsystem);
+
+					// Use the editor version of `RegisterGraphWithFrontend` so it re-registers any open MetaSound editors
+					AssetSubsystem->AddOrUpdateAsset(InAssetData, false /* bRegisterWithFrontend */);
+					
+					// May no longer be necessary, but left for safety
+					if (InAssetData.IsAssetLoaded())
 					{
-						FGraphBuilder::RegisterGraphWithFrontend(*AssetObject);
+						if (UObject* AssetObject = InAssetData.GetAsset())
+						{
+							FGraphBuilder::RegisterGraphWithFrontend(*AssetObject);
+						}
 					}
 				}
 			}
@@ -185,8 +195,12 @@ namespace Metasound
 					{
 						if (IsMetaSoundAssetClass(Obj->GetClass()->GetFName()))
 						{
+							check(GEngine);
+							UMetaSoundAssetSubsystem* AssetSubsystem = GEngine->GetEngineSubsystem<UMetaSoundAssetSubsystem>();
+							check(AssetSubsystem);
+
 							// Use the editor version of UnregisterWithFrontend so it refreshes any open MetaSound editors
-							UMetaSoundAssetSubsystem::Get().RemoveAsset(*Pair.Key, false /* bUnregisterWithFrontend */);
+							AssetSubsystem->RemoveAsset(*Pair.Key, false /* bUnregisterWithFrontend */);
 							FGraphBuilder::UnregisterGraphWithFrontend(*Pair.Key);
 						}
 					}
@@ -195,8 +209,11 @@ namespace Metasound
 					{
 						if (IsMetaSoundAssetClass(Obj->GetClass()->GetFName()))
 						{
+							check(GEngine);
+							UMetaSoundAssetSubsystem* AssetSubsystem = GEngine->GetEngineSubsystem<UMetaSoundAssetSubsystem>();
+							check(AssetSubsystem);
 							// Use the editor version of RegisterWithFrontend so it refreshes any open MetaSound editors
-							UMetaSoundAssetSubsystem::Get().AddOrUpdateAsset(*Pair.Value, false /* bRegisterWithFrontend */);
+							AssetSubsystem->AddOrUpdateAsset(*Pair.Value, false /* bRegisterWithFrontend */);
 							FGraphBuilder::RegisterGraphWithFrontend(*Pair.Value);
 						}
 					}
@@ -224,8 +241,12 @@ namespace Metasound
 			{
 				if (IsMetaSoundAssetClass(InAssetData.AssetClass))
 				{
+					check(GEngine);
+					UMetaSoundAssetSubsystem* AssetSubsystem = GEngine->GetEngineSubsystem<UMetaSoundAssetSubsystem>();
+					check(AssetSubsystem);
+
 					// Use the editor version of UnregisterWithFrontend so it refreshes any open MetaSound editors
-					UMetaSoundAssetSubsystem::Get().RemoveAsset(InAssetData, false /* bUnregisterWithFrontend */);
+					AssetSubsystem->RemoveAsset(InAssetData, false /* bUnregisterWithFrontend */);
 					if (UObject* AssetObject = InAssetData.GetAsset())
 					{
 						FGraphBuilder::UnregisterGraphWithFrontend(*AssetObject);
@@ -237,7 +258,11 @@ namespace Metasound
 			{
 				if (IsMetaSoundAssetClass(InAssetData.AssetClass))
 				{
-					UMetaSoundAssetSubsystem::Get().RenameAsset(InAssetData, false /* bReregisterWithFrontend */);
+					check(GEngine);
+					UMetaSoundAssetSubsystem* AssetSubsystem = GEngine->GetEngineSubsystem<UMetaSoundAssetSubsystem>();
+					check(AssetSubsystem);
+
+					AssetSubsystem->RenameAsset(InAssetData, false /* bReregisterWithFrontend */);
 					if (UObject* AssetObject = InAssetData.GetAsset())
 					{
 						FGraphBuilder::RegisterGraphWithFrontend(*AssetObject);
diff --git a/Engine/Plugins/Runtime/Metasound/Source/MetasoundEditor/Private/MetasoundNodeDetailCustomization.cpp b/Engine/Plugins/Runtime/Metasound/Source/MetasoundEditor/Private/MetasoundNodeDetailCustomization.cpp
--- a/Engine/Plugins/Runtime/Metasound/Source/MetasoundEditor/Private/MetasoundNodeDetailCustomization.cpp
+++ b/Engine/Plugins/Runtime/Metasound/Source/MetasoundEditor/Private/MetasoundNodeDetailCustomization.cpp
@@ -292,9 +292,6 @@ namespace Metasound
 				return FString();
 			};
 
-			TArray<const UClass*> AllowedClasses;
-			AllowedClasses.Add(ProxyGenClass.Get());
-
 			auto FilterAsset = [InProxyGenClass = ProxyGenClass](const FAssetData& InAsset)
 			{
 				if (InProxyGenClass.IsValid())
@@ -312,16 +309,16 @@ namespace Metasound
 			};
 
 			return SNew(SObjectPropertyEntryBox)
-				.ObjectPath_Lambda(GetAssetPath)
-				.AllowedClass(ProxyGenClass.Get())
-				.OnShouldSetAsset_Lambda(ValidateAsset)
-				.OnShouldFilterAsset_Lambda(FilterAsset)
-				.PropertyHandle(PropertyHandle)
 				.AllowClear(true)
-				.DisplayUseSelected(true)
+				.AllowedClass(ProxyGenClass.Get())
 				.DisplayBrowse(true)
 				.DisplayThumbnail(true)
-				.NewAssetFactories(PropertyCustomizationHelpers::GetNewAssetFactoriesForClasses(AllowedClasses));
+				.DisplayUseSelected(true)
+				.NewAssetFactories(PropertyCustomizationHelpers::GetNewAssetFactoriesForClasses({ ProxyGenClass.Get() }))
+				.ObjectPath_Lambda(GetAssetPath)
+				.OnShouldFilterAsset_Lambda(FilterAsset)
+				.OnShouldSetAsset_Lambda(ValidateAsset)
+				.PropertyHandle(PropertyHandle);
 		}
 
 		TSharedRef<SWidget> FMetasoundInputArrayDetailCustomizationBase::CreateNameWidget(TSharedPtr<IPropertyHandle> StructPropertyHandle) const
diff --git a/Engine/Plugins/Runtime/Metasound/Source/MetasoundEditor/Public/MetasoundEditorGraph.h b/Engine/Plugins/Runtime/Metasound/Source/MetasoundEditor/Public/MetasoundEditorGraph.h
--- a/Engine/Plugins/Runtime/Metasound/Source/MetasoundEditor/Public/MetasoundEditorGraph.h
+++ b/Engine/Plugins/Runtime/Metasound/Source/MetasoundEditor/Public/MetasoundEditorGraph.h
@@ -159,7 +159,7 @@ public:
 	bool IsPreviewing() const;
 
 	virtual bool Synchronize() override;
-	virtual bool Validate(bool bInAutoUpdate, bool bInClearUpdateNotes) override;
+	virtual bool Validate(bool bInClearUpdateNotes) override;
 
 	virtual void RegisterGraphWithFrontend() override;
 
diff --git a/Engine/Plugins/Runtime/Metasound/Source/MetasoundEditor/Public/MetasoundEditorGraphNode.h b/Engine/Plugins/Runtime/Metasound/Source/MetasoundEditor/Public/MetasoundEditorGraphNode.h
--- a/Engine/Plugins/Runtime/Metasound/Source/MetasoundEditor/Public/MetasoundEditorGraphNode.h
+++ b/Engine/Plugins/Runtime/Metasound/Source/MetasoundEditor/Public/MetasoundEditorGraphNode.h
@@ -88,6 +88,9 @@ public:
 	Metasound::Frontend::FConstGraphHandle GetConstRootGraphHandle() const;
 	Metasound::Frontend::FNodeHandle GetNodeHandle() const;
 	Metasound::Frontend::FConstNodeHandle GetConstNodeHandle() const;
+	Metasound::Frontend::FDataTypeRegistryInfo GetPinDataTypeInfo(const UEdGraphPin& InPin) const;
+
+	TSet<FString> GetDisallowedPinClassNames(const UEdGraphPin& InPin) const;
 
 	virtual FMetasoundFrontendClassName GetClassName() const { return FMetasoundFrontendClassName(); }
 	virtual FGuid GetNodeID() const { return FGuid(); }
diff --git a/Engine/Plugins/Runtime/Metasound/Source/MetasoundEditor/Public/MetasoundEditorGraphSchema.h b/Engine/Plugins/Runtime/Metasound/Source/MetasoundEditor/Public/MetasoundEditorGraphSchema.h
--- a/Engine/Plugins/Runtime/Metasound/Source/MetasoundEditor/Public/MetasoundEditorGraphSchema.h
+++ b/Engine/Plugins/Runtime/Metasound/Source/MetasoundEditor/Public/MetasoundEditorGraphSchema.h
@@ -259,6 +259,8 @@ class UMetasoundEditorGraphSchema : public UEdGraphSchema
 	virtual void CreateDefaultNodesForGraph(UEdGraph& Graph) const override;
 	virtual const FPinConnectionResponse CanCreateConnection(const UEdGraphPin* A, const UEdGraphPin* B) const override;
 	virtual bool TryCreateConnection(UEdGraphPin* A, UEdGraphPin* B) const override;
+	virtual void TrySetDefaultObject(UEdGraphPin& Pin, UObject* NewDefaultObject, bool bInMarkAsModified) const override;
+	virtual void TrySetDefaultValue(UEdGraphPin& Pin, const FString& InNewDefaultValue, bool bInMarkAsModified) const override;
 	virtual bool ShouldHidePinDefaultValue(UEdGraphPin* Pin) const override;
 	virtual FLinearColor GetPinTypeColor(const FEdGraphPinType& PinType) const override;
 	virtual void BreakNodeLinks(UEdGraphNode& TargetNode) const override;
diff --git a/Engine/Plugins/Runtime/Metasound/Source/MetasoundEngine/Private/Metasound.cpp b/Engine/Plugins/Runtime/Metasound/Source/MetasoundEngine/Private/Metasound.cpp
--- a/Engine/Plugins/Runtime/Metasound/Source/MetasoundEngine/Private/Metasound.cpp
+++ b/Engine/Plugins/Runtime/Metasound/Source/MetasoundEngine/Private/Metasound.cpp
@@ -12,6 +12,7 @@
 #include "MetasoundFrontendController.h"
 #include "MetasoundFrontendQuery.h"
 #include "MetasoundFrontendQuerySteps.h"
+#include "MetasoundFrontendRegistries.h"
 #include "MetasoundFrontendSearchEngine.h"
 #include "MetasoundGenerator.h"
 #include "MetasoundInstanceTransmitter.h"
@@ -114,6 +115,16 @@ Metasound::Frontend::FNodeClassInfo UMetaSound::GetAssetClassInfo() const
 	return { GetDocumentChecked().RootGraph, *GetPathName() };
 }
 
+void UMetaSound::SetReferencedAssetClassKeys(TSet<Metasound::Frontend::FNodeRegistryKey>&& InKeys)
+{
+	ReferencedAssetClassKeys = MoveTemp(InKeys);
+}
+
+TSet<UObject*>& UMetaSound::GetReferencedAssetClassCache()
+{
+	return ReferenceAssetClassCache;
+}
+
 const FMetasoundFrontendVersion& UMetaSound::GetDefaultArchetypeVersion() const
 {
 	static const FMetasoundFrontendVersion DefaultArchetypeVersion = Metasound::Engine::MetasoundV1_0::GetVersion();
@@ -129,5 +140,4 @@ const TArray<FMetasoundFrontendVersion>& UMetaSound::GetSupportedArchetypeVersio
 
 	return Supported;
 }
-
 #undef LOCTEXT_NAMESPACE // MetaSound
diff --git a/Engine/Plugins/Runtime/Metasound/Source/MetasoundEngine/Private/MetasoundSource.cpp b/Engine/Plugins/Runtime/Metasound/Source/MetasoundEngine/Private/MetasoundSource.cpp
--- a/Engine/Plugins/Runtime/Metasound/Source/MetasoundEngine/Private/MetasoundSource.cpp
+++ b/Engine/Plugins/Runtime/Metasound/Source/MetasoundEngine/Private/MetasoundSource.cpp
@@ -112,7 +112,6 @@ void UMetaSoundSource::PostEditChangeProperty(FPropertyChangedEvent& InEvent)
 		{
 			case EMetasoundSourceAudioFormat::Mono:
 			{
-				NumChannels = 1;
 				// TODO: utilize latest version `MetasoudnSourceMono`
 				bDidModifyDocument = Metasound::Frontend::FMatchRootGraphToArchetype(Metasound::Engine::MetasoundSourceMono::GetVersion()).Transform(GetDocumentHandle());
 			}
@@ -120,7 +119,6 @@ void UMetaSoundSource::PostEditChangeProperty(FPropertyChangedEvent& InEvent)
 
 			case EMetasoundSourceAudioFormat::Stereo:
 			{
-				NumChannels = 2;
 				bDidModifyDocument = Metasound::Frontend::FMatchRootGraphToArchetype(Metasound::Engine::MetasoundSourceStereo::GetVersion()).Transform(GetDocumentHandle());
 			}
 			break;
@@ -134,12 +132,47 @@ void UMetaSoundSource::PostEditChangeProperty(FPropertyChangedEvent& InEvent)
 
 		if (bDidModifyDocument)
 		{
+			ConformObjectDataToArchetype();
+
+			// Use the editor form of register to ensure other editors'
+			// MetaSounds are auto-updated if they are referencing this graph.
+			if (Graph)
+			{
+				Graph->RegisterGraphWithFrontend();
+			}
 			MarkMetasoundDocumentDirty();
 		}
 	}
 }
 #endif // WITH_EDITOR
 
+bool UMetaSoundSource::ConformObjectDataToArchetype()
+{
+	const FMetasoundFrontendVersion& ArchetypeVersion = GetDocumentHandle()->GetArchetypeVersion();
+	if (ArchetypeVersion == Metasound::Engine::MetasoundSourceMono::GetVersion())
+	{
+		if (OutputFormat != EMetasoundSourceAudioFormat::Mono || NumChannels != 1)
+		{
+			OutputFormat = EMetasoundSourceAudioFormat::Mono;
+			NumChannels = 1;
+			return true;
+		}
+	}
+
+	if (ArchetypeVersion == Metasound::Engine::MetasoundSourceStereo::GetVersion())
+	{
+		if (OutputFormat != EMetasoundSourceAudioFormat::Stereo || NumChannels != 2)
+		{
+			OutputFormat = EMetasoundSourceAudioFormat::Stereo;
+			NumChannels = 2;
+			return true;
+		}
+	}
+
+	return false;
+}
+
+
 void UMetaSoundSource::BeginDestroy()
 {
 	UnregisterGraphWithFrontend();
@@ -158,6 +191,16 @@ void UMetaSoundSource::Serialize(FArchive& InArchive)
 	Metasound::SerializeToArchive(*this, InArchive);
 }
 
+void UMetaSoundSource::SetReferencedAssetClassKeys(TSet<Metasound::Frontend::FNodeRegistryKey>&& InKeys)
+{
+	ReferencedAssetClassKeys = MoveTemp(InKeys);
+}
+
+TSet<UObject*>& UMetaSoundSource::GetReferencedAssetClassCache()
+{
+	return ReferenceAssetClassCache;
+}
+
 #if WITH_EDITORONLY_DATA
 UEdGraph* UMetaSoundSource::GetGraph()
 {
@@ -193,6 +236,16 @@ void UMetaSoundSource::SetRegistryAssetClassInfo(const Metasound::Frontend::FNod
 }
 #endif // WITH_EDITORONLY_DATA
 
+void UMetaSoundSource::InitResources()
+{
+	METASOUND_TRACE_CPUPROFILER_EVENT_SCOPE(UMetaSoundAssetSubsystem::AddOrUpdateAsset);
+
+	FMetaSoundAssetRegistrationOptions RegOptions;
+	RegOptions.bForceReregister = false;
+	RegOptions.bRegisterDependencies = true;
+	RegisterGraphWithFrontend(RegOptions);
+}
+
 Metasound::Frontend::FNodeClassInfo UMetaSoundSource::GetAssetClassInfo() const
 {
 	return { GetDocumentChecked().RootGraph, *GetPathName() };
@@ -228,7 +281,7 @@ ISoundGeneratorPtr UMetaSoundSource::CreateSoundGenerator(const FSoundGeneratorI
 	using namespace Metasound::Frontend;
 	using namespace Metasound::Engine;
 
-	METASOUND_TRACE_CPUPROFILER_EVENT_SCOPE(MetaSoundSource::CreateSoundGenerator);
+	METASOUND_TRACE_CPUPROFILER_EVENT_SCOPE(UMetaSoundSource::CreateSoundGenerator);
 
 	Duration = INDEFINITELY_LOOPING_DURATION;
 	bLooping = true;
@@ -280,10 +333,10 @@ Metasound::FOperatorSettings UMetaSoundSource::GetOperatorSettings(Metasound::FS
 const TArray<FMetasoundFrontendVersion>& UMetaSoundSource::GetSupportedArchetypeVersions() const 
 {
 	static const TArray<FMetasoundFrontendVersion> Supported(
-		{
-			Metasound::Engine::MetasoundSourceMono::GetVersion(),
-			Metasound::Engine::MetasoundSourceStereo::GetVersion()
-		});
+	{
+		Metasound::Engine::MetasoundSourceMono::GetVersion(),
+		Metasound::Engine::MetasoundSourceStereo::GetVersion()
+	});
 
 	return Supported;
 }
@@ -301,6 +354,7 @@ Metasound::FMetasoundEnvironment UMetaSoundSource::CreateEnvironment() const
 	{
 		DeviceHandle = World->GetAudioDevice();
 	}
+
 	if (!DeviceHandle.IsValid())
 	{
 		if (FAudioDeviceManager* DeviceManager = FAudioDeviceManager::Get())
diff --git a/Engine/Plugins/Runtime/Metasound/Source/MetasoundEngine/Private/MetasoundUObjectRegistry.cpp b/Engine/Plugins/Runtime/Metasound/Source/MetasoundEngine/Private/MetasoundUObjectRegistry.cpp
--- a/Engine/Plugins/Runtime/Metasound/Source/MetasoundEngine/Private/MetasoundUObjectRegistry.cpp
+++ b/Engine/Plugins/Runtime/Metasound/Source/MetasoundEngine/Private/MetasoundUObjectRegistry.cpp
@@ -76,6 +76,7 @@ namespace Metasound
 
 void UMetaSoundAssetSubsystem::Initialize(FSubsystemCollectionBase& InCollection)
 {
+	IMetaSoundAssetManager::Set(*this);
 	FCoreDelegates::OnPostEngineInit.AddUObject(this, &UMetaSoundAssetSubsystem::PostEngineInit);
 }
 
@@ -106,7 +107,7 @@ void UMetaSoundAssetSubsystem::PostInitAssetScan()
 	AssetManager.SearchAssetRegistryPaths(MetaSoundAssets, Rules);
 	for (const FAssetData& AssetData : MetaSoundAssets)
 	{
-		AddOrUpdateAsset(AssetData);
+		AddOrUpdateAsset(AssetData, false /* bRegisterWithFrontend */);
 	}
 
 	Rules.AssetBaseClass = UMetaSoundSource::StaticClass();
@@ -114,14 +115,10 @@ void UMetaSoundAssetSubsystem::PostInitAssetScan()
 	AssetManager.SearchAssetRegistryPaths(MetaSoundSourceAssets, Rules);
 	for (const FAssetData& AssetData : MetaSoundSourceAssets)
 	{
-		AddOrUpdateAsset(AssetData);
+		AddOrUpdateAsset(AssetData, false /* bRegisterWithFrontend */);
 	}
 }
 
-void UMetaSoundAssetSubsystem::Deinitialize()
-{
-}
-
 void UMetaSoundAssetSubsystem::AddOrUpdateAsset(UObject& InObject, bool bInRegisterWithFrontend)
 {
 	using namespace Metasound;
@@ -149,14 +146,11 @@ void UMetaSoundAssetSubsystem::AddOrUpdateAsset(const FAssetData& InAssetData, b
 	using namespace Metasound::AssetSubsystemPrivate;
 	using namespace Metasound::Frontend;
 
-	// TODO: Set to false for builds once registering without loading asset is supported.
-	static const bool bLoadRequiredToRegisterAssetClasses = true;
-
 	METASOUND_TRACE_CPUPROFILER_EVENT_SCOPE(UMetaSoundAssetSubsystem::AddOrUpdateAsset);
 
 	FNodeClassInfo ClassInfo;
 	bool bClassInfoFound = GetAssetClassInfo(InAssetData, ClassInfo);
-	if (!bClassInfoFound || bLoadRequiredToRegisterAssetClasses || bInRegisterWithFrontend)
+	if (!bClassInfoFound || bInRegisterWithFrontend)
 	{
 		UObject* Object = nullptr;
 
@@ -167,7 +161,7 @@ void UMetaSoundAssetSubsystem::AddOrUpdateAsset(const FAssetData& InAssetData, b
 		}
 		else
 		{
-			if (!bLoadRequiredToRegisterAssetClasses)
+			if (!bInRegisterWithFrontend)
 			{
 				UE_LOG(LogMetaSound, Warning,
 					TEXT("Failed to find serialized MetaSound asset registry data for asset '%s'. "
@@ -177,7 +171,7 @@ void UMetaSoundAssetSubsystem::AddOrUpdateAsset(const FAssetData& InAssetData, b
 			Object = Path.TryLoad();
 		}
 
-		if (ensure(Object))
+		if (Object)
 		{
 			AddOrUpdateAsset(*Object, bInRegisterWithFrontend);
 			return;
diff --git a/Engine/Plugins/Runtime/Metasound/Source/MetasoundEngine/Public/Metasound.h b/Engine/Plugins/Runtime/Metasound/Source/MetasoundEngine/Public/Metasound.h
--- a/Engine/Plugins/Runtime/Metasound/Source/MetasoundEngine/Public/Metasound.h
+++ b/Engine/Plugins/Runtime/Metasound/Source/MetasoundEngine/Public/Metasound.h
@@ -30,7 +30,7 @@ public:
 	virtual bool NeedsLoadForEditorGame() const override { return false; }
 
 	virtual bool Synchronize() { return false; }
-	virtual bool Validate(bool bInAutoUpdate, bool bInClearUpdateNotes) { return false; }
+	virtual bool Validate(bool bInClearUpdateNotes) { return false; }
 	virtual void RegisterGraphWithFrontend() { }
 };
 
@@ -53,7 +53,7 @@ namespace Metasound
 	{
 		if (UMetasoundEditorGraphBase* MetaSoundGraph = Cast<UMetasoundEditorGraphBase>(InMetaSound.GetGraph()))
 		{
-			if (MetaSoundGraph->Validate(false /* bAutoUpdate */, false /* bClearUpdateNotes */))
+			if (MetaSoundGraph->Validate(false /* bClearUpdateNotes */))
 			{
 				MetaSoundGraph->RegisterGraphWithFrontend();
 			}
@@ -82,9 +82,8 @@ namespace Metasound
 #if WITH_EDITORONLY_DATA
 		if (UMetasoundEditorGraphBase* MetaSoundGraph = Cast<UMetasoundEditorGraphBase>(InMetaSound.GetGraph()))
 		{
-			const bool bAutoUpdate = false;
 			const bool bClearUpdateNotes = true;
-			const bool bIsValid = MetaSoundGraph->Validate(bAutoUpdate, bClearUpdateNotes);
+			const bool bIsValid = MetaSoundGraph->Validate(bClearUpdateNotes);
 			if (bIsValid)
 			{
 				MetaSoundGraph->RegisterGraphWithFrontend();
@@ -98,12 +97,7 @@ namespace Metasound
 	{
 		if (InArchive.IsLoading())
 		{
-			const bool bUpdateReferences = false;
-			const bool bMarkDirty = false;
-
-			UMetaSoundAssetSubsystem& AssetSubsystem = UMetaSoundAssetSubsystem::Get();
-			InMetaSound.VersionAsset(AssetSubsystem, bMarkDirty, bUpdateReferences);
-			InMetaSound.AutoUpdate(AssetSubsystem, bMarkDirty, bUpdateReferences);
+			InMetaSound.VersionAsset();
 		}
 	}
 
@@ -171,7 +165,10 @@ public:
 	UMetaSound(const FObjectInitializer& ObjectInitializer);
 
 	UPROPERTY()
-	TSet<FSoftObjectPath> ReferencedAssets;
+	TSet<FString> ReferencedAssetClassKeys;
+
+	UPROPERTY(Transient)
+	TSet<UObject*> ReferenceAssetClassCache;
 
 	UPROPERTY(AssetRegistrySearchable)
 	FGuid AssetClassID;
@@ -229,21 +226,20 @@ public:
 	virtual void PreSave(FObjectPreSaveContext InSaveContext) override;
 	virtual void Serialize(FArchive& InArchive) override;
 
-	virtual TSet<FSoftObjectPath>& GetReferencedAssets() override
-	{
-		return ReferencedAssets;
-	}
-
-	virtual const TSet<FSoftObjectPath>& GetReferencedAssets() const override
-	{
-		return ReferencedAssets;
-	}
+	virtual TSet<UObject*>& GetReferencedAssetClassCache() override;
 
 	// Returns Asset Metadata associated with this MetaSound
 	virtual Metasound::Frontend::FNodeClassInfo GetAssetClassInfo() const override;
 
 	virtual const FMetasoundFrontendVersion& GetDefaultArchetypeVersion() const override;
 
+	virtual bool ConformObjectDataToArchetype() override { return false; }
+
+	virtual const TSet<FString>& GetReferencedAssetClassKeys() const override
+	{
+		return ReferencedAssetClassKeys;
+	}
+
 	UObject* GetOwningAsset() override
 	{
 		return this;
@@ -258,6 +254,7 @@ public:
 	const TArray<FMetasoundFrontendVersion>& GetSupportedArchetypeVersions() const override;
 
 protected:
+	virtual void SetReferencedAssetClassKeys(TSet<Metasound::Frontend::FNodeRegistryKey>&& InKeys) override;
 
 	Metasound::Frontend::FDocumentAccessPtr GetDocument() override
 	{
diff --git a/Engine/Plugins/Runtime/Metasound/Source/MetasoundEngine/Public/MetasoundSource.h b/Engine/Plugins/Runtime/Metasound/Source/MetasoundEngine/Public/MetasoundSource.h
--- a/Engine/Plugins/Runtime/Metasound/Source/MetasoundEngine/Public/MetasoundSource.h
+++ b/Engine/Plugins/Runtime/Metasound/Source/MetasoundEngine/Public/MetasoundSource.h
@@ -7,6 +7,7 @@
 #include "MetasoundAssetBase.h"
 #include "MetasoundFrontend.h"
 #include "MetasoundFrontendDocument.h"
+#include "MetasoundFrontendRegistries.h"
 #include "MetasoundInstanceTransmitter.h"
 #include "MetasoundOperatorSettings.h"
 #include "MetasoundRouter.h"
@@ -43,7 +44,10 @@ protected:
 	FMetasoundFrontendDocument RootMetasoundDocument;
 
 	UPROPERTY()
-	TSet<FSoftObjectPath> ReferencedAssets;
+	TSet<FString> ReferencedAssetClassKeys;
+
+	UPROPERTY(Transient)
+	TSet<UObject*> ReferenceAssetClassCache;
 
 #if WITH_EDITORONLY_DATA
 	UPROPERTY()
@@ -125,25 +129,24 @@ public:
 	virtual void PostEditChangeProperty(FPropertyChangedEvent& InEvent) override;
 #endif // WITH_EDITOR
 
+	virtual const TSet<FString>& GetReferencedAssetClassKeys() const override
+	{
+		return ReferencedAssetClassKeys;
+	}
+
 	virtual void BeginDestroy() override;
 	virtual void PreSave(FObjectPreSaveContext InSaveContext) override;
 	virtual void Serialize(FArchive& Ar) override;
 
-	virtual TSet<FSoftObjectPath>& GetReferencedAssets() override
-	{
-		return ReferencedAssets;
-	}
-
-	virtual const TSet<FSoftObjectPath>& GetReferencedAssets() const override
-	{
-		return ReferencedAssets;
-	}
+	virtual TSet<UObject*>& GetReferencedAssetClassCache() override;
 
 	// Returns Asset Metadata associated with this MetaSoundSource
 	virtual Metasound::Frontend::FNodeClassInfo GetAssetClassInfo() const override;
 
 	virtual const FMetasoundFrontendVersion& GetDefaultArchetypeVersion() const override;
 
+	virtual bool ConformObjectDataToArchetype() override;
+
 	UObject* GetOwningAsset() override
 	{
 		return this;
@@ -154,6 +157,8 @@ public:
 		return this;
 	}
 
+	virtual void InitResources() override;
+
 	virtual bool IsPlayable() const override;
 	virtual bool SupportsSubtitles() const override;
 	virtual float GetDuration() override;
@@ -165,6 +170,8 @@ public:
 
 
 protected:
+	virtual void SetReferencedAssetClassKeys(TSet<Metasound::Frontend::FNodeRegistryKey>&& InKeys) override;
+
 	Metasound::Frontend::FDocumentAccessPtr GetDocument() override
 	{
 		using namespace Metasound::Frontend;
diff --git a/Engine/Plugins/Runtime/Metasound/Source/MetasoundEngine/Public/MetasoundUObjectRegistry.h b/Engine/Plugins/Runtime/Metasound/Source/MetasoundEngine/Public/MetasoundUObjectRegistry.h
--- a/Engine/Plugins/Runtime/Metasound/Source/MetasoundEngine/Public/MetasoundUObjectRegistry.h
+++ b/Engine/Plugins/Runtime/Metasound/Source/MetasoundEngine/Public/MetasoundUObjectRegistry.h
@@ -16,19 +16,12 @@ class UAssetManager;
 
 /** The subsystem in charge of the MetaSound asset registry */
 UCLASS()
-class METASOUNDENGINE_API UMetaSoundAssetSubsystem : public UEngineSubsystem, public IMetaSoundAssetInterface
+class METASOUNDENGINE_API UMetaSoundAssetSubsystem : public UEngineSubsystem, public IMetaSoundAssetManager
 {
 	GENERATED_BODY()
 
 public:
 	virtual void Initialize(FSubsystemCollectionBase& InCollection) override;
-	virtual void Deinitialize() override;
-
-	static UMetaSoundAssetSubsystem& Get()
-	{
-		check(GEngine);
-		return *GEngine->GetEngineSubsystem<UMetaSoundAssetSubsystem>();
-	}
 
 	void AddOrUpdateAsset(UObject& InObject, bool bInRegisterWithFrontend = true);
 	void AddOrUpdateAsset(const FAssetData& InAssetData, bool bInRegisterWithFrontend = true);
diff --git a/Engine/Plugins/Runtime/Metasound/Source/MetasoundFrontend/Private/MetasoundAssetBase.cpp b/Engine/Plugins/Runtime/Metasound/Source/MetasoundFrontend/Private/MetasoundAssetBase.cpp
--- a/Engine/Plugins/Runtime/Metasound/Source/MetasoundFrontend/Private/MetasoundAssetBase.cpp
+++ b/Engine/Plugins/Runtime/Metasound/Source/MetasoundFrontend/Private/MetasoundAssetBase.cpp
@@ -14,6 +14,7 @@
 #include "MetasoundFrontendDocument.h"
 #include "MetasoundFrontendGraph.h"
 #include "MetasoundFrontendInjectReceiveNodes.h"
+#include "MetasoundFrontendRegistries.h"
 #include "MetasoundFrontendSearchEngine.h"
 #include "MetasoundFrontendTransform.h"
 #include "MetasoundInstanceTransmitter.h"
@@ -67,20 +68,51 @@ namespace Metasound
 
 const FString FMetasoundAssetBase::FileExtension(TEXT(".metasound"));
 
-void FMetasoundAssetBase::RegisterGraphWithFrontend()
+IMetaSoundAssetManager* IMetaSoundAssetManager::Instance = nullptr;
+
+void FMetasoundAssetBase::RegisterGraphWithFrontend(FMetaSoundAssetRegistrationOptions InRegistrationOptions)
 {
 	using namespace Metasound;
 	using namespace Metasound::Frontend;
 
 	METASOUND_TRACE_CPUPROFILER_EVENT_SCOPE(MetaSoundAssetBase::RegisterGraphWithFrontend);
+	if (!InRegistrationOptions.bForceReregister)
+	{
+		if (IsRegistered())
+		{
+			return;
+		}
+	}
 
-	FConstDocumentAccessPtr DocumentPtr = GetDocument();
-	FString AssetName;
-	FString AssetPath;
-	if (const UObject* OwningAsset = GetOwningAsset())
+	GetReferencedAssetClassCache().Reset();
+
+	if (InRegistrationOptions.bRebuildReferencedAssetClassKeys)
 	{
-		AssetName = OwningAsset->GetName();
-		AssetPath = OwningAsset->GetPathName();
+		RebuildReferencedAssetClassKeys();
+	}
+
+	if (InRegistrationOptions.bRegisterDependencies)
+	{
+		const TArray<FMetasoundAssetBase*> References = FindOrLoadReferencedAssets();
+		for (FMetasoundAssetBase* Reference : References)
+		{
+			if (InRegistrationOptions.bForceReregister || !Reference->IsRegistered())
+			{
+				// TODO: Check for infinite recursion and error if so
+				Reference->RegisterGraphWithFrontend(InRegistrationOptions);
+			}
+
+			if (UObject* RefAsset = Reference->GetOwningAsset())
+			{
+				GetReferencedAssetClassCache().Add(RefAsset);
+			}
+		}
+	}
+
+	// Auto update must be done after all referenced asset classes are registered
+	if (InRegistrationOptions.bAutoUpdate)
+	{
+		AutoUpdate();
 	}
 
 	// Registers node by copying document. Updates to document require re-registration.
@@ -144,14 +176,20 @@ void FMetasoundAssetBase::RegisterGraphWithFrontend()
 
 	UnregisterGraphWithFrontend();
 
-	FNodeClassInfo AssetClassInfo = GetAssetClassInfo();
-	if (const FMetasoundFrontendDocument* Doc = DocumentPtr.Get())
+	FString AssetName;
+	FString AssetPath;
+	const UObject* OwningAsset = GetOwningAsset();
+	if (ensure(OwningAsset))
 	{
-		RegistryKey = FMetasoundFrontendRegistryContainer::Get()->RegisterNode(MakeUnique<FNodeRegistryEntry>(AssetName, *Doc, AssetClassInfo.AssetPath));
+		AssetName = OwningAsset->GetName();
+		AssetPath = OwningAsset->GetPathName();
 	}
-	else
+
+	FNodeClassInfo AssetClassInfo = GetAssetClassInfo();
+	const FMetasoundFrontendDocument* Doc = GetDocument().Get();
+	if (Doc)
 	{
-		RegistryKey = FNodeRegistryKey();
+		RegistryKey = FMetasoundFrontendRegistryContainer::Get()->RegisterNode(MakeUnique<FNodeRegistryEntry>(AssetName, *Doc, AssetClassInfo.AssetPath));
 	}
 
 	if (NodeRegistryKey::IsValid(RegistryKey))
@@ -178,7 +216,15 @@ void FMetasoundAssetBase::RegisterGraphWithFrontend()
 	}
 	else
 	{
-		UE_LOG(LogMetaSound, Error, TEXT("Failed to register node for MetaSoundSource [Name:%s]"), *AssetName);
+		FString ClassName;
+		if (OwningAsset)
+		{
+			if (UClass* Class = OwningAsset->GetClass())
+			{
+				ClassName = Class->GetName();
+			}
+		}
+		UE_LOG(LogMetaSound, Error, TEXT("Registration failed for MetaSound node class '%s' of UObject class '%s'"), *AssetName, *ClassName);
 	}
 }
 
@@ -265,6 +311,26 @@ FMetasoundFrontendVersion FMetasoundAssetBase::GetPreferredArchetypeVersion(cons
 	return GetDefaultArchetypeVersion();
 }
 
+TArray<FMetasoundAssetBase*> FMetasoundAssetBase::FindOrLoadReferencedAssets() const
+{
+	using namespace Metasound::Frontend;
+
+	TArray<FMetasoundAssetBase*> ReferencedAssets;
+	for (const FNodeRegistryKey& Key : GetReferencedAssetClassKeys())
+	{
+		if (FMetasoundAssetBase* MetaSound = IMetaSoundAssetManager::GetChecked().FindAssetFromKey(Key))
+		{
+			ReferencedAssets.Add(MetaSound);
+		}
+		else
+		{
+			UE_LOG(LogMetaSound, Error, TEXT("Failed to find referenced MetaSound asset with key '%s'"), *Key);
+		}
+	}
+
+	return ReferencedAssets;
+}
+
 bool FMetasoundAssetBase::GetArchetype(FMetasoundFrontendArchetype& OutArchetype) const
 {
 	using namespace Metasound;
@@ -306,42 +372,27 @@ void FMetasoundAssetBase::ConformDocumentToArchetype()
 		Metasound::Frontend::FMatchRootGraphToArchetype Transform(PreferredArchetypeVersion);
 		if (Transform.Transform(GetDocumentHandle()))
 		{
+			ConformObjectDataToArchetype();
 			MarkMetasoundDocumentDirty();
 		}
 	}
 }
 
-bool FMetasoundAssetBase::AutoUpdate(const IMetaSoundAssetInterface& InAssetInterface, bool bInMarkDirty, bool bInUpdateReferencedAssets)
+bool FMetasoundAssetBase::AutoUpdate(bool bInMarkDirty)
 {
 	using namespace Metasound::Frontend;
 	METASOUND_TRACE_CPUPROFILER_EVENT_SCOPE(MetaSoundAssetBase::AutoUpdate);
 
-	bool bDidEdit = false;
-
-	if (bInUpdateReferencedAssets)
-	{
-		const TSet<FSoftObjectPath>& ReferencedAssets = GetReferencedAssets();
-		for (const FSoftObjectPath& ReferencedAsset : ReferencedAssets)
-		{
-			FMetasoundAssetBase* ReferencedMetaSound = InAssetInterface.TryLoadAsset(ReferencedAsset);
-			if (ReferencedMetaSound)
-			{
-				bDidEdit |= FAutoUpdateRootGraph(*ReferencedMetaSound, InAssetInterface).Transform(ReferencedMetaSound->GetDocumentHandle());
-			}
-		}
-	}
-
-	bDidEdit |= FAutoUpdateRootGraph(*this, InAssetInterface).Transform(GetDocumentHandle());
-
-	if (bDidEdit && bInMarkDirty)
+	const bool bUpdated = FAutoUpdateRootGraph().Transform(GetDocumentHandle());
+	if (bUpdated && bInMarkDirty)
 	{
 		MarkMetasoundDocumentDirty();
 	}
 
-	return bDidEdit;
+	return bUpdated;
 }
 
-bool FMetasoundAssetBase::VersionAsset(const IMetaSoundAssetInterface& InAssetInterface, bool bInMarkDirty, bool bInVersionReferencedAssets)
+bool FMetasoundAssetBase::VersionAsset()
 {
 	using namespace Metasound;
 	using namespace Metasound::Frontend;
@@ -356,38 +407,22 @@ bool FMetasoundAssetBase::VersionAsset(const IMetaSoundAssetInterface& InAssetIn
 	}
 
 	bool bDidEdit = false;
-
-	if (bInVersionReferencedAssets)
-	{
-		const TSet<FSoftObjectPath>& ReferencedAssets = GetReferencedAssets();
-		for (const FSoftObjectPath& ReferencedAsset : ReferencedAssets)
-		{
-			FMetasoundAssetBase* ReferencedMetaSound = InAssetInterface.TryLoadAsset(ReferencedAsset);
-			if (ensure(ReferencedMetaSound))
-			{
-				bDidEdit |= ReferencedMetaSound->VersionAsset(InAssetInterface, bInMarkDirty);
-			}
-		}
-	}
-
 	FDocumentHandle DocumentHandle = GetDocumentHandle();
 
 	bDidEdit |= FVersionDocument(AssetName, AssetPath).Transform(DocumentHandle);
 	if (FMetasoundFrontendDocument* Doc = GetDocument().Get())
 	{
-		FMetasoundFrontendVersion ArchetypeVerison = GetPreferredArchetypeVersion(*Doc);
-		bDidEdit |= FMatchRootGraphToArchetype(ArchetypeVerison).Transform(DocumentHandle);
-	}
-
-	if (bDidEdit && bInMarkDirty)
-	{
-		MarkMetasoundDocumentDirty();
+		FMetasoundFrontendVersion ArchetypeVersion = GetPreferredArchetypeVersion(*Doc);
+		bDidEdit |= FMatchRootGraphToArchetype(ArchetypeVersion).Transform(DocumentHandle);
+		if (bDidEdit)
+		{
+			ConformObjectDataToArchetype();
+		}
 	}
 
 	return bDidEdit;
 }
 
-
 TUniquePtr<Metasound::IGraph> FMetasoundAssetBase::BuildMetasoundDocument() const
 {
 	using namespace Metasound;
@@ -417,39 +452,45 @@ TUniquePtr<Metasound::IGraph> FMetasoundAssetBase::BuildMetasoundDocument() cons
 	return MoveTemp(FrontendGraph);
 }
 
-void FMetasoundAssetBase::RebuildReferencedAssets(const IMetaSoundAssetInterface& InAssetInterface)
+void FMetasoundAssetBase::RebuildReferencedAssetClassKeys()
 {
 	using namespace Metasound::Frontend;
 
-	METASOUND_TRACE_CPUPROFILER_EVENT_SCOPE(MetaSoundAssetBase::RebuildReferencedAssets);
+	METASOUND_TRACE_CPUPROFILER_EVENT_SCOPE(MetaSoundAssetBase::RebuildReferencedAssetClassKeys);
+
+	TSet<FNodeRegistryKey> AssetKeys;
 
-	TSet<FSoftObjectPath>& ReferencedAssets = GetReferencedAssets();
-	ReferencedAssets.Reset();
 
-	GetRootGraphHandle()->IterateConstNodes([AssetInterface = &InAssetInterface, RefAssets = &ReferencedAssets](FConstNodeHandle NodeHandle)
+	GetRootGraphHandle()->IterateConstNodes([RefAssetKeys = &AssetKeys](FConstNodeHandle NodeHandle)
 	{
 		using namespace Metasound::Frontend;
 
-		const FMetasoundFrontendClassMetadata& Metadata = NodeHandle->GetClassMetadata();
-		const FNodeRegistryKey LocalRegistryKey = NodeRegistryKey::CreateKey(Metadata);
-		if (const FSoftObjectPath* Path = AssetInterface->FindObjectPathFromKey(LocalRegistryKey))
+		const FMetasoundFrontendClassMetadata& RefMetadata = NodeHandle->GetClassMetadata();
+		const FNodeRegistryKey RefRegistryKey = NodeRegistryKey::CreateKey(RefMetadata);
+		if (const FMetasoundAssetBase* RefAsset = IMetaSoundAssetManager::GetChecked().FindAssetFromKey(RefRegistryKey))
 		{
-			if (Path->IsValid())
-			{
-				RefAssets->Add(*Path);
-			}
+			RefAssetKeys->Add(RefRegistryKey);
 		}
 	}, EMetasoundFrontendClassType::External);
+
+	SetReferencedAssetClassKeys(MoveTemp(AssetKeys));
 }
 
-bool FMetasoundAssetBase::ContainReferenceLoop(const IMetaSoundAssetInterface& IMetaSoundAssetInterface) const
+bool FMetasoundAssetBase::IsRegistered() const
 {
-	return false;
+	using namespace Metasound::Frontend;
+
+	if (!NodeRegistryKey::IsValid(RegistryKey))
+	{
+		return false;
+	}
+
+	return FMetasoundFrontendRegistryContainer::Get()->IsNodeRegistered(RegistryKey);
 }
 
-bool FMetasoundAssetBase::AddingReferenceCausesLoop(const FSoftObjectPath& InReferencePath, const IMetaSoundAssetInterface& InAssetInterface) const
+bool FMetasoundAssetBase::AddingReferenceCausesLoop(const FSoftObjectPath& InReferencePath) const
 {
-	const FMetasoundAssetBase* ReferenceAsset = InAssetInterface.TryLoadAsset(InReferencePath);
+	const FMetasoundAssetBase* ReferenceAsset = IMetaSoundAssetManager::GetChecked().TryLoadAsset(InReferencePath);
 	if (!ensureAlways(ReferenceAsset))
 	{
 		return false;
@@ -467,12 +508,12 @@ bool FMetasoundAssetBase::AddingReferenceCausesLoop(const FSoftObjectPath& InRef
 			return Children;
 		}
 
-		for (const FSoftObjectPath& Path : ChildAsset.GetReferencedAssets())
+		TArray<FMetasoundAssetBase*> ChildRefs = ChildAsset.FindOrLoadReferencedAssets();
+		for (const FMetasoundAssetBase* ChildRef : ChildRefs)
 		{
-			const FMetasoundAssetBase* ChildReference = InAssetInterface.TryLoadAsset(Path);
-			if (ensureAlways(ChildReference))
+			if (ChildRef)
 			{
-				Children.Add(ChildReference);
+				Children.Add(ChildRef);
 			}
 		}
 		return Children;
diff --git a/Engine/Plugins/Runtime/Metasound/Source/MetasoundFrontend/Private/MetasoundFrontendInvalidController.h b/Engine/Plugins/Runtime/Metasound/Source/MetasoundFrontend/Private/MetasoundFrontendInvalidController.h
--- a/Engine/Plugins/Runtime/Metasound/Source/MetasoundFrontend/Private/MetasoundFrontendInvalidController.h
+++ b/Engine/Plugins/Runtime/Metasound/Source/MetasoundFrontend/Private/MetasoundFrontendInvalidController.h
@@ -176,7 +176,7 @@ namespace Metasound
 			bool IsRequired() const override { return false; }
 
 			bool DiffAgainstRegistryInterface(FClassInterfaceUpdates& OutInterfaceUpdates, bool bInUseHighestMinorVersion) const override { return false; }
-			bool CanAutoUpdate(const IMetaSoundAssetInterface& AssetInterface, FClassInterfaceUpdates* OutInterfaceUpdates = nullptr) const override { return false; }
+			bool CanAutoUpdate(FClassInterfaceUpdates* OutInterfaceUpdates = nullptr) const override { return false; }
 			FMetasoundFrontendVersionNumber FindHighestVersionInRegistry() const override { return FMetasoundFrontendVersionNumber::GetInvalid(); }
 			FMetasoundFrontendVersionNumber FindHighestMinorVersionInRegistry() const override { return FMetasoundFrontendVersionNumber::GetInvalid(); }
 
diff --git a/Engine/Plugins/Runtime/Metasound/Source/MetasoundFrontend/Private/MetasoundFrontendStandardController.cpp b/Engine/Plugins/Runtime/Metasound/Source/MetasoundFrontend/Private/MetasoundFrontendStandardController.cpp
--- a/Engine/Plugins/Runtime/Metasound/Source/MetasoundFrontend/Private/MetasoundFrontendStandardController.cpp
+++ b/Engine/Plugins/Runtime/Metasound/Source/MetasoundFrontend/Private/MetasoundFrontendStandardController.cpp
@@ -7,6 +7,7 @@
 #include "Algo/Transform.h"
 #include "CoreMinimal.h"
 #include "HAL/FileManager.h"
+#include "HAL/IConsoleManager.h"
 #include "MetasoundAccessPtr.h"
 #include "MetasoundAssetBase.h"
 #include "MetasoundFrontendArchetypeRegistry.h"
@@ -19,12 +20,21 @@
 #include "MetasoundFrontendSubgraphNodeController.h"
 #include "MetasoundJsonBackend.h"
 #include "MetasoundOperatorBuilder.h"
+#include "MetasoundTrace.h"
 #include "StructDeserializer.h"
 #include "StructSerializer.h"
 #include "Serialization/MemoryReader.h"
 
 #define LOCTEXT_NAMESPACE "MetasoundFrontendStandardController"
 
+static int32 MetaSoundAutoUpdateNativeClassCVar = 1;
+FAutoConsoleVariableRef CVarMetaSoundAutoUpdateNativeClass(
+	TEXT("au.MetaSounds.AutoUpdate.NativeClasses"),
+	MetaSoundAutoUpdateNativeClassCVar,
+	TEXT("If true, node references to native class that share a version number will attempt to auto-update if the interface is different, which results in slower graph load times.\n")
+	TEXT("0: Don't auto-update native classes, !0: Auto-update native classes (default)"),
+	ECVF_Default);
+
 namespace Metasound
 {
 	namespace Frontend
@@ -1744,6 +1754,8 @@ namespace Metasound
 
 		bool FBaseNodeController::DiffAgainstRegistryInterface(FClassInterfaceUpdates& OutInterfaceUpdates, bool bInUseHighestMinorVersion) const
 		{
+			METASOUND_TRACE_CPUPROFILER_EVENT_SCOPE(BaseNodeController::DiffAgainstRegistryInterface);
+
 			OutInterfaceUpdates = FClassInterfaceUpdates();
 
 			const FMetasoundFrontendClassMetadata& NodeClassMetadata = GetClassMetadata();
@@ -1751,20 +1763,18 @@ namespace Metasound
 
 			Metasound::FNodeClassName NodeClassName = NodeClassMetadata.GetClassName().ToNodeClassName();
 
-			const bool bSortByVersion = true;
-
-			FMetasoundFrontendClass RegistryClass;
 			if (bInUseHighestMinorVersion)
 			{
-				if (!ISearchEngine::Get().FindClassWithMajorVersion(NodeClassName, NodeClassMetadata.GetVersion().Major, RegistryClass))
+				if (!ISearchEngine::Get().FindClassWithMajorVersion(NodeClassName, NodeClassMetadata.GetVersion().Major, OutInterfaceUpdates.RegistryClass))
 				{
-					OutInterfaceUpdates.RemovedInputs = NodeClassInterface.Inputs;
-					OutInterfaceUpdates.RemovedOutputs = NodeClassInterface.Outputs;
+					Algo::Transform(NodeClassInterface.Inputs, OutInterfaceUpdates.RemovedInputs, [&](const FMetasoundFrontendClassInput& Input) { return &Input; });
+					Algo::Transform(NodeClassInterface.Outputs, OutInterfaceUpdates.RemovedOutputs, [&](const FMetasoundFrontendClassOutput& Output) { return &Output; });
 					return false;
 				}
 			}
 			else
 			{
+				constexpr bool bSortByVersion = true;
 				const TArray<FMetasoundFrontendClass> Classes = ISearchEngine::Get().FindClassesWithName(NodeClassName, bSortByVersion);
 				const FMetasoundFrontendClass* ExactClass = Classes.FindByPredicate([CurrentVersion = &NodeClassMetadata.GetVersion()](const FMetasoundFrontendClass& AvailableClass)
 				{
@@ -1773,58 +1783,60 @@ namespace Metasound
 
 				if (!ExactClass)
 				{
-					OutInterfaceUpdates.RemovedInputs = NodeClassInterface.Inputs;
-					OutInterfaceUpdates.RemovedOutputs = NodeClassInterface.Outputs;
+					Algo::Transform(NodeClassInterface.Inputs, OutInterfaceUpdates.RemovedInputs, [&](const FMetasoundFrontendClassInput& Input) { return &Input; });
+					Algo::Transform(NodeClassInterface.Outputs, OutInterfaceUpdates.RemovedOutputs, [&](const FMetasoundFrontendClassOutput& Output) { return &Output; });
 					return false;
 				}
-				RegistryClass = *ExactClass;
+				OutInterfaceUpdates.RegistryClass = *ExactClass;
 			}
 
+			Algo::Transform(OutInterfaceUpdates.RegistryClass.Interface.Inputs, OutInterfaceUpdates.AddedInputs, [&](const FMetasoundFrontendClassInput& Input) { return &Input; });
 			for (const FMetasoundFrontendClassInput& Input : NodeClassInterface.Inputs)
 			{
-				auto IsFunctionalEquivalent = [NodeClassInput = &Input](const FMetasoundFrontendClassInput& Iter)
+				auto IsFunctionalEquivalent = [NodeClassInput = &Input](const FMetasoundFrontendClassInput* Iter)
 				{
-					return FMetasoundFrontendClassVertex::IsFunctionalEquivalent(*NodeClassInput, Iter);
+					return FMetasoundFrontendClassVertex::IsFunctionalEquivalent(*NodeClassInput, *Iter);
 				};
 
-				const int32 Index = RegistryClass.Interface.Inputs.FindLastByPredicate(IsFunctionalEquivalent);
+				const int32 Index = OutInterfaceUpdates.AddedInputs.FindLastByPredicate(IsFunctionalEquivalent);
 				if (Index == INDEX_NONE)
 				{
-					OutInterfaceUpdates.RemovedInputs.Add(Input);
+					OutInterfaceUpdates.RemovedInputs.Add(&Input);
 				}
 				else
 				{
-					RegistryClass.Interface.Inputs.RemoveAtSwap(Index, 1, false /* bAllowShrinking */);
+					OutInterfaceUpdates.AddedInputs.RemoveAtSwap(Index, 1, false /* bAllowShrinking */);
 				}
 			}
-			OutInterfaceUpdates.AddedInputs = MoveTemp(RegistryClass.Interface.Inputs);
 
+			Algo::Transform(OutInterfaceUpdates.RegistryClass.Interface.Outputs, OutInterfaceUpdates.AddedOutputs, [&](const FMetasoundFrontendClassOutput& Output) { return &Output; });
 			for (const FMetasoundFrontendClassOutput& Output : NodeClassInterface.Outputs)
 			{
-				auto IsFunctionalEquivalent = [NodeClassOutput = &Output](const FMetasoundFrontendClassOutput& Iter)
+				auto IsFunctionalEquivalent = [NodeClassOutput = &Output](const FMetasoundFrontendClassOutput* Iter)
 				{
-					return FMetasoundFrontendClassVertex::IsFunctionalEquivalent(*NodeClassOutput, Iter);
+					return FMetasoundFrontendClassVertex::IsFunctionalEquivalent(*NodeClassOutput, *Iter);
 				};
 
-				const int32 Index = RegistryClass.Interface.Outputs.FindLastByPredicate(IsFunctionalEquivalent);
+				const int32 Index = OutInterfaceUpdates.AddedOutputs.FindLastByPredicate(IsFunctionalEquivalent);
 				if (Index == INDEX_NONE)
 				{
-					OutInterfaceUpdates.RemovedOutputs.Add(Output);
+					OutInterfaceUpdates.RemovedOutputs.Add(&Output);
 				}
 				else
 				{
-					RegistryClass.Interface.Outputs.RemoveAtSwap(Index, 1, false /* bAllowShrinking */);
+					OutInterfaceUpdates.AddedOutputs.RemoveAtSwap(Index, 1, false /* bAllowShrinking */);
 				}
 			}
-			OutInterfaceUpdates.AddedOutputs = MoveTemp(RegistryClass.Interface.Outputs);
 
 			return true;
 		}
 
-		bool FBaseNodeController::CanAutoUpdate(const IMetaSoundAssetInterface& AssetInterface, FClassInterfaceUpdates* OutInterfaceUpdates) const
+		bool FBaseNodeController::CanAutoUpdate(FClassInterfaceUpdates* OutInterfaceUpdates) const
 		{
+			METASOUND_TRACE_CPUPROFILER_EVENT_SCOPE(BaseNodeController::CanAutoUpdate);
+
 			const FMetasoundFrontendClassMetadata& NodeClassMetadata = GetClassMetadata();
-			if (!AssetInterface.CanAutoUpdate(NodeClassMetadata.GetClassName()))
+			if (!IMetaSoundAssetManager::GetChecked().CanAutoUpdate(NodeClassMetadata.GetClassName()))
 			{
 				return false;
 			}
@@ -1851,6 +1863,11 @@ namespace Metasound
 				const bool bIsClassNative = FMetasoundFrontendRegistryContainer::Get()->IsNodeNative(RegistryKey);
 				if (bIsClassNative)
 				{
+					if (!MetaSoundAutoUpdateNativeClassCVar)
+					{
+						return false;
+					}
+
 					FClassInterfaceUpdates InterfaceUpdates;
 					DiffAgainstRegistryInterface(InterfaceUpdates, true /* bUseHighestMinorVersion */);
 					if (OutInterfaceUpdates)
diff --git a/Engine/Plugins/Runtime/Metasound/Source/MetasoundFrontend/Private/MetasoundFrontendStandardController.h b/Engine/Plugins/Runtime/Metasound/Source/MetasoundFrontend/Private/MetasoundFrontendStandardController.h
--- a/Engine/Plugins/Runtime/Metasound/Source/MetasoundFrontend/Private/MetasoundFrontendStandardController.h
+++ b/Engine/Plugins/Runtime/Metasound/Source/MetasoundFrontend/Private/MetasoundFrontendStandardController.h
@@ -11,7 +11,7 @@
 #include "UObject/Object.h"
 
 
-class IMetaSoundAssetInterface;
+class IMetaSoundAssetManager;
 
 namespace Metasound
 {
@@ -379,7 +379,7 @@ namespace Metasound
 
 			bool DiffAgainstRegistryInterface(FClassInterfaceUpdates& OutInterfaceUpdates, bool bInUseHighestMinorVersion) const override;
 
-			bool CanAutoUpdate(const IMetaSoundAssetInterface& AssetInterface, FClassInterfaceUpdates* OutInterfaceUpdates = nullptr) const override;
+			bool CanAutoUpdate(FClassInterfaceUpdates* OutInterfaceUpdates = nullptr) const override;
 			FNodeHandle ReplaceWithVersion(const FMetasoundFrontendVersionNumber& InNewVersion) override;
 			FMetasoundFrontendVersionNumber FindHighestVersionInRegistry() const override;
 			FMetasoundFrontendVersionNumber FindHighestMinorVersionInRegistry() const override;
diff --git a/Engine/Plugins/Runtime/Metasound/Source/MetasoundFrontend/Private/MetasoundFrontendTransform.cpp b/Engine/Plugins/Runtime/Metasound/Source/MetasoundFrontend/Private/MetasoundFrontendTransform.cpp
--- a/Engine/Plugins/Runtime/Metasound/Source/MetasoundFrontend/Private/MetasoundFrontendTransform.cpp
+++ b/Engine/Plugins/Runtime/Metasound/Source/MetasoundFrontend/Private/MetasoundFrontendTransform.cpp
@@ -384,12 +384,9 @@ namespace Metasound
 
 		bool FAutoUpdateRootGraph::Transform(FDocumentHandle InDocument) const
 		{
-			check(Asset);
-			check(AssetInterface);
-
 			bool bDidEdit = false;
 
-			FDocumentHandle PresetReferenceDocument = IDocumentController::GetInvalidHandle();
+			FMetasoundAssetBase* PresetReferencedMetaSoundAsset = nullptr;
 			TArray<TPair<FNodeHandle, FMetasoundFrontendVersionNumber>> NodesToUpdate;
 
 			FGraphHandle RootGraph = InDocument->GetRootGraph();
@@ -397,7 +394,7 @@ namespace Metasound
 			{
 				using namespace Metasound::Frontend;
 				FClassInterfaceUpdates InterfaceUpdates;
-				if (!NodeHandle->CanAutoUpdate(*AssetInterface, &InterfaceUpdates))
+				if (!NodeHandle->CanAutoUpdate(&InterfaceUpdates))
 				{
 					return;
 				}
@@ -420,10 +417,9 @@ namespace Metasound
 					if (ensure(ClassMetadata.GetType() == EMetasoundFrontendClassType::External))
 					{
 						const FNodeRegistryKey RegistryKey = FMetasoundFrontendRegistryContainer::Get()->GetRegistryKey(ClassMetadata);
-						FMetasoundAssetBase* ReferencedMetaSoundAsset = AssetInterface->FindAssetFromKey(RegistryKey);
-						if (ReferencedMetaSoundAsset)
+						PresetReferencedMetaSoundAsset = IMetaSoundAssetManager::GetChecked().FindAssetFromKey(RegistryKey);
+						if (PresetReferencedMetaSoundAsset)
 						{
-							PresetReferenceDocument = ReferencedMetaSoundAsset->GetDocumentHandle();
 							return;
 						}
 					}
@@ -432,9 +428,10 @@ namespace Metasound
 				NodesToUpdate.Add(TPair<FNodeHandle, FMetasoundFrontendVersionNumber>(NodeHandle, UpdateVersion));
 			}, EMetasoundFrontendClassType::External);
 
-			if (PresetReferenceDocument->IsValid())
+			if (PresetReferencedMetaSoundAsset)
 			{
-				bDidEdit |= FRebuildPresetRootGraph(PresetReferenceDocument).Transform(InDocument);
+				bDidEdit |= FRebuildPresetRootGraph(PresetReferencedMetaSoundAsset->GetDocumentHandle()).Transform(InDocument);
+				bDidEdit |= PresetReferencedMetaSoundAsset->ConformObjectDataToArchetype();
 			}
 			else
 			{
@@ -452,7 +449,6 @@ namespace Metasound
 			}
 
 			InDocument->SynchronizeDependencies();
-
 			return bDidEdit;
 		}
 
@@ -475,6 +471,10 @@ namespace Metasound
 				return false;
 			}
 
+			// 2. Run transform to ensure preset matches reference archetype
+			const FMetasoundFrontendVersion& RefArchetypeVersion = ReferencedDocument->GetArchetypeVersion();
+			FMatchRootGraphToArchetype(RefArchetypeVersion).Transform(InDocument);
+
 			using FDefaultKey = TPair<FString, FName>;
 			struct FDefaultInputValue
 			{
@@ -522,11 +522,6 @@ namespace Metasound
 			}, EMetasoundFrontendClassType::Output);
 
 			RootGraphHandle->ClearGraph();
-			const FMetasoundFrontendVersion& RefArchetype = ReferencedDocument->GetArchetypeVersion();
-// 			if (!ensure(IsArchetypeSupported(RefArchetype)))
-// 			{
-// 				return false;
-// 			}
 
 			// 2. Set initial locations of 3 columns (inputs, reference node, and outputs)
 			const TArray<FConstNodeHandle> ReferencedInputs = ReferencedGraphHandle->GetConstInputNodes();
@@ -674,9 +669,6 @@ namespace Metasound
 			},
 			EMetasoundFrontendClassType::Output);
 
-			// 5. Run transform to set to match archetype (marking inputs/outputs as required that were copied over)
-			FMatchRootGraphToArchetype(RefArchetype).Transform(InDocument);
-
 			return true;
 		}
 
diff --git a/Engine/Plugins/Runtime/Metasound/Source/MetasoundFrontend/Private/MetasoundReceiveNode.cpp b/Engine/Plugins/Runtime/Metasound/Source/MetasoundFrontend/Private/MetasoundReceiveNode.cpp
--- a/Engine/Plugins/Runtime/Metasound/Source/MetasoundFrontend/Private/MetasoundReceiveNode.cpp
+++ b/Engine/Plugins/Runtime/Metasound/Source/MetasoundFrontend/Private/MetasoundReceiveNode.cpp
@@ -8,25 +8,25 @@ namespace Metasound
 	{
 		const FString& GetAddressInputName()
 		{
-			static const FString InputName = FString(TEXT("Address"));
+			static const FString InputName(TEXT("Address"));
 			return InputName;
 		}
 
 		const FString& GetDefaultDataInputName()
 		{
-			static const FString DefaultDataName = FString(TEXT("Default"));
+			static const FString DefaultDataName(TEXT("Default"));
 			return DefaultDataName;
 		}
 
 		const FString& GetOutputName()
 		{
-			static const FString OutputName = FString(TEXT("Out"));
+			static const FString OutputName(TEXT("Out"));
 			return OutputName;
 		}
 
 		FNodeClassName GetClassNameForDataType(const FName& InDataTypeName)
 		{
-			return FNodeClassName{TEXT("Receive"), InDataTypeName, TEXT("")};
+			return FNodeClassName{ "Receive", InDataTypeName, FName() };
 		}
 
 		int32 GetCurrentMajorVersion() { return 1; }
diff --git a/Engine/Plugins/Runtime/Metasound/Source/MetasoundFrontend/Public/MetasoundAssetBase.h b/Engine/Plugins/Runtime/Metasound/Source/MetasoundFrontend/Public/MetasoundAssetBase.h
--- a/Engine/Plugins/Runtime/Metasound/Source/MetasoundFrontend/Public/MetasoundAssetBase.h
+++ b/Engine/Plugins/Runtime/Metasound/Source/MetasoundFrontend/Public/MetasoundAssetBase.h
@@ -37,9 +37,39 @@ namespace Metasound
 	} // namespace AssetTags
 } // namespace Metasound
 
-class METASOUNDFRONTEND_API IMetaSoundAssetInterface
+struct METASOUNDFRONTEND_API FMetaSoundAssetRegistrationOptions
 {
+	// If true, forces a re-register of this class (and all class dependencies
+	// if the following option 'bRegisterDependencies' is enabled).
+	bool bForceReregister = true;
+
+	// If true, recursively attempts to register dependencies.
+	bool bRegisterDependencies = false;
+
+	// Attempt to auto-update (only runs if class not registered or set to force re-register)
+	bool bAutoUpdate = true;
+
+	// Attempt to rebuild referenced class keys (only run if class not registered or set to force re-register)
+	bool bRebuildReferencedAssetClassKeys = true;
+};
+
+class METASOUNDFRONTEND_API IMetaSoundAssetManager
+{
+	static IMetaSoundAssetManager* Instance;
+
 public:
+	static void Set(IMetaSoundAssetManager& InInterface)
+	{
+		check(!Instance);
+		Instance = &InInterface;
+	}
+
+	static IMetaSoundAssetManager& GetChecked()
+	{
+		check(Instance);
+		return *Instance;
+	}
+
 	virtual bool CanAutoUpdate(const FMetasoundFrontendClassName& InClassName) const = 0;
 	virtual FMetasoundAssetBase* FindAssetFromKey(const Metasound::Frontend::FNodeRegistryKey& InRegistryKey) const = 0;
 	virtual const FSoftObjectPath* FindObjectPathFromKey(const Metasound::Frontend::FNodeRegistryKey& InRegistryKey) const = 0;
@@ -85,13 +115,17 @@ public:
 	virtual void SetRegistryAssetClassInfo(const Metasound::Frontend::FNodeClassInfo& InClassInfo) = 0;
 #endif // WITH_EDITORONLY_DATA
 
+	// Called when the archetype is changed, presenting the opportunity for
+	// any reflected object data to be updated based on the new archetype.
+	// Returns whether or not any edits were made.
+	virtual bool ConformObjectDataToArchetype() = 0;
+
 	// Registers the root graph of the given asset with the MetaSound Frontend.
-	void RegisterGraphWithFrontend();
+	void RegisterGraphWithFrontend(FMetaSoundAssetRegistrationOptions InRegistrationOptions = FMetaSoundAssetRegistrationOptions());
 
 	// Unregisters the root graph of the given asset with the MetaSound Frontend.
 	void UnregisterGraphWithFrontend();
 
-
 	// Sets/overwrites the root class metadata
 	virtual void SetMetadata(FMetasoundFrontendClassMetadata& InMetadata);
 
@@ -111,16 +145,22 @@ public:
 	// Gets the asset class info.
 	virtual Metasound::Frontend::FNodeClassInfo GetAssetClassInfo() const = 0;
 
+	// Returns all the class keys of this asset's referenced assets
+	virtual const TSet<FString>& GetReferencedAssetClassKeys() const = 0;
 
-	virtual TSet<FSoftObjectPath>& GetReferencedAssets() = 0;
-	virtual const TSet<FSoftObjectPath>& GetReferencedAssets() const = 0;
+	// Returns set of cached class references required by this class
+	// to persist during the duration of this class being registered.
+	// TODO: May no longer be needed, but leaving for now to avoid
+	// potential edge cases with GC/UObject references.  Could potentially
+	// move to cook only and making auto-update optional (to avoid large
+	// MetaSound patches caused by references being updated).
+	virtual TSet<UObject*>& GetReferencedAssetClassCache() = 0;
 
+	bool AddingReferenceCausesLoop(const FSoftObjectPath& InReferencePath) const;
 	void ConvertFromPreset();
-
-	void RebuildReferencedAssets(const IMetaSoundAssetInterface& InAssetInterface);
-
-	bool AddingReferenceCausesLoop(const FSoftObjectPath& InReferencePath, const IMetaSoundAssetInterface& InAssetInterface) const;
-	bool ContainReferenceLoop(const IMetaSoundAssetInterface& IMetaSoundAssetInterface) const;
+	TArray<FMetasoundAssetBase*> FindOrLoadReferencedAssets() const;
+	bool IsRegistered() const;
+	void RebuildReferencedAssetClassKeys();
 
 	// Imports data from a JSON string directly
 	bool ImportFromJSON(const FString& InJSON);
@@ -128,6 +168,7 @@ public:
 	// Imports the asset from a JSON file at provided path
 	bool ImportFromJSONAsset(const FString& InAbsolutePath);
 
+
 	// Returns handle for the root metasound graph of this asset.
 	Metasound::Frontend::FDocumentHandle GetDocumentHandle();
 	Metasound::Frontend::FConstDocumentHandle GetDocumentHandle() const;
@@ -145,12 +186,9 @@ public:
 	FMetasoundFrontendDocument& GetDocumentChecked();
 	const FMetasoundFrontendDocument& GetDocumentChecked() const;
 
-	// This must be called on UObject::PostLoad, as well as in this asset's UFactory, to fix up the root document based on the most recent version of the archetype.
 	void ConformDocumentToArchetype();
-
-	bool AutoUpdate(const IMetaSoundAssetInterface& InAssetInterface, bool bInMarkDirty, bool bInUpdateReferencedAssets = true);
-
-	bool VersionAsset(const IMetaSoundAssetInterface& InAssetInterface, bool bInMarkDirty, bool bInVersionReferencedAssets = true);
+	bool AutoUpdate(bool bInMarkDirty = false);
+	bool VersionAsset();
 
 	// Calls the outermost package and marks it dirty.
 	bool MarkMetasoundDocumentDirty() const;
@@ -165,6 +203,7 @@ public:
 	virtual TUniquePtr<Metasound::IGraph> BuildMetasoundDocument() const;
 
 protected:
+	virtual void SetReferencedAssetClassKeys(TSet<Metasound::Frontend::FNodeRegistryKey>&& InKeys) = 0;
 
 	TArray<FSendInfoAndVertexName> GetSendInfos(uint64 InInstanceID) const;
 
@@ -172,7 +211,6 @@ protected:
 	FText GetDisplayName(FString&& InTypeName) const;
 #endif // WITH_EDITORONLY_DATA
 
-
 	// Returns an access pointer to the document.
 	virtual Metasound::Frontend::FDocumentAccessPtr GetDocument() = 0;
 
diff --git a/Engine/Plugins/Runtime/Metasound/Source/MetasoundFrontend/Public/MetasoundFrontendController.h b/Engine/Plugins/Runtime/Metasound/Source/MetasoundFrontend/Public/MetasoundFrontendController.h
--- a/Engine/Plugins/Runtime/Metasound/Source/MetasoundFrontend/Public/MetasoundFrontendController.h
+++ b/Engine/Plugins/Runtime/Metasound/Source/MetasoundFrontend/Public/MetasoundFrontendController.h
@@ -11,7 +11,7 @@
 #include "MetasoundFrontendRegistries.h"
 #include "MetasoundGraph.h"
 
-class IMetaSoundAssetInterface;
+class IMetaSoundAssetManager;
 
 /* Metasound Controllers and Handles provide a object oriented interface for  manipulating Metasound Documents. 
  *
@@ -90,10 +90,11 @@ namespace Metasound
 		  * when querying if a node's class has been updated */
 		struct FClassInterfaceUpdates
 		{
-			TArray<FMetasoundFrontendClassInput> AddedInputs;
-			TArray<FMetasoundFrontendClassOutput> AddedOutputs;
-			TArray<FMetasoundFrontendClassInput> RemovedInputs;
-			TArray<FMetasoundFrontendClassOutput> RemovedOutputs;
+
+			TArray<const FMetasoundFrontendClassInput*> AddedInputs;
+			TArray<const FMetasoundFrontendClassOutput*> AddedOutputs;
+			TArray<const FMetasoundFrontendClassInput*> RemovedInputs;
+			TArray<const FMetasoundFrontendClassOutput*> RemovedOutputs;
 
 			bool ContainsRemovedMembers() const
 			{
@@ -109,6 +110,9 @@ namespace Metasound
 			{
 				return ContainsRemovedMembers() || ContainsAddedMembers();
 			}
+
+			// Cached copy of registry class potentially referenced by added members
+			FMetasoundFrontendClass RegistryClass;
 		};
 
 
@@ -436,7 +440,7 @@ namespace Metasound
 			  * has undergone minor revision or the interface has changed, but
 			  * no higher major revision is available. Optionally provide
 			  * interface updates to be populate with any information regarding interface updates. */
-			virtual bool CanAutoUpdate(const IMetaSoundAssetInterface& AssetInterface, FClassInterfaceUpdates* OutInterfaceUpdates = nullptr) const = 0;
+			virtual bool CanAutoUpdate(FClassInterfaceUpdates* OutInterfaceUpdates = nullptr) const = 0;
 
 			/** Description of the given node. */
 			virtual const FText& GetDescription() const = 0;
diff --git a/Engine/Plugins/Runtime/Metasound/Source/MetasoundFrontend/Public/MetasoundFrontendTransform.h b/Engine/Plugins/Runtime/Metasound/Source/MetasoundFrontend/Public/MetasoundFrontendTransform.h
--- a/Engine/Plugins/Runtime/Metasound/Source/MetasoundFrontend/Public/MetasoundFrontendTransform.h
+++ b/Engine/Plugins/Runtime/Metasound/Source/MetasoundFrontend/Public/MetasoundFrontendTransform.h
@@ -84,17 +84,9 @@ namespace Metasound
 		class METASOUNDFRONTEND_API FAutoUpdateRootGraph : public IDocumentTransform
 		{
 		public:
-			FAutoUpdateRootGraph(FMetasoundAssetBase& InAsset, const IMetaSoundAssetInterface& InAssetInterface)
-				: Asset(&InAsset)
-				, AssetInterface(&InAssetInterface)
-			{
-			}
+			FAutoUpdateRootGraph() = default;
 
 			bool Transform(FDocumentHandle InDocument) const override;
-
-		private:
-			FMetasoundAssetBase* Asset = nullptr;
-			const IMetaSoundAssetInterface* AssetInterface = nullptr;
 		};
 
 		/** Synchronizes the document's root graph's display name with that of the asset. */
diff --git a/Engine/Plugins/Runtime/Metasound/Source/MetasoundFrontend/Public/MetasoundReceiveNode.h b/Engine/Plugins/Runtime/Metasound/Source/MetasoundFrontend/Public/MetasoundReceiveNode.h
--- a/Engine/Plugins/Runtime/Metasound/Source/MetasoundFrontend/Public/MetasoundReceiveNode.h
+++ b/Engine/Plugins/Runtime/Metasound/Source/MetasoundFrontend/Public/MetasoundReceiveNode.h
@@ -53,13 +53,13 @@ namespace Metasound
 				Info.ClassName = ReceiveNodeInfo::GetClassNameForDataType(GetMetasoundDataTypeName<TDataType>());
 				Info.MajorVersion = ReceiveNodeInfo::GetCurrentMajorVersion();
 				Info.MinorVersion = ReceiveNodeInfo::GetCurrentMinorVersion();
-				Info.DisplayName = FText::Format(LOCTEXT("Metasound_ReceiveNodeDisplayNameFormat", "Receive {0}"), FText::FromName(GetMetasoundDataTypeName<TDataType>()));
+				Info.DisplayName = FText::Format(LOCTEXT("Metasound_ReceiveNodeDisplayNameFormat", "Receive {0}"), GetMetasoundDataTypeDisplayText<TDataType>());
 				Info.Description = LOCTEXT("Metasound_ReceiveNodeDescription", "Receives data from a send node with the same name.");
 				Info.Author = PluginAuthor;
 				Info.PromptIfMissing = PluginNodeMissingPrompt;
 				Info.DefaultInterface = DeclareVertexInterface();
 				Info.CategoryHierarchy = { LOCTEXT("Metasound_TransmissionNodeCategory", "Transmission") };
-				Info.Keywords = { LOCTEXT("Metasound_ReceiveKeyword", "Receive"), FText::FromName(GetMetasoundDataTypeName<TDataType>()) };
+				Info.Keywords = { };
 
 				return Info;
 			};
diff --git a/Engine/Plugins/Runtime/Metasound/Source/MetasoundFrontend/Public/MetasoundSendNode.h b/Engine/Plugins/Runtime/Metasound/Source/MetasoundFrontend/Public/MetasoundSendNode.h
--- a/Engine/Plugins/Runtime/Metasound/Source/MetasoundFrontend/Public/MetasoundSendNode.h
+++ b/Engine/Plugins/Runtime/Metasound/Source/MetasoundFrontend/Public/MetasoundSendNode.h
@@ -52,16 +52,16 @@ namespace Metasound
 				const FString& InputName = GetSendInputName();
 				FNodeClassMetadata Info;
 
-				Info.ClassName = {TEXT("Send"), GetMetasoundDataTypeName<TDataType>(), TEXT("")};
+				Info.ClassName = { "Send", GetMetasoundDataTypeName<TDataType>(), FName() };
 				Info.MajorVersion = 1;
 				Info.MinorVersion = 0;
-				Info.DisplayName = FText::Format(LOCTEXT("Metasound_SendNodeDisplayNameFormat", "Send {0}"), FText::FromName(GetMetasoundDataTypeName<TDataType>()));
+				Info.DisplayName = FText::Format(LOCTEXT("Metasound_SendNodeDisplayNameFormat", "Send {0}"), GetMetasoundDataTypeDisplayText<TDataType>());
 				Info.Description = LOCTEXT("Metasound_SendNodeDescription", "Sends data from a send node with the same name.");
 				Info.Author = PluginAuthor;
 				Info.PromptIfMissing = PluginNodeMissingPrompt;
 				Info.DefaultInterface = DeclareVertexInterface();
 				Info.CategoryHierarchy = { LOCTEXT("Metasound_TransmissionNodeCategory", "Transmission") };
-				Info.Keywords = { LOCTEXT("Metasound_SendKeyword", "Send"), FText::FromName(GetMetasoundDataTypeName<TDataType>()) };
+				Info.Keywords = { };
 
 				return Info;
 			};
diff --git a/Engine/Source/Runtime/Engine/Classes/Sound/SoundBase.h b/Engine/Source/Runtime/Engine/Classes/Sound/SoundBase.h
--- a/Engine/Source/Runtime/Engine/Classes/Sound/SoundBase.h
+++ b/Engine/Source/Runtime/Engine/Classes/Sound/SoundBase.h
@@ -280,6 +280,9 @@ public:
 	virtual const TArray<UAssetUserData*>* GetAssetUserDataArray() const override;
 	//~ End IInterface_AssetUserData Interface
 
+	/** Called prior to attempting to play sound from Game Thread to enable initializing any additional resources required for playback. */
+	virtual void InitResources() { }
+
 	/** Creates a sound generator instance from this sound base. Return true if this is being implemented by a subclass. Sound generators procedurally generate audio in the audio render thread. */
 	virtual ISoundGeneratorPtr CreateSoundGenerator(const FSoundGeneratorInitParams& InParams) { return nullptr; }
 
diff --git a/Engine/Source/Runtime/Engine/Private/AudioDevice.cpp b/Engine/Source/Runtime/Engine/Private/AudioDevice.cpp
--- a/Engine/Source/Runtime/Engine/Private/AudioDevice.cpp
+++ b/Engine/Source/Runtime/Engine/Private/AudioDevice.cpp
@@ -4679,6 +4679,11 @@ void FAudioDevice::NotifyPluginListenersWorldChanged(UWorld* World)
 
 void FAudioDevice::AddNewActiveSound(const FActiveSound& NewActiveSound)
 {
+	if (USoundBase* Sound = NewActiveSound.GetSound())
+	{
+		Sound->InitResources();
+	}
+
 	AddNewActiveSoundInternal(NewActiveSound, nullptr);
 }
 
