Factor out metasound data type registration

diff --git a/Engine/Plugins/Runtime/Metasound/Source/MetasoundEditor/Private/MetasoundEditorGraphBuilder.cpp b/Engine/Plugins/Runtime/Metasound/Source/MetasoundEditor/Private/MetasoundEditorGraphBuilder.cpp
--- a/Engine/Plugins/Runtime/Metasound/Source/MetasoundEditor/Private/MetasoundEditorGraphBuilder.cpp
+++ b/Engine/Plugins/Runtime/Metasound/Source/MetasoundEditor/Private/MetasoundEditorGraphBuilder.cpp
@@ -18,6 +18,7 @@
 #include "MetasoundEditorGraphValidation.h"
 #include "MetasoundEditorModule.h"
 #include "MetasoundEditorSettings.h"
+#include "MetasoundFrontendDataTypeRegistry.h"
 #include "MetasoundFrontendQuery.h"
 #include "MetasoundFrontendQuerySteps.h"
 #include "MetasoundFrontendRegistries.h"
@@ -483,10 +484,7 @@ namespace Metasound
 					bool bObjectFound = false;
 					if (!InInputPin.DefaultValue.IsEmpty())
 					{
-						FMetasoundFrontendRegistryContainer* FrontendRegistry = FMetasoundFrontendRegistryContainer::Get();
-						check(FrontendRegistry);
-
-						if (UClass* Class = FrontendRegistry->GetLiteralUClassForDataType(TypeName))
+						if (UClass* Class = IDataTypeRegistry::Get().GetUClassForDataType(TypeName))
 						{
 							FAssetRegistryModule& AssetRegistryModule = FModuleManager::LoadModuleChecked<FAssetRegistryModule>("AssetRegistry");
 
@@ -1555,7 +1553,8 @@ namespace Metasound
 			}
 
 			FMetasoundFrontendLiteral DefaultLiteral;
-			DefaultLiteral.SetFromLiteral(Frontend::GetDefaultParamForDataType(InputHandle->GetDataType()));
+			DefaultLiteral.SetFromLiteral(IDataTypeRegistry::Get().CreateDefaultLiteral(InputHandle->GetDataType()));
+
 			InPin.DefaultValue = DefaultLiteral.ToString();
 			return OldValue != InPin.DefaultValue;
 		}
diff --git a/Engine/Plugins/Runtime/Metasound/Source/MetasoundEditor/Private/MetasoundEditorGraphNode.cpp b/Engine/Plugins/Runtime/Metasound/Source/MetasoundEditor/Private/MetasoundEditorGraphNode.cpp
--- a/Engine/Plugins/Runtime/Metasound/Source/MetasoundEditor/Private/MetasoundEditorGraphNode.cpp
+++ b/Engine/Plugins/Runtime/Metasound/Source/MetasoundEditor/Private/MetasoundEditorGraphNode.cpp
@@ -17,6 +17,7 @@
 #include "MetasoundEditorModule.h"
 #include "MetasoundEditorSettings.h"
 #include "MetasoundFrontend.h"
+#include "MetasoundFrontendDataTypeRegistry.h"
 #include "MetasoundFrontendRegistries.h"
 #include "MetasoundFrontendSearchEngine.h"
 #include "MetasoundLiteral.h"
@@ -277,18 +278,11 @@ Metasound::Frontend::FDataTypeRegistryInfo UMetasoundEditorGraphNode::GetPinData
 	using namespace Metasound::Editor;
 	using namespace Metasound::Frontend;
 
-	FMetasoundFrontendRegistryContainer* Registry = FMetasoundFrontendRegistryContainer::Get();
-	if (!ensure(Registry))
-	{
-		return FDataTypeRegistryInfo();
-	}
+ 	Metasound::Frontend::FDataTypeRegistryInfo DataTypeInfo;
 
-	Metasound::Frontend::FDataTypeRegistryInfo DataTypeInfo;
 	FConstInputHandle Handle = FGraphBuilder::GetConstInputHandleFromPin(&InPin);
-	if (!ensure(Registry->GetInfoForDataType(Handle->GetDataType(), DataTypeInfo)))
-	{
-		return FDataTypeRegistryInfo();
-	}
+	ensure(IDataTypeRegistry::Get().GetDataTypeInfo(Handle->GetDataType(), DataTypeInfo));
+
 
 	return DataTypeInfo;
 }
diff --git a/Engine/Plugins/Runtime/Metasound/Source/MetasoundEditor/Private/MetasoundEditorModule.cpp b/Engine/Plugins/Runtime/Metasound/Source/MetasoundEditor/Private/MetasoundEditorModule.cpp
--- a/Engine/Plugins/Runtime/Metasound/Source/MetasoundEditor/Private/MetasoundEditorModule.cpp
+++ b/Engine/Plugins/Runtime/Metasound/Source/MetasoundEditor/Private/MetasoundEditorModule.cpp
@@ -24,6 +24,7 @@
 #include "MetasoundEditorGraphNodeFactory.h"
 #include "MetasoundEditorGraphSchema.h"
 #include "MetasoundEditorSettings.h"
+#include "MetasoundFrontendDataTypeRegistry.h"
 #include "MetasoundFrontendDocument.h"
 #include "MetasoundFrontendRegistries.h"
 #include "MetasoundFrontendTransform.h"
@@ -300,105 +301,113 @@ namespace Metasound
 
 			void RegisterCoreDataTypes()
 			{
-				TArray<FName> DataTypeNames = Frontend::GetAllAvailableDataTypes();
+				using namespace Metasound::Frontend;
+
+				const IDataTypeRegistry& DataTypeRegistry = IDataTypeRegistry::Get();
+
+				TArray<FName> DataTypeNames;
+				DataTypeRegistry.GetRegisteredDataTypeNames(DataTypeNames);
+
 				for (FName DataTypeName : DataTypeNames)
 				{
 					FDataTypeRegistryInfo RegistryInfo;
-					Frontend::GetTraitsForDataType(DataTypeName, RegistryInfo);
-
-					FName PinCategory = DataTypeName;
-					FName PinSubCategory;
-
-					// Execution path triggers are specialized
-					if (DataTypeName == Frontend::GetDataTypeName<FTrigger>())
+					if (ensure(DataTypeRegistry.GetDataTypeInfo(DataTypeName, RegistryInfo)))
 					{
-						PinCategory = FGraphBuilder::PinCategoryTrigger;
-					}
 
-					// GraphEditor by default designates specialized connection
-					// specification for Int64, so use it even though literal is
-					// boiled down to int32
-					//else if (DataTypeName == Frontend::GetDataTypeName<int64>())
-					//{
-					//	PinCategory = FGraphBuilder::PinCategoryInt64;
-					//}
+						FName PinCategory = DataTypeName;
+						FName PinSubCategory;
 
-					// Primitives
-					else
-					{
-						switch (RegistryInfo.PreferredLiteralType)
+						// Execution path triggers are specialized
+						if (DataTypeName == GetMetasoundDataTypeName<FTrigger>())
 						{
-							case ELiteralType::Boolean:
-							case ELiteralType::BooleanArray:
-							{
-								PinCategory = FGraphBuilder::PinCategoryBoolean;
-							}
-							break;
+							PinCategory = FGraphBuilder::PinCategoryTrigger;
+						}
 
-							case ELiteralType::Float:
-							case ELiteralType::FloatArray:
-							{
-								PinCategory = FGraphBuilder::PinCategoryFloat;
+						// GraphEditor by default designates specialized connection
+						// specification for Int64, so use it even though literal is
+						// boiled down to int32
+						//else if (DataTypeName == Frontend::GetDataTypeName<int64>())
+						//{
+						//	PinCategory = FGraphBuilder::PinCategoryInt64;
+						//}
 
-								// Doubles use the same preferred literal
-								// but different colorization
-								//if (DataTypeName == Frontend::GetDataTypeName<double>())
-								//{
-								//	PinCategory = FGraphBuilder::PinCategoryDouble;
-								//}
+						// Primitives
+						else
+						{
+							switch (RegistryInfo.PreferredLiteralType)
+							{
+								case ELiteralType::Boolean:
+								case ELiteralType::BooleanArray:
+								{
+									PinCategory = FGraphBuilder::PinCategoryBoolean;
+								}
+								break;
 
-								// Differentiate stronger numeric types associated with audio
-								if (DataTypeName == Frontend::GetDataTypeName<FTime>())
+								case ELiteralType::Float:
+								case ELiteralType::FloatArray:
 								{
-									PinSubCategory = FGraphBuilder::PinSubCategoryTime;
+									PinCategory = FGraphBuilder::PinCategoryFloat;
+
+									// Doubles use the same preferred literal
+									// but different colorization
+									//if (DataTypeName == Frontend::GetDataTypeName<double>())
+									//{
+									//	PinCategory = FGraphBuilder::PinCategoryDouble;
+									//}
+
+									// Differentiate stronger numeric types associated with audio
+									if (DataTypeName == GetMetasoundDataTypeName<FTime>())
+									{
+										PinSubCategory = FGraphBuilder::PinSubCategoryTime;
+									}
 								}
-							}
-							break;
+								break;
 
-							case ELiteralType::Integer:
-							case ELiteralType::IntegerArray:
-							{
-								PinCategory = FGraphBuilder::PinCategoryInt32;
-							}
-							break;
+								case ELiteralType::Integer:
+								case ELiteralType::IntegerArray:
+								{
+									PinCategory = FGraphBuilder::PinCategoryInt32;
+								}
+								break;
 
-							case ELiteralType::String:
-							case ELiteralType::StringArray:
-							{
-								PinCategory = FGraphBuilder::PinCategoryString;
-							}
-							break;
+								case ELiteralType::String:
+								case ELiteralType::StringArray:
+								{
+									PinCategory = FGraphBuilder::PinCategoryString;
+								}
+								break;
 
-							case ELiteralType::UObjectProxy:
-							case ELiteralType::UObjectProxyArray:
-							{
-								PinCategory = FGraphBuilder::PinCategoryObject;
-							}
-							break;
+								case ELiteralType::UObjectProxy:
+								case ELiteralType::UObjectProxyArray:
+								{
+									PinCategory = FGraphBuilder::PinCategoryObject;
+								}
+								break;
 
-							case ELiteralType::None:
-							case ELiteralType::Invalid:
-							default:
-							{
-								// Audio types are ubiquitous, so added as subcategory
-								// to be able to stylize connections (i.e. wire color & wire animation)
-								if (DataTypeName == Frontend::GetDataTypeName<FAudioBuffer>())
+								case ELiteralType::None:
+								case ELiteralType::Invalid:
+								default:
 								{
-									PinCategory = FGraphBuilder::PinCategoryAudio;
+									// Audio types are ubiquitous, so added as subcategory
+									// to be able to stylize connections (i.e. wire color & wire animation)
+									if (DataTypeName == GetMetasoundDataTypeName<FAudioBuffer>())
+									{
+										PinCategory = FGraphBuilder::PinCategoryAudio;
+									}
+									static_assert(static_cast<int32>(ELiteralType::Invalid) == 12, "Possible missing binding of pin category to primitive type");
 								}
-								static_assert(static_cast<int32>(ELiteralType::Invalid) == 12, "Possible missing binding of pin category to primitive type");
+								break;
 							}
-							break;
 						}
-					}
 
-					const bool bIsArray = RegistryInfo.IsArrayType();
-					const EPinContainerType ContainerType = bIsArray ? EPinContainerType::Array : EPinContainerType::None;
-					FEdGraphPinType PinType(PinCategory, PinSubCategory, nullptr, ContainerType, false, FEdGraphTerminalType());
- 					UClass* ClassToUse = FMetasoundFrontendRegistryContainer::Get()->GetLiteralUClassForDataType(DataTypeName);
- 					PinType.PinSubCategoryObject = Cast<UObject>(ClassToUse);
+						const bool bIsArray = RegistryInfo.IsArrayType();
+						const EPinContainerType ContainerType = bIsArray ? EPinContainerType::Array : EPinContainerType::None;
+						FEdGraphPinType PinType(PinCategory, PinSubCategory, nullptr, ContainerType, false, FEdGraphTerminalType());
+						UClass* ClassToUse = DataTypeRegistry.GetUClassForDataType(DataTypeName);
+						PinType.PinSubCategoryObject = Cast<UObject>(ClassToUse);
 
-					DataTypeInfo.Emplace(DataTypeName, FEditorDataType(MoveTemp(PinType), RegistryInfo));
+						DataTypeInfo.Emplace(DataTypeName, FEditorDataType(MoveTemp(PinType), RegistryInfo));
+					}
 				}
 			}
 
diff --git a/Engine/Plugins/Runtime/Metasound/Source/MetasoundEditor/Private/MetasoundEditorModule.h b/Engine/Plugins/Runtime/Metasound/Source/MetasoundEditor/Private/MetasoundEditorModule.h
--- a/Engine/Plugins/Runtime/Metasound/Source/MetasoundEditor/Private/MetasoundEditorModule.h
+++ b/Engine/Plugins/Runtime/Metasound/Source/MetasoundEditor/Private/MetasoundEditorModule.h
@@ -4,6 +4,7 @@
 #include "CoreMinimal.h"
 #include "EdGraph/EdGraphNode.h"
 #include "EdGraph/EdGraphPin.h"
+#include "MetasoundFrontendDataTypeRegistry.h"
 #include "MetasoundFrontendLiteral.h"
 #include "MetasoundFrontendRegistries.h"
 #include "Modules/ModuleInterface.h"
diff --git a/Engine/Plugins/Runtime/Metasound/Source/MetasoundEditor/Private/MetasoundNodeDetailCustomization.cpp b/Engine/Plugins/Runtime/Metasound/Source/MetasoundEditor/Private/MetasoundNodeDetailCustomization.cpp
--- a/Engine/Plugins/Runtime/Metasound/Source/MetasoundEditor/Private/MetasoundNodeDetailCustomization.cpp
+++ b/Engine/Plugins/Runtime/Metasound/Source/MetasoundEditor/Private/MetasoundNodeDetailCustomization.cpp
@@ -150,7 +150,7 @@ namespace Metasound
 				TSharedPtr<IPropertyHandle> ValueProperty = StructPropertyHandle->GetChildHandle(GET_MEMBER_NAME_CHECKED(FMetasoundEditorGraphInputIntRef, Value));
 				if (ValueProperty.IsValid())
 				{
-					TSharedPtr<const IEnumDataTypeInterface> EnumInterface = Registry->GetEnumInterfaceForDataType(DataTypeName);
+					TSharedPtr<const IEnumDataTypeInterface> EnumInterface = IDataTypeRegistry::Get().GetEnumInterfaceForDataType(DataTypeName);
 
 					// Not an enum, so just display as underlying type (int32)
 					if (!EnumInterface.IsValid())
@@ -828,7 +828,7 @@ namespace Metasound
 			InDefaultPropertyHandle->SetInstanceMetaData(VariableCustomizationPrivate::DataTypeNameIdentifier, TypeNameString);
 
 			FDataTypeRegistryInfo DataTypeInfo;
-			if (!ensure(Registry->GetInfoForDataType(TypeName, DataTypeInfo)))
+			if (!ensure(IDataTypeRegistry::Get().GetDataTypeInfo(TypeName, DataTypeInfo)))
 			{
 				return;
 			}
@@ -975,7 +975,7 @@ namespace Metasound
 			InDefaultPropertyHandle->SetInstanceMetaData(VariableCustomizationPrivate::DataTypeNameIdentifier, TypeNameString);
 
 			FDataTypeRegistryInfo DataTypeInfo;
-			if (!ensure(Registry->GetInfoForDataType(TypeName, DataTypeInfo)))
+			if (!ensure(IDataTypeRegistry::Get().GetDataTypeInfo(TypeName, DataTypeInfo)))
 			{
 				return;
 			}
diff --git a/Engine/Plugins/Runtime/Metasound/Source/MetasoundEditor/Private/SMetasoundEnumPin.cpp b/Engine/Plugins/Runtime/Metasound/Source/MetasoundEditor/Private/SMetasoundEnumPin.cpp
--- a/Engine/Plugins/Runtime/Metasound/Source/MetasoundEditor/Private/SMetasoundEnumPin.cpp
+++ b/Engine/Plugins/Runtime/Metasound/Source/MetasoundEditor/Private/SMetasoundEnumPin.cpp
@@ -3,6 +3,7 @@
 #include "SMetasoundEnumPin.h"
 #include "MetasoundEditorGraph.h"
 #include "MetasoundEditorGraphNode.h"
+#include "MetasoundFrontendDataTypeRegistry.h"
 #include "ScopedTransaction.h"
 
 void SMetasoundEnumPin::Construct(const FArguments& InArgs, UEdGraphPin* InGraphPinObj)
@@ -36,7 +37,7 @@ SMetasoundEnumPin::FindEnumInterfaceFromPin(UEdGraphPin* InPin)
 	if (Inputs.Num() > 0)
 	{
 		FName DataType = Inputs[0]->GetDataType();
-		return FMetasoundFrontendRegistryContainer::Get()->GetEnumInterfaceForDataType(DataType);
+		return Metasound::Frontend::IDataTypeRegistry::Get().GetEnumInterfaceForDataType(DataType);
 	}
 	return nullptr;
 }
diff --git a/Engine/Plugins/Runtime/Metasound/Source/MetasoundEngine/Private/MetasoundEngineArchetypes.cpp b/Engine/Plugins/Runtime/Metasound/Source/MetasoundEngine/Private/MetasoundEngineArchetypes.cpp
--- a/Engine/Plugins/Runtime/Metasound/Source/MetasoundEngine/Private/MetasoundEngineArchetypes.cpp
+++ b/Engine/Plugins/Runtime/Metasound/Source/MetasoundEngine/Private/MetasoundEngineArchetypes.cpp
@@ -121,7 +121,7 @@ namespace Metasound
 				FMetasoundFrontendClassVertex OnPlayTrigger;
 				OnPlayTrigger.Name = GetOnPlayInputName();
 
-				OnPlayTrigger.TypeName = Metasound::Frontend::GetDataTypeName<Metasound::FTrigger>();
+				OnPlayTrigger.TypeName = GetMetasoundDataTypeName<FTrigger>();
 				OnPlayTrigger.Metadata.DisplayName = LOCTEXT("OnPlay", "On Play");
 				OnPlayTrigger.Metadata.Description = LOCTEXT("OnPlayTriggerToolTip", "Trigger executed when this source is played.");
 				OnPlayTrigger.VertexID = FGuid::NewGuid();
@@ -132,7 +132,7 @@ namespace Metasound
 				FMetasoundFrontendClassVertex OnFinished;
 				OnFinished.Name = GetIsFinishedOutputName();
 
-				OnFinished.TypeName = Metasound::Frontend::GetDataTypeName<Metasound::FTrigger>();
+				OnFinished.TypeName = GetMetasoundDataTypeName<FTrigger>();
 				OnFinished.Metadata.DisplayName = LOCTEXT("OnFinished", "On Finished");
 				OnFinished.Metadata.Description = LOCTEXT("OnFinishedToolTip", "Trigger executed to initiate stopping the source.");
 				OnFinished.VertexID = FGuid::NewGuid();
@@ -141,7 +141,7 @@ namespace Metasound
 
 				FMetasoundFrontendClassVertex GeneratedAudio;
 				GeneratedAudio.Name = GetAudioOutputName();
-				GeneratedAudio.TypeName = Metasound::Frontend::GetDataTypeName<Metasound::FMonoAudioFormat>();
+				GeneratedAudio.TypeName = GetMetasoundDataTypeName<FMonoAudioFormat>();
 				GeneratedAudio.Metadata.DisplayName = LOCTEXT("GeneratedMono", "Audio");
 				GeneratedAudio.Metadata.Description = LOCTEXT("GeneratedAudioToolTip", "The resulting output audio from this source.");
 				GeneratedAudio.VertexID = FGuid::NewGuid();
@@ -219,7 +219,7 @@ namespace Metasound
 				FMetasoundFrontendClassVertex OnPlayTrigger;
 				OnPlayTrigger.Name = GetOnPlayInputName();
 
-				OnPlayTrigger.TypeName = Metasound::Frontend::GetDataTypeName<Metasound::FTrigger>();
+				OnPlayTrigger.TypeName = GetMetasoundDataTypeName<FTrigger>();
 				OnPlayTrigger.Metadata.DisplayName = LOCTEXT("OnPlay", "On Play");
 				OnPlayTrigger.Metadata.Description = LOCTEXT("OnPlayTriggerToolTip", "Trigger executed when this source is played.");
 				OnPlayTrigger.VertexID = FGuid::NewGuid();
@@ -230,7 +230,7 @@ namespace Metasound
 				FMetasoundFrontendClassVertex OnFinished;
 				OnFinished.Name = GetIsFinishedOutputName();
 
-				OnFinished.TypeName = Metasound::Frontend::GetDataTypeName<Metasound::FTrigger>();
+				OnFinished.TypeName = GetMetasoundDataTypeName<FTrigger>();
 				OnFinished.Metadata.DisplayName = LOCTEXT("OnFinished", "On Finished");
 				OnFinished.Metadata.Description = LOCTEXT("OnFinishedToolTip", "Trigger executed to initiate stopping the source.");
 				OnFinished.VertexID = FGuid::NewGuid();
@@ -239,7 +239,7 @@ namespace Metasound
 
 				FMetasoundFrontendClassVertex GeneratedAudio;
 				GeneratedAudio.Name = GetAudioOutputName();
-				GeneratedAudio.TypeName = Metasound::Frontend::GetDataTypeName<Metasound::FStereoAudioFormat>();
+				GeneratedAudio.TypeName = GetMetasoundDataTypeName<FStereoAudioFormat>();
 				GeneratedAudio.Metadata.DisplayName = LOCTEXT("GeneratedStereo", "Audio");
 				GeneratedAudio.Metadata.Description = LOCTEXT("GeneratedAudioToolTip", "The resulting output audio from this source.");
 				GeneratedAudio.VertexID = FGuid::NewGuid();
@@ -318,7 +318,7 @@ namespace Metasound
 			{
 				FMetasoundFrontendClassVertex GeneratedAudio;
 				GeneratedAudio.Name = GetAudioOutputName();
-				GeneratedAudio.TypeName = Metasound::Frontend::GetDataTypeName<Metasound::FAudioBuffer>();
+				GeneratedAudio.TypeName = GetMetasoundDataTypeName<FAudioBuffer>();
 				GeneratedAudio.Metadata.DisplayName = LOCTEXT("GeneratedMono", "Audio");
 				GeneratedAudio.Metadata.Description = LOCTEXT("GeneratedAudioToolTip", "The resulting output audio from this source.");
 				GeneratedAudio.VertexID = FGuid::NewGuid();
@@ -335,7 +335,7 @@ namespace Metasound
 				FMetasoundFrontendClassVertex OnPlayTrigger;
 				OnPlayTrigger.Name = GetOnPlayInputName();
 
-				OnPlayTrigger.TypeName = Metasound::Frontend::GetDataTypeName<Metasound::FTrigger>();
+				OnPlayTrigger.TypeName = GetMetasoundDataTypeName<FTrigger>();
 				OnPlayTrigger.Metadata.DisplayName = LOCTEXT("OnPlay", "On Play");
 				OnPlayTrigger.Metadata.Description = LOCTEXT("OnPlayTriggerToolTip", "Trigger executed when this source is played.");
 				OnPlayTrigger.VertexID = FGuid::NewGuid();
@@ -346,7 +346,7 @@ namespace Metasound
 				FMetasoundFrontendClassVertex OnFinished;
 				OnFinished.Name = GetIsFinishedOutputName();
 
-				OnFinished.TypeName = Metasound::Frontend::GetDataTypeName<Metasound::FTrigger>();
+				OnFinished.TypeName = GetMetasoundDataTypeName<FTrigger>();
 				OnFinished.Metadata.DisplayName = LOCTEXT("OnFinished", "On Finished");
 				OnFinished.Metadata.Description = LOCTEXT("OnFinishedToolTip", "Trigger executed to initiate stopping the source.");
 				OnFinished.VertexID = FGuid::NewGuid();
@@ -496,7 +496,7 @@ namespace Metasound
 			{
 				FMetasoundFrontendClassVertex GeneratedLeftAudio;
 				GeneratedLeftAudio.Name = GetLeftAudioOutputName();
-				GeneratedLeftAudio.TypeName = Metasound::Frontend::GetDataTypeName<Metasound::FAudioBuffer>();
+				GeneratedLeftAudio.TypeName = GetMetasoundDataTypeName<FAudioBuffer>();
 				GeneratedLeftAudio.Metadata.DisplayName = LOCTEXT("GeneratedStereoLeft", "Left Audio");
 				GeneratedLeftAudio.Metadata.Description = LOCTEXT("GeneratedLeftAudioToolTip", "The resulting output audio from this source.");
 				GeneratedLeftAudio.VertexID = FGuid::NewGuid();
@@ -508,7 +508,7 @@ namespace Metasound
 			{
 				FMetasoundFrontendClassVertex GeneratedRightAudio;
 				GeneratedRightAudio.Name = GetRightAudioOutputName();
-				GeneratedRightAudio.TypeName = Metasound::Frontend::GetDataTypeName<Metasound::FAudioBuffer>();
+				GeneratedRightAudio.TypeName = GetMetasoundDataTypeName<FAudioBuffer>();
 				GeneratedRightAudio.Metadata.DisplayName = LOCTEXT("GeneratedStereoRight", "Right Audio");
 				GeneratedRightAudio.Metadata.Description = LOCTEXT("GeneratedRightAudioToolTip", "The resulting output audio from this source.");
 				GeneratedRightAudio.VertexID = FGuid::NewGuid();
@@ -525,7 +525,7 @@ namespace Metasound
 				FMetasoundFrontendClassVertex OnPlayTrigger;
 				OnPlayTrigger.Name = GetOnPlayInputName();
 
-				OnPlayTrigger.TypeName = Metasound::Frontend::GetDataTypeName<Metasound::FTrigger>();
+				OnPlayTrigger.TypeName = GetMetasoundDataTypeName<FTrigger>();
 				OnPlayTrigger.Metadata.DisplayName = LOCTEXT("OnPlay", "On Play");
 				OnPlayTrigger.Metadata.Description = LOCTEXT("OnPlayTriggerToolTip", "Trigger executed when this source is played.");
 				OnPlayTrigger.VertexID = FGuid::NewGuid();
@@ -536,7 +536,7 @@ namespace Metasound
 				FMetasoundFrontendClassVertex OnFinished;
 				OnFinished.Name = GetIsFinishedOutputName();
 
-				OnFinished.TypeName = Metasound::Frontend::GetDataTypeName<Metasound::FTrigger>();
+				OnFinished.TypeName = GetMetasoundDataTypeName<FTrigger>();
 				OnFinished.Metadata.DisplayName = LOCTEXT("OnFinished", "On Finished");
 				OnFinished.Metadata.Description = LOCTEXT("OnFinishedToolTip", "Trigger executed to initiate stopping the source.");
 				OnFinished.VertexID = FGuid::NewGuid();
diff --git a/Engine/Plugins/Runtime/Metasound/Source/MetasoundFrontend/Private/MetasoundAssetBase.cpp b/Engine/Plugins/Runtime/Metasound/Source/MetasoundFrontend/Private/MetasoundAssetBase.cpp
--- a/Engine/Plugins/Runtime/Metasound/Source/MetasoundFrontend/Private/MetasoundAssetBase.cpp
+++ b/Engine/Plugins/Runtime/Metasound/Source/MetasoundFrontend/Private/MetasoundAssetBase.cpp
@@ -20,7 +20,6 @@
 #include "MetasoundInstanceTransmitter.h"
 #include "MetasoundJsonBackend.h"
 #include "MetasoundLog.h"
-#include "MetasoundReceiveNode.h"
 #include "MetasoundTrace.h"
 #include "StructSerializer.h"
 #include "UObject/MetaData.h"
@@ -601,13 +600,12 @@ TSet<Metasound::FVertexKey> FMetasoundAssetBase::GetNonTransmittableInputVertice
 
 	// Do not transmit vertices which are not transmittable. Async communication 
 	// is not supported without transmission.
-	FMetasoundFrontendRegistryContainer* Registry = FMetasoundFrontendRegistryContainer::Get();
-	check(nullptr != Registry);
+	const IDataTypeRegistry& Registry = IDataTypeRegistry::Get();
 
 	auto IsNotTransmittable = [&Registry](const FMetasoundFrontendClassVertex& InVertex) -> bool
 	{	
 		FDataTypeRegistryInfo Info;
-		if (Registry->GetInfoForDataType(InVertex.TypeName, Info))
+		if (Registry.GetDataTypeInfo(InVertex.TypeName, Info))
 		{
 			return !Info.bIsTransmittable;
 		}
@@ -651,11 +649,13 @@ TArray<FString> FMetasoundAssetBase::GetTransmittableInputVertexNames() const
 
 	auto IsDataTypeTransmittable = [&](const FString& InVertexName)
 	{
+		using namespace Metasound::Frontend;
+
 		Frontend::FConstClassInputAccessPtr ClassInputPtr = RootGraph->FindClassInputWithName(InVertexName);
 		if (const FMetasoundFrontendClassInput* ClassInput = ClassInputPtr.Get())
 		{
 			Frontend::FDataTypeRegistryInfo TypeInfo;
-			if (Frontend::GetTraitsForDataType(ClassInput->TypeName, TypeInfo))
+			if (IDataTypeRegistry::Get().GetDataTypeInfo(ClassInput->TypeName, TypeInfo))
 			{
 				if (TypeInfo.bIsTransmittable)
 				{
diff --git a/Engine/Plugins/Runtime/Metasound/Source/MetasoundFrontend/Private/MetasoundFrontend.cpp b/Engine/Plugins/Runtime/Metasound/Source/MetasoundFrontend/Private/MetasoundFrontend.cpp
--- a/Engine/Plugins/Runtime/Metasound/Source/MetasoundFrontend/Private/MetasoundFrontend.cpp
+++ b/Engine/Plugins/Runtime/Metasound/Source/MetasoundFrontend/Private/MetasoundFrontend.cpp
@@ -136,16 +136,6 @@ namespace Metasound
 			return OutClass;
 		}
 
-		TArray<FName> GetAllAvailableDataTypes()
-		{
-			return FMetasoundFrontendRegistryContainer::Get()->GetAllValidDataTypes();
-		}
-
-		bool GetTraitsForDataType(FName InDataType, FDataTypeRegistryInfo& OutInfo)
-		{
-			return FMetasoundFrontendRegistryContainer::Get()->GetInfoForDataType(InDataType, OutInfo);
-		}
-
 		bool ImportJSONToMetasound(const FString& InJSON, FMetasoundFrontendDocument& OutMetasoundDocument)
 		{
 			TArray<uint8> ReadBuffer;
diff --git a/Engine/Plugins/Runtime/Metasound/Source/MetasoundFrontend/Private/MetasoundFrontendDataTypeRegistry.cpp b/Engine/Plugins/Runtime/Metasound/Source/MetasoundFrontend/Private/MetasoundFrontendDataTypeRegistry.cpp
new file mode 100644
--- /dev/null
+++ b/Engine/Plugins/Runtime/Metasound/Source/MetasoundFrontend/Private/MetasoundFrontendDataTypeRegistry.cpp
@@ -0,0 +1,681 @@
+// Copyright Epic Games, Inc. All Rights Reserved.
+
+
+#include "MetasoundFrontendDataTypeRegistry.h"
+
+namespace Metasound
+{
+	namespace Frontend
+	{
+		namespace MetasoundFrontendDataTypeRegistryPrivate
+		{
+			// Return the compatible literal with the most descriptive type.
+			// TODO: Currently TIsParsable<> allows for implicit conversion of
+			// constructor arguments of integral types which can cause some confusion
+			// here when trying to match a literal type to a constructor. For example:
+			//
+			// struct FBoolConstructibleType
+			// {
+			// 	FBoolConstructibleType(bool InValue);
+			// };
+			//
+			// static_assert(TIsParsable<FBoolConstructible, double>::Value); 
+			//
+			// Implicit conversions are currently allowed in TIsParsable because this
+			// is perfectly legal syntax.
+			//
+			// double Value = 10.0;
+			// FBoolConstructibleType BoolConstructible = Value;
+			//
+			// There are some tricks to possibly disable implicit conversions when
+			// checking for specific constructors, but they are yet to be implemented 
+			// and are untested. Here's the basic idea.
+			//
+			// template<DataType, DesiredIntegralArgType>
+			// struct TOnlyConvertIfIsSame
+			// {
+			// 		// Implicit conversion only defined if types match.
+			// 		template<typename SuppliedIntegralArgType, std::enable_if<std::is_same<std::decay<SuppliedIntegralArgType>::type, DesiredIntegralArgType>::value, int> = 0>
+			// 		operator DesiredIntegralArgType()
+			// 		{
+			// 			return DesiredIntegralArgType{};
+			// 		}
+			// };
+			//
+			// static_assert(false == std::is_constructible<FBoolConstructibleType, TOnlyConvertIfSame<double>>::value);
+			// static_assert(true == std::is_constructible<FBoolConstructibleType, TOnlyConvertIfSame<bool>>::value);
+			ELiteralType GetMostDescriptiveLiteralForDataType(const FDataTypeRegistryInfo& InDataTypeInfo)
+			{
+				if (InDataTypeInfo.bIsProxyArrayParsable)
+				{
+					return ELiteralType::UObjectProxyArray;
+				}
+				else if (InDataTypeInfo.bIsProxyParsable)
+				{
+					return ELiteralType::UObjectProxy;
+				}
+				else if (InDataTypeInfo.bIsEnum && InDataTypeInfo.bIsIntParsable)
+				{
+					return ELiteralType::Integer;
+				}
+				else if (InDataTypeInfo.bIsStringArrayParsable)
+				{
+					return ELiteralType::StringArray;
+				}
+				else if (InDataTypeInfo.bIsFloatArrayParsable)
+				{
+					return ELiteralType::FloatArray;
+				}
+				else if (InDataTypeInfo.bIsIntArrayParsable)
+				{
+					return ELiteralType::IntegerArray;
+				}
+				else if (InDataTypeInfo.bIsBoolArrayParsable)
+				{
+					return ELiteralType::BooleanArray;
+				}
+				else if (InDataTypeInfo.bIsStringParsable)
+				{
+					return ELiteralType::String;
+				}
+				else if (InDataTypeInfo.bIsFloatParsable)
+				{
+					return ELiteralType::Float;
+				}
+				else if (InDataTypeInfo.bIsIntParsable)
+				{
+					return ELiteralType::Integer;
+				}
+				else if (InDataTypeInfo.bIsBoolParsable)
+				{
+					return ELiteralType::Boolean;
+				}
+				else if (InDataTypeInfo.bIsDefaultArrayParsable)
+				{
+					return ELiteralType::NoneArray; 
+				}
+				else if (InDataTypeInfo.bIsDefaultParsable)
+				{
+					return ELiteralType::None;
+				}
+				else
+				{
+					// if we ever hit this, something has gone wrong with the REGISTER_METASOUND_DATATYPE macro.
+					// we should have failed to compile if any of these are false.
+					checkNoEntry();
+					return ELiteralType::Invalid;
+				}
+			}
+
+			// Node registry entry for input nodes created from a data type registry entry.
+			class FInputNodeRegistryEntry : public INodeRegistryEntry
+			{
+			public:
+				FInputNodeRegistryEntry() = delete;
+
+				FInputNodeRegistryEntry(TUniquePtr<IDataTypeRegistryEntry>&& InDataTypeEntry)
+				: DataTypeEntry(MoveTemp(InDataTypeEntry))
+				{
+					if (DataTypeEntry.IsValid())
+					{
+						FrontendClass = DataTypeEntry->GetFrontendInputClass();
+						ClassInfo = FNodeClassInfo(FrontendClass.Metadata);
+					}
+				}
+
+				virtual ~FInputNodeRegistryEntry() = default;
+
+				virtual const FNodeClassInfo& GetClassInfo() const override
+				{
+					return ClassInfo;
+				}
+
+				virtual TUniquePtr<INode> CreateNode(FDefaultNodeConstructorParams&&) const override
+				{
+					return nullptr;
+				}
+
+				virtual TUniquePtr<INode> CreateNode(FDefaultLiteralNodeConstructorParams&& InParams) const override
+				{
+					if (DataTypeEntry.IsValid())
+					{
+						return DataTypeEntry->CreateInputNode(MoveTemp(InParams));
+					}
+					return nullptr;
+				}
+
+				virtual TUniquePtr<INode> CreateNode(const FNodeInitData&) const override
+				{
+					return nullptr;
+				}
+
+				virtual const FMetasoundFrontendClass& GetFrontendClass() const override
+				{
+					return FrontendClass;
+				}
+
+				virtual TUniquePtr<INodeRegistryEntry> Clone() const override
+				{
+					if (DataTypeEntry.IsValid())
+					{
+						return MakeUnique<FInputNodeRegistryEntry>(DataTypeEntry->Clone());
+					}
+					return MakeUnique<FInputNodeRegistryEntry>(TUniquePtr<IDataTypeRegistryEntry>());
+				}
+
+				virtual bool IsNative() const override
+				{
+					return true;
+				}
+
+			private:
+				
+				TUniquePtr<IDataTypeRegistryEntry> DataTypeEntry;
+				FNodeClassInfo ClassInfo;
+				FMetasoundFrontendClass FrontendClass;
+			};
+
+			// Node registry entry for output nodes created from a data type registry entry.
+			class FOutputNodeRegistryEntry : public INodeRegistryEntry
+			{
+			public:
+				FOutputNodeRegistryEntry() = delete;
+
+				FOutputNodeRegistryEntry(TUniquePtr<IDataTypeRegistryEntry>&& InDataTypeEntry)
+				: DataTypeEntry(MoveTemp(InDataTypeEntry))
+				{
+					if (DataTypeEntry.IsValid())
+					{
+						FrontendClass = DataTypeEntry->GetFrontendOutputClass();
+						ClassInfo = FNodeClassInfo(FrontendClass.Metadata);
+					}
+				}
+
+				virtual ~FOutputNodeRegistryEntry() = default;
+
+				virtual const FNodeClassInfo& GetClassInfo() const override
+				{
+					return ClassInfo;
+				}
+
+				virtual TUniquePtr<INode> CreateNode(FDefaultNodeConstructorParams&& InParams) const override
+				{
+					if (DataTypeEntry.IsValid())
+					{
+						return DataTypeEntry->CreateOutputNode(MoveTemp(InParams));
+					}
+					return nullptr;
+				}
+
+				virtual TUniquePtr<INode> CreateNode(FDefaultLiteralNodeConstructorParams&& InParams) const override
+				{
+					return nullptr;
+				}
+
+				virtual TUniquePtr<INode> CreateNode(const FNodeInitData&) const override
+				{
+					return nullptr;
+				}
+
+				virtual const FMetasoundFrontendClass& GetFrontendClass() const override
+				{
+					return FrontendClass;
+				}
+
+				virtual TUniquePtr<INodeRegistryEntry> Clone() const override
+				{
+					if (DataTypeEntry.IsValid())
+					{
+						return MakeUnique<FOutputNodeRegistryEntry>(DataTypeEntry->Clone());
+					}
+					return MakeUnique<FOutputNodeRegistryEntry>(TUniquePtr<IDataTypeRegistryEntry>());
+				}
+
+				virtual bool IsNative() const override
+				{
+					return true;
+				}
+
+			private:
+				
+				TUniquePtr<IDataTypeRegistryEntry> DataTypeEntry;
+				FNodeClassInfo ClassInfo;
+				FMetasoundFrontendClass FrontendClass;
+			};
+
+			// Node registry entry for variable nodes created from a data type registry entry.
+			class FVariableNodeRegistryEntry : public INodeRegistryEntry
+			{
+			public:
+				FVariableNodeRegistryEntry() = delete;
+
+				FVariableNodeRegistryEntry(TUniquePtr<IDataTypeRegistryEntry>&& InDataTypeEntry)
+				: DataTypeEntry(MoveTemp(InDataTypeEntry))
+				{
+					if (DataTypeEntry.IsValid())
+					{
+						FrontendClass = DataTypeEntry->GetFrontendVariableClass();
+						ClassInfo = FNodeClassInfo(FrontendClass.Metadata);
+					}
+				}
+
+				virtual ~FVariableNodeRegistryEntry() = default;
+
+				virtual const FNodeClassInfo& GetClassInfo() const override
+				{
+					return ClassInfo;
+				}
+
+				virtual TUniquePtr<INode> CreateNode(FDefaultNodeConstructorParams&& InParams) const override
+				{
+					return nullptr;
+				}
+
+				virtual TUniquePtr<INode> CreateNode(FDefaultLiteralNodeConstructorParams&& InParams) const override
+				{
+					if (DataTypeEntry.IsValid())
+					{
+						return DataTypeEntry->CreateVariableNode(MoveTemp(InParams));
+					}
+					return nullptr;
+				}
+
+				virtual TUniquePtr<INode> CreateNode(const FNodeInitData&) const override
+				{
+					return nullptr;
+				}
+
+				virtual const FMetasoundFrontendClass& GetFrontendClass() const override
+				{
+					return FrontendClass;
+				}
+
+				virtual TUniquePtr<INodeRegistryEntry> Clone() const override
+				{
+					if (DataTypeEntry.IsValid())
+					{
+						return MakeUnique<FVariableNodeRegistryEntry>(DataTypeEntry->Clone());
+					}
+					return MakeUnique<FVariableNodeRegistryEntry>(TUniquePtr<IDataTypeRegistryEntry>());
+				}
+
+				virtual bool IsNative() const override
+				{
+					return true;
+				}
+
+			private:
+				
+				TUniquePtr<IDataTypeRegistryEntry> DataTypeEntry;
+				FNodeClassInfo ClassInfo;
+				FMetasoundFrontendClass FrontendClass;
+			};
+
+			class FDataTypeRegistry : public IDataTypeRegistry
+			{
+			public:
+				virtual ~FDataTypeRegistry() = default;
+
+				/** Register a data type
+				 * @param InName - Name of data type.
+				 * @param InEntry - TUniquePtr to data type registry entry.
+				 *
+				 * @return True on success, false on failure.
+				 */
+				virtual bool RegisterDataType(TUniquePtr<IDataTypeRegistryEntry>&& InEntry) override;
+
+				virtual void GetRegisteredDataTypeNames(TArray<FName>& OutNames) const override;
+
+				virtual bool GetDataTypeInfo(const FName& InDataType, FDataTypeRegistryInfo& OutInfo) const override;
+
+				// Return the enum interface for a data type. If the data type does not have 
+				// an enum interface, returns a nullptr.
+				virtual TSharedPtr<const IEnumDataTypeInterface> GetEnumInterfaceForDataType(const FName& InDataType) const override;
+
+				virtual ELiteralType GetDesiredLiteralType(const FName& InDataType) const override;
+
+				virtual bool IsLiteralTypeSupported(const FName& InDataType, ELiteralType InLiteralType) const override;
+				virtual bool IsLiteralTypeSupported(const FName& InDataType, EMetasoundFrontendLiteralType InLiteralType) const override;
+
+				virtual UClass* GetUClassForDataType(const FName& InDataType) const override;
+
+				virtual FLiteral CreateDefaultLiteral(const FName& InDataType) const override;
+				virtual FLiteral CreateLiteralFromUObject(const FName& InDataType, UObject* InObject) const override;
+				virtual FLiteral CreateLiteralFromUObjectArray(const FName& InDataType, TArray<UObject*> InObjectArray) const override;
+
+				virtual TSharedPtr<IDataChannel, ESPMode::ThreadSafe> CreateDataChannel(const FName& InDataType, const FOperatorSettings& InOperatorSettings) const override;
+
+				virtual bool GetFrontendInputClass(const FName& InDataType, FMetasoundFrontendClass& OutClass) const override;
+
+				virtual bool GetFrontendVariableClass(const FName& InDataType, FMetasoundFrontendClass& OutClass) const override;
+
+				virtual bool GetFrontendOutputClass(const FName& InDataType, FMetasoundFrontendClass& OutClass) const override;
+
+				// Create a new instance of a C++ implemented node from the registry.
+				virtual TUniquePtr<INode> CreateInputNode(const FName& InInputType, FInputNodeConstructorParams&& InParams) const override;
+				virtual TUniquePtr<INode> CreateVariableNode(const FName& InVariableType, FVariableNodeConstructorParams&& InParams) const override;
+				virtual TUniquePtr<INode> CreateOutputNode(const FName& InOutputType, FOutputNodeConstructorParams&& InParams) const override;
+				virtual TUniquePtr<INode> CreateReceiveNode(const FName& InDataType, const FNodeInitData& InParams) const override;
+
+			private:
+
+				const IDataTypeRegistryEntry* FindDataTypeEntry(const FName& InDataTypeName) const;
+
+				TMap<FName, TUniquePtr<IDataTypeRegistryEntry>> RegisteredDataTypes;
+			};
+
+			bool FDataTypeRegistry::RegisterDataType(TUniquePtr<IDataTypeRegistryEntry>&& InEntry)
+			{
+				if (InEntry.IsValid())
+				{
+					const FName Name = InEntry->GetDataTypeInfo().DataTypeName;
+
+					if (!ensureAlwaysMsgf(!RegisteredDataTypes.Contains(Name),
+						TEXT("Name collision when trying to register Metasound Data Type [Name:%s]. DataType must have "
+							"unique name and REGISTER_METASOUND_DATATYPE cannot be called in a public header."),
+							*Name.ToString()))
+					{
+						return false;
+					}
+
+					// Register nodes associated with data type.
+					FMetasoundFrontendRegistryContainer* NodeRegistry = FMetasoundFrontendRegistryContainer::Get();
+					if (ensure(nullptr != NodeRegistry))
+					{
+						NodeRegistry->RegisterNode(MakeUnique<FInputNodeRegistryEntry>(InEntry->Clone()));
+						NodeRegistry->RegisterNode(MakeUnique<FOutputNodeRegistryEntry>(InEntry->Clone()));
+						NodeRegistry->RegisterNode(MakeUnique<FVariableNodeRegistryEntry>(InEntry->Clone()));
+					}
+
+					RegisteredDataTypes.Add(Name, MoveTemp(InEntry));
+
+					UE_LOG(LogMetaSound, Verbose, TEXT("Registered Metasound Datatype [Name:%s]."), *Name.ToString());
+					return true;
+				}
+
+				return false;
+			}
+
+			void FDataTypeRegistry::GetRegisteredDataTypeNames(TArray<FName>& OutNames) const
+			{
+				RegisteredDataTypes.GetKeys(OutNames);
+			}
+
+			bool FDataTypeRegistry::GetDataTypeInfo(const FName& InDataType, FDataTypeRegistryInfo& OutInfo) const
+			{
+				if (const IDataTypeRegistryEntry* Entry = FindDataTypeEntry(InDataType))
+				{
+					OutInfo = Entry->GetDataTypeInfo();
+					return true;
+				}
+				return false;
+			}
+
+			// Return the enum interface for a data type. If the data type does not have 
+			// an enum interface, returns a nullptr.
+			TSharedPtr<const IEnumDataTypeInterface> FDataTypeRegistry::GetEnumInterfaceForDataType(const FName& InDataType) const
+			{
+				if (const IDataTypeRegistryEntry* Entry = FindDataTypeEntry(InDataType))
+				{
+					return Entry->GetEnumInterface();
+				}
+				return nullptr;
+			}
+
+			ELiteralType FDataTypeRegistry::GetDesiredLiteralType(const FName& InDataType) const
+			{
+				if (const IDataTypeRegistryEntry* Entry = FindDataTypeEntry(InDataType))
+				{
+					const FDataTypeRegistryInfo& Info = Entry->GetDataTypeInfo();
+
+					// If there's a designated preferred literal type for this datatype, use that.
+					if (Info.PreferredLiteralType != Metasound::ELiteralType::Invalid)
+					{
+						return Info.PreferredLiteralType;
+					}
+
+					// Otherwise, we opt for the highest precision construction option available.
+					return MetasoundFrontendDataTypeRegistryPrivate::GetMostDescriptiveLiteralForDataType(Info);
+				}
+				return Metasound::ELiteralType::Invalid;
+			}
+
+			bool FDataTypeRegistry::IsLiteralTypeSupported(const FName& InDataType, ELiteralType InLiteralType) const
+			{
+				if (const IDataTypeRegistryEntry* Entry = FindDataTypeEntry(InDataType))
+				{
+					const FDataTypeRegistryInfo& Info = Entry->GetDataTypeInfo();
+
+					switch (InLiteralType)
+					{
+						case Metasound::ELiteralType::Boolean:
+						{
+							return Info.bIsBoolParsable;
+						}
+						case Metasound::ELiteralType::BooleanArray:
+						{
+							return Info.bIsBoolArrayParsable;
+						}
+
+						case Metasound::ELiteralType::Integer:
+						{
+							return Info.bIsIntParsable;
+						}
+						case Metasound::ELiteralType::IntegerArray:
+						{
+							return Info.bIsIntArrayParsable;
+						}
+
+						case Metasound::ELiteralType::Float:
+						{
+							return Info.bIsFloatParsable;
+						}
+						case Metasound::ELiteralType::FloatArray:
+						{
+							return Info.bIsFloatArrayParsable;
+						}
+
+						case Metasound::ELiteralType::String:
+						{
+							return Info.bIsStringParsable;
+						}
+						case Metasound::ELiteralType::StringArray:
+						{
+							return Info.bIsStringArrayParsable;
+						}
+
+						case Metasound::ELiteralType::UObjectProxy:
+						{
+							return Info.bIsProxyParsable;
+						}
+						case Metasound::ELiteralType::UObjectProxyArray:
+						{
+							return Info.bIsProxyArrayParsable;
+						}
+
+						case Metasound::ELiteralType::None:
+						{
+							return Info.bIsDefaultParsable;
+						}
+						case Metasound::ELiteralType::NoneArray:
+						{
+							return Info.bIsDefaultArrayParsable;
+						}
+
+						case Metasound::ELiteralType::Invalid:
+						default:
+						{
+							static_assert(static_cast<int32>(Metasound::ELiteralType::COUNT) == 13, "Possible missing case coverage for ELiteralType");
+							return false;
+						}
+					}
+				}
+
+				return false;
+			}
+
+			bool FDataTypeRegistry::IsLiteralTypeSupported(const FName& InDataType, EMetasoundFrontendLiteralType InLiteralType) const
+			{
+				return IsLiteralTypeSupported(InDataType, GetMetasoundLiteralType(InLiteralType));
+			}
+
+			UClass* FDataTypeRegistry::GetUClassForDataType(const FName& InDataType) const
+			{
+				if (const IDataTypeRegistryEntry* Entry = FindDataTypeEntry(InDataType))
+				{
+					return Entry->GetDataTypeInfo().ProxyGeneratorClass;
+				}
+
+				return nullptr;
+			}
+
+			FLiteral FDataTypeRegistry::CreateDefaultLiteral(const FName& InDataType) const 
+			{
+				return FLiteral::GetDefaultForType(GetDesiredLiteralType(InDataType));
+			}
+
+			FLiteral FDataTypeRegistry::CreateLiteralFromUObject(const FName& InDataType, UObject* InObject) const
+			{
+				if (const IDataTypeRegistryEntry* Entry = FindDataTypeEntry(InDataType))
+				{
+					if (Audio::IProxyDataPtr ProxyPtr = Entry->CreateProxy(InObject))
+					{
+						if (InObject)
+						{
+							ensureAlwaysMsgf(ProxyPtr.IsValid(), TEXT("UObject failed to create a valid proxy."));
+						}
+						return Metasound::FLiteral(MoveTemp(ProxyPtr));
+					}
+				}
+
+				return Metasound::FLiteral();
+			}
+
+			FLiteral FDataTypeRegistry::CreateLiteralFromUObjectArray(const FName& InDataType, TArray<UObject*> InObjectArray) const
+			{
+				if (const IDataTypeRegistryEntry* Entry = FindDataTypeEntry(InDataType))
+				{
+					TArray<Audio::IProxyDataPtr> ProxyArray;
+
+					for (UObject* InObject : InObjectArray)
+					{
+						Audio::IProxyDataPtr ProxyPtr = Entry->CreateProxy(InObject);
+
+						if (InObject)
+						{
+							if (!ProxyPtr.IsValid())
+							{
+								UE_LOG(LogMetaSound, Error, TEXT("Failed to create a valid proxy from UObject [Name:%s]."), *InObject->GetName());
+							}
+						}
+
+						ProxyArray.Add(MoveTemp(ProxyPtr));
+					}
+
+					return Metasound::FLiteral(MoveTemp(ProxyArray));
+				}
+				else
+				{
+					return Metasound::FLiteral();
+				}
+			}
+
+			TSharedPtr<IDataChannel, ESPMode::ThreadSafe> FDataTypeRegistry::CreateDataChannel(const FName& InDataType, const FOperatorSettings& InOperatorSettings) const
+			{
+				if (const IDataTypeRegistryEntry* Entry = FindDataTypeEntry(InDataType))
+				{
+					return Entry->CreateDataChannel(InOperatorSettings);
+				}
+				return nullptr;
+			}
+
+			bool FDataTypeRegistry::GetFrontendInputClass(const FName& InDataType, FMetasoundFrontendClass& OutClass) const
+			{
+				if (const IDataTypeRegistryEntry* Entry = FindDataTypeEntry(InDataType))
+				{
+					OutClass = Entry->GetFrontendInputClass();
+					return true;
+				}
+				return false;
+			}
+
+			bool FDataTypeRegistry::GetFrontendVariableClass(const FName& InDataType, FMetasoundFrontendClass& OutClass) const
+			{
+				if (const IDataTypeRegistryEntry* Entry = FindDataTypeEntry(InDataType))
+				{
+					OutClass = Entry->GetFrontendVariableClass();
+					return true;
+				}
+				return false;
+			}
+
+			bool FDataTypeRegistry::GetFrontendOutputClass(const FName& InDataType, FMetasoundFrontendClass& OutClass) const
+			{
+				if (const IDataTypeRegistryEntry* Entry = FindDataTypeEntry(InDataType))
+				{
+					OutClass = Entry->GetFrontendOutputClass();
+					return true;
+				}
+				return false;
+			}
+
+			TUniquePtr<INode> FDataTypeRegistry::CreateInputNode(const FName& InDataType, FInputNodeConstructorParams&& InParams) const
+			{
+				if (const IDataTypeRegistryEntry* Entry = FindDataTypeEntry(InDataType))
+				{
+					return Entry->CreateInputNode(MoveTemp(InParams));
+				}
+				return nullptr;
+			}
+
+			TUniquePtr<INode> FDataTypeRegistry::CreateVariableNode(const FName& InDataType, FVariableNodeConstructorParams&& InParams) const
+			{
+				if (const IDataTypeRegistryEntry* Entry = FindDataTypeEntry(InDataType))
+				{
+					return Entry->CreateVariableNode(MoveTemp(InParams));
+				}
+				return nullptr;
+			}
+
+			TUniquePtr<INode> FDataTypeRegistry::CreateOutputNode(const FName& InDataType, FOutputNodeConstructorParams&& InParams) const
+			{
+				if (const IDataTypeRegistryEntry* Entry = FindDataTypeEntry(InDataType))
+				{
+					return Entry->CreateOutputNode(MoveTemp(InParams));
+				}
+				return nullptr;
+			}
+
+			TUniquePtr<INode> FDataTypeRegistry::CreateReceiveNode(const FName& InDataType, const FNodeInitData& InParams) const
+			{
+				if (const IDataTypeRegistryEntry* Entry = FindDataTypeEntry(InDataType))
+				{
+					return Entry->CreateReceiveNode(InParams);
+				}
+				return nullptr;
+			}
+
+			const IDataTypeRegistryEntry* FDataTypeRegistry::FindDataTypeEntry(const FName& InDataTypeName) const
+			{
+				const TUniquePtr<IDataTypeRegistryEntry>* Entry = RegisteredDataTypes.Find(InDataTypeName);
+
+				if (ensureMsgf(nullptr != Entry, TEXT("Data type not registered [Name:%s]"), *InDataTypeName.ToString()))
+				{
+					if (Entry->IsValid())
+					{
+						return Entry->Get();
+					}
+				}
+
+				return nullptr;
+			}
+		}
+
+		IDataTypeRegistry& IDataTypeRegistry::Get()
+		{
+			static MetasoundFrontendDataTypeRegistryPrivate::FDataTypeRegistry Registry;
+			return Registry;
+		}
+
+	}
+}
diff --git a/Engine/Plugins/Runtime/Metasound/Source/MetasoundFrontend/Private/MetasoundFrontendDocument.cpp b/Engine/Plugins/Runtime/Metasound/Source/MetasoundFrontend/Private/MetasoundFrontendDocument.cpp
--- a/Engine/Plugins/Runtime/Metasound/Source/MetasoundFrontend/Private/MetasoundFrontendDocument.cpp
+++ b/Engine/Plugins/Runtime/Metasound/Source/MetasoundFrontend/Private/MetasoundFrontendDocument.cpp
@@ -133,7 +133,9 @@ FMetasoundFrontendClassMetadata::FMetasoundFrontendClassMetadata(const Metasound
 FMetasoundFrontendClassInput::FMetasoundFrontendClassInput(const FMetasoundFrontendClassVertex& InOther)
 :	FMetasoundFrontendClassVertex(InOther)
 {
-	EMetasoundFrontendLiteralType DefaultType = Metasound::Frontend::GetMetasoundFrontendLiteralType(FMetasoundFrontendRegistryContainer::Get()->GetDesiredLiteralTypeForDataType(InOther.TypeName));
+	using namespace Metasound::Frontend;
+
+	EMetasoundFrontendLiteralType DefaultType = GetMetasoundFrontendLiteralType(IDataTypeRegistry::Get().GetDesiredLiteralType(InOther.TypeName));
 
 	DefaultLiteral.SetType(DefaultType);
 }
@@ -141,7 +143,9 @@ FMetasoundFrontendClassInput::FMetasoundFrontendClassInput(const FMetasoundFront
 FMetasoundFrontendClassVariable::FMetasoundFrontendClassVariable(const FMetasoundFrontendClassVertex& InOther)
 	: FMetasoundFrontendClassVertex(InOther)
 {
-	EMetasoundFrontendLiteralType DefaultType = Metasound::Frontend::GetMetasoundFrontendLiteralType(FMetasoundFrontendRegistryContainer::Get()->GetDesiredLiteralTypeForDataType(InOther.TypeName));
+	using namespace Metasound::Frontend;
+
+	EMetasoundFrontendLiteralType DefaultType = GetMetasoundFrontendLiteralType(IDataTypeRegistry::Get().GetDesiredLiteralType(InOther.TypeName));
 
 	DefaultLiteral.SetType(DefaultType);
 }
diff --git a/Engine/Plugins/Runtime/Metasound/Source/MetasoundFrontend/Private/MetasoundFrontendGraph.cpp b/Engine/Plugins/Runtime/Metasound/Source/MetasoundFrontend/Private/MetasoundFrontendGraph.cpp
--- a/Engine/Plugins/Runtime/Metasound/Source/MetasoundFrontend/Private/MetasoundFrontendGraph.cpp
+++ b/Engine/Plugins/Runtime/Metasound/Source/MetasoundFrontend/Private/MetasoundFrontendGraph.cpp
@@ -8,6 +8,7 @@
 #include "Algo/Transform.h"
 #include "CoreMinimal.h"
 #include "MetasoundFrontend.h"
+#include "MetasoundFrontendDataTypeRegistry.h"
 #include "MetasoundGraph.h"
 #include "MetasoundLog.h"
 #include "MetasoundNodeInterface.h"
@@ -160,15 +161,18 @@ namespace Metasound
 
 	TUniquePtr<INode> FFrontendGraphBuilder::CreateInputNode(const FMetasoundFrontendNode& InNode, const FMetasoundFrontendClass& InClass, const FMetasoundFrontendClassInput& InOwningGraphClassInput)
 	{
+		using namespace Metasound::Frontend;
+
 		const FMetasoundFrontendLiteral* FrontendLiteral = FindInputLiteralForInputNode(InNode, InClass, InOwningGraphClassInput);
 
 		if (nullptr != FrontendLiteral)
 		{
 			if (ensure(InNode.Interface.Inputs.Num() == 1))
 			{
+				IDataTypeRegistry& DataTypeRegistry = IDataTypeRegistry::Get();
 				const FMetasoundFrontendVertex& InputVertex = InNode.Interface.Inputs[0];
 
-				const bool IsLiteralParsableByDataType = Frontend::DoesDataTypeSupportLiteralType(InputVertex.TypeName, FrontendLiteral->GetType());
+				const bool IsLiteralParsableByDataType = DataTypeRegistry.IsLiteralTypeSupported(InputVertex.TypeName, FrontendLiteral->GetType());
 
 				if (IsLiteralParsableByDataType)
 				{
@@ -182,8 +186,7 @@ namespace Metasound
 						MoveTemp(Literal)
 					};
 
-					// TODO: create input node using external class definition
-					return FMetasoundFrontendRegistryContainer::Get()->CreateInputNode(InputVertex.TypeName, MoveTemp(InitParams));
+					return DataTypeRegistry.CreateInputNode(InputVertex.TypeName, MoveTemp(InitParams));
 				}
 				else
 				{
@@ -201,6 +204,8 @@ namespace Metasound
 
 	TUniquePtr<INode> FFrontendGraphBuilder::CreateOutputNode(const FMetasoundFrontendNode& InNode, const FMetasoundFrontendClass& InClass, FBuildGraphContext& InGraphContext)
 	{
+		using namespace Metasound::Frontend;
+
 		check(InClass.Metadata.GetType() == EMetasoundFrontendClassType::Output);
 		check(InNode.ClassID == InClass.ID);
 
@@ -224,8 +229,7 @@ namespace Metasound
 				}
 			}
 
-			// TODO: create output node using external class definition
-			return FMetasoundFrontendRegistryContainer::Get()->CreateOutputNode(OutputVertex.TypeName, MoveTemp(InitParams));
+			return IDataTypeRegistry::Get().CreateOutputNode(OutputVertex.TypeName, MoveTemp(InitParams));
 		}
 		return TUniquePtr<INode>(nullptr);
 	}
@@ -522,6 +526,8 @@ namespace Metasound
 
 	void FFrontendGraphBuilder::AddDefaultInputVariables(FBuildGraphContext& InGraphContext)
 	{
+		using namespace Metasound::Frontend;
+
 		using FIterator = FDefaultInputByIDMap::TIterator;
 
 		for (FDefaultInputByIDMap::ElementType& Pair : InGraphContext.DefaultInputs)
@@ -531,7 +537,7 @@ namespace Metasound
 
 			// 1. Construct and add the default variable to the graph
 			{
-				TUniquePtr<const INode> DefaultVariable = FMetasoundFrontendRegistryContainer::Get()->CreateVariableNode(VariableData.TypeName, MoveTemp(VariableData.InitParams));
+				TUniquePtr<const INode> DefaultVariable = IDataTypeRegistry::Get().CreateVariableNode(VariableData.TypeName, MoveTemp(VariableData.InitParams));
 				InGraphContext.Graph->AddNode(VariableNodeID, TSharedPtr<const INode>(DefaultVariable.Release()));
 			}
 
diff --git a/Engine/Plugins/Runtime/Metasound/Source/MetasoundFrontend/Private/MetasoundFrontendInjectReceiveNodes.cpp b/Engine/Plugins/Runtime/Metasound/Source/MetasoundFrontend/Private/MetasoundFrontendInjectReceiveNodes.cpp
--- a/Engine/Plugins/Runtime/Metasound/Source/MetasoundFrontend/Private/MetasoundFrontendInjectReceiveNodes.cpp
+++ b/Engine/Plugins/Runtime/Metasound/Source/MetasoundFrontend/Private/MetasoundFrontendInjectReceiveNodes.cpp
@@ -120,7 +120,7 @@ namespace Metasound
 					FString::Format(TEXT("ReceiveInject_{0}"), { InVertexKey }),
 					InID
 				};
-				return FMetasoundFrontendRegistryContainer::Get()->CreateReceiveNode(InDataType, MoveTemp(ReceiveNodeInitData));
+				return IDataTypeRegistry::Get().CreateReceiveNode(InDataType, ReceiveNodeInitData);
 			}
 		}
 
diff --git a/Engine/Plugins/Runtime/Metasound/Source/MetasoundFrontend/Private/MetasoundFrontendLiteral.cpp b/Engine/Plugins/Runtime/Metasound/Source/MetasoundFrontend/Private/MetasoundFrontendLiteral.cpp
--- a/Engine/Plugins/Runtime/Metasound/Source/MetasoundFrontend/Private/MetasoundFrontendLiteral.cpp
+++ b/Engine/Plugins/Runtime/Metasound/Source/MetasoundFrontend/Private/MetasoundFrontendLiteral.cpp
@@ -4,6 +4,7 @@
 
 #include "CoreMinimal.h"
 #include "MetasoundLog.h"
+#include "MetasoundFrontendDataTypeRegistry.h"
 #include "MetasoundFrontendRegistries.h"
 
 FString LexToString(Metasound::FLiteral::FNone InValue)
@@ -483,15 +484,17 @@ Metasound::FLiteral FMetasoundFrontendLiteral::ToLiteral(const FName& InMetasoun
 	using namespace Metasound;
 	using namespace Metasound::Frontend;
 
+
 	FLiteral Literal = FLiteral::CreateInvalid();
+	const IDataTypeRegistry& DataTypeRegistry = IDataTypeRegistry::Get();
 
-	const bool bIsTypeSupported = DoesDataTypeSupportLiteralType(InMetasoundDataTypeName, Type);
+	const bool bIsTypeSupported = DataTypeRegistry.IsLiteralTypeSupported(InMetasoundDataTypeName, Type);
 
 	if (!bIsTypeSupported)
 	{
 		UE_LOG(LogMetaSound, Error, TEXT("Reverting to default literal type for data type. Failed to create supported Metasound::FLiteral for data type [Name:%s] with FMetasoundFrontendLiteral [Literal:%s]"), *InMetasoundDataTypeName.ToString(), *ToString());
 
-		Literal = GetDefaultParamForDataType(InMetasoundDataTypeName);
+		Literal = DataTypeRegistry.CreateDefaultLiteral(InMetasoundDataTypeName);
 	}
 	else if (EMetasoundFrontendLiteralType::UObject == Type)
 	{
@@ -499,22 +502,14 @@ Metasound::FLiteral FMetasoundFrontendLiteral::ToLiteral(const FName& InMetasoun
 		{
 			// UObject proxies must go through the registry. The registry contains the information
 			// needed to convert UObjects to proxies. 
-			FMetasoundFrontendRegistryContainer* Registry = FMetasoundFrontendRegistryContainer::Get();
-			if (ensure(Registry))
-			{
-				Literal = Registry->CreateLiteralFromUObject(InMetasoundDataTypeName, AsUObject[0]);
-			}
+			Literal = DataTypeRegistry.CreateLiteralFromUObject(InMetasoundDataTypeName, AsUObject[0]);
 		}
 	}
 	else if (EMetasoundFrontendLiteralType::UObjectArray == Type)
 	{
 		// UObject proxies must go through the registry. The registry contains the information
 		// needed to convert UObjects to proxies. 
-		FMetasoundFrontendRegistryContainer* Registry = FMetasoundFrontendRegistryContainer::Get();
-		if (ensure(Registry))
-		{
-			Literal = Registry->CreateLiteralFromUObjectArray(InMetasoundDataTypeName, AsUObject);
-		}
+		Literal = DataTypeRegistry.CreateLiteralFromUObjectArray(InMetasoundDataTypeName, AsUObject);
 	}
 	else 
 	{
@@ -524,7 +519,7 @@ Metasound::FLiteral FMetasoundFrontendLiteral::ToLiteral(const FName& InMetasoun
 
 	// The support for the data type should be the same whether we pass in an
 	// Metasound::ELiteralType or a EMetasoundFrontendLiteralType
-	check(bIsTypeSupported == DoesDataTypeSupportLiteralType(InMetasoundDataTypeName, Literal.GetType()));
+	check(bIsTypeSupported == DataTypeRegistry.IsLiteralTypeSupported(InMetasoundDataTypeName, Literal.GetType()));
 
 	return Literal;
 }
@@ -679,34 +674,6 @@ namespace Metasound
 			static_assert(static_cast<uint8>(EMetasoundFrontendLiteralType::Invalid) == static_cast<uint8>(ELiteralType::Invalid), "Possible literal type enum mismatch");
 			return static_cast<ELiteralType>(InLiteralType);
 		}
-
-		bool DoesDataTypeSupportLiteralType(FName InDataType, EMetasoundFrontendLiteralType InFrontendLiteralType)
-		{
-			ELiteralType CoreLiteralType = GetMetasoundLiteralType(InFrontendLiteralType);
-			return DoesDataTypeSupportLiteralType(InDataType, CoreLiteralType);
-		}
-
-		bool DoesDataTypeSupportLiteralType(FName InDataType, ELiteralType InLiteralType)
-		{
-			FMetasoundFrontendRegistryContainer* Registry = FMetasoundFrontendRegistryContainer::Get();
-			if (ensure(Registry))
-			{
-				return Registry->DoesDataTypeSupportLiteralType(InDataType, InLiteralType);
-			}
-			return false;
-		}
-
-		Metasound::FLiteral GetDefaultParamForDataType(FName InDataType)
-		{
-			FMetasoundFrontendRegistryContainer* Registry = FMetasoundFrontendRegistryContainer::Get();
-			if (ensure(Registry))
-			{
-				ELiteralType DesiredArgType = Registry->GetDesiredLiteralTypeForDataType(InDataType);
-
-				return FLiteral::GetDefaultForType(DesiredArgType);
-			}
-			return FLiteral::CreateInvalid();
-		}
 	}
 }
 
diff --git a/Engine/Plugins/Runtime/Metasound/Source/MetasoundFrontend/Private/MetasoundFrontendRegistries.cpp b/Engine/Plugins/Runtime/Metasound/Source/MetasoundFrontend/Private/MetasoundFrontendRegistries.cpp
--- a/Engine/Plugins/Runtime/Metasound/Source/MetasoundFrontend/Private/MetasoundFrontendRegistries.cpp
+++ b/Engine/Plugins/Runtime/Metasound/Source/MetasoundFrontend/Private/MetasoundFrontendRegistries.cpp
@@ -4,6 +4,7 @@
 
 #include "Algo/ForEach.h"
 #include "CoreMinimal.h"
+#include "MetasoundFrontendDataTypeRegistry.h"
 #include "MetasoundFrontendRegistryTransaction.h"
 #include "MetasoundLog.h"
 #include "MetasoundRouter.h"
@@ -54,308 +55,6 @@ namespace Metasound
 			}
 
 
-			// Return the compatible literal with the most descriptive type.
-			// TODO: Currently TIsParsable<> allows for implicit conversion of
-			// constructor arguments of integral types which can cause some confusion
-			// here when trying to match a literal type to a constructor. For example:
-			//
-			// struct FBoolConstructibleType
-			// {
-			// 	FBoolConstructibleType(bool InValue);
-			// };
-			//
-			// static_assert(TIsParsable<FBoolConstructible, double>::Value); 
-			//
-			// Implicit conversions are currently allowed in TIsParsable because this
-			// is perfectly legal syntax.
-			//
-			// double Value = 10.0;
-			// FBoolConstructibleType BoolConstructible = Value;
-			//
-			// There are some tricks to possibly disable implicit conversions when
-			// checking for specific constructors, but they are yet to be implemented 
-			// and are untested. Here's the basic idea.
-			//
-			// template<DataType, DesiredIntegralArgType>
-			// struct TOnlyConvertIfIsSame
-			// {
-			// 		// Implicit conversion only defined if types match.
-			// 		template<typename SuppliedIntegralArgType, std::enable_if<std::is_same<std::decay<SuppliedIntegralArgType>::type, DesiredIntegralArgType>::value, int> = 0>
-			// 		operator DesiredIntegralArgType()
-			// 		{
-			// 			return DesiredIntegralArgType{};
-			// 		}
-			// };
-			//
-			// static_assert(false == std::is_constructible<FBoolConstructibleType, TOnlyConvertIfSame<double>>::value);
-			// static_assert(true == std::is_constructible<FBoolConstructibleType, TOnlyConvertIfSame<bool>>::value);
-			ELiteralType GetMostDescriptiveLiteralForDataType(const FDataTypeRegistryInfo& InDataTypeInfo)
-			{
-				if (InDataTypeInfo.bIsProxyArrayParsable)
-				{
-					return ELiteralType::UObjectProxyArray;
-				}
-				else if (InDataTypeInfo.bIsProxyParsable)
-				{
-					return ELiteralType::UObjectProxy;
-				}
-				else if (InDataTypeInfo.bIsEnum && InDataTypeInfo.bIsIntParsable)
-				{
-					return ELiteralType::Integer;
-				}
-				else if (InDataTypeInfo.bIsStringArrayParsable)
-				{
-					return ELiteralType::StringArray;
-				}
-				else if (InDataTypeInfo.bIsFloatArrayParsable)
-				{
-					return ELiteralType::FloatArray;
-				}
-				else if (InDataTypeInfo.bIsIntArrayParsable)
-				{
-					return ELiteralType::IntegerArray;
-				}
-				else if (InDataTypeInfo.bIsBoolArrayParsable)
-				{
-					return ELiteralType::BooleanArray;
-				}
-				else if (InDataTypeInfo.bIsStringParsable)
-				{
-					return ELiteralType::String;
-				}
-				else if (InDataTypeInfo.bIsFloatParsable)
-				{
-					return ELiteralType::Float;
-				}
-				else if (InDataTypeInfo.bIsIntParsable)
-				{
-					return ELiteralType::Integer;
-				}
-				else if (InDataTypeInfo.bIsBoolParsable)
-				{
-					return ELiteralType::Boolean;
-				}
-				else if (InDataTypeInfo.bIsDefaultArrayParsable)
-				{
-					return ELiteralType::NoneArray; 
-				}
-				else if (InDataTypeInfo.bIsDefaultParsable)
-				{
-					return ELiteralType::None;
-				}
-				else
-				{
-					// if we ever hit this, something has gone wrong with the REGISTER_METASOUND_DATATYPE macro.
-					// we should have failed to compile if any of these are false.
-					checkNoEntry();
-					return ELiteralType::Invalid;
-				}
-			}
-
-			// Node registry entry for input nodes created from a data type registry entry.
-			class FInputNodeRegistryEntry : public INodeRegistryEntry
-			{
-			public:
-				FInputNodeRegistryEntry() = delete;
-
-				FInputNodeRegistryEntry(TUniquePtr<IDataTypeRegistryEntry>&& InDataTypeEntry)
-				: DataTypeEntry(MoveTemp(InDataTypeEntry))
-				{
-					if (DataTypeEntry.IsValid())
-					{
-						FrontendClass = DataTypeEntry->GetFrontendInputClass();
-						ClassInfo = FNodeClassInfo(FrontendClass.Metadata);
-					}
-				}
-
-				virtual ~FInputNodeRegistryEntry() = default;
-
-				virtual const FNodeClassInfo& GetClassInfo() const override
-				{
-					return ClassInfo;
-				}
-
-				virtual TUniquePtr<INode> CreateNode(FDefaultNodeConstructorParams&&) const override
-				{
-					return nullptr;
-				}
-
-				virtual TUniquePtr<INode> CreateNode(FDefaultLiteralNodeConstructorParams&& InParams) const override
-				{
-					if (DataTypeEntry.IsValid())
-					{
-						return DataTypeEntry->CreateInputNode(MoveTemp(InParams));
-					}
-					return nullptr;
-				}
-
-				virtual TUniquePtr<INode> CreateNode(const FNodeInitData&) const override
-				{
-					return nullptr;
-				}
-
-				virtual const FMetasoundFrontendClass& GetFrontendClass() const override
-				{
-					return FrontendClass;
-				}
-
-				virtual TUniquePtr<INodeRegistryEntry> Clone() const override
-				{
-					if (DataTypeEntry.IsValid())
-					{
-						return MakeUnique<FInputNodeRegistryEntry>(DataTypeEntry->Clone());
-					}
-					return MakeUnique<FInputNodeRegistryEntry>(TUniquePtr<IDataTypeRegistryEntry>());
-				}
-
-				virtual bool IsNative() const override
-				{
-					return true;
-				}
-
-			private:
-				
-				TUniquePtr<IDataTypeRegistryEntry> DataTypeEntry;
-				FNodeClassInfo ClassInfo;
-				FMetasoundFrontendClass FrontendClass;
-			};
-
-			// Node registry entry for output nodes created from a data type registry entry.
-			class FOutputNodeRegistryEntry : public INodeRegistryEntry
-			{
-			public:
-				FOutputNodeRegistryEntry() = delete;
-
-				FOutputNodeRegistryEntry(TUniquePtr<IDataTypeRegistryEntry>&& InDataTypeEntry)
-				: DataTypeEntry(MoveTemp(InDataTypeEntry))
-				{
-					if (DataTypeEntry.IsValid())
-					{
-						FrontendClass = DataTypeEntry->GetFrontendOutputClass();
-						ClassInfo = FNodeClassInfo(FrontendClass.Metadata);
-					}
-				}
-
-				virtual ~FOutputNodeRegistryEntry() = default;
-
-				virtual const FNodeClassInfo& GetClassInfo() const override
-				{
-					return ClassInfo;
-				}
-
-				virtual TUniquePtr<INode> CreateNode(FDefaultNodeConstructorParams&& InParams) const override
-				{
-					if (DataTypeEntry.IsValid())
-					{
-						return DataTypeEntry->CreateOutputNode(MoveTemp(InParams));
-					}
-					return nullptr;
-				}
-
-				virtual TUniquePtr<INode> CreateNode(FDefaultLiteralNodeConstructorParams&& InParams) const override
-				{
-					return nullptr;
-				}
-
-				virtual TUniquePtr<INode> CreateNode(const FNodeInitData&) const override
-				{
-					return nullptr;
-				}
-
-				virtual const FMetasoundFrontendClass& GetFrontendClass() const override
-				{
-					return FrontendClass;
-				}
-
-				virtual TUniquePtr<INodeRegistryEntry> Clone() const override
-				{
-					if (DataTypeEntry.IsValid())
-					{
-						return MakeUnique<FOutputNodeRegistryEntry>(DataTypeEntry->Clone());
-					}
-					return MakeUnique<FOutputNodeRegistryEntry>(TUniquePtr<IDataTypeRegistryEntry>());
-				}
-
-				virtual bool IsNative() const override
-				{
-					return true;
-				}
-
-			private:
-				
-				TUniquePtr<IDataTypeRegistryEntry> DataTypeEntry;
-				FNodeClassInfo ClassInfo;
-				FMetasoundFrontendClass FrontendClass;
-			};
-
-			// Node registry entry for variable nodes created from a data type registry entry.
-			class FVariableNodeRegistryEntry : public INodeRegistryEntry
-			{
-			public:
-				FVariableNodeRegistryEntry() = delete;
-
-				FVariableNodeRegistryEntry(TUniquePtr<IDataTypeRegistryEntry>&& InDataTypeEntry)
-				: DataTypeEntry(MoveTemp(InDataTypeEntry))
-				{
-					if (DataTypeEntry.IsValid())
-					{
-						FrontendClass = DataTypeEntry->GetFrontendVariableClass();
-						ClassInfo = FNodeClassInfo(FrontendClass.Metadata);
-					}
-				}
-
-				virtual ~FVariableNodeRegistryEntry() = default;
-
-				virtual const FNodeClassInfo& GetClassInfo() const override
-				{
-					return ClassInfo;
-				}
-
-				virtual TUniquePtr<INode> CreateNode(FDefaultNodeConstructorParams&& InParams) const override
-				{
-					return nullptr;
-				}
-
-				virtual TUniquePtr<INode> CreateNode(FDefaultLiteralNodeConstructorParams&& InParams) const override
-				{
-					if (DataTypeEntry.IsValid())
-					{
-						return DataTypeEntry->CreateVariableNode(MoveTemp(InParams));
-					}
-					return nullptr;
-				}
-
-				virtual TUniquePtr<INode> CreateNode(const FNodeInitData&) const override
-				{
-					return nullptr;
-				}
-
-				virtual const FMetasoundFrontendClass& GetFrontendClass() const override
-				{
-					return FrontendClass;
-				}
-
-				virtual TUniquePtr<INodeRegistryEntry> Clone() const override
-				{
-					if (DataTypeEntry.IsValid())
-					{
-						return MakeUnique<FVariableNodeRegistryEntry>(DataTypeEntry->Clone());
-					}
-					return MakeUnique<FVariableNodeRegistryEntry>(TUniquePtr<IDataTypeRegistryEntry>());
-				}
-
-				virtual bool IsNative() const override
-				{
-					return true;
-				}
-
-			private:
-				
-				TUniquePtr<IDataTypeRegistryEntry> DataTypeEntry;
-				FNodeClassInfo ClassInfo;
-				FMetasoundFrontendClass FrontendClass;
-			};
-
 
 			// Registry container private implementation.
 			class FRegistryContainerImpl : public FMetasoundFrontendRegistryContainer
@@ -401,9 +100,6 @@ namespace Metasound
 
 				virtual void ForEachNodeRegistryTransactionSince(FRegistryTransactionID InSince, FRegistryTransactionID* OutCurrentRegistryTransactionID, TFunctionRef<void(const FNodeRegistryTransaction&)> InFunc) const override;
 
-				// Return any data types that can be used as a metasound input type or output type.
-				TArray<FName> GetAllValidDataTypes() override;
-
 				void IterateRegistry(Metasound::FIterateMetasoundFrontendClassFunction InIterFunc, EMetasoundFrontendClassType InClassType = EMetasoundFrontendClassType::Invalid) const override;
 
 				// Find Frontend Document data.
@@ -414,11 +110,6 @@ namespace Metasound
 				bool FindOutputNodeRegistryKeyForDataType(const FName& InDataTypeName, FNodeRegistryKey& OutKey) override;
 
 				// Create a new instance of a C++ implemented node from the registry.
-				TUniquePtr<Metasound::INode> CreateInputNode(const FName& InInputType, Metasound::FInputNodeConstructorParams&& InParams) override;
-				TUniquePtr<Metasound::INode> CreateVariableNode(const FName& InVariableType, FVariableNodeConstructorParams&& InParams) override;
-				TUniquePtr<Metasound::INode> CreateOutputNode(const FName& InOutputType, Metasound::FOutputNodeConstructorParams&& InParams) override;
-				TUniquePtr<Metasound::INode> CreateReceiveNode(const FName& InVariableType, FReceiveNodeConstructorParams&& InParams) override;
-
 				virtual TUniquePtr<INode> CreateNode(const FNodeRegistryKey& InKey, FDefaultNodeConstructorParams&&) const override;
 				virtual TUniquePtr<INode> CreateNode(const FNodeRegistryKey& InKey, FDefaultLiteralNodeConstructorParams&&) const override;
 				virtual TUniquePtr<Metasound::INode> CreateNode(const FNodeRegistryKey& InKey, const Metasound::FNodeInitData& InInitData) const override;
@@ -427,29 +118,8 @@ namespace Metasound
 				// Returns an empty array if none are available.
 				TArray<FConverterNodeInfo> GetPossibleConverterNodes(const FName& FromDataType, const FName& ToDataType) override;
 
-				virtual bool RegisterDataType(const FName& InName, TUniquePtr<Metasound::Frontend::IDataTypeRegistryEntry>&&) override;
-
-				// Get the desired kind of literal for a given data type. Returns EConstructorArgType::Invalid if the data type couldn't be found.
-				Metasound::ELiteralType GetDesiredLiteralTypeForDataType(FName InDataType) const override;
-
-				// Get whether we can build a literal of this specific type for InDataType.
-				bool DoesDataTypeSupportLiteralType(FName InDataType, Metasound::ELiteralType InLiteralType) const override;
-
-				// Handle uobjects and literals
-				UClass* GetLiteralUClassForDataType(FName InDataType) const override;
-				Metasound::FLiteral CreateLiteralFromUObject(const FName& InDataType, UObject* InObject) override;
-				Metasound::FLiteral CreateLiteralFromUObjectArray(const FName& InDataType, TArray<UObject*> InObjectArray) override;
-
-				// Get info about a specific data type (what kind of literals we can use, etc.)
-				// @returns false if InDataType wasn't found in the registry. 
-				bool GetInfoForDataType(FName InDataType, FDataTypeRegistryInfo& OutInfo) override;
-
-				TSharedPtr<const Metasound::Frontend::IEnumDataTypeInterface> GetEnumInterfaceForDataType(FName InDataType) const override;
-
-				TSharedPtr<Metasound::IDataChannel, ESPMode::ThreadSafe> CreateDataChannelForDataType(const FName& InDataType, const Metasound::FOperatorSettings& InOperatorSettings) const override;
 			private:
 
-				const IDataTypeRegistryEntry* FindDataTypeEntry(const FName& InDataTypeName) const;
 				const INodeRegistryEntry* FindNodeEntry(const FNodeRegistryKey& InKey) const;
 
 
@@ -464,7 +134,6 @@ namespace Metasound
 
 				// Registry in which we keep all information about nodes implemented in C++.
 				TMap<FNodeRegistryKey, TUniquePtr<INodeRegistryEntry>> RegisteredNodes;
-				TMap<FName, TUniquePtr<IDataTypeRegistryEntry>> RegisteredDataTypes;
 
 				// Registry in which we keep lists of possible nodes to use to convert between two datatypes
 				TMap<FConverterNodeRegistryKey, FConverterNodeRegistryValue> ConverterNodeRegistry;
@@ -507,110 +176,7 @@ namespace Metasound
 				return RegistryTransactionHistory.ForEachTransactionSince(InSince, OutCurrentRegistryTransactionID, InFunc);
 			}
 
-			TUniquePtr<Metasound::INode> FRegistryContainerImpl::CreateInputNode(const FName& InDataType, Metasound::FInputNodeConstructorParams&& InParams)
-			{
-				const IDataTypeRegistryEntry* Entry = FindDataTypeEntry(InDataType);
-
-				if (ensureAlwaysMsgf(nullptr != Entry, TEXT("Could not find data type [Name:%s]"), *InDataType.ToString()))
-				{
-					return Entry->CreateInputNode(MoveTemp(InParams));
-				}
-				else
-				{
-					return nullptr;
-				}
-			}
-
-			TUniquePtr<Metasound::INode> FRegistryContainerImpl::CreateVariableNode(const FName& InDataType, Metasound::FVariableNodeConstructorParams&& InParams)
-			{
-				const IDataTypeRegistryEntry* Entry = FindDataTypeEntry(InDataType);
-
-				if (ensureAlwaysMsgf(nullptr != Entry, TEXT("Could not find data type [Name:%s]"), *InDataType.ToString()))
-				{
-					return Entry->CreateVariableNode(MoveTemp(InParams));
-				}
-				else
-				{
-					return nullptr;
-				}
-			}
-
-			TUniquePtr<Metasound::INode> FRegistryContainerImpl::CreateOutputNode(const FName& InDataType, Metasound::FOutputNodeConstructorParams&& InParams)
-			{
-				const IDataTypeRegistryEntry* Entry = FindDataTypeEntry(InDataType);
-
-				if (ensureAlwaysMsgf(nullptr != Entry, TEXT("Could not find data type [Name:%s]"), *InDataType.ToString()))
-				{
-					return Entry->CreateOutputNode(MoveTemp(InParams));
-				}
-				else
-				{
-					return nullptr;
-				}
-			}
-
-			TUniquePtr<Metasound::INode> FRegistryContainerImpl::CreateReceiveNode(const FName& InDataType, Metasound::FReceiveNodeConstructorParams&& InParams)
-			{
-				const IDataTypeRegistryEntry* Entry = FindDataTypeEntry(InDataType);
-
-				if (ensureAlwaysMsgf(nullptr != Entry, TEXT("Could not find data type [Name:%s]"), *InDataType.ToString()))
-				{
-					return Entry->CreateReceiveNode(MoveTemp(InParams));
-				}
-				else
-				{
-					return nullptr;
-				}
-			}
 
-			Metasound::FLiteral FRegistryContainerImpl::CreateLiteralFromUObject(const FName& InDataType, UObject* InObject)
-			{
-				const IDataTypeRegistryEntry* Entry = FindDataTypeEntry(InDataType);
-
-				if (ensureAlwaysMsgf(nullptr != Entry, TEXT("Could not find data type [Name:%s]"), *InDataType.ToString()))
-
-				{
-					if (Audio::IProxyDataPtr ProxyPtr = Entry->CreateProxy(InObject))
-					{
-						if (InObject)
-						{
-							ensureAlwaysMsgf(ProxyPtr.IsValid(), TEXT("UObject failed to create a valid proxy."));
-						}
-						return Metasound::FLiteral(MoveTemp(ProxyPtr));
-					}
-				}
-
-				return Metasound::FLiteral();
-			}
-
-			Metasound::FLiteral FRegistryContainerImpl::CreateLiteralFromUObjectArray(const FName& InDataType, TArray<UObject*> InObjectArray)
-			{
-				const IDataTypeRegistryEntry* Entry = FindDataTypeEntry(InDataType);
-
-				if (ensureAlwaysMsgf(nullptr != Entry, TEXT("Could not find data type [Name:%s]"), *InDataType.ToString()))
-
-				{
-					TArray<Audio::IProxyDataPtr> ProxyArray;
-
-					for (UObject* InObject : InObjectArray)
-					{
-						Audio::IProxyDataPtr ProxyPtr = Entry->CreateProxy(InObject);
-
-						if (InObject)
-						{
-							ensureAlwaysMsgf(ProxyPtr.IsValid(), TEXT("UObject failed to create a valid proxy!"));
-						}
-
-						ProxyArray.Add(MoveTemp(ProxyPtr));
-					}
-
-					return Metasound::FLiteral(MoveTemp(ProxyArray));
-				}
-				else
-				{
-					return Metasound::FLiteral();
-				}
-			}
 			TUniquePtr<Metasound::INode> FRegistryContainerImpl::CreateNode(const FNodeRegistryKey& InKey, FDefaultNodeConstructorParams&& InParams) const
 			{
 				const INodeRegistryEntry* Entry = FindNodeEntry(InKey);
@@ -660,141 +226,6 @@ namespace Metasound
 				}
 			}
 
-			Metasound::ELiteralType FRegistryContainerImpl::GetDesiredLiteralTypeForDataType(FName InDataType) const
-			{
-				const IDataTypeRegistryEntry* Entry = FindDataTypeEntry(InDataType);
-
-				if (nullptr == Entry)
-				{
-					return Metasound::ELiteralType::Invalid;
-				}
-
-				const FDataTypeRegistryInfo& Info = Entry->GetDataTypeInfo();
-
-				// If there's a designated preferred literal type for this datatype, use that.
-				if (Info.PreferredLiteralType != Metasound::ELiteralType::Invalid)
-				{
-					return Info.PreferredLiteralType;
-				}
-
-				// Otherwise, we opt for the highest precision construction option available.
-				return Metasound::Frontend::MetasoundFrontendRegistryPrivate::GetMostDescriptiveLiteralForDataType(Info);
-			}
-
-			UClass* FRegistryContainerImpl::GetLiteralUClassForDataType(FName InDataType) const
-			{
-				const IDataTypeRegistryEntry* Entry = FindDataTypeEntry(InDataType);
-
-				if (ensureAlwaysMsgf(nullptr != Entry, TEXT("Could not find data type [Name:%s]."), *InDataType.ToString()))
-				{
-					return Entry->GetDataTypeInfo().ProxyGeneratorClass;
-				}
-
-				return nullptr;
-			}
-
-			bool FRegistryContainerImpl::DoesDataTypeSupportLiteralType(FName InDataType, Metasound::ELiteralType InLiteralType) const
-			{
-				const IDataTypeRegistryEntry* Entry = FindDataTypeEntry(InDataType);
-
-				if (ensureAlwaysMsgf(nullptr != Entry, TEXT("Could not find data type [Name:%s]."), *InDataType.ToString()))
-				{
-					const FDataTypeRegistryInfo& Info = Entry->GetDataTypeInfo();
-
-					switch (InLiteralType)
-					{
-						case Metasound::ELiteralType::Boolean:
-						{
-							return Info.bIsBoolParsable;
-						}
-						case Metasound::ELiteralType::BooleanArray:
-						{
-							return Info.bIsBoolArrayParsable;
-						}
-
-						case Metasound::ELiteralType::Integer:
-						{
-							return Info.bIsIntParsable;
-						}
-						case Metasound::ELiteralType::IntegerArray:
-						{
-							return Info.bIsIntArrayParsable;
-						}
-
-						case Metasound::ELiteralType::Float:
-						{
-							return Info.bIsFloatParsable;
-						}
-						case Metasound::ELiteralType::FloatArray:
-						{
-							return Info.bIsFloatArrayParsable;
-						}
-
-						case Metasound::ELiteralType::String:
-						{
-							return Info.bIsStringParsable;
-						}
-						case Metasound::ELiteralType::StringArray:
-						{
-							return Info.bIsStringArrayParsable;
-						}
-
-						case Metasound::ELiteralType::UObjectProxy:
-						{
-							return Info.bIsProxyParsable;
-						}
-						case Metasound::ELiteralType::UObjectProxyArray:
-						{
-							return Info.bIsProxyArrayParsable;
-						}
-
-						case Metasound::ELiteralType::None:
-						{
-							return Info.bIsDefaultParsable;
-						}
-						case Metasound::ELiteralType::NoneArray:
-						{
-							return Info.bIsDefaultArrayParsable;
-						}
-
-						case Metasound::ELiteralType::Invalid:
-						default:
-						{
-							static_assert(static_cast<int32>(Metasound::ELiteralType::COUNT) == 13, "Possible missing case coverage for ELiteralType");
-							return false;
-						}
-					}
-				}
-
-				return false;
-			}
-
-			bool FRegistryContainerImpl::RegisterDataType(const FName& InName, TUniquePtr<IDataTypeRegistryEntry>&& InRegistryEntry)
-			{
-				if (InRegistryEntry.IsValid())
-				{
-					if (!ensureAlwaysMsgf(!RegisteredDataTypes.Contains(InName),
-						TEXT("Name collision when trying to register Metasound Data Type [Name:%s]. DataType must have "
-							"unique name and REGISTER_METASOUND_DATATYPE cannot be called in a public header."),
-							*InName.ToString()))
-					{
-						return false;
-					}
-
-
-					RegisteredDataTypes.Add(InName, InRegistryEntry->Clone());
-
-					RegisterNode(MakeUnique<FInputNodeRegistryEntry>(InRegistryEntry->Clone()));
-					RegisterNode(MakeUnique<FOutputNodeRegistryEntry>(InRegistryEntry->Clone()));
-					RegisterNode(MakeUnique<FVariableNodeRegistryEntry>(MoveTemp(InRegistryEntry)));
-
-					UE_LOG(LogMetaSound, Verbose, TEXT("Registered Metasound Datatype [Name:%s]."), *InName.ToString());
-					return true;
-				}
-
-				return false;
-			}
-
 			FNodeRegistryKey FRegistryContainerImpl::RegisterNode(TUniquePtr<INodeRegistryEntry>&& InEntry)
 			{
 				FNodeRegistryKey Key;
@@ -874,45 +305,6 @@ namespace Metasound
 				return false;
 			}
 
-			TArray<FName> FRegistryContainerImpl::GetAllValidDataTypes()
-			{
-				TArray<FName> OutDataTypes;
-
-				for (auto& DataTypeTuple : RegisteredDataTypes)
-				{
-					OutDataTypes.Add(DataTypeTuple.Key);
-				}
-
-				return OutDataTypes;
-			}
-
-			bool FRegistryContainerImpl::GetInfoForDataType(FName InDataType, Metasound::Frontend::FDataTypeRegistryInfo& OutInfo)
-			{
-				if (const IDataTypeRegistryEntry* Entry = FindDataTypeEntry(InDataType))
-				{
-					OutInfo = Entry->GetDataTypeInfo();
-					return true;
-				}
-				return false;
-			}
-
-			TSharedPtr<const Metasound::Frontend::IEnumDataTypeInterface> FRegistryContainerImpl::GetEnumInterfaceForDataType(FName InDataType) const
-			{
-				if (const IDataTypeRegistryEntry* Entry = FindDataTypeEntry(InDataType))
-				{
-					return Entry->GetEnumInterface();
-				}
-				return nullptr;
-			}
-
-			TSharedPtr<Metasound::IDataChannel, ESPMode::ThreadSafe> FRegistryContainerImpl::CreateDataChannelForDataType(const FName& InDataType, const Metasound::FOperatorSettings& InSettings) const
-			{
-				if (const IDataTypeRegistryEntry* Entry = FindDataTypeEntry(InDataType))
-				{
-					return Entry->CreateDataChannel(InSettings);
-				}
-				return nullptr;
-			}
 
 			bool FRegistryContainerImpl::FindFrontendClassFromRegistered(const FNodeRegistryKey& InKey, FMetasoundFrontendClass& OutClass)
 			{
@@ -937,9 +329,10 @@ namespace Metasound
 
 			bool FRegistryContainerImpl::FindInputNodeRegistryKeyForDataType(const FName& InDataTypeName, FNodeRegistryKey& OutKey)
 			{
-				if (const IDataTypeRegistryEntry* Entry = FindDataTypeEntry(InDataTypeName))
+				FMetasoundFrontendClass Class;
+				if (IDataTypeRegistry::Get().GetFrontendInputClass(InDataTypeName, Class))
 				{
-					OutKey = NodeRegistryKey::CreateKey(Entry->GetFrontendInputClass().Metadata);
+					OutKey = NodeRegistryKey::CreateKey(Class.Metadata);
 					return true;
 				}
 				return false;
@@ -947,9 +340,10 @@ namespace Metasound
 
 			bool FRegistryContainerImpl::FindVariableNodeRegistryKeyForDataType(const FName& InDataTypeName, FNodeRegistryKey& OutKey)
 			{
-				if (const IDataTypeRegistryEntry* Entry = FindDataTypeEntry(InDataTypeName))
+				FMetasoundFrontendClass Class;
+				if (IDataTypeRegistry::Get().GetFrontendVariableClass(InDataTypeName, Class))
 				{
-					OutKey = NodeRegistryKey::CreateKey(Entry->GetFrontendVariableClass().Metadata);
+					OutKey = NodeRegistryKey::CreateKey(Class.Metadata);
 					return true;
 				}
 				return false;
@@ -957,9 +351,10 @@ namespace Metasound
 
 			bool FRegistryContainerImpl::FindOutputNodeRegistryKeyForDataType(const FName& InDataTypeName, FNodeRegistryKey& OutKey)
 			{
-				if (const IDataTypeRegistryEntry* Entry = FindDataTypeEntry(InDataTypeName))
+				FMetasoundFrontendClass Class;
+				if (IDataTypeRegistry::Get().GetFrontendOutputClass(InDataTypeName, Class))
 				{
-					OutKey = NodeRegistryKey::CreateKey(Entry->GetFrontendOutputClass().Metadata);
+					OutKey = NodeRegistryKey::CreateKey(Class.Metadata);
 					return true;
 				}
 				return false;
@@ -988,18 +383,6 @@ namespace Metasound
 				}
 			}
 
-			const IDataTypeRegistryEntry* FRegistryContainerImpl::FindDataTypeEntry(const FName& InDataTypeName) const
-			{
-				if (const TUniquePtr<IDataTypeRegistryEntry>* Entry = RegisteredDataTypes.Find(InDataTypeName))
-				{
-					if (Entry->IsValid())
-					{
-						return Entry->Get();
-					}
-				}
-
-				return nullptr;
-			}
 
 			const INodeRegistryEntry* FRegistryContainerImpl::FindNodeEntry(const FNodeRegistryKey& InKey) const
 			{
diff --git a/Engine/Plugins/Runtime/Metasound/Source/MetasoundFrontend/Private/MetasoundFrontendStandardController.cpp b/Engine/Plugins/Runtime/Metasound/Source/MetasoundFrontend/Private/MetasoundFrontendStandardController.cpp
--- a/Engine/Plugins/Runtime/Metasound/Source/MetasoundFrontend/Private/MetasoundFrontendStandardController.cpp
+++ b/Engine/Plugins/Runtime/Metasound/Source/MetasoundFrontend/Private/MetasoundFrontendStandardController.cpp
@@ -2927,7 +2927,7 @@ namespace Metasound
 		{
 			if (const FMetasoundFrontendClassInput* Desc = FindInputDescriptionWithName(InInputName))
 			{
-				return FRegistry::Get()->GetDesiredLiteralTypeForDataType(Desc->TypeName);
+				return IDataTypeRegistry::Get().GetDesiredLiteralType(Desc->TypeName);
 			}
 			return ELiteralType::Invalid;
 		}
@@ -2937,7 +2937,7 @@ namespace Metasound
 		{
 			if (const FMetasoundFrontendClassInput* Desc = FindInputDescriptionWithName(InInputName))
 			{
-				return FRegistry::Get()->GetLiteralUClassForDataType(Desc->TypeName);
+				return IDataTypeRegistry::Get().GetUClassForDataType(Desc->TypeName);
 			}
 			return nullptr;
 		}
@@ -2955,7 +2955,7 @@ namespace Metasound
 		{
 			if (FMetasoundFrontendClassInput* Desc = FindInputDescriptionWithVertexID(InVertexID))
 			{
-				if (ensure(DoesDataTypeSupportLiteralType(Desc->TypeName, InLiteral.GetType())))
+				if (ensure(IDataTypeRegistry::Get().IsLiteralTypeSupported(Desc->TypeName, InLiteral.GetType())))
 				{
 					Desc->DefaultLiteral = InLiteral;
 					return true;
@@ -2973,7 +2973,7 @@ namespace Metasound
 		{
 			if (FMetasoundFrontendClassInput* Desc = FindInputDescriptionWithVertexID(InVertexID))
 			{
-				Metasound::FLiteral Literal = Frontend::GetDefaultParamForDataType(Desc->TypeName);
+				Metasound::FLiteral Literal = IDataTypeRegistry::Get().CreateDefaultLiteral(Desc->TypeName);
 				Desc->DefaultLiteral.SetFromLiteral(Literal);
 				return Desc->DefaultLiteral.IsValid();
 			}
diff --git a/Engine/Plugins/Runtime/Metasound/Source/MetasoundFrontend/Private/MetasoundRouter.cpp b/Engine/Plugins/Runtime/Metasound/Source/MetasoundFrontend/Private/MetasoundRouter.cpp
--- a/Engine/Plugins/Runtime/Metasound/Source/MetasoundFrontend/Private/MetasoundRouter.cpp
+++ b/Engine/Plugins/Runtime/Metasound/Source/MetasoundFrontend/Private/MetasoundRouter.cpp
@@ -2,6 +2,7 @@
 
 #include "MetasoundRouter.h"
 
+#include "MetasoundFrontendDataTypeRegistry.h"
 #include "MetasoundFrontendRegistries.h"
 #include "MetasoundOperatorInterface.h"
 #include "CoreMinimal.h"
@@ -279,15 +280,10 @@ namespace Metasound
 			const FDataChannelKey ChannelKey = GetDataChannelKey(InDataTypeName, InChannelName);
 			
 			// This is the first time we're seeing this, add it to the map.
-			FMetasoundFrontendRegistryContainer* Registry = FMetasoundFrontendRegistryContainer::Get();
-
-			if (ensure(nullptr != Registry))
+			DataChannel = Metasound::Frontend::IDataTypeRegistry::Get().CreateDataChannel(InDataTypeName, InOperatorSettings);
+			if (DataChannel.IsValid())
 			{
-				DataChannel = Registry->CreateDataChannelForDataType(InDataTypeName, InOperatorSettings);
-				if (DataChannel.IsValid())
-				{
-					DataChannelMap.Add(ChannelKey, DataChannel.ToSharedRef());
-				}
+				DataChannelMap.Add(ChannelKey, DataChannel.ToSharedRef());
 			}
 		}
 
diff --git a/Engine/Plugins/Runtime/Metasound/Source/MetasoundFrontend/Public/MetasoundDataTypeRegistrationMacro.h b/Engine/Plugins/Runtime/Metasound/Source/MetasoundFrontend/Public/MetasoundDataTypeRegistrationMacro.h
--- a/Engine/Plugins/Runtime/Metasound/Source/MetasoundFrontend/Public/MetasoundDataTypeRegistrationMacro.h
+++ b/Engine/Plugins/Runtime/Metasound/Source/MetasoundFrontend/Public/MetasoundDataTypeRegistrationMacro.h
@@ -423,7 +423,7 @@ namespace Metasound
 					return MakeUnique<TGetVariableNode<TDataType>>(InParams.NodeName, InParams.InstanceID, InParams.VertexName, MoveTemp(InParams.InitParam));
 				}
 
-				virtual TUniquePtr<INode> CreateReceiveNode(FReceiveNodeConstructorParams&& InParams) const override
+				virtual TUniquePtr<INode> CreateReceiveNode(const FNodeInitData& InParams) const override
 				{
 					return MakeUnique<TReceiveNode<TDataType>>(InParams);
 				}
@@ -470,7 +470,7 @@ namespace Metasound
 				TSharedPtr<Frontend::IEnumDataTypeInterface> EnumInterface;
 			};
 
-			bool bSucceeded = FMetasoundFrontendRegistryContainer::Get()->RegisterDataType(GetMetasoundDataTypeName<TDataType>(), MakeUnique<FDataTypeRegistryEntry>());
+			bool bSucceeded = Frontend::IDataTypeRegistry::Get().RegisterDataType(MakeUnique<FDataTypeRegistryEntry>());
 			ensureAlwaysMsgf(bSucceeded, TEXT("Failed to register data type %s in the node registry!"), *GetMetasoundDataTypeString<TDataType>());
 
 			if (bSucceeded)
diff --git a/Engine/Plugins/Runtime/Metasound/Source/MetasoundFrontend/Public/MetasoundFrontend.h b/Engine/Plugins/Runtime/Metasound/Source/MetasoundFrontend/Public/MetasoundFrontend.h
--- a/Engine/Plugins/Runtime/Metasound/Source/MetasoundFrontend/Public/MetasoundFrontend.h
+++ b/Engine/Plugins/Runtime/Metasound/Source/MetasoundFrontend/Public/MetasoundFrontend.h
@@ -3,10 +3,12 @@
 #pragma once
 
 #include "CoreMinimal.h"
+
+#include "MetasoundBuilderInterface.h"
 #include "MetasoundGraph.h"
-#include "MetasoundFrontendRegistries.h"
+#include "MetasoundFrontendDataTypeRegistry.h"
 #include "MetasoundFrontendDocument.h"
-#include "MetasoundBuilderInterface.h"
+#include "MetasoundFrontendRegistries.h"
 
 // Forward Declarations
 class FMetasoundAssetBase;
@@ -77,21 +79,6 @@ namespace Metasound
 			return GenerateClassDescription<NodeType>(InitData);
 		}
 
-
-		template<typename DataType>
-		FName GetDataTypeName()
-		{
-			static FName DataTypeName = FName(TDataReferenceTypeInfo<DataType>::TypeName);
-			return DataTypeName;
-		}
-
-		// Returns a list of all available data types.
-		METASOUNDFRONTEND_API TArray<FName> GetAllAvailableDataTypes();
-
-		// outputs the traits for a given data type.
-		// returns false if InDataType couldn't be found.
-		METASOUNDFRONTEND_API bool GetTraitsForDataType(FName InDataType, FDataTypeRegistryInfo& OutInfo);
-
 		// Takes a JSON string and deserializes it into a Metasound document struct.
 		// @returns false if the file couldn't be found or parsed into a document.
 		METASOUNDFRONTEND_API bool ImportJSONToMetasound(const FString& InJSON, FMetasoundFrontendDocument& OutMetasoundDocument);
diff --git a/Engine/Plugins/Runtime/Metasound/Source/MetasoundFrontend/Public/MetasoundFrontendDataTypeRegistry.h b/Engine/Plugins/Runtime/Metasound/Source/MetasoundFrontend/Public/MetasoundFrontendDataTypeRegistry.h
new file mode 100644
--- /dev/null
+++ b/Engine/Plugins/Runtime/Metasound/Source/MetasoundFrontend/Public/MetasoundFrontendDataTypeRegistry.h
@@ -0,0 +1,210 @@
+// Copyright Epic Games, Inc. All Rights Reserved.
+
+#pragma once
+
+#include "CoreMinimal.h"
+
+#include "MetasoundEnum.h"
+#include "MetasoundFrontendDocument.h"
+#include "MetasoundFrontendRegistries.h"
+#include "MetasoundLiteral.h"
+
+namespace Metasound
+{
+	namespace Frontend
+	{
+		using FOutputNodeConstructorParams = FDefaultNodeConstructorParams;
+		using FVariableNodeConstructorParams = FDefaultLiteralNodeConstructorParams;
+		using FInputNodeConstructorParams = FDefaultLiteralNodeConstructorParams;
+
+		/** FDataTypeRegsitryInfo contains runtime inspectable behavior of a registered
+		 * MetaSound data type.
+		 */
+		struct FDataTypeRegistryInfo
+		{
+			// The name of the data type.
+			FName DataTypeName;
+
+			// The preferred constructor argument type for creating instances of the data type.
+			ELiteralType PreferredLiteralType = ELiteralType::Invalid;
+
+			// Constructor argument support in TDataTypeLiteralFactory<TDataType>;
+			bool bIsDefaultParsable = false;
+			bool bIsBoolParsable = false;
+			bool bIsIntParsable = false;
+			bool bIsFloatParsable = false;
+			bool bIsStringParsable = false;
+			bool bIsProxyParsable = false;
+			bool bIsDefaultArrayParsable = false;
+			bool bIsBoolArrayParsable = false;
+			bool bIsIntArrayParsable = false;
+			bool bIsFloatArrayParsable = false;
+			bool bIsStringArrayParsable = false;
+			bool bIsProxyArrayParsable = false;
+
+			// Is an TEnum wrapped enum
+			bool bIsEnum = false;
+
+			// Determines whether the type can be used with send/receive transmitters
+			bool bIsTransmittable = false;
+
+			// If this datatype was registered with a specific UClass to use to filter with, that will be used here:
+			UClass* ProxyGeneratorClass = nullptr;
+
+			FORCEINLINE bool IsArrayType() const
+			{
+				return bIsDefaultArrayParsable
+				|| bIsBoolArrayParsable
+				|| bIsIntArrayParsable
+				|| bIsFloatArrayParsable
+				|| bIsStringArrayParsable
+				|| bIsProxyArrayParsable;
+			}
+		};
+
+		/** Interface for metadata of a registered MetaSound enum type. */
+		struct METASOUNDFRONTEND_API IEnumDataTypeInterface
+		{
+			using FGenericInt32Entry = TEnumEntry<int32>;
+
+			virtual ~IEnumDataTypeInterface() = default;
+
+			virtual const TArray<FGenericInt32Entry>& GetAllEntries() const = 0;
+			virtual FName GetNamespace() const = 0;
+			virtual int32 GetDefaultValue() const = 0;
+
+			template<typename Predicate>
+			TOptional<FGenericInt32Entry> FindEntryBy(Predicate Pred) const
+			{
+				TArray<FGenericInt32Entry> Entries = GetAllEntries();
+				if (FGenericInt32Entry* Found = Entries.FindByPredicate(Pred))
+				{
+					return *Found;
+				}
+				return {};
+			}
+			TOptional<FGenericInt32Entry> FindByValue(int32 InEnumValue) const
+			{
+				return FindEntryBy([InEnumValue](const FGenericInt32Entry& i) -> bool { return i.Value == InEnumValue; });
+			}
+			TOptional<FGenericInt32Entry> FindByName(FName InEnumName) const
+			{
+				return FindEntryBy([InEnumName](const FGenericInt32Entry& i) -> bool { return i.Name == InEnumName; });
+			}
+			TOptional<FName> ToName(int32 InEnumValue) const
+			{
+				if(TOptional<FGenericInt32Entry> Result = FindByValue(InEnumValue))
+				{
+					return Result->Name;
+				}
+				return {};
+			}
+			TOptional<int32> ToValue(FName InName) const
+			{
+				if (TOptional<FGenericInt32Entry> Result = FindByName(InName))
+				{
+					return Result->Value;
+				}
+				return {};
+			}
+		};
+
+		/** Registry entry interface for a MetaSound data type. */
+		class IDataTypeRegistryEntry
+		{
+		public:
+			virtual ~IDataTypeRegistryEntry() = default;
+
+			/** Return the FDataTypeRegistryInfo for the data type */
+			virtual const FDataTypeRegistryInfo& GetDataTypeInfo() const = 0;
+
+			/** Return the enum interface for the data type. If the data type does
+			 * not support an enum interface, the returned pointer should be invalid.
+			 */
+			virtual TSharedPtr<const IEnumDataTypeInterface> GetEnumInterface() const = 0;
+
+			/** Return an FMetasoundFrontenClass representing an input node of the data type. */
+			virtual const FMetasoundFrontendClass& GetFrontendInputClass() const = 0;
+
+			/** Return an FMetasoundFrontenClass representing a variable node of the data type. */
+			virtual const FMetasoundFrontendClass& GetFrontendVariableClass() const = 0;
+
+			/** Return an FMetasoundFrontenClass representing an output node of the data type. */
+			virtual const FMetasoundFrontendClass& GetFrontendOutputClass() const = 0;
+
+			/** Create an input node */
+			virtual TUniquePtr<INode> CreateInputNode(FInputNodeConstructorParams&&) const = 0;
+
+			/** Create an output node */
+			virtual TUniquePtr<INode> CreateOutputNode(FOutputNodeConstructorParams&&) const = 0;
+			
+			/** Create a variable node */
+			virtual TUniquePtr<INode> CreateVariableNode(FVariableNodeConstructorParams&&) const = 0;
+
+			/** Create a receive node */
+			virtual TUniquePtr<INode> CreateReceiveNode(const FNodeInitData&) const = 0;
+
+			/** Create a proxy from a UObject. If this data type does not support
+			 * UObject proxies, return a nullptr. */
+			virtual Audio::IProxyDataPtr CreateProxy(UObject*) const = 0;
+
+			/** Create a data channel for transmission. If this data type does not
+			 * support transmission, return a nullptr. */
+			virtual TSharedPtr<IDataChannel, ESPMode::ThreadSafe> CreateDataChannel(const FOperatorSettings&) const = 0;
+
+			/** Clone this registry entry. */
+			virtual TUniquePtr<IDataTypeRegistryEntry> Clone() const = 0;
+		};
+
+		class METASOUNDFRONTEND_API IDataTypeRegistry
+		{
+		public:
+			static IDataTypeRegistry& Get();
+			virtual ~IDataTypeRegistry() = default;
+
+			/** Register a data type
+			 * @param InName - Name of data type.
+			 * @param InEntry - TUniquePtr to data type registry entry.
+			 *
+			 * @return True on success, false on failure.
+			 */
+			virtual bool RegisterDataType(TUniquePtr<IDataTypeRegistryEntry>&& InEntry) = 0;
+
+			virtual void GetRegisteredDataTypeNames(TArray<FName>& OutNames) const = 0;
+
+			virtual bool GetDataTypeInfo(const FName& InDataType, FDataTypeRegistryInfo& OutInfo) const = 0;
+
+			// Return the enum interface for a data type. If the data type does not have 
+			// an enum interface, returns a nullptr.
+			virtual TSharedPtr<const IEnumDataTypeInterface> GetEnumInterfaceForDataType(const FName& InDataType) const = 0;
+
+			virtual ELiteralType GetDesiredLiteralType(const FName& InDataType) const = 0;
+
+			virtual bool IsLiteralTypeSupported(const FName& InDataType, ELiteralType InLiteralType) const = 0;
+			virtual bool IsLiteralTypeSupported(const FName& InDataType, EMetasoundFrontendLiteralType InLiteralType) const = 0;
+
+			virtual UClass* GetUClassForDataType(const FName& InDataType) const = 0;
+
+			virtual FLiteral CreateDefaultLiteral(const FName& InDataType) const = 0;
+			virtual FLiteral CreateLiteralFromUObject(const FName& InDataType, UObject* InObject) const = 0;
+			virtual FLiteral CreateLiteralFromUObjectArray(const FName& InDataType, TArray<UObject*> InObjectArray) const = 0;
+
+			virtual TSharedPtr<IDataChannel, ESPMode::ThreadSafe> CreateDataChannel(const FName& InDataType, const FOperatorSettings& InOperatorSettings) const = 0;
+
+			/** Return an FMetasoundFrontenClass representing an input node of the data type. */
+			virtual bool GetFrontendInputClass(const FName& InDataType, FMetasoundFrontendClass& OutClass) const = 0;
+
+			/** Return an FMetasoundFrontenClass representing a variable node of the data type. */
+			virtual bool GetFrontendVariableClass(const FName& InDataType, FMetasoundFrontendClass& OutClass) const = 0;
+
+			/** Return an FMetasoundFrontenClass representing an output node of the data type. */
+			virtual bool GetFrontendOutputClass(const FName& InDataType, FMetasoundFrontendClass& OutClass) const = 0;
+
+			// Create a new instance of a C++ implemented node 
+			virtual TUniquePtr<INode> CreateInputNode(const FName& InInputType, FInputNodeConstructorParams&& InParams) const = 0;
+			virtual TUniquePtr<INode> CreateVariableNode(const FName& InVariableType, FVariableNodeConstructorParams&& InParams) const = 0;
+			virtual TUniquePtr<INode> CreateOutputNode(const FName& InOutputType, FOutputNodeConstructorParams&& InParams) const = 0;
+			virtual TUniquePtr<INode> CreateReceiveNode(const FName& InDataType, const FNodeInitData& InParams) const = 0;
+		};
+	}
+}
diff --git a/Engine/Plugins/Runtime/Metasound/Source/MetasoundFrontend/Public/MetasoundFrontendLiteral.h b/Engine/Plugins/Runtime/Metasound/Source/MetasoundFrontend/Public/MetasoundFrontendLiteral.h
--- a/Engine/Plugins/Runtime/Metasound/Source/MetasoundFrontend/Public/MetasoundFrontendLiteral.h
+++ b/Engine/Plugins/Runtime/Metasound/Source/MetasoundFrontend/Public/MetasoundFrontendLiteral.h
@@ -128,16 +128,6 @@ namespace Metasound
 {
 	namespace Frontend
 	{
-		// Returns true if the data type can be constructed using the literal type.
-		METASOUNDFRONTEND_API bool DoesDataTypeSupportLiteralType(FName InDataType, EMetasoundFrontendLiteralType InLiteralType);
-
-		// Returns true if the data type can be constructed using the literal type.
-		METASOUNDFRONTEND_API bool DoesDataTypeSupportLiteralType(FName InDataType, ELiteralType InLiteralType);
-
-		// Returns the defaulted version of a literal param for the given data type.
-		// @Returns an invalid init param if the data type couldn't be found.
-		METASOUNDFRONTEND_API FLiteral GetDefaultParamForDataType(FName InDataType);
-
 		// Convenience function to convert Metasound::ELiteralType to EMetasoundFrontendLiteralType.
 		METASOUNDFRONTEND_API EMetasoundFrontendLiteralType GetMetasoundFrontendLiteralType(ELiteralType InLiteralType);
 
diff --git a/Engine/Plugins/Runtime/Metasound/Source/MetasoundFrontend/Public/MetasoundFrontendRegistries.h b/Engine/Plugins/Runtime/Metasound/Source/MetasoundFrontend/Public/MetasoundFrontendRegistries.h
--- a/Engine/Plugins/Runtime/Metasound/Source/MetasoundFrontend/Public/MetasoundFrontendRegistries.h
+++ b/Engine/Plugins/Runtime/Metasound/Source/MetasoundFrontend/Public/MetasoundFrontendRegistries.h
@@ -36,7 +36,6 @@ namespace Metasound
 	using FOutputNodeConstructorParams = FDefaultNodeConstructorParams;
 	using FVariableNodeConstructorParams = FDefaultLiteralNodeConstructorParams;
 	using FInputNodeConstructorParams = FDefaultLiteralNodeConstructorParams;
-	using FReceiveNodeConstructorParams = FNodeInitData;
 
 	using FIterateMetasoundFrontendClassFunction = TFunctionRef<void(const FMetasoundFrontendClass&)>;
 
@@ -144,144 +143,8 @@ namespace Metasound
 			virtual bool IsNative() const = 0;
 		};
 
-		/** FDataTypeRegsitryInfo contains runtime inspectable behavior of a registered
-		 * MetaSound data type.
-		 */
-		struct FDataTypeRegistryInfo
-		{
-			// The name of the data type.
-			FName DataTypeName;
-
-			// The preferred constructor argument type for creating instances of the data type.
-			ELiteralType PreferredLiteralType = ELiteralType::Invalid;
-
-			// Constructor argument support in TDataTypeLiteralFactory<TDataType>;
-			bool bIsDefaultParsable = false;
-			bool bIsBoolParsable = false;
-			bool bIsIntParsable = false;
-			bool bIsFloatParsable = false;
-			bool bIsStringParsable = false;
-			bool bIsProxyParsable = false;
-			bool bIsDefaultArrayParsable = false;
-			bool bIsBoolArrayParsable = false;
-			bool bIsIntArrayParsable = false;
-			bool bIsFloatArrayParsable = false;
-			bool bIsStringArrayParsable = false;
-			bool bIsProxyArrayParsable = false;
-
-			// Is an TEnum wrapped enum
-			bool bIsEnum = false;
-
-			// Determines whether the type can be used with send/receive transmitters
-			bool bIsTransmittable = false;
-
-			// If this datatype was registered with a specific UClass to use to filter with, that will be used here:
-			UClass* ProxyGeneratorClass = nullptr;
-
-			FORCEINLINE bool IsArrayType() const
-			{
-				return bIsDefaultArrayParsable
-				|| bIsBoolArrayParsable
-				|| bIsIntArrayParsable
-				|| bIsFloatArrayParsable
-				|| bIsStringArrayParsable
-				|| bIsProxyArrayParsable;
-			}
-		};
-
-		/** Interface for metadata of a registered MetaSound enum type. */
-		struct METASOUNDFRONTEND_API IEnumDataTypeInterface
-		{
-			using FGenericInt32Entry = TEnumEntry<int32>;
-
-			virtual ~IEnumDataTypeInterface() = default;
-
-			virtual const TArray<FGenericInt32Entry>& GetAllEntries() const = 0;
-			virtual FName GetNamespace() const = 0;
-			virtual int32 GetDefaultValue() const = 0;
-
-			template<typename Predicate>
-			TOptional<FGenericInt32Entry> FindEntryBy(Predicate Pred) const
-			{
-				TArray<FGenericInt32Entry> Entries = GetAllEntries();
-				if (FGenericInt32Entry* Found = Entries.FindByPredicate(Pred))
-				{
-					return *Found;
-				}
-				return {};
-			}
-			TOptional<FGenericInt32Entry> FindByValue(int32 InEnumValue) const
-			{
-				return FindEntryBy([InEnumValue](const FGenericInt32Entry& i) -> bool { return i.Value == InEnumValue; });
-			}
-			TOptional<FGenericInt32Entry> FindByName(FName InEnumName) const
-			{
-				return FindEntryBy([InEnumName](const FGenericInt32Entry& i) -> bool { return i.Name == InEnumName; });
-			}
-			TOptional<FName> ToName(int32 InEnumValue) const
-			{
-				if(TOptional<FGenericInt32Entry> Result = FindByValue(InEnumValue))
-				{
-					return Result->Name;
-				}
-				return {};
-			}
-			TOptional<int32> ToValue(FName InName) const
-			{
-				if (TOptional<FGenericInt32Entry> Result = FindByName(InName))
-				{
-					return Result->Value;
-				}
-				return {};
-			}
-		};
-
-		/** Registry entry interface for a MetaSound data type. */
-		class IDataTypeRegistryEntry
-		{
-		public:
-			virtual ~IDataTypeRegistryEntry() = default;
-
-			/** Return the FDataTypeRegistryInfo for the data type */
-			virtual const FDataTypeRegistryInfo& GetDataTypeInfo() const = 0;
-
-			/** Return the enum interface for the data type. If the data type does
-			 * not support an enum interface, the returned pointer should be invalid.
-			 */
-			virtual TSharedPtr<const IEnumDataTypeInterface> GetEnumInterface() const = 0;
-
-			/** Return an FMetasoundFrontenClass representing an input node of the data type. */
-			virtual const FMetasoundFrontendClass& GetFrontendInputClass() const = 0;
-
-			/** Return an FMetasoundFrontenClass representing a variable node of the data type. */
-			virtual const FMetasoundFrontendClass& GetFrontendVariableClass() const = 0;
 
-			/** Return an FMetasoundFrontenClass representing an output node of the data type. */
-			virtual const FMetasoundFrontendClass& GetFrontendOutputClass() const = 0;
 
-			/** Create an input node */
-			virtual TUniquePtr<INode> CreateInputNode(FInputNodeConstructorParams&&) const = 0;
-
-			/** Create an output node */
-			virtual TUniquePtr<INode> CreateOutputNode(FOutputNodeConstructorParams&&) const = 0;
-			
-			/** Create a variable node */
-			virtual TUniquePtr<INode> CreateVariableNode(FVariableNodeConstructorParams&&) const = 0;
-
-			/** Create a recevie node. */
-			virtual TUniquePtr<INode> CreateReceiveNode(FReceiveNodeConstructorParams&&) const = 0;
-
-			/** Create a proxy from a UObject. If this data type does not support
-			 * UObject proxies, return a nullptr. */
-			virtual Audio::IProxyDataPtr CreateProxy(UObject*) const = 0;
-
-			/** Create a data channel for transmission. If this data type does not
-			 * support transmission, return a nullptr. */
-			virtual TSharedPtr<IDataChannel, ESPMode::ThreadSafe> CreateDataChannel(const FOperatorSettings&) const = 0;
-
-			/** Clone this registry entry. */
-			virtual TUniquePtr<IDataTypeRegistryEntry> Clone() const = 0;
-		};
 
 		struct METASOUNDFRONTEND_API FConverterNodeRegistryKey
 		{
@@ -396,9 +259,6 @@ public:
 	using FNodeRegistryKey = Metasound::Frontend::FNodeRegistryKey;
 	using FNodeClassMetadata = Metasound::FNodeClassMetadata;
 
-	using FDataTypeRegistryInfo = Metasound::Frontend::FDataTypeRegistryInfo;
-	using IEnumDataTypeInterface = Metasound::Frontend::IEnumDataTypeInterface;
-
 	static FMetasoundFrontendRegistryContainer* Get();
 	static void ShutdownMetasoundFrontend();
 
@@ -429,45 +289,6 @@ public:
 	/** Perform function for each registry transaction since a given transaction ID. */
 	virtual void ForEachNodeRegistryTransactionSince(Metasound::Frontend::FRegistryTransactionID InSince, Metasound::Frontend::FRegistryTransactionID* OutCurrentRegistryTransactionID, TFunctionRef<void(const Metasound::Frontend::FNodeRegistryTransaction&)> InFunc) const = 0;
 
-	/** Register a data type
-	 * @param InName - Name of data type.
-	 * @param InEntry - TUniquePtr to data type registry entry.
-	 *
-	 * @return True on success, false on failure.
-	 */
-	virtual bool RegisterDataType(const FName& InName, TUniquePtr<Metasound::Frontend::IDataTypeRegistryEntry>&& InEntry) = 0;
-
-	// Return any data types that can be used as a metasound input type or output type.
-	virtual TArray<FName> GetAllValidDataTypes() = 0;
-
-	// Get info about a specific data type (what kind of literals we can use, etc.)
-	// @returns false if InDataType wasn't found in the registry. 
-	virtual bool GetInfoForDataType(FName InDataType, FDataTypeRegistryInfo& OutInfo) = 0;
-
-	// Get the desired kind of literal for a given data type. Returns EConstructorArgType::Invalid if the data type couldn't be found.
-	virtual Metasound::ELiteralType GetDesiredLiteralTypeForDataType(FName InDataType) const = 0;
-
-	// Get whether we can build a literal of this specific type for InDataType.
-	virtual bool DoesDataTypeSupportLiteralType(FName InDataType, Metasound::ELiteralType InLiteralType) const = 0;
-
-	template<typename ArgType>
-	bool DoesDataTypeSupportLiteralType(FName InDataType) const
-	{
-		return DoesDataTypeSupportLiteralType(InDataType, Metasound::TLiteralTypeInfo<ArgType>::GetLiteralArgTypeEnum());
-	}
-
-	virtual UClass* GetLiteralUClassForDataType(FName InDataType) const = 0;
-	virtual Metasound::FLiteral CreateLiteralFromUObject(const FName& InDataType, UObject* InObject) = 0;
-	virtual Metasound::FLiteral CreateLiteralFromUObjectArray(const FName& InDataType, TArray<UObject*> InObjectArray) = 0;
-
-	// Return the enum interface for a data type. If the data type does not have 
-	// an enum interface, returns a nullptr.
-	virtual TSharedPtr<const Metasound::Frontend::IEnumDataTypeInterface> GetEnumInterfaceForDataType(FName InDataType) const = 0;
-
-	// Create a data channel for a data type. If the data type cannot be used in
-	// a data channel, return a nullptr.
-	virtual TSharedPtr<Metasound::IDataChannel, ESPMode::ThreadSafe> CreateDataChannelForDataType(const FName& InDataType, const Metasound::FOperatorSettings& InOperatorSettings) const = 0;
-
 
 	/** Register an external node with the frontend.
 	 *
@@ -504,11 +325,6 @@ public:
 	virtual bool FindVariableNodeRegistryKeyForDataType(const FName& InDataTypeName, FNodeRegistryKey& OutKey) = 0;
 	virtual bool FindOutputNodeRegistryKeyForDataType(const FName& InDataTypeName, FNodeRegistryKey& OutKey) = 0;
 
-	// Create a new instance of a C++ implemented node from the registry.
-	virtual TUniquePtr<Metasound::INode> CreateInputNode(const FName& InInputType, Metasound::FInputNodeConstructorParams&& InParams) = 0;
-	virtual TUniquePtr<Metasound::INode> CreateVariableNode(const FName& InVariableType, Metasound::FVariableNodeConstructorParams&& InParams) = 0;
-	virtual TUniquePtr<Metasound::INode> CreateOutputNode(const FName& InOutputType, Metasound::FOutputNodeConstructorParams&& InParams) = 0;
-	virtual TUniquePtr<Metasound::INode> CreateReceiveNode(const FName& InVariableType, Metasound::FReceiveNodeConstructorParams&& InParams) = 0;
 
 	virtual TUniquePtr<Metasound::INode> CreateNode(const FNodeRegistryKey& InKey, Metasound::FDefaultNodeConstructorParams&&) const = 0;
 	virtual TUniquePtr<Metasound::INode> CreateNode(const FNodeRegistryKey& InKey, Metasound::FDefaultLiteralNodeConstructorParams&&) const = 0;
