Disable asyn generation of metasound

diff --git a/Engine/Plugins/Runtime/Metasound/Source/MetasoundGenerator/Private/MetasoundGenerator.cpp b/Engine/Plugins/Runtime/Metasound/Source/MetasoundGenerator/Private/MetasoundGenerator.cpp
--- a/Engine/Plugins/Runtime/Metasound/Source/MetasoundGenerator/Private/MetasoundGenerator.cpp
+++ b/Engine/Plugins/Runtime/Metasound/Source/MetasoundGenerator/Private/MetasoundGenerator.cpp
@@ -11,7 +11,7 @@
 #include "MetasoundTrigger.h"
 
 #ifndef METASOUND_GENERATOR_BUILD_GRAPH_ASYNC
-#define METASOUND_GENERATOR_BUILD_GRAPH_ASYNC 1
+#define METASOUND_GENERATOR_BUILD_GRAPH_ASYNC 0
 #endif
 
 namespace Metasound
