Add Peak Mode input and Audio-rate output to Envelope Follower Node in Metasounds

diff --git a/Engine/Plugins/Runtime/Metasound/Source/MetasoundStandardNodes/Private/MetasoundEnvelopeFollowerNode.cpp b/Engine/Plugins/Runtime/Metasound/Source/MetasoundStandardNodes/Private/MetasoundEnvelopeFollowerNode.cpp
--- a/Engine/Plugins/Runtime/Metasound/Source/MetasoundStandardNodes/Private/MetasoundEnvelopeFollowerNode.cpp
+++ b/Engine/Plugins/Runtime/Metasound/Source/MetasoundStandardNodes/Private/MetasoundEnvelopeFollowerNode.cpp
@@ -2,9 +2,11 @@
 
 #include "MetasoundEnvelopeFollowerNode.h"
 
+#include "Internationalization/Text.h"
 #include "MetasoundAudioBuffer.h"
 #include "MetasoundExecutableOperator.h"
 #include "MetasoundFacade.h"
+#include "MetasoundEnumRegistrationMacro.h"
 #include "MetasoundNodeRegistrationMacro.h"
 #include "MetasoundDataTypeRegistrationMacro.h"
 #include "MetasoundOperatorSettings.h"
@@ -13,6 +15,7 @@
 #include "MetasoundTrigger.h"
 #include "MetasoundTime.h"
 #include "MetasoundVertex.h"
+#include "MetasoundParamHelper.h"
 #include "DSP/EnvelopeFollower.h"
 
 #define LOCTEXT_NAMESPACE "MetasoundStandardNodes_EnvelopeFollower"
@@ -24,8 +27,26 @@ namespace Metasound
 		static const TCHAR* InParamNameAudioInput = TEXT("In");
 		static const TCHAR* InParamNameAttackTime = TEXT("Attack Time");
 		static const TCHAR* InParamNameReleaseTime = TEXT("Release Time");
+		static const TCHAR* InParamNameFollowMode = TEXT("Peak Mode");
 		static const TCHAR* OutParamNameEnvelope = TEXT("Envelope");
 	}
+	METASOUND_PARAM(OutputAudioEnvelope, "Audio Envelope", "The output envelope value of the audio signal (audio rate).");
+
+	enum class EEnvelopeFollowMode
+	{
+		MeanSquared = 0,
+		RootMeanSquared,
+		Peak
+	};
+
+	DECLARE_METASOUND_ENUM(EEnvelopeFollowMode, EEnvelopeFollowMode::MeanSquared, METASOUNDSTANDARDNODES_API,
+	FEnumEnvelopeFollowMode, FEnumEnvelopeFollowModeInfo, FEnvelopeFollowModeReadRef, FEnumEnvelopeFollowModeWriteRef);
+
+	DEFINE_METASOUND_ENUM_BEGIN(EEnvelopeFollowMode, FEnumEnvelopeFollowMode, "EnvelopeFollowMode")
+		DEFINE_METASOUND_ENUM_ENTRY(EEnvelopeFollowMode::MeanSquared, LOCTEXT("EnvelopeFollowModeMSDescription", "MS"), LOCTEXT("EnvelopeFollowModeMSDescriptionTT", "Envelope follows a running Mean Squared of the audio signal.")),
+		DEFINE_METASOUND_ENUM_ENTRY(EEnvelopeFollowMode::RootMeanSquared, LOCTEXT("EnvelopeFollowModeRMSDescription", "RMS"), LOCTEXT("EnvelopeFollowModeRMSDescriptionTT", "Envelope follows a running Root Mean Squared of the audio signal.")),
+		DEFINE_METASOUND_ENUM_ENTRY(EEnvelopeFollowMode::Peak, LOCTEXT("EnvelopeFollowModePeakDescription", "Peak"), LOCTEXT("EnvelopeFollowModePeakDescriptionTT", "Envelope follows the peaks in the audio signal.")),
+	DEFINE_METASOUND_ENUM_END()
 
 	class FEnvelopeFollowerOperator : public TExecutableOperator<FEnvelopeFollowerOperator>
 	{
@@ -35,10 +56,11 @@ namespace Metasound
 		static const FVertexInterface& GetVertexInterface();
 		static TUniquePtr<IOperator> CreateOperator(const FCreateOperatorParams& InParams, FBuildErrorArray& OutErrors);
 
-		FEnvelopeFollowerOperator(const FOperatorSettings& InSettings, 
-			const FAudioBufferReadRef& InAudioInput, 
-			const FTimeReadRef& InAttackTime, 
-			const FTimeReadRef& InReleaseTime);
+		FEnvelopeFollowerOperator(const FOperatorSettings& InSettings,
+			const FAudioBufferReadRef& InAudioInput,
+			const FTimeReadRef& InAttackTime,
+			const FTimeReadRef& InReleaseTime,
+			const FEnvelopeFollowModeReadRef& InEnvelopeMode);
 
 		virtual FDataReferenceCollection GetInputs() const override;
 		virtual FDataReferenceCollection GetOutputs() const override;
@@ -49,32 +71,40 @@ namespace Metasound
 		FAudioBufferReadRef AudioInput;
 
 		// The amount of attack time
-		FTimeReadRef AttackTime;
+		FTimeReadRef AttackTimeInput;
 
 		// The amount of release time
-		FTimeReadRef ReleaseTime;
+		FTimeReadRef ReleaseTimeInput;
+
+		// The Envelope-Following method
+		FEnvelopeFollowModeReadRef FollowModeInput;
 
-		// The audio output
-		FFloatWriteRef EnvelopeOutput;
+		// The envelope outputs
+		FFloatWriteRef EnvelopeFloatOutput;
+		FAudioBufferWriteRef EnvelopeAudioOutput;
 
 		// The envelope follower DSP object
 		Audio::FEnvelopeFollower EnvelopeFollower;
 
 		double PrevAttackTime = 0.0;
 		double PrevReleaseTime = 0.0;
+		EEnvelopeFollowMode PrevFollowMode = EEnvelopeFollowMode::Peak;
 	};
 
-	FEnvelopeFollowerOperator::FEnvelopeFollowerOperator(const FOperatorSettings& InSettings, 
-		const FAudioBufferReadRef& InAudioInput, 
+	FEnvelopeFollowerOperator::FEnvelopeFollowerOperator(const FOperatorSettings& InSettings,
+		const FAudioBufferReadRef& InAudioInput,
 		const FTimeReadRef& InAttackTime,
-		const FTimeReadRef& InReleaseTime)
+		const FTimeReadRef& InReleaseTime,
+		const FEnvelopeFollowModeReadRef& InEnvelopeMode)
 		: AudioInput(InAudioInput)
-		, AttackTime(InAttackTime)
-		, ReleaseTime(InReleaseTime)
-		, EnvelopeOutput(FFloatWriteRef::CreateNew(0.0f))
+		, AttackTimeInput(InAttackTime)
+		, ReleaseTimeInput(InReleaseTime)
+		, FollowModeInput(InEnvelopeMode)
+		, EnvelopeFloatOutput(FFloatWriteRef::CreateNew())
+		, EnvelopeAudioOutput(FAudioBufferWriteRef::CreateNew(InSettings))
 	{
-		PrevAttackTime = FMath::Max(FTime::ToMilliseconds(*AttackTime), 0.0);
-		PrevReleaseTime = FMath::Max(FTime::ToMilliseconds(*ReleaseTime), 0.0);
+		PrevAttackTime = FMath::Max(FTime::ToMilliseconds(*AttackTimeInput), 0.0);
+		PrevReleaseTime = FMath::Max(FTime::ToMilliseconds(*ReleaseTimeInput), 0.0);
 
 		EnvelopeFollower.Init(InSettings.GetSampleRate(), PrevAttackTime, PrevReleaseTime);
 	}
@@ -83,42 +113,59 @@ namespace Metasound
 	{
 		FDataReferenceCollection InputDataReferences;
 		InputDataReferences.AddDataReadReference(EnvelopeFollower::InParamNameAudioInput, AudioInput);
-		InputDataReferences.AddDataReadReference(EnvelopeFollower::InParamNameAttackTime, AttackTime);
-		InputDataReferences.AddDataReadReference(EnvelopeFollower::InParamNameReleaseTime, ReleaseTime);
+		InputDataReferences.AddDataReadReference(EnvelopeFollower::InParamNameAttackTime, AttackTimeInput);
+		InputDataReferences.AddDataReadReference(EnvelopeFollower::InParamNameReleaseTime, ReleaseTimeInput);
+		InputDataReferences.AddDataReadReference(EnvelopeFollower::InParamNameFollowMode, FollowModeInput);
+
 		return InputDataReferences;
 	}
 
 	FDataReferenceCollection FEnvelopeFollowerOperator::GetOutputs() const
 	{
 		FDataReferenceCollection OutputDataReferences;
-		OutputDataReferences.AddDataReadReference(EnvelopeFollower::OutParamNameEnvelope, EnvelopeOutput);
+		OutputDataReferences.AddDataReadReference(EnvelopeFollower::OutParamNameEnvelope, EnvelopeFloatOutput);
+		OutputDataReferences.AddDataReadReference(METASOUND_GET_PARAM_NAME(OutputAudioEnvelope), EnvelopeAudioOutput);
 		return OutputDataReferences;
 	}
 
 	void FEnvelopeFollowerOperator::Execute()
 	{
 		// Check for any input changes
-		double CurrentAttackTime = FMath::Max(FTime::ToMilliseconds(*AttackTime), 0.0);
+		double CurrentAttackTime = FMath::Max(FTime::ToMilliseconds(*AttackTimeInput), 0.0);
 		if (!FMath::IsNearlyEqual(CurrentAttackTime, PrevAttackTime))
 		{
 			PrevAttackTime = CurrentAttackTime;
 			EnvelopeFollower.SetAttackTime(CurrentAttackTime);
 		}
 
-		double CurrentReleaseTime = FMath::Max(FTime::ToMilliseconds(*ReleaseTime), 0.0);
+		double CurrentReleaseTime = FMath::Max(FTime::ToMilliseconds(*ReleaseTimeInput), 0.0);
 		if (!FMath::IsNearlyEqual(CurrentReleaseTime, PrevReleaseTime))
 		{
 			PrevReleaseTime = CurrentReleaseTime;
 			EnvelopeFollower.SetReleaseTime(CurrentReleaseTime);
 		}
 
-		// Process the audio through the envelope follower
-		const float* InputAudio = AudioInput->GetData();
-		int32 NumFrames = AudioInput->Num();
-		EnvelopeFollower.ProcessAudio(InputAudio, NumFrames);
+		if (PrevFollowMode != *FollowModeInput)
+		{
+			PrevFollowMode = *FollowModeInput;
+			switch (PrevFollowMode)
+			{
+			case EEnvelopeFollowMode::MeanSquared:
+			default:
+				EnvelopeFollower.SetMode(Audio::EPeakMode::Type::MeanSquared);
+				break;
+			case EEnvelopeFollowMode::RootMeanSquared:
+				EnvelopeFollower.SetMode(Audio::EPeakMode::Type::RootMeanSquared);
+				break;
+			case EEnvelopeFollowMode::Peak:
+				EnvelopeFollower.SetMode(Audio::EPeakMode::Type::Peak);
+				break;
+			}
+		}
 
-		// Write the current envelope follower value to the output
-		*EnvelopeOutput = EnvelopeFollower.GetCurrentValue();
+		// Process the audio through the envelope follower
+		EnvelopeFollower.ProcessAudio(AudioInput->GetData(), EnvelopeAudioOutput->GetData(), AudioInput->Num());
+		*EnvelopeFloatOutput = EnvelopeFollower.GetCurrentValue();
 	}
 
 	const FVertexInterface& FEnvelopeFollowerOperator::GetVertexInterface()
@@ -127,10 +174,12 @@ namespace Metasound
 			FInputVertexInterface(
 				TInputDataVertexModel<FAudioBuffer>(EnvelopeFollower::InParamNameAudioInput, LOCTEXT("AudioInputToolTT", "Audio input.")),
 				TInputDataVertexModel<FTime>(EnvelopeFollower::InParamNameAttackTime, LOCTEXT("AttackTimeTT", "The attack time of the envelope follower."), 0.01f),
-				TInputDataVertexModel<FTime>(EnvelopeFollower::InParamNameReleaseTime, LOCTEXT("ReleaseTimeTT", "The release time of the envelope follower."), 0.1f)
+				TInputDataVertexModel<FTime>(EnvelopeFollower::InParamNameReleaseTime, LOCTEXT("ReleaseTimeTT", "The release time of the envelope follower."), 0.1f),
+				TInputDataVertexModel<FEnumEnvelopeFollowMode>(EnvelopeFollower::InParamNameFollowMode, LOCTEXT("FollowModeTT", "The following-method of the envelope follower."))
 			),
 			FOutputVertexInterface(
-				TOutputDataVertexModel<float>(EnvelopeFollower::OutParamNameEnvelope, LOCTEXT("EnvelopeFollowerOutputTT", "The output envelope value of the audio signal."))
+				TOutputDataVertexModel<float>(EnvelopeFollower::OutParamNameEnvelope, LOCTEXT("EnvelopeFollowerOutputTT", "The output envelope value of the audio signal.")),
+				TOutputDataVertexModel<FAudioBuffer>(METASOUND_GET_PARAM_NAME_AND_TT(OutputAudioEnvelope))
 			)
 		);
 
@@ -144,7 +193,7 @@ namespace Metasound
 			FNodeClassMetadata Info;
 			Info.ClassName = { StandardNodes::Namespace, TEXT("Envelope Follower"), TEXT("") };
 			Info.MajorVersion = 1;
-			Info.MinorVersion = 2;
+			Info.MinorVersion = 3;
 			Info.DisplayName = LOCTEXT("Metasound_EnvelopeFollowerDisplayName", "Envelope Follower");
 			Info.Description = LOCTEXT("Metasound_EnvelopeFollowerDescription", "Outputs an envelope from an input audio signal.");
 			Info.Author = PluginAuthor;
@@ -168,8 +217,10 @@ namespace Metasound
 		FAudioBufferReadRef AudioIn = InputCollection.GetDataReadReferenceOrConstruct<FAudioBuffer>(EnvelopeFollower::InParamNameAudioInput, InParams.OperatorSettings);
 		FTimeReadRef AttackTime = InputCollection.GetDataReadReferenceOrConstructWithVertexDefault<FTime>(InputInterface, EnvelopeFollower::InParamNameAttackTime, InParams.OperatorSettings);
 		FTimeReadRef ReleaseTime = InputCollection.GetDataReadReferenceOrConstructWithVertexDefault<FTime>(InputInterface, EnvelopeFollower::InParamNameReleaseTime, InParams.OperatorSettings);
+		FEnvelopeFollowModeReadRef EnvelopeModeIn = InputCollection.GetDataReadReferenceOrConstruct<FEnumEnvelopeFollowMode>(EnvelopeFollower::InParamNameFollowMode);
+
 
-		return MakeUnique<FEnvelopeFollowerOperator>(InParams.OperatorSettings, AudioIn, AttackTime, ReleaseTime);
+		return MakeUnique<FEnvelopeFollowerOperator>(InParams.OperatorSettings, AudioIn, AttackTime, ReleaseTime, EnvelopeModeIn);
 	}
 
 	FEnvelopeFollowerNode::FEnvelopeFollowerNode(const FNodeInitData& InitData)
diff --git a/Engine/Source/Runtime/SignalProcessing/Private/EnvelopeFollower.cpp b/Engine/Source/Runtime/SignalProcessing/Private/EnvelopeFollower.cpp
--- a/Engine/Source/Runtime/SignalProcessing/Private/EnvelopeFollower.cpp
+++ b/Engine/Source/Runtime/SignalProcessing/Private/EnvelopeFollower.cpp
@@ -15,7 +15,6 @@ namespace Audio
 		: EnvMode(EPeakMode::Peak)
 		, MeanWindowSize(DefaultWindowSize)
 		, MeanHopSize(DefaultHopSize)
-		, SumBuffer(DefaultWindowSize, DefaultHopSize)
 		, SampleRate(44100.0f)
 		, AttackTimeMsec(0.0f)
 		, AttackTimeSamples(0.0f)
@@ -32,9 +31,9 @@ namespace Audio
 		const EPeakMode::Type InMode,
 		const bool bInIsAnalog,
 		const int32 InWindowSizeForMean,
-		const int32 InHopSizeForMean) : SumBuffer(InWindowSizeForMean, InHopSizeForMean)
+		const int32 InHopSizeForMean)		
 	{
-		Init(InSampleRate, InAttackTimeMsec, InReleaseTimeMSec, InMode, bInIsAnalog,InWindowSizeForMean, InHopSizeForMean);
+		Init(InSampleRate, InAttackTimeMsec, InReleaseTimeMSec, InMode, bInIsAnalog, InWindowSizeForMean, InHopSizeForMean);
 	}
 
 	FEnvelopeFollower::~FEnvelopeFollower()
@@ -49,16 +48,23 @@ namespace Audio
 		const int32 InWindowSizeForMean,
 		const int32 InHopSizeForMean)
 	{
+		check(InWindowSizeForMean > 0);
+		check(InHopSizeForMean > 0);
+		check(InWindowSizeForMean > InHopSizeForMean);
+
 		SampleRate = InSampleRate;
 
 		bIsAnalog = bInIsAnalog;
 		EnvMode = InMode;
 
-		SumBuffer = TSlidingBuffer<float>(InWindowSizeForMean, InHopSizeForMean);
+		MeanWindowSize = InWindowSizeForMean;
+		MeanHopSize = InHopSizeForMean;
 
 		// Set the attack and release times using the default values
 		SetAttackTime(InAttackTimeMsec);
 		SetReleaseTime(InReleaseTimeMSec);
+
+		ScratchBuffer.SetNum(InWindowSizeForMean);
 	}
 
 	void FEnvelopeFollower::Reset()
@@ -105,23 +111,31 @@ namespace Audio
 		// MS/RMS
 		if (EnvMode == EPeakMode::MeanSquared || EnvMode == EPeakMode::RootMeanSquared)
 		{
-			TAutoSlidingWindow<float> SlidingWindow(SumBuffer, TArrayView<const float>(InAudioBuffer, InNumSamples), ScratchBuffer, false);
-
-			for (auto& Window : SlidingWindow)
+			TCircularAudioBuffer<float> SumBuffer(MeanWindowSize);
+			int32 SampleIndex = 0;
+			// This loop iterates per-window, not per-sample
+			while (SampleIndex < InNumSamples)
 			{
-				float CurrentMean;
+				// Get a new Window
+				const int32 SamplesToCopy = MeanHopSize < (InNumSamples - SampleIndex) ? MeanHopSize : (InNumSamples - SampleIndex);
+				SumBuffer.Push(InAudioBuffer + SampleIndex, SamplesToCopy);
+				float* CurrentWindow = ScratchBuffer.GetData();
+				SumBuffer.Pop(CurrentWindow, SamplesToCopy);
 				
-				ArrayMeanSquared(Window, CurrentMean);
-
+				// Calculate MS/RMS
+				float CurrentMean;
+				ArrayMeanSquared(TArrayView<float>(CurrentWindow, SamplesToCopy), CurrentMean);
 				if (EnvMode == EPeakMode::RootMeanSquared)
 				{
 					CurrentMean = FMath::Sqrt(CurrentMean);
 				}
 
-				for (int j = 0; j < MeanWindowSize; ++j)
+				// Use MS/RMS to drive envelope following
+				for (int i = 0; i < SamplesToCopy; ++i)
 				{
 					ProcessAudioNonClamped(CurrentMean);
 				}
+				SampleIndex += SamplesToCopy;
 			}
 		}
 		// Peak mode
@@ -140,24 +154,31 @@ namespace Audio
 		// MS/RMS
 		if (EnvMode == EPeakMode::MeanSquared || EnvMode == EPeakMode::RootMeanSquared)
 		{
-			TAutoSlidingWindow<float> SlidingWindow(SumBuffer, TArrayView<const float>(InAudioBuffer, InNumSamples), ScratchBuffer, false);
-
+			TCircularAudioBuffer<float> SumBuffer(MeanWindowSize);
 			int32 SampleIndex = 0;
-			for (auto& Window : SlidingWindow)
+			// This loop iterates per-window, not per-sample
+			while (SampleIndex < InNumSamples)
 			{
-				float CurrentMean;
-				ArrayMeanSquared(Window, CurrentMean);
+				// Get a new Window
+				const int32 SamplesToCopy = MeanHopSize < (InNumSamples - SampleIndex) ? MeanHopSize : (InNumSamples - SampleIndex);
+				SumBuffer.Push(InAudioBuffer + SampleIndex, SamplesToCopy);
+				float* CurrentWindow = ScratchBuffer.GetData();
+				SumBuffer.Pop(CurrentWindow, MeanWindowSize);
 
+				// Calculate MS/RMS
+				float CurrentMean;
+				ArrayMeanSquared(TArrayView<float>(CurrentWindow, SamplesToCopy), CurrentMean);
 				if (EnvMode == EPeakMode::RootMeanSquared)
 				{
 					CurrentMean = FMath::Sqrt(CurrentMean);
 				}
 
-				for (int j = 0; j < MeanWindowSize; ++j)
+				// Use MS/RMS to drive envelope following
+				for (int i = 0; i < SamplesToCopy; ++i)
 				{
-					OutAudioBuffer[SampleIndex] = ProcessAudioNonClamped(CurrentMean);
-					++SampleIndex;
+					OutAudioBuffer[SampleIndex + i] = ProcessAudioNonClamped(CurrentMean);
 				}
+				SampleIndex += SamplesToCopy;
 			}
 		}
 		// Peak
diff --git a/Engine/Source/Runtime/SignalProcessing/Public/DSP/EnvelopeFollower.h b/Engine/Source/Runtime/SignalProcessing/Public/DSP/EnvelopeFollower.h
--- a/Engine/Source/Runtime/SignalProcessing/Public/DSP/EnvelopeFollower.h
+++ b/Engine/Source/Runtime/SignalProcessing/Public/DSP/EnvelopeFollower.h
@@ -3,7 +3,7 @@
 #pragma once
 
 #include "CoreMinimal.h"
-#include "SlidingWindow.h"
+#include "DSP/Dsp.h"
 
 namespace Audio
 {
@@ -85,8 +85,7 @@ namespace Audio
 	protected:
 		EPeakMode::Type EnvMode;
 		int32 MeanWindowSize;
-		int32 MeanHopSize;
-		TSlidingBuffer<float> SumBuffer;
+		int32 MeanHopSize;		
 		TArray<float> ScratchBuffer;
 		float SampleRate;
 		float AttackTimeMsec;
