Add Trigger Sequence Node to Metasounds

diff --git a/Engine/Plugins/Runtime/Metasound/Source/MetasoundStandardNodes/Private/MetasoundTriggerSequenceNode.cpp b/Engine/Plugins/Runtime/Metasound/Source/MetasoundStandardNodes/Private/MetasoundTriggerSequenceNode.cpp
new file mode 100644
--- /dev/null
+++ b/Engine/Plugins/Runtime/Metasound/Source/MetasoundStandardNodes/Private/MetasoundTriggerSequenceNode.cpp
@@ -0,0 +1,238 @@
+// Copyright Epic Games, Inc. All Rights Reserved.
+
+#include "CoreMinimal.h"
+#include "MetasoundBuilderInterface.h"
+#include "MetasoundDataReferenceCollection.h"
+#include "MetasoundExecutableOperator.h"
+#include "MetasoundFacade.h"
+#include "MetasoundNode.h"
+#include "MetasoundNodeInterface.h"
+#include "MetasoundNodeRegistrationMacro.h"
+#include "MetasoundOperatorInterface.h"
+#include "MetasoundPrimitives.h"
+#include "MetasoundSampleCounter.h"
+#include "MetasoundStandardNodesNames.h"
+#include "MetasoundStandardNodesCategories.h"
+#include "MetasoundParamHelper.h"
+
+#define LOCTEXT_NAMESPACE "MetasoundStandardNodes_TriggerSequenceNode"
+
+#define REGISTER_TRIGGER_SEQUENCE_NODE(Number) \
+	using FTriggerSequenceNode_##Number = TTriggerSequenceNode<Number>; \
+	METASOUND_REGISTER_NODE(FTriggerSequenceNode_##Number) \
+
+namespace Metasound
+{
+
+	namespace TriggerSequenceVertexNames
+	{
+		METASOUND_PARAM(InputTrigger, "In", "Input trigger.");
+		METASOUND_PARAM(InputReset, "Reset", "Resets the sequence. When triggered, the next trigger will be Out 0.");		
+		METASOUND_PARAM(InputLoop, "Loop", "Whether the sequence will automatically loop back to Out 0 after all triggers have been called.");
+
+		METASOUND_PARAM(OutputTrigger, "Out {0}", "Trigger Output {0} in the sequence.");
+
+	}
+
+	template<uint32 NumOutputs>
+	class TTriggerSequenceOperator : public TExecutableOperator<TTriggerSequenceOperator<NumOutputs>>
+	{
+	public:
+		static const FNodeClassMetadata& GetNodeInfo();
+		static const FVertexInterface& GetVertexInterface();
+		static TUniquePtr<IOperator> CreateOperator(const FCreateOperatorParams& InParams, FBuildErrorArray& OutErrors);
+
+		TTriggerSequenceOperator(const FOperatorSettings& InSettings, 
+			const FTriggerReadRef& InTrigger, 
+			const FTriggerReadRef& InReset,
+			const FBoolReadRef& bInLoop);
+
+		virtual FDataReferenceCollection GetInputs() const override;
+
+		virtual FDataReferenceCollection GetOutputs() const override;
+
+		void Execute();
+
+	private:
+		FTriggerReadRef TriggerInput;
+		FTriggerReadRef TriggerResetInput;
+		FBoolReadRef bLoopInput;
+
+		TArray<FTriggerWriteRef> TriggerOutputs;
+
+		uint32 CurrentIndex = 0;
+
+	};
+
+	template<uint32 NumOutputs>
+	TTriggerSequenceOperator<NumOutputs>::TTriggerSequenceOperator(const FOperatorSettings& InSettings, 
+		const FTriggerReadRef& InTrigger, 
+		const FTriggerReadRef& InReset,
+		const FBoolReadRef& bInLoop)
+		: TriggerInput(InTrigger)
+		, TriggerResetInput(InReset)
+		, bLoopInput(bInLoop)
+	{
+		
+		for (uint32 i = 0; i < NumOutputs; ++i)
+		{
+			TriggerOutputs.Add(FTriggerWriteRef::CreateNew(InSettings));
+		}
+	}
+
+	template<uint32 NumOutputs>
+	FDataReferenceCollection TTriggerSequenceOperator<NumOutputs>::GetInputs() const
+	{
+		using namespace TriggerSequenceVertexNames;
+
+		FDataReferenceCollection InputDataReferences;
+		InputDataReferences.AddDataReadReference(METASOUND_GET_PARAM_NAME(InputTrigger), TriggerInput);
+		InputDataReferences.AddDataReadReference(METASOUND_GET_PARAM_NAME(InputReset), TriggerResetInput);
+		InputDataReferences.AddDataReadReference(METASOUND_GET_PARAM_NAME(InputLoop), bLoopInput);
+		return InputDataReferences;
+	}
+
+	template<uint32 NumOutputs>
+	FDataReferenceCollection TTriggerSequenceOperator<NumOutputs>::GetOutputs() const
+	{
+		using namespace TriggerSequenceVertexNames;
+
+		FDataReferenceCollection OutputDataReferences;
+		for (uint32 i = 0; i < NumOutputs; ++i)
+		{
+			OutputDataReferences.AddDataReadReference(METASOUND_GET_VARIABLE_PARAM_NAME(OutputTrigger, i), TriggerOutputs[i]);
+		}
+
+		return OutputDataReferences;
+	}
+
+	template<uint32 NumOutputs>
+	void TTriggerSequenceOperator<NumOutputs>::Execute()
+	{
+		// Advance all output blocks
+		for (uint32 i = 0; i < NumOutputs; ++i)
+		{
+			TriggerOutputs[i]->AdvanceBlock();
+		}
+
+		// Reset the Sequence
+		TriggerResetInput->ExecuteBlock(
+			[](int32, int32)
+			{
+			},
+			[this](int32 StartFrame, int32 EndFrame)
+			{
+				CurrentIndex = 0;
+			}
+		);
+
+		// Continue sequence
+		TriggerInput->ExecuteBlock(
+			[](int32, int32) 
+			{
+			},
+			[this](int32 StartFrame, int32 EndFrame)
+			{
+				if (CurrentIndex < NumOutputs)
+				{
+					// Send Trigger Out
+					TriggerOutputs[CurrentIndex]->TriggerFrame(StartFrame);
+
+					// Update index (wrap if necessary)
+					++CurrentIndex;
+				}
+				if (*bLoopInput && CurrentIndex >= NumOutputs)
+				{
+					CurrentIndex = 0;
+				}
+			}
+		);
+	}
+
+	template<uint32 NumOutputs>
+	TUniquePtr<IOperator> TTriggerSequenceOperator<NumOutputs>::CreateOperator(const FCreateOperatorParams& InParams, FBuildErrorArray& OutErrors)
+	{
+		using namespace TriggerSequenceVertexNames;
+
+		const FInputVertexInterface& InputInterface = GetVertexInterface().GetInputInterface();
+
+		FTriggerReadRef TriggerIn = InParams.InputDataReferences.GetDataReadReferenceOrConstruct<FTrigger>(METASOUND_GET_PARAM_NAME(InputTrigger), InParams.OperatorSettings);
+		FTriggerReadRef TriggerResetIn = InParams.InputDataReferences.GetDataReadReferenceOrConstruct<FTrigger>(METASOUND_GET_PARAM_NAME(InputReset), InParams.OperatorSettings);
+		FBoolReadRef bLoopIn = InParams.InputDataReferences.GetDataReadReferenceOrConstructWithVertexDefault<bool>(InputInterface, METASOUND_GET_PARAM_NAME(InputLoop), InParams.OperatorSettings);
+
+		return MakeUnique<TTriggerSequenceOperator>(InParams.OperatorSettings, TriggerIn, TriggerResetIn, bLoopIn);
+	}
+
+	template<uint32 NumOutputs>
+	const FVertexInterface& TTriggerSequenceOperator<NumOutputs>::GetVertexInterface()
+	{
+		using namespace TriggerSequenceVertexNames;
+
+		auto CreateVertexInterface = []() -> FVertexInterface
+		{
+			FInputVertexInterface InputInterface;
+			InputInterface.Add(TInputDataVertexModel<FTrigger>(METASOUND_GET_PARAM_NAME_AND_TT(InputTrigger)));
+			InputInterface.Add(TInputDataVertexModel<FTrigger>(METASOUND_GET_PARAM_NAME_AND_TT(InputReset)));
+			InputInterface.Add(TInputDataVertexModel<bool>(METASOUND_GET_PARAM_NAME_AND_TT(InputLoop), true));
+			
+			FOutputVertexInterface OutputInterface;
+			for (uint32 i = 0; i < NumOutputs; ++i)
+			{
+				OutputInterface.Add(TOutputDataVertexModel<FTrigger>(METASOUND_GET_VARIABLE_PARAM_NAME_AND_TT(OutputTrigger, i)));
+			}
+			return FVertexInterface(InputInterface, OutputInterface);
+		};
+
+		static const FVertexInterface Interface = CreateVertexInterface();
+		return Interface;		
+	}
+
+	template<uint32 NumOutputs>
+	const FNodeClassMetadata& TTriggerSequenceOperator<NumOutputs>::GetNodeInfo()
+	{
+		auto InitNodeInfo = []() -> FNodeClassMetadata
+		{
+			FName OperatorName = *FString::Printf(TEXT("Trigger Sequence (%d)"), NumOutputs);
+			FText NodeDisplayName = FText::Format(LOCTEXT("TriggerSequenceDisplayNamePattern", "Trigger Sequence ({0})"), NumOutputs);
+			FText NodeDescription = LOCTEXT("TriggerSequenceDescription", "Each time this node is triggered, it sends the next output trigger in the sequence.");
+
+			FNodeClassMetadata Info;
+			Info.ClassName = { StandardNodes::Namespace, OperatorName, TEXT("") };
+			Info.MajorVersion = 1;
+			Info.MinorVersion = 0;
+			Info.DisplayName = NodeDisplayName;
+			Info.Description = NodeDescription;
+			Info.Author = PluginAuthor;
+			Info.PromptIfMissing = PluginNodeMissingPrompt;
+			Info.DefaultInterface = GetVertexInterface();
+			Info.CategoryHierarchy.Emplace(NodeCategories::Trigger);
+
+			return Info;
+		};
+
+		static const FNodeClassMetadata Info = InitNodeInfo();
+
+		return Info;
+	}
+
+	// Node Class
+	template <uint32 NumOutputs>
+	class TTriggerSequenceNode : public FNodeFacade
+	{
+	public:
+		TTriggerSequenceNode(const FNodeInitData& InitData)
+			: FNodeFacade(InitData.InstanceName, InitData.InstanceID, TFacadeOperatorClass<TTriggerSequenceOperator<NumOutputs>>())
+		{
+		}
+	};
+
+	REGISTER_TRIGGER_SEQUENCE_NODE(2)
+	REGISTER_TRIGGER_SEQUENCE_NODE(3)
+	REGISTER_TRIGGER_SEQUENCE_NODE(4)
+	REGISTER_TRIGGER_SEQUENCE_NODE(5)
+	REGISTER_TRIGGER_SEQUENCE_NODE(6)
+	REGISTER_TRIGGER_SEQUENCE_NODE(7)
+	REGISTER_TRIGGER_SEQUENCE_NODE(8)
+}
+
+#undef LOCTEXT_NAMESPACE
