Allow async metasound generator building W/O deadlocks

diff --git a/Engine/Plugins/Runtime/Metasound/Source/MetasoundEngine/Private/MetasoundEngineArchetypes.cpp b/Engine/Plugins/Runtime/Metasound/Source/MetasoundEngine/Private/MetasoundEngineArchetypes.cpp
--- a/Engine/Plugins/Runtime/Metasound/Source/MetasoundEngine/Private/MetasoundEngineArchetypes.cpp
+++ b/Engine/Plugins/Runtime/Metasound/Source/MetasoundEngine/Private/MetasoundEngineArchetypes.cpp
@@ -93,10 +93,10 @@ namespace Metasound
 				return OnFinishedOutputName;
 			}
 
-			const FString& GetAudioDeviceHandleVariableName()
+			const FString& GetAudioDeviceIDVariableName()
 			{
-				static const FString AudioDeviceHandleVarName = TEXT("AudioDeviceHandle");
-				return AudioDeviceHandleVarName;
+				static const FString AudioDeviceIDVarName = TEXT("AudioDeviceID");
+				return AudioDeviceIDVarName;
 			}
 
 			const FString& GetSoundUniqueIdName()
@@ -149,12 +149,12 @@ namespace Metasound
 				Archetype.Interface.Outputs.Add(GeneratedAudio);
 
 				// Environment 
-				FMetasoundFrontendEnvironmentVariable AudioDeviceHandle;
-				AudioDeviceHandle.Name = GetAudioDeviceHandleVariableName();
-				AudioDeviceHandle.Metadata.DisplayName = FText::FromString(AudioDeviceHandle.Name);
-				AudioDeviceHandle.Metadata.Description = LOCTEXT("AudioDeviceHandleToolTip", "Audio device handle");
+				FMetasoundFrontendEnvironmentVariable AudioDeviceID;
+				AudioDeviceID.Name = GetAudioDeviceIDVariableName();
+				AudioDeviceID.Metadata.DisplayName = FText::FromString(AudioDeviceID.Name);
+				AudioDeviceID.Metadata.Description = LOCTEXT("AudioDeviceIDToolTip", "Audio device handle");
 
-				Archetype.Interface.Environment.Add(AudioDeviceHandle);
+				Archetype.Interface.Environment.Add(AudioDeviceID);
 
 				return Archetype;
 			}
@@ -186,10 +186,10 @@ namespace Metasound
 				return OnFinishedOutputName;
 			}
 
-			const FString& GetAudioDeviceHandleVariableName()
+			const FString& GetAudioDeviceIDVariableName()
 			{
-				static const FString AudioDeviceHandleVarName = TEXT("AudioDeviceHandle");
-				return AudioDeviceHandleVarName;
+				static const FString AudioDeviceIDVarName = TEXT("AudioDeviceID");
+				return AudioDeviceIDVarName;
 			}
 
 			const FString& GetSoundUniqueIdName()
@@ -247,12 +247,12 @@ namespace Metasound
 				Archetype.Interface.Outputs.Add(GeneratedAudio);
 
 				// Environment
-				FMetasoundFrontendEnvironmentVariable AudioDeviceHandle;
-				AudioDeviceHandle.Name = GetAudioDeviceHandleVariableName();
-				AudioDeviceHandle.Metadata.DisplayName = FText::FromString(AudioDeviceHandle.Name);
-				AudioDeviceHandle.Metadata.Description = LOCTEXT("AudioDeviceHandleToolTip", "Audio device handle");
+				FMetasoundFrontendEnvironmentVariable AudioDeviceID;
+				AudioDeviceID.Name = GetAudioDeviceIDVariableName();
+				AudioDeviceID.Metadata.DisplayName = FText::FromString(AudioDeviceID.Name);
+				AudioDeviceID.Metadata.Description = LOCTEXT("AudioDeviceIDToolTip", "Audio device handle");
 
-				Archetype.Interface.Environment.Add(AudioDeviceHandle);
+				Archetype.Interface.Environment.Add(AudioDeviceID);
 
 				return Archetype;
 			}
@@ -284,10 +284,10 @@ namespace Metasound
 				return OnFinishedOutputName;
 			}
 
-			const FString& GetAudioDeviceHandleVariableName()
+			const FString& GetAudioDeviceIDVariableName()
 			{
-				static const FString AudioDeviceHandleVarName = TEXT("AudioDeviceHandle");
-				return AudioDeviceHandleVarName;
+				static const FString AudioDeviceIDVarName = TEXT("AudioDeviceID");
+				return AudioDeviceIDVarName;
 			}
 
 			const FString& GetSoundUniqueIdName()
@@ -357,12 +357,12 @@ namespace Metasound
 				Archetype.Interface.Outputs.Add(GeneratedAudio);
 
 				// Environment
-				FMetasoundFrontendEnvironmentVariable AudioDeviceHandle;
-				AudioDeviceHandle.Name = GetAudioDeviceHandleVariableName();
-				AudioDeviceHandle.Metadata.DisplayName = FText::FromString(AudioDeviceHandle.Name);
-				AudioDeviceHandle.Metadata.Description = LOCTEXT("AudioDeviceHandleToolTip", "Audio device handle");
+				FMetasoundFrontendEnvironmentVariable AudioDeviceID;
+				AudioDeviceID.Name = GetAudioDeviceIDVariableName();
+				AudioDeviceID.Metadata.DisplayName = FText::FromString(AudioDeviceID.Name);
+				AudioDeviceID.Metadata.Description = LOCTEXT("AudioDeviceIDToolTip", "Audio device handle");
 
-				Archetype.Interface.Environment.Add(AudioDeviceHandle);
+				Archetype.Interface.Environment.Add(AudioDeviceID);
 
 				return Archetype;
 			}
@@ -462,10 +462,10 @@ namespace Metasound
 				return OnFinishedOutputName;
 			}
 
-			const FString& GetAudioDeviceHandleVariableName()
+			const FString& GetAudioDeviceIDVariableName()
 			{
-				static const FString AudioDeviceHandleVarName = TEXT("AudioDeviceHandle");
-				return AudioDeviceHandleVarName;
+				static const FString AudioDeviceIDVarName = TEXT("AudioDeviceID");
+				return AudioDeviceIDVarName;
 			}
 
 			const FString& GetSoundUniqueIdName()
@@ -550,11 +550,11 @@ namespace Metasound
 				Archetype.Interface.Outputs.Add(GeneratedRightAudio);
 
 				// Environment
-				FMetasoundFrontendEnvironmentVariable AudioDeviceHandle;
-				AudioDeviceHandle.Name = GetAudioDeviceHandleVariableName();
-				AudioDeviceHandle.Metadata.DisplayName = FText::FromString(AudioDeviceHandle.Name);
-				AudioDeviceHandle.Metadata.Description = LOCTEXT("AudioDeviceHandleToolTip", "Audio device handle");
-				Archetype.Interface.Environment.Add(AudioDeviceHandle);
+				FMetasoundFrontendEnvironmentVariable AudioDeviceID;
+				AudioDeviceID.Name = GetAudioDeviceIDVariableName();
+				AudioDeviceID.Metadata.DisplayName = FText::FromString(AudioDeviceID.Name);
+				AudioDeviceID.Metadata.Description = LOCTEXT("AudioDeviceIDToolTip", "Audio device handle");
+				Archetype.Interface.Environment.Add(AudioDeviceID);
 
 				return Archetype;
 			};
@@ -659,10 +659,10 @@ namespace Metasound
 				return CurrentMonoVersion::GetIsFinishedOutputName();
 			}
 
-			const FString& GetAudioDeviceHandleVariableName()
+			const FString& GetAudioDeviceIDVariableName()
 			{
-				check(CurrentMonoVersion::GetAudioDeviceHandleVariableName() == CurrentStereoVersion::GetAudioDeviceHandleVariableName());
-				return CurrentMonoVersion::GetAudioDeviceHandleVariableName();
+				check(CurrentMonoVersion::GetAudioDeviceIDVariableName() == CurrentStereoVersion::GetAudioDeviceIDVariableName());
+				return CurrentMonoVersion::GetAudioDeviceIDVariableName();
 			}
 
 			const FString& GetSoundUniqueIdName()
diff --git a/Engine/Plugins/Runtime/Metasound/Source/MetasoundEngine/Private/MetasoundEngineArchetypes.h b/Engine/Plugins/Runtime/Metasound/Source/MetasoundEngine/Private/MetasoundEngineArchetypes.h
--- a/Engine/Plugins/Runtime/Metasound/Source/MetasoundEngine/Private/MetasoundEngineArchetypes.h
+++ b/Engine/Plugins/Runtime/Metasound/Source/MetasoundEngine/Private/MetasoundEngineArchetypes.h
@@ -22,7 +22,7 @@ namespace Metasound
 			const FString& GetOnPlayInputName();
 			const FString& GetAudioOutputName();
 			const FString& GetIsFinishedOutputName();
-			const FString& GetAudioDeviceHandleVariableName();
+			const FString& GetAudioDeviceIDVariableName();
 			const FString& GetSoundUniqueIdName();
 			const FString& GetIsPreviewSoundName();
 		}
@@ -34,7 +34,7 @@ namespace Metasound
 			const FString& GetOnPlayInputName();
 			const FString& GetAudioOutputName();
 			const FString& GetIsFinishedOutputName();
-			const FString& GetAudioDeviceHandleVariableName();
+			const FString& GetAudioDeviceIDVariableName();
 			const FString& GetSoundUniqueIdName();
 			const FString& GetIsPreviewSoundName();
 			const FString& GetGraphName();
@@ -47,7 +47,7 @@ namespace Metasound
 			const FString& GetOnPlayInputName();
 			const FString& GetAudioOutputName();
 			const FString& GetIsFinishedOutputName();
-			const FString& GetAudioDeviceHandleVariableName();
+			const FString& GetAudioDeviceIDVariableName();
 			const FString& GetSoundUniqueIdName();
 			const FString& GetIsPreviewSoundName();
 			const FString& GetInstanceIDName();
@@ -62,7 +62,7 @@ namespace Metasound
 			const FString& GetLeftAudioOutputName();
 			const FString& GetRightAudioOutputName();
 			const FString& GetIsFinishedOutputName();
-			const FString& GetAudioDeviceHandleVariableName();
+			const FString& GetAudioDeviceIDVariableName();
 			const FString& GetSoundUniqueIdName();
 			const FString& GetIsPreviewSoundName();
 			const FString& GetInstanceIDName();
@@ -74,7 +74,7 @@ namespace Metasound
 		{
 			const FString& GetOnPlayInputName();
 			const FString& GetIsFinishedOutputName();
-			const FString& GetAudioDeviceHandleVariableName();
+			const FString& GetAudioDeviceIDVariableName();
 			const FString& GetSoundUniqueIdName();
 			const FString& GetIsPreviewSoundName();
 			const FString& GetInstanceIDName();
diff --git a/Engine/Plugins/Runtime/Metasound/Source/MetasoundEngine/Private/MetasoundEngineEnvironment.cpp b/Engine/Plugins/Runtime/Metasound/Source/MetasoundEngine/Private/MetasoundEngineEnvironment.cpp
--- a/Engine/Plugins/Runtime/Metasound/Source/MetasoundEngine/Private/MetasoundEngineEnvironment.cpp
+++ b/Engine/Plugins/Runtime/Metasound/Source/MetasoundEngine/Private/MetasoundEngineEnvironment.cpp
@@ -5,5 +5,5 @@
 #include "AudioDeviceManager.h"
 #include "MetasoundEnvironment.h"
 
-DEFINE_METASOUND_ENVIRONMENT_VARIABLE_TYPE(FAudioDeviceHandle);
+//DEFINE_METASOUND_ENVIRONMENT_VARIABLE_TYPE(FAudioDeviceHandle);
 
diff --git a/Engine/Plugins/Runtime/Metasound/Source/MetasoundEngine/Private/MetasoundSource.cpp b/Engine/Plugins/Runtime/Metasound/Source/MetasoundEngine/Private/MetasoundSource.cpp
--- a/Engine/Plugins/Runtime/Metasound/Source/MetasoundEngine/Private/MetasoundSource.cpp
+++ b/Engine/Plugins/Runtime/Metasound/Source/MetasoundEngine/Private/MetasoundSource.cpp
@@ -349,21 +349,26 @@ Metasound::FMetasoundEnvironment UMetaSoundSource::CreateEnvironment() const
 
 	FMetasoundEnvironment Environment;
 	// Add audio device ID to environment.
-	FAudioDeviceHandle DeviceHandle;
+	Audio::FDeviceId AudioDeviceID = INDEX_NONE;
+
 	if (UWorld* World = GetWorld())
 	{
-		DeviceHandle = World->GetAudioDevice();
+		FAudioDeviceHandle DeviceHandle = World->GetAudioDevice();
+		if (DeviceHandle.IsValid())
+		{
+			AudioDeviceID = DeviceHandle.GetDeviceID();
+		}
 	}
 
-	if (!DeviceHandle.IsValid())
+	if (INDEX_NONE == AudioDeviceID)
 	{
 		if (FAudioDeviceManager* DeviceManager = FAudioDeviceManager::Get())
 		{
-			DeviceHandle = DeviceManager->GetMainAudioDeviceHandle();
+			AudioDeviceID = DeviceManager->GetMainAudioDeviceID();
 		}
 	}
 
-	Environment.SetValue<FAudioDeviceHandle>(MetasoundSource::GetAudioDeviceHandleVariableName(), DeviceHandle);
+	Environment.SetValue<Audio::FDeviceId>(MetasoundSource::GetAudioDeviceIDVariableName(), AudioDeviceID);
 	Environment.SetValue<uint32>(MetasoundSource::GetSoundUniqueIdName(), GetUniqueID());
 
 	return Environment;
diff --git a/Engine/Plugins/Runtime/Metasound/Source/MetasoundEngine/Public/MetasoundEngineEnvironment.h b/Engine/Plugins/Runtime/Metasound/Source/MetasoundEngine/Public/MetasoundEngineEnvironment.h
--- a/Engine/Plugins/Runtime/Metasound/Source/MetasoundEngine/Public/MetasoundEngineEnvironment.h
+++ b/Engine/Plugins/Runtime/Metasound/Source/MetasoundEngine/Public/MetasoundEngineEnvironment.h
@@ -5,4 +5,7 @@
 #include "AudioDeviceManager.h"
 #include "MetasoundEnvironment.h"
 
-DECLARE_METASOUND_ENVIRONMENT_VARIABLE_TYPE(METASOUNDENGINE_API, FAudioDeviceHandle);
+// Passing around an FAudioDeviceHandle which was held on sound generators
+// caused a deadlock on shutdown.  If needed, pass an audio device ID instead 
+// and make sure to release the handle when it's no longer needed.
+//DECLARE_METASOUND_ENVIRONMENT_VARIABLE_TYPE(METASOUNDENGINE_API, FAudioDeviceHandle);
diff --git a/Engine/Plugins/Runtime/Metasound/Source/MetasoundGenerator/Private/MetasoundGenerator.cpp b/Engine/Plugins/Runtime/Metasound/Source/MetasoundGenerator/Private/MetasoundGenerator.cpp
--- a/Engine/Plugins/Runtime/Metasound/Source/MetasoundGenerator/Private/MetasoundGenerator.cpp
+++ b/Engine/Plugins/Runtime/Metasound/Source/MetasoundGenerator/Private/MetasoundGenerator.cpp
@@ -10,12 +10,34 @@
 #include "MetasoundTrace.h"
 #include "MetasoundTrigger.h"
 
-#ifndef METASOUND_GENERATOR_BUILD_GRAPH_ASYNC
-#define METASOUND_GENERATOR_BUILD_GRAPH_ASYNC 0
-#endif
 
 namespace Metasound
 {
+	namespace ConsoleVariables
+	{
+		static bool bEnableAsyncMetaSoundGeneratorBuilder = true;
+	}
+}
+
+FAutoConsoleVariableRef CVarMetaSoundEnableAsyncGeneratorBuilder(
+	TEXT("au.MetaSound.EnableAsyncGeneratorBuilder"),
+	Metasound::ConsoleVariables::bEnableAsyncMetaSoundGeneratorBuilder,
+	TEXT("Enables async building of FMetaSoundGenerators\n")
+	TEXT("Default: true"),
+	ECVF_Default);
+
+
+namespace Metasound
+{
+	void FMetasoundGeneratorInitParams::Release()
+	{
+		Graph.Reset();
+		Environment = {};
+		MetaSoundName = {};
+		AudioOutputNames = {};
+		OnPlayInputName = {};
+		IsFinishedOutputName = {};
+	}
 
 	FAsyncMetaSoundBuilder::FAsyncMetaSoundBuilder(FMetasoundGenerator* InGenerator, FMetasoundGeneratorInitParams&& InInitParams, bool bInTriggerGenerator)
 		: Generator(InGenerator)
@@ -81,6 +103,8 @@ namespace Metasound
 			// Otherwise, generator will continually wait for a new generator.
 			Generator->SetPendingGraphBuildFailed();
 		}
+
+		InitParams.Release();
 	}
 
 	FMetasoundGenerator::FMetasoundGenerator(FMetasoundGeneratorInitParams&& InParams)
@@ -101,15 +125,19 @@ namespace Metasound
 
 		BuilderTask = MakeUnique<FBuilderTask>(this, MoveTemp(InParams), true /* bTriggerGenerator */);
 		
-#if METASOUND_GENERATOR_BUILD_GRAPH_ASYNC
-		BuilderTask->StartBackgroundTask(GBackgroundPriorityThreadPool);
-		
-#else
-		BuilderTask->StartSynchronousTask();
-		BuilderTask = nullptr;
-		UpdateGraphIfPending();
-		bIsWaitingForFirstGraph = false;
-#endif
+		if (Metasound::ConsoleVariables::bEnableAsyncMetaSoundGeneratorBuilder)
+		{
+			// Build operator asynchronously
+			BuilderTask->StartBackgroundTask(GBackgroundPriorityThreadPool);
+		}
+		else
+		{
+			// Build operator synchronously
+			BuilderTask->StartSynchronousTask();
+			BuilderTask = nullptr;
+			UpdateGraphIfPending();
+			bIsWaitingForFirstGraph = false;
+		}
 	}
 
 	FMetasoundGenerator::~FMetasoundGenerator()
diff --git a/Engine/Plugins/Runtime/Metasound/Source/MetasoundGenerator/Public/MetasoundGenerator.h b/Engine/Plugins/Runtime/Metasound/Source/MetasoundGenerator/Public/MetasoundGenerator.h
--- a/Engine/Plugins/Runtime/Metasound/Source/MetasoundGenerator/Public/MetasoundGenerator.h
+++ b/Engine/Plugins/Runtime/Metasound/Source/MetasoundGenerator/Public/MetasoundGenerator.h
@@ -16,7 +16,7 @@
 namespace Metasound
 {
 	// Struct needed for building the metasound graph
-	struct FMetasoundGeneratorInitParams
+	struct METASOUNDGENERATOR_API FMetasoundGeneratorInitParams
 	{
 		FOperatorSettings OperatorSettings;
 		TUniquePtr<IGraph> Graph;
@@ -25,6 +25,8 @@ namespace Metasound
 		TArray<FString> AudioOutputNames;
 		FString OnPlayInputName;
 		FString IsFinishedOutputName;
+
+		void Release();
 	};
 
 	struct FMetasoundGeneratorData
