Fix static analysis warning

diff --git a/Engine/Plugins/Runtime/Metasound/Source/MetasoundFrontend/Private/MetasoundFrontendDataTypeRegistry.cpp b/Engine/Plugins/Runtime/Metasound/Source/MetasoundFrontend/Private/MetasoundFrontendDataTypeRegistry.cpp
--- a/Engine/Plugins/Runtime/Metasound/Source/MetasoundFrontend/Private/MetasoundFrontendDataTypeRegistry.cpp
+++ b/Engine/Plugins/Runtime/Metasound/Source/MetasoundFrontend/Private/MetasoundFrontendDataTypeRegistry.cpp
@@ -584,12 +584,9 @@ namespace Metasound
 				if (const IDataTypeRegistryEntry* Entry = FindDataTypeEntry(InDataType))
 				{
 					ProxyPtr = Entry->CreateProxy(InObject);
-					if (ProxyPtr)
+					if (!ProxyPtr && InObject)
 					{
-						if (!InObject)
-						{
-							UE_LOG(LogMetaSound, Error, TEXT("Failed to create a valid proxy from UObject '%s'."), *InObject->GetName());
-						}
+						UE_LOG(LogMetaSound, Error, TEXT("Failed to create a valid proxy from UObject '%s'."), *InObject->GetName());
 					}
 				}
 
