Fix misreporting MetaSound validation errors when cooking. Cooking should only report validation errors if there are any remaining after registering/auto-updating.

diff --git a/Engine/Plugins/Runtime/Metasound/Source/MetasoundEditor/Public/MetasoundEditorGraphValidation.h b/Engine/Plugins/Runtime/Metasound/Source/MetasoundEditor/Public/MetasoundEditorGraphValidation.h
--- a/Engine/Plugins/Runtime/Metasound/Source/MetasoundEditor/Public/MetasoundEditorGraphValidation.h
+++ b/Engine/Plugins/Runtime/Metasound/Source/MetasoundEditor/Public/MetasoundEditorGraphValidation.h
@@ -2,7 +2,6 @@
 #pragma once
 
 #include "CoreMinimal.h"
-#include "MetasoundFrontendDocument.h"
 #include "MetasoundEditorGraph.h"
 #include "MetasoundEditorGraphNode.h"
 
diff --git a/Engine/Plugins/Runtime/Metasound/Source/MetasoundEngine/Private/Metasound.cpp b/Engine/Plugins/Runtime/Metasound/Source/MetasoundEngine/Private/Metasound.cpp
--- a/Engine/Plugins/Runtime/Metasound/Source/MetasoundEngine/Private/Metasound.cpp
+++ b/Engine/Plugins/Runtime/Metasound/Source/MetasoundEngine/Private/Metasound.cpp
@@ -66,7 +66,7 @@ void UMetaSound::BeginDestroy()
 void UMetaSound::PreSave(FObjectPreSaveContext InSaveContext)
 {
 	Super::PreSave(InSaveContext);
-	Metasound::PreSaveAsset(*this);
+	Metasound::PreSaveAsset(*this, InSaveContext);
 }
 
 void UMetaSound::Serialize(FArchive& InArchive)
diff --git a/Engine/Plugins/Runtime/Metasound/Source/MetasoundEngine/Private/MetasoundSource.cpp b/Engine/Plugins/Runtime/Metasound/Source/MetasoundEngine/Private/MetasoundSource.cpp
--- a/Engine/Plugins/Runtime/Metasound/Source/MetasoundEngine/Private/MetasoundSource.cpp
+++ b/Engine/Plugins/Runtime/Metasound/Source/MetasoundEngine/Private/MetasoundSource.cpp
@@ -183,7 +183,7 @@ void UMetaSoundSource::BeginDestroy()
 void UMetaSoundSource::PreSave(FObjectPreSaveContext InSaveContext)
 {
 	Super::PreSave(InSaveContext);
-	Metasound::PreSaveAsset(*this);
+	Metasound::PreSaveAsset(*this, InSaveContext);
 }
 
 void UMetaSoundSource::Serialize(FArchive& InArchive)
diff --git a/Engine/Plugins/Runtime/Metasound/Source/MetasoundEngine/Public/Metasound.h b/Engine/Plugins/Runtime/Metasound/Source/MetasoundEngine/Public/Metasound.h
--- a/Engine/Plugins/Runtime/Metasound/Source/MetasoundEngine/Public/Metasound.h
+++ b/Engine/Plugins/Runtime/Metasound/Source/MetasoundEngine/Public/Metasound.h
@@ -12,6 +12,7 @@
 #include "MetasoundRouter.h"
 #include "MetasoundUObjectRegistry.h"
 #include "Serialization/Archive.h"
+#include "UObject/ObjectSaveContext.h"
 
 #if WITH_EDITORONLY_DATA
 #include "Algo/Transform.h"
@@ -77,18 +78,43 @@ namespace Metasound
 #endif // WITH_EDITOR
 
 	template <typename TMetaSoundObject>
-	void PreSaveAsset(TMetaSoundObject& InMetaSound)
+	void PreSaveAsset(TMetaSoundObject& InMetaSound, FObjectPreSaveContext InSaveContext)
 	{
+		// If not editor and cooking, no need to force re-registering as if its already registered,
+		// no edits have been made between last registration and this one if a parent asset has requested
+		// it to register.
+		FMetaSoundAssetRegistrationOptions NonEditorCookRegistrationOptions;
+		NonEditorCookRegistrationOptions.bForceReregister = false;
+		NonEditorCookRegistrationOptions.bRegisterDependencies = true;
+
 #if WITH_EDITORONLY_DATA
 		if (UMetasoundEditorGraphBase* MetaSoundGraph = Cast<UMetasoundEditorGraphBase>(InMetaSound.GetGraph()))
 		{
-			const bool bClearUpdateNotes = true;
-			const bool bIsValid = MetaSoundGraph->Validate(bClearUpdateNotes);
-			if (bIsValid)
+			// When cooking, validate after registration to avoid reporting false errors
+			if (InSaveContext.IsCooking())
 			{
-				MetaSoundGraph->RegisterGraphWithFrontend();
+				// Non-editor builds use explicit options above as opposed to using
+				// editor variant via the MetaSoundEdGraph.
+				if (GIsEditor)
+				{
+					MetaSoundGraph->RegisterGraphWithFrontend();
+				}
+				else
+				{
+					InMetaSound.RegisterGraphWithFrontend(NonEditorCookRegistrationOptions);
+				}
+				MetaSoundGraph->Validate(false /* bInClearUpdateNotes */);
+			}
+			else
+			{
+				if (MetaSoundGraph->Validate(true /* bInClearUpdateNotes */))
+				{
+					MetaSoundGraph->RegisterGraphWithFrontend();
+				}
 			}
 		}
+#else
+		InMetaSound.RegisterGraphWithFrontend(NonEditorCookRegistrationOptions);
 #endif // WITH_EDITORONLY_DATA
 	}
 
diff --git a/Engine/Plugins/Runtime/Metasound/Source/MetasoundFrontend/Public/MetasoundAssetBase.h b/Engine/Plugins/Runtime/Metasound/Source/MetasoundFrontend/Public/MetasoundAssetBase.h
--- a/Engine/Plugins/Runtime/Metasound/Source/MetasoundFrontend/Public/MetasoundAssetBase.h
+++ b/Engine/Plugins/Runtime/Metasound/Source/MetasoundFrontend/Public/MetasoundAssetBase.h
@@ -46,7 +46,8 @@ struct METASOUNDFRONTEND_API FMetaSoundAssetRegistrationOptions
 	// If true, recursively attempts to register dependencies.
 	bool bRegisterDependencies = false;
 
-	// Attempt to auto-update (only runs if class not registered or set to force re-register)
+	// Attempt to auto-update (Only runs if class not registered or set to force re-register.
+	// Will not respect being set to true if project-level MetaSoundSettings specify to not run auto-update.)
 	bool bAutoUpdate = true;
 
 	// Attempt to rebuild referenced class keys (only run if class not registered or set to force re-register)
