Fix a use-after-destruct on stack found by address sanitizer.

#rnx

diff --git a/Engine/Plugins/Runtime/Metasound/Source/MetasoundGenerator/Private/MetasoundGenerator.cpp b/Engine/Plugins/Runtime/Metasound/Source/MetasoundGenerator/Private/MetasoundGenerator.cpp
--- a/Engine/Plugins/Runtime/Metasound/Source/MetasoundGenerator/Private/MetasoundGenerator.cpp
+++ b/Engine/Plugins/Runtime/Metasound/Source/MetasoundGenerator/Private/MetasoundGenerator.cpp
@@ -51,7 +51,8 @@ namespace Metasound
 		METASOUND_TRACE_CPUPROFILER_EVENT_SCOPE(AsyncMetaSoundBuilder::DoWork);
 		// Create an instance of the new graph
 		FOperatorBuilder OperatorBuilder(FOperatorBuilderSettings::GetDefaultSettings());
-		FBuildGraphParams BuildParams{*InitParams.Graph, InitParams.OperatorSettings, FDataReferenceCollection{}, InitParams.Environment};
+		FDataReferenceCollection DataReferenceCollection{};
+		FBuildGraphParams BuildParams{*InitParams.Graph, InitParams.OperatorSettings, DataReferenceCollection, InitParams.Environment};
 		TArray<IOperatorBuilder::FBuildErrorPtr> BuildErrors;
 
 		TUniquePtr<IOperator> GraphOperator = OperatorBuilder.BuildGraphOperator(BuildParams, BuildErrors);
