Implement MetasoundGraphCore Variable

diff --git a/Engine/Plugins/Runtime/Metasound/Source/MetasoundEditor/Private/MetasoundEditorGraphBuilder.cpp b/Engine/Plugins/Runtime/Metasound/Source/MetasoundEditor/Private/MetasoundEditorGraphBuilder.cpp
--- a/Engine/Plugins/Runtime/Metasound/Source/MetasoundEditor/Private/MetasoundEditorGraphBuilder.cpp
+++ b/Engine/Plugins/Runtime/Metasound/Source/MetasoundEditor/Private/MetasoundEditorGraphBuilder.cpp
@@ -1121,8 +1121,14 @@ namespace Metasound
 		{
 			using namespace Frontend;
 
+			FEdGraphPinType PinType;
+			FName DataTypeName = InInputHandle->GetDataType();
+
 			IMetasoundEditorModule& EditorModule = FModuleManager::GetModuleChecked<IMetasoundEditorModule>("MetaSoundEditor");
-			const FEdGraphPinType PinType = EditorModule.FindDataType(InInputHandle->GetDataType()).PinType;
+			if (EditorModule.IsRegisteredDataType(DataTypeName))
+			{
+				PinType = EditorModule.FindDataType(DataTypeName).PinType;
+			}
 
 			UEdGraphPin* NewPin = InEditorNode.CreatePin(EGPD_Input, PinType, FName(*InInputHandle->GetName()));
 			if (ensure(NewPin))
@@ -1136,8 +1142,14 @@ namespace Metasound
 
 		UEdGraphPin* FGraphBuilder::AddPinToNode(UMetasoundEditorGraphNode& InEditorNode, Frontend::FConstOutputHandle InOutputHandle)
 		{
+			FEdGraphPinType PinType;
+			FName DataTypeName = InOutputHandle->GetDataType();
+
 			IMetasoundEditorModule& EditorModule = FModuleManager::GetModuleChecked<IMetasoundEditorModule>("MetaSoundEditor");
-			FEdGraphPinType	PinType = EditorModule.FindDataType(InOutputHandle->GetDataType()).PinType;
+			if (EditorModule.IsRegisteredDataType(DataTypeName))
+			{
+				PinType = EditorModule.FindDataType(DataTypeName).PinType;
+			}
 
 			UEdGraphPin* NewPin = InEditorNode.CreatePin(EGPD_Output, PinType, FName(*InOutputHandle->GetName()));
 			if (ensure(NewPin))
diff --git a/Engine/Plugins/Runtime/Metasound/Source/MetasoundFrontend/Private/MetasoundAssetBase.cpp b/Engine/Plugins/Runtime/Metasound/Source/MetasoundFrontend/Private/MetasoundAssetBase.cpp
--- a/Engine/Plugins/Runtime/Metasound/Source/MetasoundFrontend/Private/MetasoundAssetBase.cpp
+++ b/Engine/Plugins/Runtime/Metasound/Source/MetasoundFrontend/Private/MetasoundAssetBase.cpp
@@ -135,9 +135,9 @@ void FMetasoundAssetBase::RegisterGraphWithFrontend(FMetaSoundAssetRegistrationO
 			return ClassInfo;
 		}
 
-		virtual TUniquePtr<INode> CreateNode(FDefaultNodeConstructorParams&&) const override
+		virtual TUniquePtr<INode> CreateNode(const FNodeInitData&) const override
 		{
-			return nullptr;
+			return FFrontendGraphBuilder().CreateGraph(Document);
 		}
 
 		virtual TUniquePtr<INode> CreateNode(FDefaultLiteralNodeConstructorParams&&) const override
@@ -145,11 +145,17 @@ void FMetasoundAssetBase::RegisterGraphWithFrontend(FMetaSoundAssetRegistrationO
 			return nullptr;
 		}
 
-		virtual TUniquePtr<INode> CreateNode(const FNodeInitData&) const override
+		virtual TUniquePtr<INode> CreateNode(FDefaultNamedVertexNodeConstructorParams&&) const override
 		{
-			return FFrontendGraphBuilder().CreateGraph(Document);
+			return nullptr;
 		}
 
+		virtual TUniquePtr<INode> CreateNode(FDefaultNamedVertexWithLiteralNodeConstructorParams&&) const override
+		{
+			return nullptr;
+		}
+
+
 		virtual const FMetasoundFrontendClass& GetFrontendClass() const override
 		{
 			return FrontendClass;
diff --git a/Engine/Plugins/Runtime/Metasound/Source/MetasoundFrontend/Private/MetasoundFrontendDataTypeRegistry.cpp b/Engine/Plugins/Runtime/Metasound/Source/MetasoundFrontend/Private/MetasoundFrontendDataTypeRegistry.cpp
--- a/Engine/Plugins/Runtime/Metasound/Source/MetasoundFrontend/Private/MetasoundFrontendDataTypeRegistry.cpp
+++ b/Engine/Plugins/Runtime/Metasound/Source/MetasoundFrontend/Private/MetasoundFrontendDataTypeRegistry.cpp
@@ -107,44 +107,35 @@
 				}
 			}
 
-			// Node registry entry for input nodes created from a data type registry entry.
-			class FInputNodeRegistryEntry : public INodeRegistryEntry
+			// Base class for INodeRegistryEntrys that come from an IDataTypeRegistryEntry
+			class FDataTypeNodeRegistryEntry : public INodeRegistryEntry
 			{
 			public:
-				FInputNodeRegistryEntry() = delete;
+				FDataTypeNodeRegistryEntry() = default;
 
-				FInputNodeRegistryEntry(TUniquePtr<IDataTypeRegistryEntry>&& InDataTypeEntry)
-				: DataTypeEntry(MoveTemp(InDataTypeEntry))
-				{
-					if (DataTypeEntry.IsValid())
-					{
-						FrontendClass = DataTypeEntry->GetFrontendInputClass();
-						ClassInfo = FNodeClassInfo(FrontendClass.Metadata);
-					}
-				}
+				virtual ~FDataTypeNodeRegistryEntry() = default;
 
-				virtual ~FInputNodeRegistryEntry() = default;
-
 				virtual const FNodeClassInfo& GetClassInfo() const override
 				{
 					return ClassInfo;
 				}
 
-				virtual TUniquePtr<INode> CreateNode(FDefaultNodeConstructorParams&&) const override
-				{
-					return nullptr;
-				}
-
-				virtual TUniquePtr<INode> CreateNode(FDefaultLiteralNodeConstructorParams&& InParams) const override
-				{
-					if (DataTypeEntry.IsValid())
-					{
-						return DataTypeEntry->CreateInputNode(MoveTemp(InParams));
-					}
-					return nullptr;
-				}
-
 				virtual TUniquePtr<INode> CreateNode(const FNodeInitData&) const override
 				{
 					return nullptr;
 				}
+
+				virtual TUniquePtr<INode> CreateNode(FDefaultLiteralNodeConstructorParams&& InParams) const override
+				{
+					return nullptr;
+				}
+
+				virtual TUniquePtr<INode> CreateNode(FDefaultNamedVertexNodeConstructorParams&&) const override
+				{
+					return nullptr;
+				}
+
+				virtual TUniquePtr<INode> CreateNode(FDefaultNamedVertexWithLiteralNodeConstructorParams&& InParams) const override
+				{
+					return nullptr;
+				}
@@ -154,29 +145,67 @@
 					return FrontendClass;
 				}
 
-				virtual TUniquePtr<INodeRegistryEntry> Clone() const override
-				{
-					if (DataTypeEntry.IsValid())
-					{
-						return MakeUnique<FInputNodeRegistryEntry>(DataTypeEntry->Clone());
-					}
-					return MakeUnique<FInputNodeRegistryEntry>(TUniquePtr<IDataTypeRegistryEntry>());
-				}
-
 				virtual bool IsNative() const override
 				{
 					return true;
 				}
 
+			protected:
+
+				void SetFrontendClass(const FMetasoundFrontendClass& InClass)
+				{
+					FrontendClass = InClass;
+					ClassInfo = FNodeClassInfo(FrontendClass.Metadata);
+				}
+
 			private:
 				
-				TUniquePtr<IDataTypeRegistryEntry> DataTypeEntry;
 				FNodeClassInfo ClassInfo;
 				FMetasoundFrontendClass FrontendClass;
 			};
 
+			// Node registry entry for input nodes created from a data type registry entry.
+			class FInputNodeRegistryEntry : public FDataTypeNodeRegistryEntry
+			{
+			public:
+				FInputNodeRegistryEntry() = delete;
+
+				FInputNodeRegistryEntry(TUniquePtr<IDataTypeRegistryEntry>&& InDataTypeEntry)
+				: DataTypeEntry(MoveTemp(InDataTypeEntry))
+				{
+					if (DataTypeEntry.IsValid())
+					{
+						SetFrontendClass(DataTypeEntry->GetFrontendInputClass());
+					}
+				}
+
+				virtual ~FInputNodeRegistryEntry() = default;
+
+				virtual TUniquePtr<INode> CreateNode(FDefaultNamedVertexWithLiteralNodeConstructorParams&& InParams) const override
+				{
+					if (DataTypeEntry.IsValid())
+					{
+						return DataTypeEntry->CreateInputNode(MoveTemp(InParams));
+					}
+					return nullptr;
+				}
+
+				virtual TUniquePtr<INodeRegistryEntry> Clone() const override
+				{
+					if (DataTypeEntry.IsValid())
+					{
+						return MakeUnique<FInputNodeRegistryEntry>(DataTypeEntry->Clone());
+					}
+					return MakeUnique<FInputNodeRegistryEntry>(TUniquePtr<IDataTypeRegistryEntry>());
+				}
+
+			private:
+				
+				TUniquePtr<IDataTypeRegistryEntry> DataTypeEntry;
+			};
+
 			// Node registry entry for output nodes created from a data type registry entry.
-			class FOutputNodeRegistryEntry : public INodeRegistryEntry
+			class FOutputNodeRegistryEntry : public FDataTypeNodeRegistryEntry
 			{
 			public:
 				FOutputNodeRegistryEntry() = delete;
@@ -186,19 +215,13 @@
 				{
 					if (DataTypeEntry.IsValid())
 					{
-						FrontendClass = DataTypeEntry->GetFrontendOutputClass();
-						ClassInfo = FNodeClassInfo(FrontendClass.Metadata);
+						SetFrontendClass(DataTypeEntry->GetFrontendOutputClass());
 					}
 				}
 
 				virtual ~FOutputNodeRegistryEntry() = default;
 
-				virtual const FNodeClassInfo& GetClassInfo() const override
-				{
-					return ClassInfo;
-				}
-
-				virtual TUniquePtr<INode> CreateNode(FDefaultNodeConstructorParams&& InParams) const override
+				virtual TUniquePtr<INode> CreateNode(FDefaultNamedVertexNodeConstructorParams&& InParams) const override
 				{
 					if (DataTypeEntry.IsValid())
 					{
@@ -207,21 +230,6 @@
 					return nullptr;
 				}
 
-				virtual TUniquePtr<INode> CreateNode(FDefaultLiteralNodeConstructorParams&& InParams) const override
-				{
-					return nullptr;
-				}
-
-				virtual TUniquePtr<INode> CreateNode(const FNodeInitData&) const override
-				{
-					return nullptr;
-				}
-
-				virtual const FMetasoundFrontendClass& GetFrontendClass() const override
-				{
-					return FrontendClass;
-				}
-
 				virtual TUniquePtr<INodeRegistryEntry> Clone() const override
 				{
 					if (DataTypeEntry.IsValid())
@@ -231,84 +239,210 @@
 					return MakeUnique<FOutputNodeRegistryEntry>(TUniquePtr<IDataTypeRegistryEntry>());
 				}
 
-				virtual bool IsNative() const override
-				{
-					return true;
-				}
-
 			private:
 				
 				TUniquePtr<IDataTypeRegistryEntry> DataTypeEntry;
-				FNodeClassInfo ClassInfo;
-				FMetasoundFrontendClass FrontendClass;
 			};
 
-			// Node registry entry for variable nodes created from a data type registry entry.
-			class FVariableNodeRegistryEntry : public INodeRegistryEntry
+			// Node registry entry for literal nodes created from a data type registry entry.
+			class FLiteralNodeRegistryEntry : public FDataTypeNodeRegistryEntry
 			{
 			public:
-				FVariableNodeRegistryEntry() = delete;
+				FLiteralNodeRegistryEntry() = delete;
 
-				FVariableNodeRegistryEntry(TUniquePtr<IDataTypeRegistryEntry>&& InDataTypeEntry)
+				FLiteralNodeRegistryEntry(TUniquePtr<IDataTypeRegistryEntry>&& InDataTypeEntry)
 				: DataTypeEntry(MoveTemp(InDataTypeEntry))
 				{
 					if (DataTypeEntry.IsValid())
 					{
-						FrontendClass = DataTypeEntry->GetFrontendVariableClass();
-						ClassInfo = FNodeClassInfo(FrontendClass.Metadata);
+						SetFrontendClass(DataTypeEntry->GetFrontendLiteralClass());
 					}
 				}
 
-				virtual ~FVariableNodeRegistryEntry() = default;
-
-				virtual const FNodeClassInfo& GetClassInfo() const override
-				{
-					return ClassInfo;
-				}
-
-				virtual TUniquePtr<INode> CreateNode(FDefaultNodeConstructorParams&& InParams) const override
-				{
-					return nullptr;
-				}
+				virtual ~FLiteralNodeRegistryEntry() = default;
 
 				virtual TUniquePtr<INode> CreateNode(FDefaultLiteralNodeConstructorParams&& InParams) const override
 				{
 					if (DataTypeEntry.IsValid())
 					{
-						return DataTypeEntry->CreateVariableNode(MoveTemp(InParams));
+						return DataTypeEntry->CreateLiteralNode(MoveTemp(InParams));
 					}
 					return nullptr;
-				}
-
-				virtual TUniquePtr<INode> CreateNode(const FNodeInitData&) const override
-				{
-					return nullptr;
-				}
-
-				virtual const FMetasoundFrontendClass& GetFrontendClass() const override
-				{
-					return FrontendClass;
 				}
 
 				virtual TUniquePtr<INodeRegistryEntry> Clone() const override
 				{
 					if (DataTypeEntry.IsValid())
 					{
-						return MakeUnique<FVariableNodeRegistryEntry>(DataTypeEntry->Clone());
+						return MakeUnique<FLiteralNodeRegistryEntry>(DataTypeEntry->Clone());
 					}
-					return MakeUnique<FVariableNodeRegistryEntry>(TUniquePtr<IDataTypeRegistryEntry>());
-				}
-
-				virtual bool IsNative() const override
-				{
-					return true;
+					return MakeUnique<FLiteralNodeRegistryEntry>(TUniquePtr<IDataTypeRegistryEntry>());
 				}
 
 			private:
 				
 				TUniquePtr<IDataTypeRegistryEntry> DataTypeEntry;
-				FNodeClassInfo ClassInfo;
-				FMetasoundFrontendClass FrontendClass;
 			};
 
+
+			// Node registry entry for init variable nodes created from a data type registry entry.
+			class FInitVariableNodeRegistryEntry : public FDataTypeNodeRegistryEntry
+			{
+			public:
+				FInitVariableNodeRegistryEntry() = delete;
+
+				FInitVariableNodeRegistryEntry(TUniquePtr<IDataTypeRegistryEntry>&& InDataTypeEntry)
+				: DataTypeEntry(MoveTemp(InDataTypeEntry))
+				{
+					if (DataTypeEntry.IsValid())
+					{
+						SetFrontendClass(DataTypeEntry->GetFrontendInitVariableClass());
+					}
+				}
+
+				virtual ~FInitVariableNodeRegistryEntry() = default;
+
+				virtual TUniquePtr<INode> CreateNode(FDefaultLiteralNodeConstructorParams&& InParams) const override
+				{
+					if (DataTypeEntry.IsValid())
+					{
+						return DataTypeEntry->CreateInitVariableNode(MoveTemp(InParams));
+					}
+					return nullptr;
+				}
+
+				virtual TUniquePtr<INodeRegistryEntry> Clone() const override
+				{
+					if (DataTypeEntry.IsValid())
+					{
+						return MakeUnique<FInitVariableNodeRegistryEntry>(DataTypeEntry->Clone());
+					}
+					return MakeUnique<FInitVariableNodeRegistryEntry>(TUniquePtr<IDataTypeRegistryEntry>());
+				}
+
+			private:
+				
+				TUniquePtr<IDataTypeRegistryEntry> DataTypeEntry;
+			};
+
+			// Node registry entry for set variable nodes created from a data type registry entry.
+			class FSetVariableNodeRegistryEntry : public FDataTypeNodeRegistryEntry
+			{
+			public:
+				FSetVariableNodeRegistryEntry() = delete;
+
+				FSetVariableNodeRegistryEntry(TUniquePtr<IDataTypeRegistryEntry>&& InDataTypeEntry)
+				: DataTypeEntry(MoveTemp(InDataTypeEntry))
+				{
+					if (DataTypeEntry.IsValid())
+					{
+						SetFrontendClass(DataTypeEntry->GetFrontendSetVariableClass());
+					}
+				}
+
+				virtual ~FSetVariableNodeRegistryEntry() = default;
+
+				virtual TUniquePtr<INode> CreateNode(const FNodeInitData& InParams) const override
+				{
+					if (DataTypeEntry.IsValid())
+					{
+						return DataTypeEntry->CreateSetVariableNode(InParams);
+					}
+					return nullptr;
+				}
+
+				virtual TUniquePtr<INodeRegistryEntry> Clone() const override
+				{
+					if (DataTypeEntry.IsValid())
+					{
+						return MakeUnique<FSetVariableNodeRegistryEntry>(DataTypeEntry->Clone());
+					}
+					return MakeUnique<FSetVariableNodeRegistryEntry>(TUniquePtr<IDataTypeRegistryEntry>());
+				}
+
+			private:
+				
+				TUniquePtr<IDataTypeRegistryEntry> DataTypeEntry;
+			};
+
+			// Node registry entry for get variable nodes created from a data type registry entry.
+			class FGetVariableNodeRegistryEntry : public FDataTypeNodeRegistryEntry
+			{
+			public:
+				FGetVariableNodeRegistryEntry() = delete;
+
+				FGetVariableNodeRegistryEntry(TUniquePtr<IDataTypeRegistryEntry>&& InDataTypeEntry)
+				: DataTypeEntry(MoveTemp(InDataTypeEntry))
+				{
+					if (DataTypeEntry.IsValid())
+					{
+						SetFrontendClass(DataTypeEntry->GetFrontendGetVariableClass());
+					}
+				}
+
+				virtual ~FGetVariableNodeRegistryEntry() = default;
+
+				virtual TUniquePtr<INode> CreateNode(const FNodeInitData& InParams) const override
+				{
+					if (DataTypeEntry.IsValid())
+					{
+						return DataTypeEntry->CreateGetVariableNode(InParams);
+					}
+					return nullptr;
+				}
+
+				virtual TUniquePtr<INodeRegistryEntry> Clone() const override
+				{
+					if (DataTypeEntry.IsValid())
+					{
+						return MakeUnique<FGetVariableNodeRegistryEntry>(DataTypeEntry->Clone());
+					}
+					return MakeUnique<FGetVariableNodeRegistryEntry>(TUniquePtr<IDataTypeRegistryEntry>());
+				}
+
+			private:
+				
+				TUniquePtr<IDataTypeRegistryEntry> DataTypeEntry;
+			};
+
+			// Node registry entry for get delayed variable nodes created from a data type registry entry.
+			class FGetDelayedVariableNodeRegistryEntry : public FDataTypeNodeRegistryEntry
+			{
+			public:
+				FGetDelayedVariableNodeRegistryEntry() = delete;
+
+				FGetDelayedVariableNodeRegistryEntry(TUniquePtr<IDataTypeRegistryEntry>&& InDataTypeEntry)
+				: DataTypeEntry(MoveTemp(InDataTypeEntry))
+				{
+					if (DataTypeEntry.IsValid())
+					{
+						SetFrontendClass(DataTypeEntry->GetFrontendGetDelayedVariableClass());
+					}
+				}
+
+				virtual ~FGetDelayedVariableNodeRegistryEntry() = default;
+
+				virtual TUniquePtr<INode> CreateNode(const FNodeInitData& InParams) const override
+				{
+					if (DataTypeEntry.IsValid())
+					{
+						return DataTypeEntry->CreateGetDelayedVariableNode(InParams);
+					}
+					return nullptr;
+				}
+
+				virtual TUniquePtr<INodeRegistryEntry> Clone() const override
+				{
+					if (DataTypeEntry.IsValid())
+					{
+						return MakeUnique<FGetDelayedVariableNodeRegistryEntry>(DataTypeEntry->Clone());
+					}
+					return MakeUnique<FGetDelayedVariableNodeRegistryEntry>(TUniquePtr<IDataTypeRegistryEntry>());
+				}
+
+			private:
+				
+				TUniquePtr<IDataTypeRegistryEntry> DataTypeEntry;
+			};
+
 			class FDataTypeRegistry : public IDataTypeRegistry
@@ -350,16 +484,22 @@
 				virtual TSharedPtr<IDataChannel, ESPMode::ThreadSafe> CreateDataChannel(const FName& InDataType, const FOperatorSettings& InOperatorSettings) const override;
 
 				virtual bool GetFrontendInputClass(const FName& InDataType, FMetasoundFrontendClass& OutClass) const override;
-
-				virtual bool GetFrontendVariableClass(const FName& InDataType, FMetasoundFrontendClass& OutClass) const override;
-
+				virtual bool GetFrontendLiteralClass(const FName& InDataType, FMetasoundFrontendClass& OutClass) const override;
 				virtual bool GetFrontendOutputClass(const FName& InDataType, FMetasoundFrontendClass& OutClass) const override;
+				virtual bool GetFrontendInitVariableClass(const FName& InDataType, FMetasoundFrontendClass& OutClass) const override;
+				virtual bool GetFrontendSetVariableClass(const FName& InDataType, FMetasoundFrontendClass& OutClass) const override;
+				virtual bool GetFrontendGetVariableClass(const FName& InDataType, FMetasoundFrontendClass& OutClass) const override;
+				virtual bool GetFrontendGetDelayedVariableClass(const FName& InDataType, FMetasoundFrontendClass& OutClass) const override;
 
 				// Create a new instance of a C++ implemented node from the registry.
 				virtual TUniquePtr<INode> CreateInputNode(const FName& InInputType, FInputNodeConstructorParams&& InParams) const override;
-				virtual TUniquePtr<INode> CreateVariableNode(const FName& InVariableType, FVariableNodeConstructorParams&& InParams) const override;
+				virtual TUniquePtr<INode> CreateLiteralNode(const FName& InLiteralType, FLiteralNodeConstructorParams&& InParams) const override;
 				virtual TUniquePtr<INode> CreateOutputNode(const FName& InOutputType, FOutputNodeConstructorParams&& InParams) const override;
-				virtual TUniquePtr<INode> CreateReceiveNode(const FName& InDataType, const FNodeInitData& InParams) const override;
+				virtual TUniquePtr<INode> CreateReceiveNode(const FName& InOutputType, const FNodeInitData& InParams) const override;
+				virtual TUniquePtr<INode> CreateInitVariableNode(const FName& InDataType, FInitVariableNodeConstructorParams&&) const override;
+				virtual TUniquePtr<INode> CreateSetVariableNode(const FName& InDataType, const FNodeInitData&) const override;
+				virtual TUniquePtr<INode> CreateGetVariableNode(const FName& InDataType, const FNodeInitData&) const override;
+				virtual TUniquePtr<INode> CreateGetDelayedVariableNode(const FName& InDataType, const FNodeInitData&) const override;
 
 			private:
 
@@ -391,7 +531,11 @@
 					{
 						NodeRegistry->RegisterNode(MakeUnique<FInputNodeRegistryEntry>(InEntry->Clone()));
 						NodeRegistry->RegisterNode(MakeUnique<FOutputNodeRegistryEntry>(InEntry->Clone()));
-						NodeRegistry->RegisterNode(MakeUnique<FVariableNodeRegistryEntry>(InEntry->Clone()));
+						NodeRegistry->RegisterNode(MakeUnique<FLiteralNodeRegistryEntry>(InEntry->Clone()));
+						//NodeRegistry->RegisterNode(MakeUnique<FInitVariableNodeRegistryEntry>(InEntry->Clone()));
+						//NodeRegistry->RegisterNode(MakeUnique<FSetVariableNodeRegistryEntry>(InEntry->Clone()));
+						//NodeRegistry->RegisterNode(MakeUnique<FGetVariableNodeRegistryEntry>(InEntry->Clone()));
+						//NodeRegistry->RegisterNode(MakeUnique<FGetDelayedVariableNodeRegistryEntry>(InEntry->Clone()));
 					}
 
 					const FDataTypeRegistryInfo& RegistryInfo = InEntry->GetDataTypeInfo();
@@ -564,7 +708,19 @@
 
 			FLiteral FDataTypeRegistry::CreateDefaultLiteral(const FName& InDataType) const
 			{
-				return FLiteral::GetDefaultForType(GetDesiredLiteralType(InDataType));
+				if (const IDataTypeRegistryEntry* Entry = FindDataTypeEntry(InDataType))
+				{
+					const FDataTypeRegistryInfo& Info = Entry->GetDataTypeInfo();
+					if (Info.bIsEnum)
+					{
+						if (TSharedPtr<const IEnumDataTypeInterface> EnumInterface = Entry->GetEnumInterface())
+						{
+							return FLiteral(EnumInterface->GetDefaultValue());
+						}
+					}
+					return FLiteral::GetDefaultForType(Info.PreferredLiteralType);
+				}
+				return FLiteral::CreateInvalid();
 			}
 
 			bool FDataTypeRegistry::IsUObjectProxyFactory(UObject* InObject) const
@@ -630,11 +786,11 @@
 				return false;
 			}
 
-			bool FDataTypeRegistry::GetFrontendVariableClass(const FName& InDataType, FMetasoundFrontendClass& OutClass) const
+			bool FDataTypeRegistry::GetFrontendLiteralClass(const FName& InDataType, FMetasoundFrontendClass& OutClass) const
 			{
 				if (const IDataTypeRegistryEntry* Entry = FindDataTypeEntry(InDataType))
 				{
-					OutClass = Entry->GetFrontendVariableClass();
+					OutClass = Entry->GetFrontendLiteralClass();
 					return true;
 				}
 				return false;
@@ -650,6 +806,49 @@
 				return false;
 			}
 
+			bool FDataTypeRegistry::GetFrontendInitVariableClass(const FName& InDataType, FMetasoundFrontendClass& OutClass) const
+			{
+				if (const IDataTypeRegistryEntry* Entry = FindDataTypeEntry(InDataType))
+				{
+					OutClass = Entry->GetFrontendInitVariableClass();
+					return true;
+				}
+				return false;
+			}
+
+			bool FDataTypeRegistry::GetFrontendSetVariableClass(const FName& InDataType, FMetasoundFrontendClass& OutClass) const
+
+			{
+				if (const IDataTypeRegistryEntry* Entry = FindDataTypeEntry(InDataType))
+				{
+					OutClass = Entry->GetFrontendSetVariableClass();
+					return true;
+				}
+				return false;
+			}
+
+			bool FDataTypeRegistry::GetFrontendGetVariableClass(const FName& InDataType, FMetasoundFrontendClass& OutClass) const
+
+			{
+				if (const IDataTypeRegistryEntry* Entry = FindDataTypeEntry(InDataType))
+				{
+					OutClass = Entry->GetFrontendGetVariableClass();
+					return true;
+				}
+				return false;
+			}
+
+			bool FDataTypeRegistry::GetFrontendGetDelayedVariableClass(const FName& InDataType, FMetasoundFrontendClass& OutClass) const
+
+			{
+				if (const IDataTypeRegistryEntry* Entry = FindDataTypeEntry(InDataType))
+				{
+					OutClass = Entry->GetFrontendGetDelayedVariableClass();
+					return true;
+				}
+				return false;
+			}
+
 			TUniquePtr<INode> FDataTypeRegistry::CreateInputNode(const FName& InDataType, FInputNodeConstructorParams&& InParams) const
 			{
 				if (const IDataTypeRegistryEntry* Entry = FindDataTypeEntry(InDataType))
@@ -659,11 +858,11 @@
 				return nullptr;
 			}
 
-			TUniquePtr<INode> FDataTypeRegistry::CreateVariableNode(const FName& InDataType, FVariableNodeConstructorParams&& InParams) const
+			TUniquePtr<INode> FDataTypeRegistry::CreateLiteralNode(const FName& InDataType, FLiteralNodeConstructorParams&& InParams) const 
 			{
 				if (const IDataTypeRegistryEntry* Entry = FindDataTypeEntry(InDataType))
 				{
-					return Entry->CreateVariableNode(MoveTemp(InParams));
+					return Entry->CreateLiteralNode(MoveTemp(InParams));
 				}
 				return nullptr;
 			}
@@ -686,6 +885,42 @@
 				return nullptr;
 			}
 
+			TUniquePtr<INode> FDataTypeRegistry::CreateInitVariableNode(const FName& InDataType, FInitVariableNodeConstructorParams&& InParams) const
+			{
+				if (const IDataTypeRegistryEntry* Entry = FindDataTypeEntry(InDataType))
+				{
+					return Entry->CreateInitVariableNode(MoveTemp(InParams));
+				}
+				return nullptr;
+			}
+
+			TUniquePtr<INode> FDataTypeRegistry::CreateSetVariableNode(const FName& InDataType, const FNodeInitData& InParams) const
+			{
+				if (const IDataTypeRegistryEntry* Entry = FindDataTypeEntry(InDataType))
+				{
+					return Entry->CreateSetVariableNode(InParams);
+				}
+				return nullptr;
+			}
+
+			TUniquePtr<INode> FDataTypeRegistry::CreateGetVariableNode(const FName& InDataType, const FNodeInitData& InParams) const
+			{
+				if (const IDataTypeRegistryEntry* Entry = FindDataTypeEntry(InDataType))
+				{
+					return Entry->CreateGetVariableNode(InParams);
+				}
+				return nullptr;
+			}
+
+			TUniquePtr<INode> FDataTypeRegistry::CreateGetDelayedVariableNode(const FName& InDataType, const FNodeInitData& InParams) const
+			{
+				if (const IDataTypeRegistryEntry* Entry = FindDataTypeEntry(InDataType))
+				{
+					return Entry->CreateGetDelayedVariableNode(InParams);
+				}
+				return nullptr;
+			}
+
 			const IDataTypeRegistryEntry* FDataTypeRegistry::FindDataTypeEntry(const FName& InDataTypeName) const
 			{
 				const TUniquePtr<IDataTypeRegistryEntry>* Entry = RegisteredDataTypes.Find(InDataTypeName);
diff --git a/Engine/Plugins/Runtime/Metasound/Source/MetasoundFrontend/Private/MetasoundFrontendGraph.cpp b/Engine/Plugins/Runtime/Metasound/Source/MetasoundFrontend/Private/MetasoundFrontendGraph.cpp
--- a/Engine/Plugins/Runtime/Metasound/Source/MetasoundFrontend/Private/MetasoundFrontendGraph.cpp
+++ b/Engine/Plugins/Runtime/Metasound/Source/MetasoundFrontend/Private/MetasoundFrontendGraph.cpp
@@ -10,6 +10,7 @@
 #include "MetasoundFrontend.h"
 #include "MetasoundFrontendDataTypeRegistry.h"
 #include "MetasoundGraph.h"
+#include "MetasoundLiteralNode.h"
 #include "MetasoundLog.h"
 #include "MetasoundNodeInterface.h"
 
@@ -222,8 +223,8 @@ namespace Metasound
 
 			{
 				const FNodeInitData InitData = FrontendGraphPrivate::CreateNodeInitData(InNode);
-				TArray<FDefaultVariableData> DefaultVariableData = GetInputDefaultVariableData(InNode, InitData);
-				for (FDefaultVariableData& Data : DefaultVariableData)
+				TArray<FDefaultLiteralData> DefaultLiteralData = GetInputDefaultLiteralData(InNode, InitData);
+				for (FDefaultLiteralData& Data : DefaultLiteralData)
 				{
 					InGraphContext.DefaultInputs.Emplace(FNodeIDVertexID { InNode.ID, Data.DestinationVertexID }, MoveTemp(Data));
 				}
@@ -241,8 +242,8 @@ namespace Metasound
 
 		const FNodeInitData InitData = FrontendGraphPrivate::CreateNodeInitData(InNode);
 		{
-			TArray<FDefaultVariableData> DefaultVariableData = GetInputDefaultVariableData(InNode, InitData);
-			for (FDefaultVariableData& Data : DefaultVariableData)
+			TArray<FDefaultLiteralData> DefaultLiteralData = GetInputDefaultLiteralData(InNode, InitData);
+			for (FDefaultLiteralData& Data : DefaultLiteralData)
 			{
 				InGraphContext.DefaultInputs.Emplace(FNodeIDVertexID { InNode.ID, Data.DestinationVertexID }, MoveTemp(Data));
 			}
@@ -524,7 +525,7 @@ namespace Metasound
 		}
 	}
 
-	void FFrontendGraphBuilder::AddDefaultInputVariables(FBuildGraphContext& InGraphContext)
+	void FFrontendGraphBuilder::AddDefaultInputLiterals(FBuildGraphContext& InGraphContext)
 	{
 		using namespace Metasound::Frontend;
 
@@ -532,29 +533,29 @@ namespace Metasound
 
 		for (FDefaultInputByIDMap::ElementType& Pair : InGraphContext.DefaultInputs)
 		{
-			FDefaultVariableData& VariableData = Pair.Value;
-			const FGuid VariableNodeID = FGuid::NewGuid();
+			FDefaultLiteralData& LiteralData = Pair.Value;
+			const FGuid LiteralNodeID = FGuid::NewGuid();
 
 			// 1. Construct and add the default variable to the graph
 			{
-				TUniquePtr<const INode> DefaultVariable = IDataTypeRegistry::Get().CreateVariableNode(VariableData.TypeName, MoveTemp(VariableData.InitParams));
-				InGraphContext.Graph->AddNode(VariableNodeID, TSharedPtr<const INode>(DefaultVariable.Release()));
+				TUniquePtr<const INode> DefaultLiteral = IDataTypeRegistry::Get().CreateLiteralNode(LiteralData.TypeName, MoveTemp(LiteralData.InitParams));
+				InGraphContext.Graph->AddNode(LiteralNodeID, TSharedPtr<const INode>(DefaultLiteral.Release()));
 			}
 
 			// 2. Connect the default variable to the expected input
-			const INode* FromNode = InGraphContext.Graph->FindNode(VariableNodeID);
+			const INode* FromNode = InGraphContext.Graph->FindNode(LiteralNodeID);
 			if (!ensure(FromNode))
 			{
 				continue;
 			}
-			const FVertexKey& FromVertexKey = VariableData.InitParams.VertexName;
+			const FVertexKey& FromVertexKey = LiteralNodeNames::GetOutputDataName();
 
-			const INode* ToNode = InGraphContext.Graph->FindNode(VariableData.DestinationNodeID);
+			const INode* ToNode = InGraphContext.Graph->FindNode(LiteralData.DestinationNodeID);
 			if (!ensure(ToNode))
 			{
 				continue;
 			}
-			const FVertexKey& ToVertexKey = VariableData.DestinationVertexKey;
+			const FVertexKey& ToVertexKey = LiteralData.DestinationVertexKey;
 
 			bool bSuccess = InGraphContext.Graph->AddDataEdge(*FromNode, FromVertexKey, *ToNode, ToVertexKey);
 			if (!bSuccess)
@@ -567,14 +568,14 @@ namespace Metasound
 		InGraphContext.DefaultInputs.Reset();
 	}
 
-	TArray<FFrontendGraphBuilder::FDefaultVariableData> FFrontendGraphBuilder::GetInputDefaultVariableData(const FMetasoundFrontendNode& InNode, const FNodeInitData& InInitData)
+	TArray<FFrontendGraphBuilder::FDefaultLiteralData> FFrontendGraphBuilder::GetInputDefaultLiteralData(const FMetasoundFrontendNode& InNode, const FNodeInitData& InInitData)
 	{
-		TArray<FDefaultVariableData> DefaultVariableData;
+		TArray<FDefaultLiteralData> DefaultLiteralData;
 
 		TArray<FMetasoundFrontendVertex> InputVertices = InNode.Interface.Inputs;
 		for (const FMetasoundFrontendVertexLiteral& Literal : InNode.InputLiterals)
 		{
-			FVariableNodeConstructorParams InitParams;
+			FLiteralNodeConstructorParams InitParams;
 			FName TypeName;
 
 			FGuid VertexID = Literal.VertexID;
@@ -583,10 +584,9 @@ namespace Metasound
 			{
 				if (Vertex.VertexID == VertexID)
 				{
-					InitParams.InitParam = Literal.Value.ToLiteral(Vertex.TypeName);
+					InitParams.Literal = Literal.Value.ToLiteral(Vertex.TypeName);
 					InitParams.InstanceID = FGuid::NewGuid();
 					InitParams.NodeName = InInitData.InstanceName + Vertex.Name + InitParams.InstanceID.ToString();
-					InitParams.VertexName = InitParams.NodeName;
 					TypeName = Vertex.TypeName;
 
 					DestinationVertexKey = Vertex.Name;
@@ -600,7 +600,7 @@ namespace Metasound
 			const bool bRemoved = InputVertices.RemoveAllSwap(RemoveAndBuildParams, false /* bAllowShrinking */) > 0;
 			if (ensure(bRemoved))
 			{
-				DefaultVariableData.Emplace(FDefaultVariableData
+				DefaultLiteralData.Emplace(FDefaultLiteralData
 				{
 					InNode.ID,
 					VertexID,
@@ -611,7 +611,7 @@ namespace Metasound
 			}
 		}
 
-		return DefaultVariableData;
+		return DefaultLiteralData;
 	}
 
 	/** Check that all dependencies are C++ class dependencies. */
@@ -719,7 +719,7 @@ namespace Metasound
 		// a datatype or node is not registered.
 		AddNodesToGraph(BuildGraphContext);
 		AddEdgesToGraph(BuildGraphContext);
-		AddDefaultInputVariables(BuildGraphContext);
+		AddDefaultInputLiterals(BuildGraphContext);
 
 		check(BuildGraphContext.Graph->OwnsAllReferencedNodes());
 		return MoveTemp(BuildGraphContext.Graph);
diff --git a/Engine/Plugins/Runtime/Metasound/Source/MetasoundFrontend/Private/MetasoundFrontendRegistries.cpp b/Engine/Plugins/Runtime/Metasound/Source/MetasoundFrontend/Private/MetasoundFrontendRegistries.cpp
--- a/Engine/Plugins/Runtime/Metasound/Source/MetasoundFrontend/Private/MetasoundFrontendRegistries.cpp
+++ b/Engine/Plugins/Runtime/Metasound/Source/MetasoundFrontend/Private/MetasoundFrontendRegistries.cpp
@@ -110,9 +110,10 @@ namespace Metasound
 				bool FindOutputNodeRegistryKeyForDataType(const FName& InDataTypeName, FNodeRegistryKey& OutKey) override;
 
 				// Create a new instance of a C++ implemented node from the registry.
-				virtual TUniquePtr<INode> CreateNode(const FNodeRegistryKey& InKey, FDefaultNodeConstructorParams&&) const override;
-				virtual TUniquePtr<INode> CreateNode(const FNodeRegistryKey& InKey, FDefaultLiteralNodeConstructorParams&&) const override;
 				virtual TUniquePtr<Metasound::INode> CreateNode(const FNodeRegistryKey& InKey, const Metasound::FNodeInitData& InInitData) const override;
+				virtual TUniquePtr<INode> CreateNode(const FNodeRegistryKey& InKey, FDefaultLiteralNodeConstructorParams&&) const override;
+				virtual TUniquePtr<INode> CreateNode(const FNodeRegistryKey& InKey, FDefaultNamedVertexNodeConstructorParams&&) const override;
+				virtual TUniquePtr<INode> CreateNode(const FNodeRegistryKey& InKey, FDefaultNamedVertexWithLiteralNodeConstructorParams&&) const override;
 
 				// Returns a list of possible nodes to use to convert from FromDataType to ToDataType.
 				// Returns an empty array if none are available.
@@ -177,7 +178,18 @@ namespace Metasound
 			}
 
 
-			TUniquePtr<Metasound::INode> FRegistryContainerImpl::CreateNode(const FNodeRegistryKey& InKey, FDefaultNodeConstructorParams&& InParams) const
+			TUniquePtr<Metasound::INode> FRegistryContainerImpl::CreateNode(const FNodeRegistryKey& InKey, const Metasound::FNodeInitData& InInitData) const
+			{
+				const INodeRegistryEntry* Entry = FindNodeEntry(InKey);
+
+				if (ensureAlwaysMsgf(nullptr != Entry, TEXT("Could not find node [RegistryKey:%s]"), *InKey))
+				{
+					return Entry->CreateNode(InInitData);
+				}
+				return nullptr;
+			}
+
+			TUniquePtr<Metasound::INode> FRegistryContainerImpl::CreateNode(const FNodeRegistryKey& InKey, FDefaultLiteralNodeConstructorParams&& InParams) const
 			{
 				const INodeRegistryEntry* Entry = FindNodeEntry(InKey);
 
@@ -189,7 +201,7 @@ namespace Metasound
 				return nullptr;
 			}
 
-			TUniquePtr<Metasound::INode> FRegistryContainerImpl::CreateNode(const FNodeRegistryKey& InKey, FDefaultLiteralNodeConstructorParams&& InParams) const
+			TUniquePtr<Metasound::INode> FRegistryContainerImpl::CreateNode(const FNodeRegistryKey& InKey, FDefaultNamedVertexNodeConstructorParams&& InParams) const
 			{
 				const INodeRegistryEntry* Entry = FindNodeEntry(InKey);
 
@@ -201,13 +213,13 @@ namespace Metasound
 				return nullptr;
 			}
 
-			TUniquePtr<Metasound::INode> FRegistryContainerImpl::CreateNode(const FNodeRegistryKey& InKey, const Metasound::FNodeInitData& InInitData) const
+			TUniquePtr<Metasound::INode> FRegistryContainerImpl::CreateNode(const FNodeRegistryKey& InKey, FDefaultNamedVertexWithLiteralNodeConstructorParams&& InParams) const
 			{
 				const INodeRegistryEntry* Entry = FindNodeEntry(InKey);
 
 				if (ensureAlwaysMsgf(nullptr != Entry, TEXT("Could not find node [RegistryKey:%s]"), *InKey))
 				{
-					return Entry->CreateNode(InInitData);
+					return Entry->CreateNode(MoveTemp(InParams));
 				}
 
 				return nullptr;
@@ -341,7 +353,7 @@ namespace Metasound
 			bool FRegistryContainerImpl::FindVariableNodeRegistryKeyForDataType(const FName& InDataTypeName, FNodeRegistryKey& OutKey)
 			{
 				FMetasoundFrontendClass Class;
-				if (IDataTypeRegistry::Get().GetFrontendVariableClass(InDataTypeName, Class))
+				if (IDataTypeRegistry::Get().GetFrontendLiteralClass(InDataTypeName, Class))
 				{
 					OutKey = NodeRegistryKey::CreateKey(Class.Metadata);
 					return true;
diff --git a/Engine/Plugins/Runtime/Metasound/Source/MetasoundFrontend/Public/MetasoundAudioBuffer.h b/Engine/Plugins/Runtime/Metasound/Source/MetasoundFrontend/Public/MetasoundAudioBuffer.h
--- a/Engine/Plugins/Runtime/Metasound/Source/MetasoundFrontend/Public/MetasoundAudioBuffer.h
+++ b/Engine/Plugins/Runtime/Metasound/Source/MetasoundFrontend/Public/MetasoundAudioBuffer.h
@@ -5,7 +5,7 @@
 #include "CoreMinimal.h"
 #include "DSP/BufferVectorOperations.h"
 #include "MetasoundGraphCoreModule.h"
-#include "MetasoundDataReference.h"
+#include "MetasoundDataReferenceMacro.h"
 #include "MetasoundLog.h"
 
 /** Define which determines whether to check that the size of the audio buffer has not changed since initialization */
diff --git a/Engine/Plugins/Runtime/Metasound/Source/MetasoundFrontend/Public/MetasoundDataTypeRegistrationMacro.h b/Engine/Plugins/Runtime/Metasound/Source/MetasoundFrontend/Public/MetasoundDataTypeRegistrationMacro.h
--- a/Engine/Plugins/Runtime/Metasound/Source/MetasoundFrontend/Public/MetasoundDataTypeRegistrationMacro.h
+++ b/Engine/Plugins/Runtime/Metasound/Source/MetasoundFrontend/Public/MetasoundDataTypeRegistrationMacro.h
@@ -11,20 +11,20 @@
 #include "MetasoundConverterNodeRegistrationMacro.h"
 #include "MetasoundDataFactory.h"
 #include "MetasoundDataReference.h"
+#include "MetasoundDataReferenceMacro.h"
 #include "MetasoundEnum.h"
+#include "MetasoundFrontendDataTypeRegistry.h"
 #include "MetasoundFrontendRegistries.h"
 #include "MetasoundInputNode.h"
 #include "MetasoundLiteral.h"
+#include "MetasoundLiteralNode.h"
 #include "MetasoundLog.h"
 #include "MetasoundNodeRegistrationMacro.h"
-#include "MetasoundOperatorInterface.h"
 #include "MetasoundOutputNode.h"
-#include "MetasoundPrimitives.h"
 #include "MetasoundReceiveNode.h"
-#include "MetasoundRouter.h"
 #include "MetasoundSendNode.h"
 #include "MetasoundTransmissionRegistration.h"
-#include "MetasoundVariableNode.h"
+#include "MetasoundVariableNodes.h"
 
 #include <type_traits>
 
@@ -356,8 +356,20 @@ namespace Metasound
 					TOutputNode<TDataType> OutputPrototype(TEXT(""), FGuid(), TEXT(""));
 					OutputClass = Metasound::Frontend::GenerateClassDescription(OutputPrototype.GetMetadata(), EMetasoundFrontendClassType::Output);
 
-					TGetVariableNode<TDataType> VariablePrototype(TEXT(""), FGuid(), TEXT(""), FLiteral());
-					VariableClass = Metasound::Frontend::GenerateClassDescription(VariablePrototype.GetMetadata(), EMetasoundFrontendClassType::Variable);
+					TLiteralNode<TDataType> LiteralPrototype(TEXT(""), FGuid(), FLiteral());
+					LiteralClass = Metasound::Frontend::GenerateClassDescription(LiteralPrototype.GetMetadata(), EMetasoundFrontendClassType::Variable);
+
+					TInitVariableNode<TDataType> InitVariablePrototype(TEXT(""), FGuid(), FLiteral());
+					InitVariableClass = Metasound::Frontend::GenerateClassDescription(InitVariablePrototype.GetMetadata(), EMetasoundFrontendClassType::Variable);
+
+					TSetVariableNode<TDataType> SetVariablePrototype(TEXT(""), FGuid());
+					SetVariableClass = Metasound::Frontend::GenerateClassDescription(SetVariablePrototype.GetMetadata(), EMetasoundFrontendClassType::External);
+
+					TGetVariableNode<TDataType> GetVariablePrototype(TEXT(""), FGuid());
+					GetVariableClass = Metasound::Frontend::GenerateClassDescription(GetVariablePrototype.GetMetadata(), EMetasoundFrontendClassType::External);
+
+					TGetDelayedVariableNode<TDataType> GetDelayedVariablePrototype(TEXT(""), FGuid());
+					GetDelayedVariableClass = Metasound::Frontend::GenerateClassDescription(GetDelayedVariablePrototype.GetMetadata(), EMetasoundFrontendClassType::External);
 				}
 				virtual ~FDataTypeRegistryEntry() {}
 
@@ -376,9 +388,9 @@ namespace Metasound
 					return InputClass;
 				}
 
-				virtual const FMetasoundFrontendClass& GetFrontendVariableClass() const override
+				virtual const FMetasoundFrontendClass& GetFrontendLiteralClass() const override
 				{
-					return VariableClass;
+					return LiteralClass;
 				}
 
 				virtual const FMetasoundFrontendClass& GetFrontendOutputClass() const override
@@ -386,6 +398,26 @@ namespace Metasound
 					return OutputClass;
 				}
 
+				virtual const FMetasoundFrontendClass& GetFrontendInitVariableClass() const override
+				{
+					return InitVariableClass;
+				}
+
+				virtual const FMetasoundFrontendClass& GetFrontendSetVariableClass() const override
+				{
+					return SetVariableClass;
+				}
+
+				virtual const FMetasoundFrontendClass& GetFrontendGetVariableClass() const override
+				{
+					return GetVariableClass;
+				}
+
+				virtual const FMetasoundFrontendClass& GetFrontendGetDelayedVariableClass() const 
+				{
+					return GetDelayedVariableClass;
+				}
+
 				virtual TUniquePtr<INode> CreateInputNode(FInputNodeConstructorParams&& InParams) const override
 				{
 					return MakeUnique<TInputNode<TDataType>>(InParams.NodeName, InParams.InstanceID, InParams.VertexName, MoveTemp(InParams.InitParam));
@@ -396,9 +428,9 @@ namespace Metasound
 					return MakeUnique<TOutputNode<TDataType>>(InParams.NodeName, InParams.InstanceID, InParams.VertexName);
 				}
 
-				virtual TUniquePtr<INode> CreateVariableNode(FVariableNodeConstructorParams&& InParams) const override
+				virtual TUniquePtr<INode> CreateLiteralNode(FLiteralNodeConstructorParams&& InParams) const override
 				{
-					return MakeUnique<TGetVariableNode<TDataType>>(InParams.NodeName, InParams.InstanceID, InParams.VertexName, MoveTemp(InParams.InitParam));
+					return MakeUnique<TLiteralNode<TDataType>>(InParams.NodeName, InParams.InstanceID, MoveTemp(InParams.Literal));
 				}
 
 				virtual TUniquePtr<INode> CreateReceiveNode(const FNodeInitData& InParams) const override
@@ -406,6 +438,26 @@ namespace Metasound
 					return MakeUnique<TReceiveNode<TDataType>>(InParams);
 				}
 
+				virtual TUniquePtr<INode> CreateInitVariableNode(FInitVariableNodeConstructorParams&& InParams) const override
+				{
+					return MakeUnique<TInitVariableNode<TDataType>>(InParams.NodeName, InParams.InstanceID, MoveTemp(InParams.Literal));
+				}
+
+				virtual TUniquePtr<INode> CreateSetVariableNode(const FNodeInitData& InParams) const override
+				{
+					return MakeUnique<TSetVariableNode<TDataType>>(InParams);
+				}
+
+				virtual TUniquePtr<INode> CreateGetVariableNode(const FNodeInitData& InParams) const override
+				{
+					return MakeUnique<TGetVariableNode<TDataType>>(InParams);
+				}
+
+				virtual TUniquePtr<INode> CreateGetDelayedVariableNode(const FNodeInitData& InParams) const override
+				{
+					return MakeUnique<TGetDelayedVariableNode<TDataType>>(InParams);
+				}
+
 				virtual Audio::IProxyDataPtr CreateProxy(UObject* InObject) const override
 				{
 					// Only attempt to create proxy if the `UClassToUse` is not void.
@@ -442,7 +494,11 @@ namespace Metasound
 				Frontend::FDataTypeRegistryInfo Info;
 				FMetasoundFrontendClass InputClass;
 				FMetasoundFrontendClass OutputClass;
-				FMetasoundFrontendClass VariableClass;
+				FMetasoundFrontendClass LiteralClass;
+				FMetasoundFrontendClass InitVariableClass;
+				FMetasoundFrontendClass SetVariableClass;
+				FMetasoundFrontendClass GetVariableClass;
+				FMetasoundFrontendClass GetDelayedVariableClass;
 				TSharedPtr<Frontend::IEnumDataTypeInterface> EnumInterface;
 			};
 
diff --git a/Engine/Plugins/Runtime/Metasound/Source/MetasoundFrontend/Public/MetasoundFrontendDataTypeRegistry.h b/Engine/Plugins/Runtime/Metasound/Source/MetasoundFrontend/Public/MetasoundFrontendDataTypeRegistry.h
--- a/Engine/Plugins/Runtime/Metasound/Source/MetasoundFrontend/Public/MetasoundFrontendDataTypeRegistry.h
+++ b/Engine/Plugins/Runtime/Metasound/Source/MetasoundFrontend/Public/MetasoundFrontendDataTypeRegistry.h
@@ -13,10 +13,6 @@ namespace Metasound
 {
 	namespace Frontend
 	{
-		using FOutputNodeConstructorParams = FDefaultNodeConstructorParams;
-		using FVariableNodeConstructorParams = FDefaultLiteralNodeConstructorParams;
-		using FInputNodeConstructorParams = FDefaultLiteralNodeConstructorParams;
-
 		/** FDataTypeRegsitryInfo contains runtime inspectable behavior of a registered
 		 * MetaSound data type.
 		 */
@@ -123,15 +119,27 @@ namespace Metasound
 			 */
 			virtual TSharedPtr<const IEnumDataTypeInterface> GetEnumInterface() const = 0;
 
-			/** Return an FMetasoundFrontenClass representing an input node of the data type. */
+			/** Return an FMetasoundFrontendClass representing an input node of the data type. */
 			virtual const FMetasoundFrontendClass& GetFrontendInputClass() const = 0;
 
-			/** Return an FMetasoundFrontenClass representing a variable node of the data type. */
-			virtual const FMetasoundFrontendClass& GetFrontendVariableClass() const = 0;
+			/** Return an FMetasoundFrontendClass representing a variable node of the data type. */
+			virtual const FMetasoundFrontendClass& GetFrontendLiteralClass() const = 0;
 
-			/** Return an FMetasoundFrontenClass representing an output node of the data type. */
+			/** Return an FMetasoundFrontendClass representing an output node of the data type. */
 			virtual const FMetasoundFrontendClass& GetFrontendOutputClass() const = 0;
 
+			/** Return an FMetasoundFrontendClass representing an init variable node of the data type. */
+			virtual const FMetasoundFrontendClass& GetFrontendInitVariableClass() const = 0;
+
+			/** Return an FMetasoundFrontendClass representing an set variable node of the data type. */
+			virtual const FMetasoundFrontendClass& GetFrontendSetVariableClass() const = 0;
+
+			/** Return an FMetasoundFrontendClass representing an get variable node of the data type. */
+			virtual const FMetasoundFrontendClass& GetFrontendGetVariableClass() const = 0;
+
+			/** Return an FMetasoundFrontendClass representing an get delayed variable node of the data type. */
+			virtual const FMetasoundFrontendClass& GetFrontendGetDelayedVariableClass() const = 0;
+
 			/** Create an input node */
 			virtual TUniquePtr<INode> CreateInputNode(FInputNodeConstructorParams&&) const = 0;
 
@@ -139,11 +147,26 @@ namespace Metasound
 			virtual TUniquePtr<INode> CreateOutputNode(FOutputNodeConstructorParams&&) const = 0;
 			
 			/** Create a variable node */
-			virtual TUniquePtr<INode> CreateVariableNode(FVariableNodeConstructorParams&&) const = 0;
+			virtual TUniquePtr<INode> CreateLiteralNode(FLiteralNodeConstructorParams&&) const = 0;
 
-			/** Create a receive node */
+			/** Create a receive node for this data type. */
 			virtual TUniquePtr<INode> CreateReceiveNode(const FNodeInitData&) const = 0;
 
+			/** Create a init variable node for this data type. 
+			 *
+			 *  @param InInitParams - Contains a literal used to create the variable.
+			 */
+			virtual TUniquePtr<INode> CreateInitVariableNode(FInitVariableNodeConstructorParams&& InInitParams) const = 0;
+
+			/** Create a set variable node for this data type. */
+			virtual TUniquePtr<INode> CreateSetVariableNode(const FNodeInitData&) const = 0;
+
+			/** Create a get variable node for this data type. */
+			virtual TUniquePtr<INode> CreateGetVariableNode(const FNodeInitData&) const = 0;
+
+			/** Create a get delayed variable node for this data type. */
+			virtual TUniquePtr<INode> CreateGetDelayedVariableNode(const FNodeInitData&) const = 0;
+
 			/** Create a proxy from a UObject. If this data type does not support
 			 * UObject proxies, return a nullptr. */
 			virtual Audio::IProxyDataPtr CreateProxy(UObject* InObject) const = 0;
@@ -195,20 +218,36 @@ namespace Metasound
 
 			virtual TSharedPtr<IDataChannel, ESPMode::ThreadSafe> CreateDataChannel(const FName& InDataType, const FOperatorSettings& InOperatorSettings) const = 0;
 
-			/** Return an FMetasoundFrontenClass representing an input node of the data type. */
+			/** Return an FMetasoundFrontendClass representing an input node of the data type. */
 			virtual bool GetFrontendInputClass(const FName& InDataType, FMetasoundFrontendClass& OutClass) const = 0;
 
-			/** Return an FMetasoundFrontenClass representing a variable node of the data type. */
-			virtual bool GetFrontendVariableClass(const FName& InDataType, FMetasoundFrontendClass& OutClass) const = 0;
+			/** Return an FMetasoundFrontendClass representing a variable node of the data type. */
+			virtual bool GetFrontendLiteralClass(const FName& InDataType, FMetasoundFrontendClass& OutClass) const = 0;
 
-			/** Return an FMetasoundFrontenClass representing an output node of the data type. */
+			/** Return an FMetasoundFrontendClass representing an output node of the data type. */
 			virtual bool GetFrontendOutputClass(const FName& InDataType, FMetasoundFrontendClass& OutClass) const = 0;
 
-			// Create a new instance of a C++ implemented node 
+			/** Return an FMetasoundFrontendClass representing an init variable node of the data type. */
+			virtual bool GetFrontendInitVariableClass(const FName& InDataType, FMetasoundFrontendClass& OutClass) const = 0;
+
+			/** Return an FMetasoundFrontendClass representing an set variable node of the data type. */
+			virtual bool GetFrontendSetVariableClass(const FName& InDataType, FMetasoundFrontendClass& OutClass) const = 0;
+
+			/** Return an FMetasoundFrontendClass representing an get variable node of the data type. */
+			virtual bool GetFrontendGetVariableClass(const FName& InDataType, FMetasoundFrontendClass& OutClass) const = 0;
+
+			/** Return an FMetasoundFrontendClass representing an get delayed variable node of the data type. */
+			virtual bool GetFrontendGetDelayedVariableClass(const FName& InDataType, FMetasoundFrontendClass& OutClass) const = 0;
+
+			// Create a new instance of a C++ implemented node from the registry.
 			virtual TUniquePtr<INode> CreateInputNode(const FName& InInputType, FInputNodeConstructorParams&& InParams) const = 0;
-			virtual TUniquePtr<INode> CreateVariableNode(const FName& InVariableType, FVariableNodeConstructorParams&& InParams) const = 0;
+			virtual TUniquePtr<INode> CreateLiteralNode(const FName& InLiteralType, FLiteralNodeConstructorParams&& InParams) const = 0;
 			virtual TUniquePtr<INode> CreateOutputNode(const FName& InOutputType, FOutputNodeConstructorParams&& InParams) const = 0;
-			virtual TUniquePtr<INode> CreateReceiveNode(const FName& InDataType, const FNodeInitData& InParams) const = 0;
+			virtual TUniquePtr<INode> CreateReceiveNode(const FName& InDataType, const FNodeInitData&) const = 0;
+			virtual TUniquePtr<INode> CreateInitVariableNode(const FName& InDataType, FInitVariableNodeConstructorParams&& InParams) const = 0;
+			virtual TUniquePtr<INode> CreateSetVariableNode(const FName& InDataType, const FNodeInitData& InParams) const = 0;
+			virtual TUniquePtr<INode> CreateGetVariableNode(const FName& InDataType, const FNodeInitData& InParams) const = 0;
+			virtual TUniquePtr<INode> CreateGetDelayedVariableNode(const FName& InDataType, const FNodeInitData& InParams) const = 0;
 		};
 	}
 }
diff --git a/Engine/Plugins/Runtime/Metasound/Source/MetasoundFrontend/Public/MetasoundFrontendGraph.h b/Engine/Plugins/Runtime/Metasound/Source/MetasoundFrontend/Public/MetasoundFrontendGraph.h
--- a/Engine/Plugins/Runtime/Metasound/Source/MetasoundFrontend/Public/MetasoundFrontendGraph.h
+++ b/Engine/Plugins/Runtime/Metasound/Source/MetasoundFrontend/Public/MetasoundFrontendGraph.h
@@ -115,20 +115,20 @@ namespace Metasound
 		static TUniquePtr<FFrontendGraph> CreateGraph(const FMetasoundFrontendGraphClass& InGraphClass, const TArray<FMetasoundFrontendGraphClass>& InSubgraphs, const TArray<FMetasoundFrontendClass>& InDependencies);
 
 	private:
-		struct FDefaultVariableData
+		struct FDefaultLiteralData
 		{
 			FGuid DestinationNodeID;
 			FGuid DestinationVertexID;
 			FVertexKey DestinationVertexKey;
 			FName TypeName;
-			FVariableNodeConstructorParams InitParams;
+			FLiteralNodeConstructorParams InitParams;
 		};
 
 		// Map of Input VertexID to variable data required to construct and connect default variable
 		using FNodeIDVertexID = TTuple<FGuid, FGuid>;
 		using FDependencyByIDMap = TMap<FGuid, const FMetasoundFrontendClass*>;
 		using FSharedNodeByIDMap = TMap<FGuid, TSharedPtr<const INode>>;
-		using FDefaultInputByIDMap = TMap<FNodeIDVertexID, FDefaultVariableData>;
+		using FDefaultInputByIDMap = TMap<FNodeIDVertexID, FDefaultLiteralData>;
 
 		// Context used throughout entire graph build process
 		// (for both a root and nested subgraphs)
@@ -148,7 +148,7 @@ namespace Metasound
 			FDefaultInputByIDMap DefaultInputs;
 		};
 
-		static TArray<FDefaultVariableData> GetInputDefaultVariableData(const FMetasoundFrontendNode& InNode, const FNodeInitData& InInitData);
+		static TArray<FDefaultLiteralData> GetInputDefaultLiteralData(const FMetasoundFrontendNode& InNode, const FNodeInitData& InInitData);
 
 		static bool SortSubgraphDependencies(TArray<const FMetasoundFrontendGraphClass*>& Subgraphs);
 
@@ -165,6 +165,6 @@ namespace Metasound
 		// TODO: add errors here. Most will be a "PromptIfMissing"...
 		static void AddNodesToGraph(FBuildGraphContext& InGraphContext);
 		static void AddEdgesToGraph(FBuildGraphContext& InGraphContext);
-		static void AddDefaultInputVariables(FBuildGraphContext& InGraphContext);
+		static void AddDefaultInputLiterals(FBuildGraphContext& InGraphContext);
 	};
 }
diff --git a/Engine/Plugins/Runtime/Metasound/Source/MetasoundFrontend/Public/MetasoundFrontendRegistries.h b/Engine/Plugins/Runtime/Metasound/Source/MetasoundFrontend/Public/MetasoundFrontendRegistries.h
--- a/Engine/Plugins/Runtime/Metasound/Source/MetasoundFrontend/Public/MetasoundFrontendRegistries.h
+++ b/Engine/Plugins/Runtime/Metasound/Source/MetasoundFrontend/Public/MetasoundFrontendRegistries.h
@@ -16,7 +16,7 @@
 namespace Metasound
 {
 	// Base implementation for NodeConstructorCallbacks
-	struct FDefaultNodeConstructorParams
+	struct FDefaultNamedVertexNodeConstructorParams
 	{
 		// the instance name and name of the specific connection that should be used.
 		FString NodeName;
@@ -24,7 +24,7 @@ namespace Metasound
 		FString VertexName;
 	};
 
-	struct FDefaultLiteralNodeConstructorParams
+	struct FDefaultNamedVertexWithLiteralNodeConstructorParams
 	{
 		// the instance name and name of the specific connection that should be used.
 		FString NodeName;
@@ -33,9 +33,17 @@ namespace Metasound
 		FLiteral InitParam = FLiteral::CreateInvalid();
 	};
 
-	using FOutputNodeConstructorParams = FDefaultNodeConstructorParams;
-	using FVariableNodeConstructorParams = FDefaultLiteralNodeConstructorParams;
-	using FInputNodeConstructorParams = FDefaultLiteralNodeConstructorParams;
+	struct FDefaultLiteralNodeConstructorParams
+	{
+		FString NodeName;
+		FGuid InstanceID;
+		FLiteral Literal = FLiteral::CreateInvalid();
+	};
+
+	using FInputNodeConstructorParams = FDefaultNamedVertexWithLiteralNodeConstructorParams;
+	using FOutputNodeConstructorParams = FDefaultNamedVertexNodeConstructorParams;
+	using FLiteralNodeConstructorParams = FDefaultLiteralNodeConstructorParams;
+	using FInitVariableNodeConstructorParams = FDefaultLiteralNodeConstructorParams;
 
 	using FIterateMetasoundFrontendClassFunction = TFunctionRef<void(const FMetasoundFrontendClass&)>;
 
@@ -112,12 +120,19 @@ namespace Metasound
 			 */
 			virtual const FNodeClassInfo& GetClassInfo() const = 0;
 
-			/** Create a node given FDefaultNodeConstructorParams.
+			/** Create a node given FDefaultNamedVertexNodeConstructorParams.
 			 *
-			 * If a node can be created with FDefaultNodeConstructorParams, this function
+			 * If a node can be created with FDefaultNamedVertexNodeConstructorParams, this function
 			 * should return a valid node pointer.
 			 */
-			virtual TUniquePtr<INode> CreateNode(FDefaultNodeConstructorParams&&) const = 0;
+			virtual TUniquePtr<INode> CreateNode(FDefaultNamedVertexNodeConstructorParams&&) const = 0;
+
+			/** Create a node given FDefaultNamedVertexWithLiteralNodeConstructorParams.
+			 *
+			 * If a node can be created with FDefaultNamedVertexWithLiteralNodeConstructorParams, this function
+			 * should return a valid node pointer.
+			 */
+			virtual TUniquePtr<INode> CreateNode(FDefaultNamedVertexWithLiteralNodeConstructorParams&&) const = 0;
 
 			/** Create a node given FDefaultLiteralNodeConstructorParams.
 			 *
@@ -326,9 +341,10 @@ public:
 	virtual bool FindOutputNodeRegistryKeyForDataType(const FName& InDataTypeName, FNodeRegistryKey& OutKey) = 0;
 
 
-	virtual TUniquePtr<Metasound::INode> CreateNode(const FNodeRegistryKey& InKey, Metasound::FDefaultNodeConstructorParams&&) const = 0;
-	virtual TUniquePtr<Metasound::INode> CreateNode(const FNodeRegistryKey& InKey, Metasound::FDefaultLiteralNodeConstructorParams&&) const = 0;
 	virtual TUniquePtr<Metasound::INode> CreateNode(const FNodeRegistryKey& InKey, const Metasound::FNodeInitData&) const = 0;
+	virtual TUniquePtr<Metasound::INode> CreateNode(const FNodeRegistryKey& InKey, Metasound::FDefaultLiteralNodeConstructorParams&&) const = 0;
+	virtual TUniquePtr<Metasound::INode> CreateNode(const FNodeRegistryKey& InKey, Metasound::FDefaultNamedVertexNodeConstructorParams&&) const = 0;
+	virtual TUniquePtr<Metasound::INode> CreateNode(const FNodeRegistryKey& InKey, Metasound::FDefaultNamedVertexWithLiteralNodeConstructorParams&&) const = 0;
 
 	virtual bool RegisterConversionNode(const FConverterNodeRegistryKey& InNodeKey, const FConverterNodeInfo& InNodeInfo) = 0;
 
diff --git a/Engine/Plugins/Runtime/Metasound/Source/MetasoundFrontend/Public/MetasoundInputNode.h b/Engine/Plugins/Runtime/Metasound/Source/MetasoundFrontend/Public/MetasoundInputNode.h
--- a/Engine/Plugins/Runtime/Metasound/Source/MetasoundFrontend/Public/MetasoundInputNode.h
+++ b/Engine/Plugins/Runtime/Metasound/Source/MetasoundFrontend/Public/MetasoundInputNode.h
@@ -83,7 +83,7 @@ namespace Metasound
 			{
 				if (TExecutableDataType<DataType>::bIsExecutable)
 				{
-					return &TInputOperator<FTrigger>::ExecuteFunction;
+					return &TInputOperator<DataType>::ExecuteFunction;
 				}
 				return nullptr;
 			}
diff --git a/Engine/Plugins/Runtime/Metasound/Source/MetasoundFrontend/Public/MetasoundNodeRegistrationMacro.h b/Engine/Plugins/Runtime/Metasound/Source/MetasoundFrontend/Public/MetasoundNodeRegistrationMacro.h
--- a/Engine/Plugins/Runtime/Metasound/Source/MetasoundFrontend/Public/MetasoundNodeRegistrationMacro.h
+++ b/Engine/Plugins/Runtime/Metasound/Source/MetasoundFrontend/Public/MetasoundNodeRegistrationMacro.h
@@ -57,21 +57,27 @@ namespace Metasound
 				return ClassInfo;
 			}
 
-			virtual TUniquePtr<INode> CreateNode(FDefaultNodeConstructorParams&& InParams) const override
+			virtual TUniquePtr<INode> CreateNode(const FNodeInitData& InParams) const override
 			{
-				return nullptr;
+				return MakeUnique<FNodeType>(InParams);
 			}
-		
+
 			virtual TUniquePtr<INode> CreateNode(FDefaultLiteralNodeConstructorParams&& InParams) const override
 			{
 				return nullptr;
 			}
 
-			virtual TUniquePtr<INode> CreateNode(const FNodeInitData& InParams) const override
+			virtual TUniquePtr<INode> CreateNode(FDefaultNamedVertexNodeConstructorParams&& InParams) const override
 			{
-				return MakeUnique<FNodeType>(InParams);
+				return nullptr;
+			}
+		
+			virtual TUniquePtr<INode> CreateNode(FDefaultNamedVertexWithLiteralNodeConstructorParams&& InParams) const override
+			{
+				return nullptr;
 			}
 
+
 			virtual const FMetasoundFrontendClass& GetFrontendClass() const override
 			{
 				return FrontendClass;
diff --git a/Engine/Plugins/Runtime/Metasound/Source/MetasoundFrontend/Public/MetasoundPrimitives.h b/Engine/Plugins/Runtime/Metasound/Source/MetasoundFrontend/Public/MetasoundPrimitives.h
--- a/Engine/Plugins/Runtime/Metasound/Source/MetasoundFrontend/Public/MetasoundPrimitives.h
+++ b/Engine/Plugins/Runtime/Metasound/Source/MetasoundFrontend/Public/MetasoundPrimitives.h
@@ -3,7 +3,7 @@
 #pragma once
 
 #include "CoreMinimal.h"
-#include "MetasoundDataReference.h"
+#include "MetasoundDataReferenceMacro.h"
 
 namespace Metasound
 {
@@ -11,12 +11,8 @@ namespace Metasound
 	
 	DECLARE_METASOUND_DATA_REFERENCE_TYPES(int32, METASOUNDFRONTEND_API, FInt32TypeInfo, FInt32ReadRef, FInt32WriteRef);
 
-	//DECLARE_METASOUND_DATA_REFERENCE_TYPES(int64, METASOUNDFRONTEND_API, FInt64TypeInfo, FInt64ReadRef, FInt64WriteRef);
-
 	DECLARE_METASOUND_DATA_REFERENCE_TYPES(float, METASOUNDFRONTEND_API, FFloatTypeInfo, FFloatReadRef, FFloatWriteRef);
 
-	//DECLARE_METASOUND_DATA_REFERENCE_TYPES(double, METASOUNDFRONTEND_API, FDoubleTypeInfo, FDoubleReadRef, FDoubleWriteRef);
-
 	DECLARE_METASOUND_DATA_REFERENCE_TYPES(FString, METASOUNDFRONTEND_API, FStringTypeInfo, FStringReadRef, FStringWriteRef);
 }
 
diff --git a/Engine/Plugins/Runtime/Metasound/Source/MetasoundFrontend/Public/MetasoundTime.h b/Engine/Plugins/Runtime/Metasound/Source/MetasoundFrontend/Public/MetasoundTime.h
--- a/Engine/Plugins/Runtime/Metasound/Source/MetasoundFrontend/Public/MetasoundTime.h
+++ b/Engine/Plugins/Runtime/Metasound/Source/MetasoundFrontend/Public/MetasoundTime.h
@@ -1,7 +1,7 @@
 // Copyright Epic Games, Inc. All Rights Reserved.
 #pragma once
 
-#include "MetasoundDataReference.h"
+#include "MetasoundDataReferenceMacro.h"
 #include "MetasoundOperatorSettings.h"
 
 namespace Metasound
diff --git a/Engine/Plugins/Runtime/Metasound/Source/MetasoundFrontend/Public/MetasoundVariableNode.h b/Engine/Plugins/Runtime/Metasound/Source/MetasoundFrontend/Public/MetasoundVariableNode.h
deleted file mode 100644
diff --git a/Engine/Plugins/Runtime/Metasound/Source/MetasoundFrontend/Public/MetasoundVariableRegistration.h b/Engine/Plugins/Runtime/Metasound/Source/MetasoundFrontend/Public/MetasoundVariableRegistration.h
new file mode 100644
--- /dev/null
+++ b/Engine/Plugins/Runtime/Metasound/Source/MetasoundFrontend/Public/MetasoundVariableRegistration.h
@@ -0,0 +1,39 @@
+// Copyright Epic Games, Inc. All Rights Reserved.
+
+#pragma once
+
+#include "MetasoundDataFactory.h"
+#include <type_traits>
+
+namespace Metasound
+{
+	namespace Frontend
+	{
+		/** Allow or disallow types to be registered as MetaSound Variables. */
+		template<typename DataType>
+		struct TEnableVariables
+		{
+			// By default, all registered MetaSound data types are variables.
+			static constexpr bool Value = true;
+		};
+
+		/** Determine if data type supports necessary operations required of all
+		 * MetaSound variable types.
+		 */
+		template<typename DataType>
+		struct TVariablesSupport
+		{
+
+		private:
+
+			static constexpr bool bIsCopyAssignable = std::is_copy_assignable<DataType>::value;
+			static constexpr bool bIsParsableFromLiteral = TLiteralTraits<DataType>::bIsParsableFromAnyLiteralType;
+
+		public:
+
+			/* MetaSound variables must be copy-assignable and parsable from a FLiteral */
+			static constexpr bool bIsVariableSupported = bIsCopyAssignable && bIsParsableFromLiteral;
+		};
+	}
+}
+
diff --git a/Engine/Plugins/Runtime/Metasound/Source/MetasoundGraphCore/Private/MetasoundLiteralNode.cpp b/Engine/Plugins/Runtime/Metasound/Source/MetasoundGraphCore/Private/MetasoundLiteralNode.cpp
new file mode 100644
--- /dev/null
+++ b/Engine/Plugins/Runtime/Metasound/Source/MetasoundGraphCore/Private/MetasoundLiteralNode.cpp
@@ -0,0 +1,15 @@
+// Copyright Epic Games, Inc. All Rights Reserved.
+
+#include "MetasoundLiteralNode.h"
+#include "CoreMinimal.h"
+
+namespace Metasound
+{
+	namespace LiteralNodeNames
+	{
+		const TCHAR* GetOutputDataName()
+		{
+			return TEXT("Value");
+		}
+	}
+}
diff --git a/Engine/Plugins/Runtime/Metasound/Source/MetasoundGraphCore/Private/MetasoundOperatorBuilder.cpp b/Engine/Plugins/Runtime/Metasound/Source/MetasoundGraphCore/Private/MetasoundOperatorBuilder.cpp
--- a/Engine/Plugins/Runtime/Metasound/Source/MetasoundGraphCore/Private/MetasoundOperatorBuilder.cpp
+++ b/Engine/Plugins/Runtime/Metasound/Source/MetasoundGraphCore/Private/MetasoundOperatorBuilder.cpp
@@ -336,8 +336,22 @@ namespace Metasound
 
 			const FDataReferenceCollection& FromDataReferenceCollection = InDataReferenceMap[Edge->From.Node].Outputs;
 
-			if (!FromDataReferenceCollection.ContainsDataReadReference(Edge->From.Vertex.GetVertexName(), Edge->From.Vertex.GetDataTypeName()))
+			bool bSuccess = false;
+
+			if (FromDataReferenceCollection.ContainsDataWriteReference(Edge->From.Vertex.GetVertexName(), Edge->From.Vertex.GetDataTypeName()))
+			{
+				// Contains data write reference
+				bSuccess = OutCollection.AddDataWriteReferenceFrom(Edge->To.Vertex.GetVertexName(), FromDataReferenceCollection, Edge->From.Vertex.GetVertexName(), Edge->From.Vertex.GetDataTypeName());
+			}
+			else if (FromDataReferenceCollection.ContainsDataReadReference(Edge->From.Vertex.GetVertexName(), Edge->From.Vertex.GetDataTypeName()))
+			{
+				// Contains data read reference
+				bSuccess = OutCollection.AddDataReadReferenceFrom(Edge->To.Vertex.GetVertexName(), FromDataReferenceCollection, Edge->From.Vertex.GetVertexName(), Edge->From.Vertex.GetDataTypeName());
+			}
+
+			if (!bSuccess)
 			{
+				// Does not contain any reference
 				// This is likely a node programming error where the edges reported by the INode interface
 				// did not match the readable parameter refs created by the operators outputs. Or, the edge description is invalid.
 
@@ -347,15 +361,6 @@ namespace Metasound
 				BuildStatus |= FBuildStatus::NonFatalError;
 				continue;
 			}
-
-
-			bool bSuccess = OutCollection.AddDataReadReferenceFrom(Edge->To.Vertex.GetVertexName(), FromDataReferenceCollection, Edge->From.Vertex.GetVertexName(), Edge->From.Vertex.GetDataTypeName());
-
-			if (!bSuccess)
-			{
-				AddBuildError<FMissingOutputDataReferenceError>(InContext.Errors, Edge->From);
-				BuildStatus |= FBuildStatus::NonFatalError;
-			}
 		}
 
 		return BuildStatus;
diff --git a/Engine/Plugins/Runtime/Metasound/Source/MetasoundGraphCore/Private/MetasoundVariableNodes.cpp b/Engine/Plugins/Runtime/Metasound/Source/MetasoundGraphCore/Private/MetasoundVariableNodes.cpp
new file mode 100644
--- /dev/null
+++ b/Engine/Plugins/Runtime/Metasound/Source/MetasoundGraphCore/Private/MetasoundVariableNodes.cpp
@@ -0,0 +1,34 @@
+// Copyright Epic Games, Inc. All Rights Reserved.
+
+#include "MetasoundVariableNodes.h"
+#include "CoreMinimal.h"
+
+namespace Metasound
+{
+	namespace VariableNodeVertexNames
+	{
+		const FString& GetInputDataName()
+		{
+			static const FString Name = TEXT("Value");
+			return Name;
+		}
+
+		const FString& GetOutputDataName()
+		{
+			static const FString Name = TEXT("Value");
+			return Name;
+		}
+
+		const FString& GetInputVariableName()
+		{
+			static const FString Name = TEXT("Variable");
+			return Name;
+		}
+
+		const FString& GetOutputVariableName()
+		{
+			static const FString Name = TEXT("Variable");
+			return Name;
+		}
+	}
+}
diff --git a/Engine/Plugins/Runtime/Metasound/Source/MetasoundGraphCore/Public/MetasoundDataReference.h b/Engine/Plugins/Runtime/Metasound/Source/MetasoundGraphCore/Public/MetasoundDataReference.h
--- a/Engine/Plugins/Runtime/Metasound/Source/MetasoundGraphCore/Public/MetasoundDataReference.h
+++ b/Engine/Plugins/Runtime/Metasound/Source/MetasoundGraphCore/Public/MetasoundDataReference.h
@@ -3,63 +3,8 @@
 #pragma once
 
 #include "CoreMinimal.h"
-#include "MetasoundOperatorSettings.h"
 #include <type_traits>
 
-
-/** Macro to make declaring a metasound parameter simple.  */
-// Declares a metasound parameter type by
-// - Adding typedefs for commonly used template types.
-// - Defining parameter type traits.
-#define DECLARE_METASOUND_DATA_REFERENCE_TYPES(DataType, ModuleApi, DataTypeInfoTypeName, DataReadReferenceTypeName, DataWriteReferenceTypeName) \
-	template<> \
-	struct ::Metasound::TDataReferenceTypeInfo<DataType> \
-	{ \
-		static ModuleApi const TCHAR* TypeName; \
-		static ModuleApi const FText& GetTypeDisplayText(); \
-		static ModuleApi const FMetasoundDataTypeId TypeId; \
-		\
-		private: \
-		\
-		static const DataType* const TypePtr; \
-	}; \
-	/*TODO remove typedefs */ \
-	typedef ::Metasound::TDataReferenceTypeInfo<DataType> DataTypeInfoTypeName; \
-	\
-	typedef ::Metasound::TDataReadReference<DataType> DataReadReferenceTypeName; \
-	typedef ::Metasound::TDataWriteReference<DataType> DataWriteReferenceTypeName;\
-	template<> \
-	struct ::Metasound::TDataReferenceTypeInfo<TArray<DataType>> \
-	{ \
-		static ModuleApi const TCHAR* TypeName; \
-		static ModuleApi const FText& GetTypeDisplayText(); \
-		static ModuleApi const FMetasoundDataTypeId TypeId; \
-		\
-		private: \
-		\
-		static const TArray<DataType>* const TypePtr; \
-	};
-
-// This only needs to be called if you don't plan on calling REGISTER_METASOUND_DATATYPE.
-#define DEFINE_METASOUND_DATA_TYPE(DataType, DataTypeName) \
-	const TCHAR* ::Metasound::TDataReferenceTypeInfo<DataType>::TypeName = TEXT(DataTypeName); \
-	const FText& ::Metasound::TDataReferenceTypeInfo<DataType>::GetTypeDisplayText() \
-	{ \
-		static const FText DisplayText = NSLOCTEXT("MetaSoundCore_DataReference", DataTypeName, DataTypeName); \
-		return DisplayText; \
-	} \
-	const DataType* const ::Metasound::TDataReferenceTypeInfo<DataType>::TypePtr = nullptr; \
-	const void* const ::Metasound::TDataReferenceTypeInfo<DataType>::TypeId = static_cast<const FMetasoundDataTypeId>(&::Metasound::TDataReferenceTypeInfo<DataType>::TypePtr); \
-	/* Array definition */ \
-	const TCHAR* ::Metasound::TDataReferenceTypeInfo<TArray<DataType>>::TypeName = TEXT(DataTypeName ":Array"); \
-	const FText& ::Metasound::TDataReferenceTypeInfo<TArray<DataType>>::GetTypeDisplayText() \
-	{ \
-		static const FText DisplayText = NSLOCTEXT("MetaSoundCore_DataReference", DataTypeName "_Array", DataTypeName ":Array"); \
-		return DisplayText; \
-	} \
-	const TArray<DataType>* const ::Metasound::TDataReferenceTypeInfo<TArray<DataType>>::TypePtr = nullptr; \
-	const void* const ::Metasound::TDataReferenceTypeInfo<TArray<DataType>>::TypeId = static_cast<const FMetasoundDataTypeId>(&::Metasound::TDataReferenceTypeInfo<TArray<DataType>>::TypePtr);
-
 namespace Audio
 {
 	// Forward declare
@@ -362,56 +307,57 @@ namespace Metasound
 		{
 		}
 
-		public:
-			typedef TDataReference<DataType> FDataReference;
+	public:
+		typedef TDataReference<DataType> FDataReference;
 
-			// This should be used to construct a new DataType object and return this TDataReadReference as a wrapper around it.
-			template <typename... ArgTypes>
-			static TDataReadReference<DataType> CreateNew(ArgTypes&&... Args)
-			{
-				static_assert(std::is_constructible<DataType, ArgTypes...>::value, "DataType constructor does not support provided types.");
-				return TDataReadReference<DataType>(EDataRefShouldConstruct::NewObject, Forward<ArgTypes>(Args)...);
-			}
+		// This should be used to construct a new DataType object and return this TDataReadReference as a wrapper around it.
+		template <typename... ArgTypes>
+		static TDataReadReference<DataType> CreateNew(ArgTypes&&... Args)
+		{
+			static_assert(std::is_constructible<DataType, ArgTypes...>::value, "DataType constructor does not support provided types.");
+			return TDataReadReference<DataType>(EDataRefShouldConstruct::NewObject, Forward<ArgTypes>(Args)...);
+		}
 
-			TDataReadReference(const TDataReadReference& Other) = default;
+		TDataReadReference(const TDataReadReference& Other) = default;
 
-			/** Construct a readable parameter ref from a writable parameter ref. */
-			explicit TDataReadReference(const TDataWriteReference<DataType>& WritableRef)
-			:	FDataReference(WritableRef)
-			{
-			}
+		/** Construct a readable parameter ref from a writable parameter ref. */
+		explicit TDataReadReference(const TDataWriteReference<DataType>& WritableRef)
+		:	FDataReference(WritableRef)
+		{
+		}
 
-			/** Assign a readable parameter ref from a writable parameter ref. */
-			TDataReadReference<DataType>& operator=(const TDataWriteReference<DataType>& Other)
-			{
-				TDataReference<DataType>::ObjectReference = Other.ObjectReference;
-				return *this;
-			}
+		/** Assign a readable parameter ref from a writable parameter ref. */
+		TDataReadReference<DataType>& operator=(const TDataWriteReference<DataType>& Other)
+		{
+			TDataReference<DataType>::ObjectReference = Other.ObjectReference;
+			return *this;
+		}
 
-			/** Enable copy operator */
-			TDataReadReference<DataType>& operator=(const TDataReadReference<DataType>& Other) = default;
+		/** Enable copy operator */
+		TDataReadReference<DataType>& operator=(const TDataReadReference<DataType>& Other) = default;
 
-			/** Enable move operator */
-			TDataReadReference<DataType>& operator=(TDataReadReference<DataType>&& Other) = default;
+		/** Enable move operator */
+		TDataReadReference<DataType>& operator=(TDataReadReference<DataType>&& Other) = default;
 
-			/** Const access to the underlying parameter object. */
-			FORCEINLINE const DataType& operator*() const
-			{
-				return *TDataReference<DataType>::ObjectReference;
-			}
+		/** Const access to the underlying parameter object. */
+		FORCEINLINE const DataType& operator*() const
+		{
+			return *TDataReference<DataType>::ObjectReference;
+		}
 
-			/** Const access to the underlying parameter object. */
-			FORCEINLINE const DataType* operator->() const
-			{
-				return TDataReference<DataType>::ObjectReference.operator->();
-			}
+		/** Const access to the underlying parameter object. */
+		FORCEINLINE const DataType* operator->() const
+		{
+			return TDataReference<DataType>::ObjectReference.operator->();
+		}
 
-			/** Create a clone of this parameter reference. */
-			virtual TUniquePtr<IDataReference> Clone() const override
-			{
-				typedef TDataReadReference<DataType> FDataReadReference;
+		/** Create a clone of this parameter reference. */
+		virtual TUniquePtr<IDataReference> Clone() const override
+		{
+			typedef TDataReadReference<DataType> FDataReadReference;
 
-				return MakeUnique< FDataReadReference >(*this);
-			}
+			return MakeUnique< FDataReadReference >(*this);
+		}
 	};
 }
+
diff --git a/Engine/Plugins/Runtime/Metasound/Source/MetasoundGraphCore/Public/MetasoundDataReferenceMacro.h b/Engine/Plugins/Runtime/Metasound/Source/MetasoundGraphCore/Public/MetasoundDataReferenceMacro.h
new file mode 100644
--- /dev/null
+++ b/Engine/Plugins/Runtime/Metasound/Source/MetasoundGraphCore/Public/MetasoundDataReferenceMacro.h
@@ -0,0 +1,61 @@
+// Copyright Epic Games, Inc. All Rights Reserved.
+
+#pragma once
+
+#include "CoreMinimal.h"
+#include "MetasoundDataReference.h"
+#include "MetasoundVariable.h"
+
+#define DECLARE_METASOUND_DATA_REFERENCE_CORE_TYPE(DataType, ModuleApi) \
+	template<> \
+	struct ::Metasound::TDataReferenceTypeInfo<DataType> \
+	{ \
+		static ModuleApi const TCHAR* TypeName; \
+		static ModuleApi const FText& GetTypeDisplayText(); \
+		static ModuleApi const FMetasoundDataTypeId TypeId; \
+		\
+		private: \
+		\
+		static const DataType* const TypePtr; \
+	};
+
+#define DECLARE_METASOUND_DATA_REFERENCE_ALIAS_TYPES(DataType, DataTypeInfoTypeName, DataReadReferenceTypeName, DataWriteReferenceTypeName) \
+	typedef ::Metasound::TDataReferenceTypeInfo<DataType> DataTypeInfoTypeName; \
+	\
+	typedef ::Metasound::TDataReadReference<DataType> DataReadReferenceTypeName; \
+	typedef ::Metasound::TDataWriteReference<DataType> DataWriteReferenceTypeName;
+
+/** Macro to make declaring a metasound parameter simple.  */
+// Declares a metasound parameter type by
+// - Adding typedefs for commonly used template types.
+// - Defining parameter type traits.
+#define DECLARE_METASOUND_DATA_REFERENCE_TYPES(DataType, ModuleApi, DataTypeInfoTypeName, DataReadReferenceTypeName, DataWriteReferenceTypeName) \
+	DECLARE_METASOUND_DATA_REFERENCE_CORE_TYPE(DataType, ModuleApi); \
+	DECLARE_METASOUND_DATA_REFERENCE_ALIAS_TYPES(DataType, DataTypeInfoTypeName, DataReadReferenceTypeName, DataWriteReferenceTypeName); \
+	DECLARE_METASOUND_DATA_REFERENCE_CORE_TYPE(::Metasound::TVariable<DataType>, ModuleApi); \
+	DECLARE_METASOUND_DATA_REFERENCE_CORE_TYPE(::Metasound::TDelayedVariable<DataType>, ModuleApi); \
+	DECLARE_METASOUND_DATA_REFERENCE_CORE_TYPE(TArray<DataType>, ModuleApi); \
+	DECLARE_METASOUND_DATA_REFERENCE_CORE_TYPE(::Metasound::TVariable<TArray<DataType>>, ModuleApi); \
+	DECLARE_METASOUND_DATA_REFERENCE_CORE_TYPE(::Metasound::TDelayedVariable<TArray<DataType>>, ModuleApi);
+
+
+#define DEFINE_METASOUND_DATA_REFERENCE_CORE_TYPE(DataType, DataTypeName, DataTypeLoctextKey) \
+	const TCHAR* ::Metasound::TDataReferenceTypeInfo<DataType>::TypeName = TEXT(DataTypeName); \
+	const FText& ::Metasound::TDataReferenceTypeInfo<DataType>::GetTypeDisplayText() \
+	{ \
+		static const FText DisplayText = NSLOCTEXT("MetaSoundCore_DataReference", DataTypeLoctextKey, DataTypeName); \
+		return DisplayText; \
+	} \
+	const DataType* const ::Metasound::TDataReferenceTypeInfo<DataType>::TypePtr = nullptr; \
+	const void* const ::Metasound::TDataReferenceTypeInfo<DataType>::TypeId = static_cast<const FMetasoundDataTypeId>(&::Metasound::TDataReferenceTypeInfo<DataType>::TypePtr);
+
+
+// This only needs to be called if you don't plan on calling REGISTER_METASOUND_DATATYPE.
+#define DEFINE_METASOUND_DATA_TYPE(DataType, DataTypeName) \
+	DEFINE_METASOUND_DATA_REFERENCE_CORE_TYPE(DataType, DataTypeName, DataTypeName); \
+	DEFINE_METASOUND_DATA_REFERENCE_CORE_TYPE(::Metasound::TVariable<DataType>, DataTypeName":Variable", DataTypeName"_Variable"); \
+	DEFINE_METASOUND_DATA_REFERENCE_CORE_TYPE(::Metasound::TDelayedVariable<DataType>, DataTypeName":DelayedVariable", DataTypeName"_DelayedVariable"); \
+	DEFINE_METASOUND_DATA_REFERENCE_CORE_TYPE(TArray<DataType>, DataTypeName":Array", DataTypeName"_Array"); \
+	DEFINE_METASOUND_DATA_REFERENCE_CORE_TYPE(::Metasound::TVariable<TArray<DataType>>, DataTypeName":Array:Variable", DataTypeName"_Array_Variable"); \
+	DEFINE_METASOUND_DATA_REFERENCE_CORE_TYPE(::Metasound::TDelayedVariable<TArray<DataType>>, DataTypeName":Array:DelayedVariable", DataTypeName"_Array_DelayedVariable"); 
+
diff --git a/Engine/Plugins/Runtime/Metasound/Source/MetasoundGraphCore/Public/MetasoundLiteralNode.h b/Engine/Plugins/Runtime/Metasound/Source/MetasoundGraphCore/Public/MetasoundLiteralNode.h
new file mode 100644
--- /dev/null
+++ b/Engine/Plugins/Runtime/Metasound/Source/MetasoundGraphCore/Public/MetasoundLiteralNode.h
@@ -0,0 +1,175 @@
+// Copyright Epic Games, Inc. All Rights Reserved.
+#pragma once
+
+#include "CoreMinimal.h"
+#include "Internationalization/Text.h"
+#include "MetasoundBuilderInterface.h"
+#include "MetasoundExecutableOperator.h"
+#include "MetasoundNode.h"
+#include "MetasoundNodeInterface.h"
+#include "MetasoundOperatorInterface.h"
+#include "MetasoundVertex.h"
+
+#define LOCTEXT_NAMESPACE "MetasoundGraphCore"
+
+
+namespace Metasound
+{
+	namespace LiteralNodeNames
+	{
+		METASOUNDGRAPHCORE_API const TCHAR* GetOutputDataName();
+	}
+
+	template<typename DataType>
+	class TLiteralOperator : public IOperator
+	{
+	public:
+		using FDataWriteReference = TDataWriteReference<DataType>;
+
+		TLiteralOperator(FDataWriteReference InDataReference)
+			// Executable DataTypes require a copy of the output to operate on whereas non-executable
+			// types do not. Avoid copy by assigning to reference for non-executable types.
+			: InputValue(InDataReference)
+			, OutputValue(TExecutableDataType<DataType>::bIsExecutable ? FDataWriteReference::CreateNew(*InDataReference) : InDataReference)
+		{
+		}
+
+		virtual ~TLiteralOperator() = default;
+
+		virtual FDataReferenceCollection GetInputs() const override
+		{
+			FDataReferenceCollection Inputs;
+			return Inputs;
+		}
+
+		virtual FDataReferenceCollection GetOutputs() const override
+		{
+			FDataReferenceCollection Outputs;
+			Outputs.AddDataReadReference<DataType>(LiteralNodeNames::GetOutputDataName(), OutputValue);
+			return Outputs;
+		}
+
+		void Execute()
+		{
+			TExecutableDataType<DataType>::Execute(*InputValue, *OutputValue);
+		}
+
+		static void ExecuteFunction(IOperator* InOperator)
+		{
+			static_cast<TLiteralOperator<DataType>*>(InOperator)->Execute();
+		}
+
+		virtual FExecuteFunction GetExecuteFunction() override
+		{
+			if (TExecutableDataType<DataType>::bIsExecutable)
+			{
+				return &TLiteralOperator<DataType>::ExecuteFunction;
+			}
+			return nullptr;
+		}
+
+	private:
+		FDataWriteReference InputValue;
+		FDataWriteReference OutputValue;
+	};
+
+
+	/** TLiteralOperatorLiteralFactory creates an input by passing it a literal. */
+	template<typename DataType>
+	class TLiteralOperatorLiteralFactory : public IOperatorFactory
+	{
+	public:
+
+		using FDataWriteReference = TDataWriteReference<DataType>;
+
+		TLiteralOperatorLiteralFactory(FLiteral&& InInitParam)
+			: InitParam(MoveTemp(InInitParam))
+		{
+		}
+
+		virtual TUniquePtr<IOperator> CreateOperator(const FCreateOperatorParams& InParams, FBuildErrorArray& OutErrors) override;
+
+	private:
+		FLiteral InitParam;
+	};
+
+	/** TLiteralNode represents a variable within a metasound graph. */
+	template<typename DataType>
+	class TLiteralNode : public FNode
+	{
+	public:
+
+		static FVertexInterface DeclareVertexInterface()
+		{
+			return FVertexInterface(
+				FInputVertexInterface(),
+				FOutputVertexInterface(
+					TOutputDataVertexModel<DataType>(LiteralNodeNames::GetOutputDataName(), FText::GetEmpty())
+				)
+			);
+		}
+
+		static FNodeClassMetadata GetNodeInfo()
+		{
+			FNodeClassMetadata Info;
+
+			Info.ClassName = {"Literal", GetMetasoundDataTypeName<DataType>(), ""};
+			Info.MajorVersion = 1;
+			Info.MinorVersion = 0;
+			Info.DisplayName = FText::Format(LOCTEXT("Metasound_LiteralNodeDisplayNameFormat", "Literal {0}"), FText::FromName(GetMetasoundDataTypeName<DataType>()));
+			Info.Description = LOCTEXT("Metasound_LiteralNodeDescription", "Literal accessible within a parent Metasound graph.");
+			Info.Author = PluginAuthor;
+			Info.PromptIfMissing = PluginNodeMissingPrompt;
+			Info.DefaultInterface = DeclareVertexInterface();
+
+			return Info;
+		}
+
+		/* Construct a TLiteralNode using the TLiteralOperatorLiteralFactory<> and moving
+		 * InParam to the TLiteralOperatorLiteralFactory constructor.*/
+		explicit TLiteralNode(const FString& InNodeDescription, const FGuid& InInstanceID, FLiteral&& InParam)
+		:	FNode(InNodeDescription, InInstanceID, GetNodeInfo())
+		,	Interface(DeclareVertexInterface())
+		, 	Factory(MakeOperatorFactoryRef<TLiteralOperatorLiteralFactory<DataType>>(MoveTemp(InParam)))
+		{
+		}
+
+		virtual const FVertexInterface& GetVertexInterface() const override
+		{
+			return Interface;
+		}
+
+		virtual bool SetVertexInterface(const FVertexInterface& InInterface) override
+		{
+			return Interface == InInterface;
+		}
+
+		virtual bool IsVertexInterfaceSupported(const FVertexInterface& InInterface) const override
+		{
+			return Interface == InInterface;
+		}
+
+		virtual TSharedRef<IOperatorFactory, ESPMode::ThreadSafe> GetDefaultOperatorFactory() const override
+		{
+			return Factory;
+		}
+
+	private:
+		FVertexInterface Interface;
+		FOperatorFactorySharedRef Factory;
+	};
+
+	template<typename DataType>
+	TUniquePtr<IOperator> TLiteralOperatorLiteralFactory<DataType>::CreateOperator(const FCreateOperatorParams& InParams, FBuildErrorArray& OutErrors)
+	{
+		using FLiteralNodeType = TLiteralNode<DataType>;
+
+		// Create write reference by calling compatible constructor with literal.
+		FDataWriteReference DataRef = TDataWriteReferenceLiteralFactory<DataType>::CreateExplicitArgs(InParams.OperatorSettings, InitParam);
+
+		return MakeUnique<TLiteralOperator<DataType>>(DataRef);
+	}
+
+} // namespace Metasound
+#undef LOCTEXT_NAMESPACE // MetasoundGraphCore
+
diff --git a/Engine/Plugins/Runtime/Metasound/Source/MetasoundGraphCore/Public/MetasoundVariable.h b/Engine/Plugins/Runtime/Metasound/Source/MetasoundGraphCore/Public/MetasoundVariable.h
new file mode 100644
--- /dev/null
+++ b/Engine/Plugins/Runtime/Metasound/Source/MetasoundGraphCore/Public/MetasoundVariable.h
@@ -0,0 +1,95 @@
+// Copyright Epic Games, Inc. All Rights Reserved.
+
+#pragma once
+
+#include "CoreMinimal.h"
+#include "MetasoundDataReference.h"
+#include "MetasoundDataFactory.h"
+
+namespace Metasound
+{
+
+	/** A MetaSound Variable contains a readable data reference.
+	 *
+	 * @tparam DataType - Underlying data type of data reference.
+	 */
+	template<typename DataType>
+	struct TVariable
+	{
+		using FReadReference = TDataReadReference<DataType>;
+
+		TVariable() = delete;
+		TVariable(TVariable&&) = default;
+		TVariable(const TVariable&) = default;
+
+		/** Create a variable with a given data reference. */
+		TVariable(FReadReference InDataReference)
+		: DataReference(InDataReference)
+		{
+		}
+
+		FReadReference GetDataReference() const
+		{
+			return DataReference;
+		}
+
+	private:
+
+		FReadReference DataReference;
+	};
+	
+	/** A MetaSound DelayedVariable contains a data reference's prior and current value.
+	 *
+	 * @tparam DataType - Underlying data type of data reference.
+	 */
+	template<typename DataType>
+	struct TDelayedVariable
+	{
+		using FWriteReference = TDataWriteReference<DataType>;
+		using FReadReference = TDataReadReference<DataType>;
+
+		TDelayedVariable() = delete;
+		TDelayedVariable(TDelayedVariable&&) = default;
+		TDelayedVariable(const TDelayedVariable&) = default;
+
+
+		/** Create a delayed variable.
+		 *
+		 * @param InDelayedReference - A writable reference which will hold the delayed
+		 *                           version of the data reference.
+		 */
+		TDelayedVariable(FWriteReference InDelayedReference)
+		: DelayedDataReference(InDelayedReference)
+		, DataReference(InDelayedReference)
+		{
+		}
+
+		void CopyReferencedData()
+		{
+			*DelayedDataReference = *DataReference;
+		}
+
+		void SetDataReference(FReadReference InDataReference)
+		{
+			DataReference = InDataReference;
+		}
+
+		/** Get the current data reference. */
+		FReadReference GetDataReference() const
+		{
+			return DataReference;
+		}
+
+		/** Get the delayed data reference */
+		FReadReference GetDelayedDataReference() const
+		{
+			return DelayedDataReference;
+		}
+
+	private:
+
+		FWriteReference DelayedDataReference;
+		FReadReference DataReference;
+	};
+}
+
diff --git a/Engine/Plugins/Runtime/Metasound/Source/MetasoundGraphCore/Public/MetasoundVariableNodes.h b/Engine/Plugins/Runtime/Metasound/Source/MetasoundGraphCore/Public/MetasoundVariableNodes.h
new file mode 100644
--- /dev/null
+++ b/Engine/Plugins/Runtime/Metasound/Source/MetasoundGraphCore/Public/MetasoundVariableNodes.h
@@ -0,0 +1,489 @@
+// Copyright Epic Games, Inc. All Rights Reserved.
+
+#pragma once
+
+#include "CoreMinimal.h"
+#include "MetasoundBuilderInterface.h"
+#include "MetasoundDataFactory.h"
+#include "MetasoundExecutableOperator.h"
+#include "MetasoundFacade.h"
+#include "MetasoundLog.h"
+#include "MetasoundNode.h"
+#include "MetasoundNodeInterface.h"
+#include "MetasoundVariable.h"
+
+#define LOCTEXT_NAMESPACE "MetasoundGraphCore"
+
+namespace Metasound
+{
+	namespace VariableNodeVertexNames
+	{
+		/** Input vertex name for data */
+		METASOUNDGRAPHCORE_API const FString& GetInputDataName();
+
+		/** Output vertex name for data */
+		METASOUNDGRAPHCORE_API const FString& GetOutputDataName();
+
+		/** Input vertex name for variables */
+		METASOUNDGRAPHCORE_API const FString& GetInputVariableName();
+		
+		/** Output vertex name for variables */
+		METASOUNDGRAPHCORE_API const FString& GetOutputVariableName();
+	}
+
+	/** Init Variable nodes initialize variable values. The output of a InitVariableNode
+	 * is a TDelayedVariable. This provides access to a variable's value before it is set
+	 * by using the TGetDelayedVariableNode.
+	 */
+	template<typename DataType>
+	class TInitVariableNode : public FNode
+	{
+		using FDelayedVariable = TDelayedVariable<DataType>;
+
+		class FOperator : public TExecutableOperator<FOperator>
+		{
+			using Super = TExecutableOperator<FOperator>;
+
+		public:
+			FOperator(TDataWriteReference<FDelayedVariable> InVariable)
+			: Variable(InVariable)
+			{
+			}
+
+			virtual ~FOperator() = default;
+
+			virtual FDataReferenceCollection GetInputs() const override
+			{
+				FDataReferenceCollection Collection;
+				return Collection;
+			}
+
+			virtual FDataReferenceCollection GetOutputs() const override
+			{
+				FDataReferenceCollection Collection;
+				Collection.AddDataWriteReference(VariableNodeVertexNames::GetOutputVariableName(), Variable);
+				return Collection;
+			}
+
+			void Execute()
+			{
+				Variable->CopyReferencedData();
+			}
+
+		private:
+
+			TDataWriteReference<FDelayedVariable> Variable;
+		};
+
+		class FFactory : public IOperatorFactory
+		{
+		public:
+			FFactory(FLiteral&& InLiteral)
+			: Literal(MoveTemp(InLiteral))
+			{
+			}
+
+			virtual ~FFactory() = default;
+
+			virtual TUniquePtr<IOperator> CreateOperator(const FCreateOperatorParams& InParams, FBuildErrorArray& OutErrors) override
+			{
+				TDataWriteReference<DataType> Data = TDataWriteReferenceLiteralFactory<DataType>::CreateExplicitArgs(InParams.OperatorSettings, Literal);
+				return MakeUnique<FOperator>(TDataWriteReference<FDelayedVariable>::CreateNew(Data));
+			}
+		private:
+			FLiteral Literal;
+		};
+
+
+		static FVertexInterface DeclareVertexInterface()
+		{
+			return FVertexInterface(
+				FInputVertexInterface(
+				),
+				FOutputVertexInterface(
+					TOutputDataVertexModel<FDelayedVariable>(VariableNodeVertexNames::GetOutputVariableName(), FText::GetEmpty())
+				)
+			);
+		}
+
+		static FNodeClassMetadata GetNodeMetadata()
+		{
+			FNodeClassMetadata Info;
+
+			Info.ClassName = {"InitVariable", GetMetasoundDataTypeName<DataType>(), ""};
+			Info.MajorVersion = 1;
+			Info.MinorVersion = 0;
+			Info.DisplayName = FText::Format(LOCTEXT("Metasound_InitVariableNodeDisplayNameFormat", "Init {0} Variable"), FText::FromName(GetMetasoundDataTypeName<DataType>()));
+			Info.Description = LOCTEXT("Metasound_InitVariableNodeDescription", "Initialize a variable of a MetaSound graph.");
+			Info.Author = PluginAuthor;
+			Info.PromptIfMissing = PluginNodeMissingPrompt;
+			Info.DefaultInterface = DeclareVertexInterface();
+			Info.CategoryHierarchy = { LOCTEXT("Metasound_VariableCategory", "Variable") };
+
+			return Info;
+		}
+
+	public:
+
+		TInitVariableNode(const FString& InNodeDescription, const FGuid& InInstanceID, FLiteral&& InLiteral)
+		: FNode(InNodeDescription, InInstanceID, GetNodeMetadata())
+		, Interface(DeclareVertexInterface())
+		, Factory(MakeShared<FFactory>(MoveTemp(InLiteral)))
+		{
+		}
+
+		virtual const FVertexInterface& GetVertexInterface() const override
+		{
+			return Interface;
+		}
+
+		virtual bool SetVertexInterface(const FVertexInterface& InInterface) override
+		{
+			return Interface == InInterface;
+		}
+
+		virtual bool IsVertexInterfaceSupported(const FVertexInterface& InInterface) const override
+		{
+			return Interface == InInterface;
+		}
+
+		virtual TSharedRef<IOperatorFactory, ESPMode::ThreadSafe> GetDefaultOperatorFactory() const override
+		{
+			return Factory;
+		}
+
+	private:
+
+		FVertexInterface Interface;
+		FOperatorFactorySharedRef Factory;
+	};
+
+	/** TSetVariableNode allows variable values to be set. It converts TDelayedVariables
+	 * into TVariables and assigns TDelayedVariable references.
+	 */
+	template<typename DataType>
+	class TSetVariableNode : public FNodeFacade
+	{
+		using FVariable = TVariable<DataType>;
+		using FDelayedVariable = TDelayedVariable<DataType>;
+
+		class FOperator : public FNoOpOperator 
+		{
+		public:
+			static TUniquePtr<IOperator> CreateOperator(const FCreateOperatorParams& InParams, FBuildErrorArray& OutErrors);
+			static FVertexInterface DeclareVertexInterface();
+			static const FNodeClassMetadata& GetNodeInfo();
+
+			FOperator(TDataReadReference<FVariable> InVariable)
+			: Variable(InVariable)
+			{
+			}
+
+			virtual ~FOperator() = default;
+
+			virtual FDataReferenceCollection GetOutputs() const override
+			{
+				FDataReferenceCollection Collection;
+				Collection.AddDataReadReference<FVariable>(VariableNodeVertexNames::GetOutputVariableName(), Variable);
+				return Collection;
+			}
+
+		private:
+			TDataReadReference<FVariable> Variable;
+		};
+
+	public:
+		TSetVariableNode(const FString& InNodeDescription, const FGuid& InInstanceID)
+		: FNodeFacade(InNodeDescription, InInstanceID, TFacadeOperatorClass<FOperator>())
+		{
+		}
+
+		TSetVariableNode(const FNodeInitData& InInitData)
+		: TSetVariableNode(InInitData.InstanceName, InInitData.InstanceID)
+		{}
+	};
+
+	template<typename DataType>
+	TUniquePtr<IOperator> TSetVariableNode<DataType>::FOperator::CreateOperator(const FCreateOperatorParams& InParams, FBuildErrorArray& OutErrors) 
+	{
+		using namespace VariableNodeVertexNames;
+
+		const FDataReferenceCollection& Inputs = InParams.InputDataReferences;
+
+		// Update delayed variable.
+		if (Inputs.ContainsDataWriteReference<FDelayedVariable>(GetInputVariableName()))
+		{
+			TDataWriteReference<FDelayedVariable> Delayed = Inputs.GetDataWriteReference<FDelayedVariable>(GetInputVariableName());
+			if (Inputs.ContainsDataReadReference<DataType>(GetInputDataName()))
+			{
+				// Update the delayed variable with the data reference to copy from.
+				TDataReadReference InputData = Inputs.GetDataReadReference<DataType>(GetInputDataName());
+				Delayed->SetDataReference(InputData);
+			}
+
+			// Create variable with matching reference to delayed variable.
+			TDataReadReference<FVariable> Variable = TDataReadReference<FVariable>::CreateNew(Delayed->GetDataReference());
+			return MakeUnique<FOperator>(Variable);
+		}
+		else if (Inputs.ContainsDataReadReference<DataType>(GetInputDataName()))
+		{
+			// Create Variable without DelayedVariable
+			TDataReadReference InputData = Inputs.GetDataReadReference<DataType>(GetInputDataName());
+			TDataReadReference<FVariable> Variable = TDataReadReference<FVariable>::CreateNew(InputData);
+			return MakeUnique<FOperator>(Variable);
+		}
+		else 
+		{
+			// Nothing to do if there's no input data.
+			return MakeUnique<FNoOpOperator>();
+		}
+	}
+
+	template<typename DataType>
+	FVertexInterface TSetVariableNode<DataType>::FOperator::DeclareVertexInterface()
+	{
+		return FVertexInterface(
+			FInputVertexInterface(
+				TInputDataVertexModel<DataType>(VariableNodeVertexNames::GetInputDataName(), FText::GetEmpty()), 
+				TInputDataVertexModel<FDelayedVariable>(VariableNodeVertexNames::GetInputVariableName(), FText::GetEmpty())
+			),
+			FOutputVertexInterface(
+				TOutputDataVertexModel<FVariable>(VariableNodeVertexNames::GetOutputVariableName(), FText::GetEmpty())
+			)
+		);
+	}
+
+	template<typename DataType>
+	const FNodeClassMetadata& TSetVariableNode<DataType>::FOperator::GetNodeInfo()
+	{
+		auto CreateNodeClassMetadata = []()
+		{
+			FNodeClassMetadata Info;
+
+			Info.ClassName = {"SetVariable", GetMetasoundDataTypeName<DataType>(), ""};
+			Info.MajorVersion = 1;
+			Info.MinorVersion = 0;
+			Info.DisplayName = FText::Format(LOCTEXT("Metasound_SetVariableNodeDisplayNameFormat", "Set {0}"), GetMetasoundDataTypeDisplayText<DataType>());
+			Info.Description = LOCTEXT("Metasound_SetVariableNodeDescription", "Set variable on MetaSound graph.");
+			Info.Author = PluginAuthor;
+			Info.PromptIfMissing = PluginNodeMissingPrompt;
+			Info.DefaultInterface = DeclareVertexInterface();
+			Info.CategoryHierarchy = { LOCTEXT("Metasound_VariableCategory", "Variable") };
+
+			return Info;
+		};
+
+		static const FNodeClassMetadata Metadata = CreateNodeClassMetadata();
+		return Metadata;
+	}
+
+	/** TGetDelayedVariableNode provides access to the prior executions variable value.
+	 * TGetDelayedVariableNodes must always be before TSetVariableNodes in the dependency
+	 * order.
+	 */
+	template<typename DataType>
+	class TGetDelayedVariableNode : public FNodeFacade
+	{
+		using FDelayedVariable = TDelayedVariable<DataType>;
+
+		class FOperator : public FNoOpOperator 
+		{
+		public:
+			static TUniquePtr<IOperator> CreateOperator(const FCreateOperatorParams& InParams, FBuildErrorArray& OutErrors);
+			static FVertexInterface DeclareVertexInterface();
+			static const FNodeClassMetadata& GetNodeInfo();
+
+			FOperator(TDataWriteReference<FDelayedVariable> InDelayedVariable)
+			: DelayedVariable(InDelayedVariable)
+			, DelayedData(DelayedVariable->GetDelayedDataReference())
+			{
+			}
+
+			virtual ~FOperator() = default;
+
+			virtual FDataReferenceCollection GetOutputs() const override
+			{
+				FDataReferenceCollection Collection;
+				Collection.AddDataWriteReference<FDelayedVariable>(VariableNodeVertexNames::GetOutputVariableName(), DelayedVariable);
+				Collection.AddDataReadReference<DataType>(VariableNodeVertexNames::GetOutputDataName(), DelayedVariable->GetDelayedDataReference());
+				return Collection;
+			}
+
+		private:
+			TDataWriteReference<FDelayedVariable> DelayedVariable;
+			TDataReadReference<DataType> DelayedData;
+		};
+
+	public:
+
+		TGetDelayedVariableNode(const FString& InNodeDescription, const FGuid& InInstanceID)
+		: FNodeFacade(InNodeDescription, InInstanceID, TFacadeOperatorClass<FOperator>())
+		{
+		}
+
+		TGetDelayedVariableNode(const FNodeInitData& InInitData)
+		: TGetDelayedVariableNode(InInitData.InstanceName, InInitData.InstanceID)
+		{}
+	};
+
+	template<typename DataType>
+	TUniquePtr<IOperator> TGetDelayedVariableNode<DataType>::FOperator::CreateOperator(const FCreateOperatorParams& InParams, FBuildErrorArray& OutErrors) 
+	{
+		using namespace VariableNodeVertexNames;
+
+		const FDataReferenceCollection& Inputs = InParams.InputDataReferences;
+
+		if (ensure(Inputs.ContainsDataWriteReference<FDelayedVariable>(GetInputVariableName())))
+		{
+			TDataWriteReference<FDelayedVariable> DelayedVariable = Inputs.GetDataWriteReference<FDelayedVariable>(GetInputVariableName());
+
+			return MakeUnique<FOperator>(DelayedVariable);
+		}
+		else 
+		{
+			// Nothing to do if there's no input data.
+			UE_LOG(LogMetaSound, Warning, TEXT("Missing internal variable connection. Failed to create valid \"GetDelayedVariable\" operator"));
+			return MakeUnique<FNoOpOperator>();
+		}
+	}
+
+	template<typename DataType>
+	FVertexInterface TGetDelayedVariableNode<DataType>::FOperator::DeclareVertexInterface()
+	{
+		return FVertexInterface(
+			FInputVertexInterface(
+				TInputDataVertexModel<FDelayedVariable>(VariableNodeVertexNames::GetInputVariableName(), FText::GetEmpty())
+			),
+			FOutputVertexInterface(
+				TOutputDataVertexModel<FDelayedVariable>(VariableNodeVertexNames::GetOutputVariableName(), FText::GetEmpty()),
+				TOutputDataVertexModel<DataType>(VariableNodeVertexNames::GetOutputDataName(), FText::GetEmpty())
+			)
+		);
+	}
+
+	template<typename DataType>
+	const FNodeClassMetadata& TGetDelayedVariableNode<DataType>::FOperator::GetNodeInfo()
+	{
+		auto CreateNodeClassMetadata = []()
+		{
+			FNodeClassMetadata Info;
+
+			Info.ClassName = {"GetDelayedVariable", GetMetasoundDataTypeName<DataType>(), ""};
+			Info.MajorVersion = 1;
+			Info.MinorVersion = 0;
+			Info.DisplayName = FText::Format(LOCTEXT("Metasound_GetDelayedVariableNodeDisplayNameFormat", "Get Delayed {0}"), GetMetasoundDataTypeDisplayText<DataType>());
+			Info.Description = LOCTEXT("Metasound_GetDelayedVariableNodeDescription", "Get a delayed variable on MetaSound graph.");
+			Info.Author = PluginAuthor;
+			Info.PromptIfMissing = PluginNodeMissingPrompt;
+			Info.DefaultInterface = DeclareVertexInterface();
+			Info.CategoryHierarchy = { LOCTEXT("Metasound_VariableCategory", "Variable") };
+
+			return Info;
+		};
+
+		static const FNodeClassMetadata Metadata = CreateNodeClassMetadata();
+		return Metadata;
+	}
+
+	/** FGetVariable node provides delay free, cpu free access to a set variable. */
+	template<typename DataType>
+	class TGetVariableNode : public FNodeFacade
+	{
+		using FVariable = TVariable<DataType>;
+
+		class FOperator : public FNoOpOperator 
+		{
+		public:
+			static TUniquePtr<IOperator> CreateOperator(const FCreateOperatorParams& InParams, FBuildErrorArray& OutErrors);
+			static FVertexInterface DeclareVertexInterface();
+			static const FNodeClassMetadata& GetNodeInfo();
+
+			FOperator(TDataReadReference<DataType> InData)
+			: Data(InData)
+			{
+			}
+
+			virtual ~FOperator() = default;
+
+			virtual FDataReferenceCollection GetOutputs() const override
+			{
+				FDataReferenceCollection Collection;
+				Collection.AddDataReadReference<DataType>(VariableNodeVertexNames::GetOutputDataName(), Data);
+				return Collection;
+			}
+
+		private:
+			TDataReadReference<DataType> Data;
+		};
+	public:
+
+		TGetVariableNode(const FString& InNodeDescription, const FGuid& InInstanceID)
+		: FNodeFacade(InNodeDescription, InInstanceID, TFacadeOperatorClass<FOperator>())
+		{
+		}
+
+		TGetVariableNode(const FNodeInitData& InInitData)
+		: TGetVariableNode(InInitData.InstanceName, InInitData.InstanceID)
+		{}
+	};
+
+	template<typename DataType>
+	TUniquePtr<IOperator> TGetVariableNode<DataType>::FOperator::CreateOperator(const FCreateOperatorParams& InParams, FBuildErrorArray& OutErrors) 
+	{
+		using namespace VariableNodeVertexNames;
+
+		const FDataReferenceCollection& Inputs = InParams.InputDataReferences;
+
+		// Update delayed variable.
+		if (Inputs.ContainsDataReadReference<FVariable>(GetInputVariableName()))
+		{
+			TDataReadReference<FVariable> Variable = Inputs.GetDataReadReference<FVariable>(GetInputVariableName());
+			return MakeUnique<FOperator>(Variable->GetDataReference());
+		}
+		else 
+		{
+			// Nothing to do if there's no input data.
+			UE_LOG(LogMetaSound, Warning, TEXT("Missing internal variable connection. Failed to create valid \"GetVariable\" operator"));
+			return MakeUnique<FNoOpOperator>();
+		}
+	}
+
+	template<typename DataType>
+	FVertexInterface TGetVariableNode<DataType>::FOperator::DeclareVertexInterface()
+	{
+		return FVertexInterface(
+			FInputVertexInterface(
+				TInputDataVertexModel<FVariable>(VariableNodeVertexNames::GetInputVariableName(), FText::GetEmpty())
+			),
+			FOutputVertexInterface(
+				TOutputDataVertexModel<DataType>(VariableNodeVertexNames::GetOutputDataName(), FText::GetEmpty())
+			)
+		);
+	}
+
+	template<typename DataType>
+	const FNodeClassMetadata& TGetVariableNode<DataType>::FOperator::GetNodeInfo()
+	{
+		auto CreateNodeClassMetadata = []()
+		{
+			FNodeClassMetadata Info;
+
+			Info.ClassName = {"GetVariable", GetMetasoundDataTypeName<DataType>(), ""};
+			Info.MajorVersion = 1;
+			Info.MinorVersion = 0;
+			Info.DisplayName = FText::Format(LOCTEXT("Metasound_GetVariableNodeDisplayNameFormat", "Get {0}"), GetMetasoundDataTypeDisplayText<DataType>());
+			Info.Description = LOCTEXT("Metasound_GetVariableNodeDescription", "Get variable on MetaSound graph.");
+			Info.Author = PluginAuthor;
+			Info.PromptIfMissing = PluginNodeMissingPrompt;
+			Info.DefaultInterface = DeclareVertexInterface();
+
+			Info.CategoryHierarchy = { LOCTEXT("Metasound_VariableCategory", "Variable") };
+			return Info;
+		};
+
+		static const FNodeClassMetadata Metadata = CreateNodeClassMetadata();
+		return Metadata;
+	}
+}
+
+#undef LOCTEXT_NAMESPACE // MetasoundGraphCore
