Fix MetaSound Engine Test

diff --git a/Engine/Plugins/Runtime/Metasound/Metasound.uplugin b/Engine/Plugins/Runtime/Metasound/Metasound.uplugin
--- a/Engine/Plugins/Runtime/Metasound/Metasound.uplugin
+++ b/Engine/Plugins/Runtime/Metasound/Metasound.uplugin
@@ -40,6 +40,11 @@
 			"Type": "Runtime",
 			"LoadingPhase": "Default"
 		},
+		{
+			"Name": "MetasoundEngineTest",
+			"Type": "Runtime",
+			"LoadingPhase": "Default"
+		},
 		{
 			"Name": "MetasoundEditor",
 			"Type": "UncookedOnly",
diff --git a/Engine/Plugins/Runtime/Metasound/Source/MetasoundEngineTest/MetasoundEngineTest.Build.cs b/Engine/Plugins/Runtime/Metasound/Source/MetasoundEngineTest/MetasoundEngineTest.Build.cs
new file mode 100644
--- /dev/null
+++ b/Engine/Plugins/Runtime/Metasound/Source/MetasoundEngineTest/MetasoundEngineTest.Build.cs
@@ -0,0 +1,36 @@
+// Copyright Epic Games, Inc. All Rights Reserved.
+
+using UnrealBuildTool;
+using System.IO;
+
+namespace UnrealBuildTool.Rules
+{
+	public class MetasoundEngineTest : ModuleRules
+	{
+		public MetasoundEngineTest(ReadOnlyTargetRules Target) : base(Target)
+		{
+			PrivateDependencyModuleNames.AddRange
+			(
+				new string[]
+				{
+					"Core",
+					"CoreUObject",
+					"Engine",
+					"MetasoundGraphCore",
+					"MetasoundGenerator",
+					"MetasoundFrontend",
+					"MetasoundEngine",
+					"SignalProcessing",
+					"Projects"
+				}
+			);
+
+			PublicDependencyModuleNames.AddRange
+			(
+				new string[]
+				{
+				}
+			);
+		}
+	}
+}
diff --git a/Engine/Plugins/Runtime/Metasound/Source/MetasoundEngineTest/Private/EngineTestMetasoundSource.cpp b/Engine/Plugins/Runtime/Metasound/Source/MetasoundEngineTest/Private/EngineTestMetasoundSource.cpp
new file mode 100644
--- /dev/null
+++ b/Engine/Plugins/Runtime/Metasound/Source/MetasoundEngineTest/Private/EngineTestMetasoundSource.cpp
@@ -0,0 +1,205 @@
+// Copyright Epic Games, Inc. All Rights Reserved.
+
+#include "AudioDevice.h"
+#include "Components/AudioComponent.h"
+#include "Interfaces/IPluginManager.h"
+#include "MetasoundFrontendController.h"
+#include "MetasoundFrontendSearchEngine.h"
+#include "MetasoundSource.h"
+#include "Misc/AutomationTest.h"
+#include "Tests/AutomationCommon.h"
+
+#if WITH_DEV_AUTOMATION_TESTS
+
+namespace EngineTestMetasoundSourcePrivate
+{
+	static FString GetPluginContentDirectory()
+	{
+		TSharedPtr<IPlugin> Plugin = IPluginManager::Get().FindPlugin(TEXT("Metasound"));
+		if (ensure(Plugin.IsValid()))
+		{
+			return Plugin->GetContentDir();
+		}
+		return FString();
+	}
+	static FString GetPathToTestFilesDir()
+	{
+		FString OutPath =  FPaths::Combine(GetPluginContentDirectory(), TEXT("Test"));
+
+		OutPath = FPaths::ConvertRelativePathToFull(OutPath);
+		FPaths::NormalizeDirectoryName(OutPath);
+		
+		return OutPath;
+	}
+
+	static FString GetPathToGeneratedFilesDir()
+	{
+		FString OutPath = FPaths::Combine(FPaths::ProjectSavedDir(), TEXT("Metasounds"));
+
+		OutPath = FPaths::ConvertRelativePathToFull(OutPath);
+		FPaths::NormalizeDirectoryName(OutPath);
+		
+		return OutPath;
+	}
+
+	static FString GetPathToGeneratedAssetsDir()
+	{
+		FString OutPath = TEXT("/Game/Metasound/Generated/");
+		FPaths::NormalizeDirectoryName(OutPath);
+		return OutPath;
+	}
+
+	Metasound::Frontend::FNodeHandle AddNode(Metasound::Frontend::IGraphController& InGraph, const Metasound::FNodeClassName& InClassName)
+	{
+		FMetasoundFrontendClass NodeClass;
+		ensure(Metasound::Frontend::ISearchEngine::Get().FindClassWithHighestVersion(InClassName, NodeClass));
+		Metasound::Frontend::FNodeHandle Node = InGraph.AddNode(NodeClass.Metadata);
+		check(Node->IsValid());
+		return Node;
+	}
+
+	FMetasoundFrontendDocument CreateMetaSoundMonoSourceDocument()
+	{
+		using namespace Metasound;
+		using namespace Metasound::Frontend;
+
+		FMetasoundFrontendDocument Document;
+		Document.ArchetypeVersion = FMetasoundFrontendVersion{"MonoSource", {1, 1}};
+
+		Document.RootGraph.Metadata.SetClassName(FMetasoundFrontendClassName { "Namespace", "Unit Test Node", *LexToString(FGuid::NewGuid()) });
+		Document.RootGraph.Metadata.SetType(EMetasoundFrontendClassType::Graph);
+
+		FDocumentHandle DocumentHandle = IDocumentController::CreateDocumentHandle(Document);
+		Metasound::Frontend::FGraphHandle RootGraph = DocumentHandle->GetRootGraph();
+		check(RootGraph->IsValid());
+
+		// Input on play
+		FMetasoundFrontendClassInput OnPlayInput;
+		OnPlayInput.Name = TEXT("On Play");
+		OnPlayInput.TypeName = GetMetasoundDataTypeName<FTrigger>();
+		OnPlayInput.VertexID = FGuid::NewGuid();
+
+		FNodeHandle OnPlayInputNode = RootGraph->AddInputVertex(OnPlayInput);
+		check(OnPlayInputNode->IsValid());
+
+		// Input Frequency
+		FMetasoundFrontendClassInput FrequencyInput;
+		FrequencyInput.Name = TEXT("Frequency");
+		FrequencyInput.TypeName = GetMetasoundDataTypeName<float>();
+		FrequencyInput.VertexID = FGuid::NewGuid();
+		FrequencyInput.DefaultLiteral.Set(100.f);
+
+		FNodeHandle FrequencyInputNode = RootGraph->AddInputVertex(FrequencyInput);
+		check(FrequencyInputNode->IsValid());
+
+
+		// Output Is Finished
+		FMetasoundFrontendClassOutput IsFinishedOutput;
+		IsFinishedOutput.Name = TEXT("On Finished");
+		IsFinishedOutput.TypeName = Metasound::GetMetasoundDataTypeName<FTrigger>();
+		IsFinishedOutput.VertexID = FGuid::NewGuid();
+		FNodeHandle IsFinishedOutputNode = RootGraph->AddOutputVertex(IsFinishedOutput);
+		check(IsFinishedOutputNode->IsValid());
+
+		// Output Audio
+		FMetasoundFrontendClassOutput AudioOutput;
+		AudioOutput.Name = TEXT("Audio:0");
+		AudioOutput.TypeName = Metasound::GetMetasoundDataTypeName<FAudioBuffer>();
+		AudioOutput.VertexID = FGuid::NewGuid();
+		
+		FNodeHandle OutputAudioNode = RootGraph->AddOutputVertex(AudioOutput);
+		check(OutputAudioNode->IsValid());
+
+		// osc node 
+		FNodeHandle OscNode = AddNode(*RootGraph, {TEXT("UE"), TEXT("Sine"), TEXT("Audio")});
+
+		// Make connections:
+
+		// frequency input "Frequency" -> oscillator "Frequency"
+		FOutputHandle OutputToConnect = FrequencyInputNode->GetOutputsWithVertexName(FString(TEXT("Frequency")))[0];
+		FInputHandle InputToConnect = OscNode->GetInputsWithVertexName(FString(TEXT("Frequency")))[0];
+		ensure(InputToConnect->Connect(*OutputToConnect));
+
+		// oscillator to output
+		OutputToConnect = OscNode->GetOutputsWithVertexName(FString(TEXT("Audio")))[0];
+		InputToConnect = OutputAudioNode->GetInputsWithVertexName(TEXT("Audio:0"))[0];
+		ensure(InputToConnect->Connect(*OutputToConnect));
+
+		return Document;
+	}
+}
+
+DEFINE_LATENT_AUTOMATION_COMMAND_ONE_PARAMETER(FAudioComponentPlayLatentCommand, UAudioComponent*, AudioComponent);
+
+bool FAudioComponentPlayLatentCommand::Update()
+{
+	if (AudioComponent)
+	{
+		AudioComponent->Play();
+		return true;
+	}
+	return false;
+}
+
+DEFINE_LATENT_AUTOMATION_COMMAND_ONE_PARAMETER(FAudioComponentStopLatentCommand, UAudioComponent*, AudioComponent);
+
+bool FAudioComponentStopLatentCommand::Update()
+{
+	if (AudioComponent)
+	{
+		AudioComponent->Stop();
+		return true;
+	}
+	return false;
+}
+
+DEFINE_LATENT_AUTOMATION_COMMAND_ONE_PARAMETER(FAudioComponentRemoveFromRootLatentCommand, UAudioComponent*, AudioComponent);
+
+bool FAudioComponentRemoveFromRootLatentCommand::Update()
+{
+	if (AudioComponent)
+	{
+		AudioComponent->RemoveFromRoot();
+		return true;
+	}
+	return false;
+}
+
+// This test imports a metasound source from a JSON file, attempts to play it, and exports it as a json.
+IMPLEMENT_SIMPLE_AUTOMATION_TEST(FAudioMetasoundSourceTest, "Audio.Metasound.BuildAndPlayMetasoundSource", EAutomationTestFlags::EditorContext | EAutomationTestFlags::EngineFilter)
+bool FAudioMetasoundSourceTest::RunTest(const FString& Parameters)
+{
+	UMetaSoundSource* MetaSoundSource = NewObject<UMetaSoundSource>(GetTransientPackage(), FName(*LexToString(FGuid::NewGuid())));;
+	if (ensure(nullptr != MetaSoundSource))
+	{
+		MetaSoundSource->SetDocument(EngineTestMetasoundSourcePrivate::CreateMetaSoundMonoSourceDocument());
+
+		if (FAudioDevice* AudioDevice = GEngine->GetMainAudioDeviceRaw())
+		{
+
+			UAudioComponent* AudioComponent = FAudioDevice::CreateComponent(MetaSoundSource);
+			AddErrorIfFalse(AudioComponent != nullptr, "Couldn't create audio component!");
+
+			if (AudioComponent)
+			{
+				AudioComponent->bIsUISound = true;
+				AudioComponent->bAllowSpatialization = false;
+				AudioComponent->SetVolumeMultiplier(1.0f);
+				AudioComponent->AddToRoot();
+
+				ADD_LATENT_AUTOMATION_COMMAND(FAudioComponentPlayLatentCommand(AudioComponent));
+				ADD_LATENT_AUTOMATION_COMMAND(FEngineWaitLatentCommand(2.f));
+				ADD_LATENT_AUTOMATION_COMMAND(FAudioComponentStopLatentCommand(AudioComponent));
+				ADD_LATENT_AUTOMATION_COMMAND(FEngineWaitLatentCommand(0.5f));
+				ADD_LATENT_AUTOMATION_COMMAND(FAudioComponentRemoveFromRootLatentCommand(AudioComponent));
+			}
+		}
+	}
+
+	return true;
+ }
+
+// TODO: add separate JSON test.
+
+
+#endif //WITH_DEV_AUTOMATION_TESTS
diff --git a/Engine/Plugins/Runtime/Metasound/Source/MetasoundEngineTest/Private/EngineTestMetasoundWavePlayer.cpp b/Engine/Plugins/Runtime/Metasound/Source/MetasoundEngineTest/Private/EngineTestMetasoundWavePlayer.cpp
new file mode 100644
--- /dev/null
+++ b/Engine/Plugins/Runtime/Metasound/Source/MetasoundEngineTest/Private/EngineTestMetasoundWavePlayer.cpp
@@ -0,0 +1,273 @@
+// Copyright Epic Games, Inc. All Rights Reserved.
+
+#include "CoreMinimal.h"
+#include "MetasoundEnvironment.h"
+#include "MetasoundFrontend.h"
+#include "MetasoundSource.h"
+#include "Kismet/GameplayStatics.h"
+#include "Components/AudioComponent.h"
+#include "AudioDevice.h"
+#include "Misc/AutomationTest.h"
+
+#include "Sound/SoundWave.h"
+
+#include "MetasoundWave.h"
+#include "MetasoundAudioBuffer.h"
+
+#include "IAudioCodec.h"
+#include "IAudioCodecRegistry.h"
+
+#if WITH_DEV_AUTOMATION_TESTS
+
+// Disabled for now until waveplayer implementation is extracted and place in a better location for testing
+// (as nodes are moving to headerless implementation to better support hotfixing between engine versions.)
+#if 0
+IMPLEMENT_SIMPLE_AUTOMATION_TEST(FAudioMetasoundWavePlayer_PlaySoundWave, "Audio.Metasound.WavePlayer.PlaySoundWave", EAutomationTestFlags::EditorContext | EAutomationTestFlags::EngineFilter)
+bool FAudioMetasoundWavePlayer_PlaySoundWave::RunTest(const FString& Parameters)
+{	
+	const TCHAR* TestAsset = TEXT("SoundWave'/Game/Tests/Audio/AttenuationFocusAzimuth/FocusAzimuthPinkNoise.FocusAzimuthPinkNoise'");
+	FSoftObjectPath SoftPath(TestAsset);
+	USoundWave* TestWave = Cast<USoundWave>(SoftPath.TryLoad());
+
+#if WITH_EDITOR
+	TestWave->CachePlatformData();
+#endif //WITH_EDITOR
+
+	AddErrorIfFalse(TestWave != nullptr, TEXT("Failed to load a SoundWave"));
+
+	// Make a copy of the expected PCM.
+	TArray<float> ExpectedFloat;
+	if (const void* pExpected = TestWave->RawData.LockReadOnly() )
+	{
+		int32 NumSamples = TestWave->RawData.GetBulkDataSize()/sizeof(int16);
+		ExpectedFloat.SetNum(NumSamples);
+	
+		const int16* Src16 = (const int16*) pExpected;
+		for(int32 i = 0; i < NumSamples; ++i)
+		{
+			ExpectedFloat[i] = (float)Src16[i] / 32768.f;
+		}		
+		TestWave->RawData.Unlock();
+	}
+	
+	using namespace Metasound;
+	using namespace Audio;
+	FWavePlayerNode Node(TEXT("WavePlayer"), FGuid::NewGuid());
+
+	FOperatorFactorySharedRef Factory = Node.GetDefaultOperatorFactory();
+
+	Metasound::FOperatorSettings OperatorSettings(48000.f, 1024);
+	Metasound::FMetasoundEnvironment Environment;
+	Metasound::FDataReferenceCollection DataReferenceCollection;
+
+	Audio::FProxyDataInitParams Params { TEXT("WavePlayer") };
+	TUniquePtr<IProxyData> WaveProxy = TestWave->CreateNewProxyData(Params);
+	Metasound::FWaveAssetWriteRef Wave = Metasound::FWaveAssetWriteRef::CreateNew(WaveProxy);	
+	DataReferenceCollection.AddDataReadReference(TEXT("Wave"), Wave);	
+
+	TArray<TUniquePtr<Metasound::IOperatorBuildError>> BuildErrors;
+
+	TUniquePtr<Metasound::IOperator> Operator = Factory->CreateOperator({Node, OperatorSettings, DataReferenceCollection, Environment}, BuildErrors);
+	UTEST_TRUE("Made the operator", Operator.IsValid());
+	UTEST_TRUE("We should have zero errors", BuildErrors.Num() == 0);
+	   
+	//THEN("It should produce some valid PCM output")
+	{
+
+		Metasound::FDataReferenceCollection OutputCollection = Operator->GetOutputs();
+		Metasound::IOperator::FExecuteFunction Func = Operator->GetExecuteFunction();
+
+		UTEST_TRUE("Has Audio Output", OutputCollection.ContainsDataReadReference<Metasound::FAudioBuffer>(TEXT("AudioLeft")));
+
+		if (OutputCollection.ContainsDataReadReference<Metasound::FAudioBuffer>(TEXT("AudioLeft")))
+		{
+			Metasound::FAudioBufferReadRef Buffer = OutputCollection.GetDataReadReference<Metasound::FAudioBuffer>(TEXT("AudioLeft"));
+			UTEST_NOT_NULL("Buffer is not null", Buffer->GetData());
+
+			const float* Data = Buffer->GetData();
+			float Phase = 0.0f;
+			const float Frequency = 200.f;
+
+			const int32 NumBuffers = (TestWave->RawData.GetBulkDataSize() / TestWave->NumChannels / sizeof(int16)) /
+				OperatorSettings.GetNumFramesPerBlock();
+			
+			int32 ExpectedPos = 0;
+			for (int32 j = 0; j < NumBuffers; j++)
+			{
+				Invoke(Func, Operator.Get());
+
+				/*for (int32 i = 0; i < OperatorSettings.GetNumFramesPerBlock(); ++i)
+				{
+					UTEST_EQUAL_TOLERANCE("Compare signals are the same-ish-sorta",
+						Data[i],
+						ExpectedFloat[i],
+						0.01f);
+				}*/
+
+				ExpectedPos += OperatorSettings.GetNumFramesPerBlock();
+			}
+		}
+	}
+
+	return true;
+}
+
+
+IMPLEMENT_SIMPLE_AUTOMATION_TEST(FAudioMetasoundWavePlayer_DefaultMakesSilence, "Audio.Metasound.WavePlayer.DefaultMakesSilence", EAutomationTestFlags::EditorContext | EAutomationTestFlags::EngineFilter)
+bool FAudioMetasoundWavePlayer_DefaultMakesSilence::RunTest(const FString& Parameters)
+{
+	using namespace Metasound;
+	using namespace Audio;
+	FWavePlayerNode Node(TEXT("WavePlayer"), FGuid::NewGuid());
+	
+	FOperatorFactorySharedRef Factory = Node.GetDefaultOperatorFactory();
+
+	Metasound::FOperatorSettings OperatorSettings(48000.f, 1024);
+	Metasound::FMetasoundEnvironment Environment;
+	Metasound::FDataReferenceCollection DataReferenceCollection;
+	TArray<TUniquePtr<Metasound::IOperatorBuildError>> BuildErrors;
+
+	TUniquePtr<Metasound::IOperator> Operator = Factory->CreateOperator({Node, OperatorSettings, DataReferenceCollection, Environment}, BuildErrors);
+	UTEST_TRUE("Create operator",Operator.IsValid());
+	
+	UTEST_TRUE("Creating without any inputs will log some errors about bad inputs.",BuildErrors.Num() > 0);
+
+	//THEN("It should correctly warn and produce silence")
+	{
+		Metasound::FDataReferenceCollection OutputCollection = Operator->GetOutputs();
+		Metasound::IOperator::FExecuteFunction Func = Operator->GetExecuteFunction();
+		
+		UTEST_TRUE("Has Audio Output",OutputCollection.ContainsDataReadReference<Metasound::FAudioBuffer>(TEXT("AudioLeft")));
+
+		if (OutputCollection.ContainsDataReadReference<Metasound::FAudioBuffer>(TEXT("AudioLeft")))
+		{
+			Metasound::FAudioBufferReadRef Buffer = OutputCollection.GetDataReadReference<Metasound::FAudioBuffer>(TEXT("AudioLeft"));
+			UTEST_TRUE("Output Buffer is valid", nullptr != Buffer->GetData());
+			UTEST_TRUE("Num of Frames per block match input", OperatorSettings.GetNumFramesPerBlock() == Buffer->Num());
+
+			const int32 NumBuffers = 2;
+
+			const float* Data = Buffer->GetData();
+			for (int32 j = 0; j < NumBuffers; j++)
+			{
+				Invoke(Func, Operator.Get());
+
+				for (int32 i = 0; i < OperatorSettings.GetNumFramesPerBlock(); i++)
+				{
+					// Should all be silent.
+					UTEST_TRUE("", Data[i] == 0.0f);
+				}
+			}
+		}
+	}
+	return true;
+}
+
+/*
+IMPLEMENT_SIMPLE_AUTOMATION_TEST(FAudioMetasoundWavePlayer_EncodeDecodeOsscilator, "Audio.Metasound.WavePlayer.EncodeDecodeOsscilator", EAutomationTestFlags::EditorContext | EAutomationTestFlags::EngineFilter)
+bool FAudioMetasoundWavePlayer_EncodeDecodeOsscilator::RunTest(const FString& Parameters)
+{
+	using namespace Metasound;
+	using namespace Audio;
+	FWavePlayerNode Node(TEXT("WavePlayer"));
+
+	// Make a simple mono sequence in int16
+	TArray<int16> Sequence;
+	constexpr int32 SequenceLength = 48000;
+	Sequence.SetNum(SequenceLength);
+
+	// Oscillator.
+	{
+		const float DefaultFrequency = 200.f;
+		float Phase = 0.f;
+		float Frequency = DefaultFrequency;
+		float Rate = 48000.f;
+
+		for (int32 i = 0; i < SequenceLength; i++)
+		{
+			float OscFloat = FMath::Sin(Phase);
+			Sequence[i] = static_cast<int16>(OscFloat * 32767.f);
+			Phase += 2.0f * PI * Frequency / Rate;
+
+			while (Phase > (2.0f * PI))
+			{
+				Phase -= (2.0f * PI);
+			}
+		}
+	}
+
+	// Find default codec / create encoder / encode signal.
+	IEncoderInput::FFormat Format{ 1, 48000, Int16_Interleaved };
+	ICodecRegistry::FCodecPtr Codec = ICodecRegistry::Get().FindDefaultCodec();
+	TUniquePtr<IEncoderInput> InputObj = IEncoderInput::Create(Sequence, Format);
+	UTEST_NOT_NULL("Default Codec", Codec);
+	UTEST_TRUE("Can Encode",Codec->GetDetails().Features.HasFeature(FCodecFeatures::HasEncoder));
+	UTEST_TRUE("Can Decode",Codec->GetDetails().Features.HasFeature(FCodecFeatures::HasDecoder));
+
+	ICodec::FEncoderPtr Encoder = Codec->CreateEncoder(InputObj.Get());
+	UTEST_TRUE("Encoder was created",Encoder.IsValid());
+
+	TArray<uint8> CompressedBytes;
+	if (Encoder.IsValid())
+	{
+		UTEST_TRUE("Encode passes",Encoder->Encode(InputObj.Get(), IEncoderOutput::Create(CompressedBytes).Get(), nullptr));
+		UTEST_TRUE("Encode produced some data",CompressedBytes.Num() > 0);
+	}
+
+	
+
+	//WHEN("Making a WavePlayer with valid inputs")
+	{
+		FOperatorFactorySharedRef Factory = Node.GetDefaultOperatorFactory();
+
+		Metasound::FOperatorSettings OperatorSettings(48000.f, 1024);
+		Metasound::FDataReferenceCollection DataReferenceCollection;
+
+		Metasound::FWaveAssetWriteRef Wave = Metasound::FWaveAssetWriteRef::CreateNew(CompressedBytes);
+		DataReferenceCollection.AddDataReadReference(TEXT("Wave"), Wave);
+
+		TArray<TUniquePtr<Metasound::IOperatorBuildError>> BuildErrors;
+
+		TUniquePtr<Metasound::IOperator> Operator = Factory->CreateOperator({Node, OperatorSettings, DataReferenceCollection}, BuildErrors);
+		UTEST_TRUE("Made the operator",Operator.IsValid());
+		UTEST_TRUE("We should have zero errors",BuildErrors.Num() == 0);
+
+		//THEN("It should produce some valid PCM output")
+		{
+			Metasound::FDataReferenceCollection OutputCollection = Operator->GetOutputs();
+			Metasound::IOperator::FExecuteFunction Func = Operator->GetExecuteFunction();
+			Metasound::FAudioBufferReadRef Buffer = OutputCollection.GetDataReadReference<Metasound::FAudioBuffer>(TEXT("Audio"));
+			UTEST_NOT_NULL("Buffer is not null",Buffer->GetData());
+
+			const float* Data = Buffer->GetData();
+			float Phase = 0.0f;
+			const float Frequency = 200.f;
+
+			const int32 NumBuffers = SequenceLength / OperatorSettings.GetNumFramesPerBlock();
+
+			for (int32 j = 0; j < NumBuffers; j++)
+			{
+				Invoke(Func, Operator.Get());
+
+				for (int32 i = 0; i < OperatorSettings.GetNumFramesPerBlock(); i++)
+				{
+					UTEST_EQUAL_TOLERANCE("Compare signals are the same-ish-sorta",Data[i], FMath::Sin(Phase),0.01f);
+					//REQUIRE(Data[i] == Approx(FMath::Sin(Phase)).margin(0.01f));
+					Phase += 2.0f * PI * Frequency / OperatorSettings.GetSampleRate();
+
+					while (Phase > (2.0f * PI))
+					{
+						Phase -= (2.0f * PI);
+					}
+				}
+			}
+		}
+	}
+	
+	// Got to the end.
+	return true;
+}
+*/
+#endif // 0
+
+#endif //WITH_DEV_AUTOMATION_TESTS
diff --git a/Engine/Plugins/Runtime/Metasound/Source/MetasoundEngineTest/Private/MetasoundEngineTestModule.cpp b/Engine/Plugins/Runtime/Metasound/Source/MetasoundEngineTest/Private/MetasoundEngineTestModule.cpp
new file mode 100644
--- /dev/null
+++ b/Engine/Plugins/Runtime/Metasound/Source/MetasoundEngineTest/Private/MetasoundEngineTestModule.cpp
@@ -0,0 +1,22 @@
+// Copyright Epic Games, Inc. All Rights Reserved.
+
+#include "Modules/ModuleInterface.h"
+#include "Modules/ModuleManager.h"
+
+namespace Metasound 
+{
+	class METASOUNDENGINETEST_API FMetasoundEngineTestModule : public IModuleInterface
+	{
+	public:
+		virtual void StartupModule() override
+		{
+			FModuleManager::Get().LoadModuleChecked("MetasoundEngine");
+		}
+	};
+}
+
+IMPLEMENT_MODULE(Metasound::FMetasoundEngineTestModule, MetasoundEngineTest);
+
+
+
+
