Move FString/FText VertexName identifers to FNames. Convention is as follows:
- FGuid IDs are used to identify unique verticies within an entire graph (currently may be duplicated if graph is duplicated).
- FNames are used to be readible vertex identifiers that are only unique within a node's API.
- FText is only for ***OPTIONAL*** DisplayNames that can be provided as a more descriptive, readible descriptor of a given vertex (But are not necessarily unique). Currently disabled until loc tool on FText property in MetaSoundDetailsCustomization is properly displayed.

diff --git a/Engine/Plugins/Runtime/Metasound/Source/MetasoundEditor/Private/MetasoundEditor.cpp b/Engine/Plugins/Runtime/Metasound/Source/MetasoundEditor/Private/MetasoundEditor.cpp
--- a/Engine/Plugins/Runtime/Metasound/Source/MetasoundEditor/Private/MetasoundEditor.cpp
+++ b/Engine/Plugins/Runtime/Metasound/Source/MetasoundEditor/Private/MetasoundEditor.cpp
@@ -945,14 +945,12 @@ namespace Metasound
 				{
 					if (UAudioComponent* PreviewComponent = GEditor->GetPreviewAudioComponent())
 					{
-						// TODO: fix how identifying the parameter to update is determined. It should not be done
-						// with a "DisplayName" but rather the vertex Guid.
-						if (TScriptInterface<IAudioParameterInterface> ParamInterface = PreviewComponent)
+						FConstNodeHandle NodeHandle = InputNode->GetConstNodeHandle();
+						NodeHandle->IterateConstInputs([ComponentInterface = PreviewComponent](FConstInputHandle Input)
 						{
-							FConstNodeHandle NodeHandle = InputNode->GetConstNodeHandle();
-							FVertexKey VertexKey = Metasound::FVertexKey(NodeHandle->GetDisplayName().ToString());
-							ParamInterface->SetTriggerParameter(*VertexKey);
-						}
+							FVertexName VertexName = Input->GetName();
+							ComponentInterface->SetTriggerParameter(VertexName);
+						});
 					}
 				}
 			}
@@ -1298,14 +1296,12 @@ namespace Metasound
 				{
 					for (int32 i = IndexToDelete + 1; i < Handles.Num(); ++i)
 					{
-						const FName NameToSelect(*(Handles[i]->GetDisplayName().ToString()));
-						return FNameNodeIDPair { NameToSelect, Handles[i]->GetID() };
+						return FNameNodeIDPair { Handles[i]->GetNodeName(), Handles[i]->GetID() };
 					}
 
 					for (int32 i = IndexToDelete - 1; i >= 0; --i)
 					{
-						FName NameToSelect = (*(Handles[i]->GetDisplayName().ToString()));
-						return FNameNodeIDPair { NameToSelect, Handles[i]->GetID() };
+						return FNameNodeIDPair { Handles[i]->GetNodeName(), Handles[i]->GetID() };
 					}
 				}
 
@@ -1614,7 +1610,11 @@ namespace Metasound
 						{
 							FMetasoundFrontendClass MetaSoundClass;
 							FMetasoundFrontendRegistryContainer::Get()->FindFrontendClassFromRegistered(PastedRegistryKey, MetaSoundClass);
-							const FString FriendlyClassName = MetaSoundClass.Metadata.GetDisplayName().ToString();
+							FString FriendlyClassName = MetaSoundClass.Metadata.GetDisplayName().ToString();
+							if (FriendlyClassName.IsEmpty())
+							{
+								FriendlyClassName = MetaSoundClass.Metadata.GetClassName().ToString();
+							}
 							UE_LOG(LogMetaSound, Warning, TEXT("Failed to paste node with class '%s'.  Class would introduce cyclic asset dependency."), *FriendlyClassName);
 							bNotifyReferenceLoop = true;
 							NodesToRemove.Add(GraphNode);
@@ -1822,12 +1822,8 @@ namespace Metasound
 					Frontend::FNodeHandle NodeHandle = MetasoundAction->GetNodeHandle();
 					if (InNodeID == NodeHandle->GetID())
 					{
-						TArray<Frontend::FConstOutputHandle> OutputHandles = NodeHandle->GetConstOutputs();
-						if (ensure(OutputHandles.Num() == 1))
-						{
-							const FName ActionName = *OutputHandles[0]->GetDisplayName().ToString();
-							MetasoundInterfaceMenu->SelectItemByName(ActionName, ESelectInfo::Direct, Action->GetSectionID());
-						}
+						const FName ActionName = NodeHandle->GetNodeName();
+						MetasoundInterfaceMenu->SelectItemByName(ActionName, ESelectInfo::Direct, Action->GetSectionID());
 						break;
 					}
 				}
@@ -1855,12 +1851,8 @@ namespace Metasound
 					Frontend::FNodeHandle NodeHandle = MetasoundAction->GetNodeHandle();
 					if (InNodeID == NodeHandle->GetID())
 					{
-						TArray<Frontend::FConstOutputHandle> OutputHandles = NodeHandle->GetConstOutputs();
-						if (ensure(OutputHandles.Num() == 1))
-						{
-							const FName ActionName = *OutputHandles[0]->GetDisplayName().ToString();
-							MetasoundInterfaceMenu->SelectItemByName(ActionName, ESelectInfo::Direct, Action->GetSectionID());
-						}
+						const FName ActionName = NodeHandle->GetNodeName();
+						MetasoundInterfaceMenu->SelectItemByName(ActionName, ESelectInfo::Direct, Action->GetSectionID());
 						break;
 					}
 				}
@@ -1877,9 +1869,9 @@ namespace Metasound
 			MetasoundAsset->GetRootGraphHandle()->IterateConstNodes([this, ActionList = &OutAllActions](const Frontend::FConstNodeHandle& Input)
 			{
 				const FText Tooltip = Input->GetDescription();
-				const FText InputName = Input->GetDisplayName();
 				UMetasoundEditorGraph& Graph = GetMetaSoundGraphChecked();
 
+				const FText InputName = Input->GetDisplayName();
 				TSharedPtr<FMetasoundGraphNodeSchemaAction> NewFuncAction = MakeShared<FMetasoundGraphNodeSchemaAction>(FText::GetEmpty(), InputName, Tooltip, 1, ENodeSection::Inputs);
 				NewFuncAction->Graph = &Graph;
 				ActionList->AddAction(NewFuncAction);
@@ -1900,9 +1892,9 @@ namespace Metasound
 			MetasoundAsset->GetRootGraphHandle()->IterateConstNodes([this, ActionList = &OutAllActions](const Frontend::FConstNodeHandle& Output)
 			{
 				const FText Tooltip = Output->GetDescription();
-				const FText OutputName = Output->GetDisplayName();
 				UMetasoundEditorGraph& Graph = GetMetaSoundGraphChecked();
 
+				const FText OutputName = Output->GetDisplayName();
 				TSharedPtr<FMetasoundGraphNodeSchemaAction> NewFuncAction = MakeShared<FMetasoundGraphNodeSchemaAction>(FText::GetEmpty(), OutputName, Tooltip, 1, ENodeSection::Outputs);
 				NewFuncAction->Graph = &Graph;
 				NewFuncAction->NodeID = Output->GetID();
@@ -1935,6 +1927,11 @@ namespace Metasound
 
 		bool FEditor::HandleActionMatchesName(FEdGraphSchemaAction* InAction, const FName& InName) const
 		{
+			if (FMetasoundGraphNodeSchemaAction* Action = static_cast<FMetasoundGraphNodeSchemaAction*>(InAction))
+			{
+				return InName == Action->GetNodeHandle()->GetNodeName();
+			}
+
 			return false;
 		}
 
@@ -2076,7 +2073,7 @@ namespace Metasound
 					Frontend::FNodeHandle NodeHandle = FGraphBuilder::AddInputNodeHandle(*Metasound, DataTypeName, FText::GetEmpty());
 					if (ensure(NodeHandle->IsValid()))
 					{
-						NameToSelect = *NodeHandle->GetDisplayName().ToString();
+						NameToSelect = NodeHandle->GetNodeName();
 
 						TObjectPtr<UMetasoundEditorGraphInput> Input = Graph.FindOrAddInput(NodeHandle);
 						if (ensure(Input))
@@ -2101,7 +2098,7 @@ namespace Metasound
 					Frontend::FNodeHandle NodeHandle = FGraphBuilder::AddOutputNodeHandle(*Metasound, DataTypeName, FText::GetEmpty());
 					if (ensure(NodeHandle->IsValid()))
 					{
-						NameToSelect = *NodeHandle->GetDisplayName().ToString();
+						NameToSelect = NodeHandle->GetNodeName();
 
 						TObjectPtr<UMetasoundEditorGraphOutput> Output = Graph.FindOrAddOutput(NodeHandle);
 						if (ensure(Output))
diff --git a/Engine/Plugins/Runtime/Metasound/Source/MetasoundEditor/Private/MetasoundEditorGraph.cpp b/Engine/Plugins/Runtime/Metasound/Source/MetasoundEditor/Private/MetasoundEditorGraph.cpp
--- a/Engine/Plugins/Runtime/Metasound/Source/MetasoundEditor/Private/MetasoundEditorGraph.cpp
+++ b/Engine/Plugins/Runtime/Metasound/Source/MetasoundEditor/Private/MetasoundEditorGraph.cpp
@@ -11,6 +11,7 @@
 #include "MetasoundEditorGraphValidation.h"
 #include "MetasoundEditorModule.h"
 #include "MetasoundUObjectRegistry.h"
+#include "MetasoundVertex.h"
 #include "ScopedTransaction.h"
 
 #define LOCTEXT_NAMESPACE "MetaSoundEditor"
@@ -58,11 +59,30 @@ void UMetasoundEditorGraphVariable::SetDescription(const FText& InDescription)
 	}
 }
 
+void UMetasoundEditorGraphVariable::SetNodeName(const FName& InNewName)
+{
+	using namespace Metasound::Frontend;
+
+	const FText TransactionLabel = FText::Format(LOCTEXT("RenameVariableNameFormat", "Rename Metasound {0}"), GetVariableLabel());
+	const FScopedTransaction Transaction(TransactionLabel);
+
+	Modify();
+	if (UMetasoundEditorGraph* Graph = Cast<UMetasoundEditorGraph>(GetOuter()))
+	{
+		Graph->GetMetasoundChecked().Modify();
+	}
+
+	FNodeHandle NodeHandle = GetNodeHandle();
+	NodeHandle->SetNodeName(InNewName);
+
+	NameChanged.Broadcast(NodeID);
+}
+
 void UMetasoundEditorGraphVariable::SetDisplayName(const FText& InNewName)
 {
 	using namespace Metasound::Frontend;
 
-	const FText TransactionLabel = FText::Format(LOCTEXT("SetVariableDisplayNameFormat", "Rename Metasound {0}"), GetVariableLabel());
+	const FText TransactionLabel = FText::Format(LOCTEXT("RenameVariableDisplayNameFormat", "Set Metasound {0} DisplayName"), GetVariableLabel());
 	const FScopedTransaction Transaction(TransactionLabel);
 
 	Modify();
@@ -105,6 +125,7 @@ void UMetasoundEditorGraphVariable::SetDataType(FName InNewType)
 
 	// 2. Cache the old version's Frontend data.
 	FNodeHandle NodeHandle = GetNodeHandle();
+	const FName NodeName = NodeHandle->GetNodeName();
 	const FText NodeDisplayName = NodeHandle->GetDisplayName();
 
 	// 3. Delete the Frontend variable
@@ -112,7 +133,8 @@ void UMetasoundEditorGraphVariable::SetDataType(FName InNewType)
 
 	// 4. Add the new input node with the same identifier data but new datatype.
 	UObject& Metasound = Graph->GetMetasoundChecked();
-	FNodeHandle NewNodeHandle = AddNodeHandle(NodeDisplayName, InNewType);
+	FNodeHandle NewNodeHandle = AddNodeHandle(NodeName, InNewType);
+	NewNodeHandle->SetDisplayName(NodeDisplayName);
 
 	if (!ensure(NewNodeHandle->IsValid()))
 	{
@@ -173,7 +195,7 @@ bool UMetasoundEditorGraphVariable::CanRename(const FText& InNewName, FText& Out
 {
 	using namespace Metasound::Frontend;
 
-	if (InNewName.IsEmpty())
+	if (InNewName.IsEmptyOrWhitespace())
 	{
 		OutError = FText::Format(LOCTEXT("VariableRenameInvalid_NameEmpty", "{0} cannot be empty string."), InNewName);
 		return false;
@@ -186,14 +208,14 @@ bool UMetasoundEditorGraphVariable::CanRename(const FText& InNewName, FText& Out
 	}
 
 	bool bIsNameValid = true;
-
+	const FName NewFName = *InNewName.ToString();
 	FConstNodeHandle NodeHandle = GetConstNodeHandle();
 	FConstGraphHandle GraphHandle = NodeHandle->GetOwningGraph();
 	GraphHandle->IterateConstNodes([&](FConstNodeHandle NodeToCompare)
 	{
 		if (NodeID != NodeToCompare->GetID())
 		{
-			if (InNewName.CompareToCaseIgnored(NodeToCompare->GetDisplayName()) == 0)
+			if (NewFName == NodeToCompare->GetNodeName())
 			{
 				bIsNameValid = false;
 				OutError = FText::Format(LOCTEXT("VariableRenameInvalid_NameTaken", "{0} is already in use"), InNewName);
@@ -239,8 +261,8 @@ void UMetasoundEditorGraphInput::UpdateDocumentInput(bool bPostTransaction)
 
 	FGraphHandle GraphHandle = MetasoundAsset->GetRootGraphHandle();
 	FNodeHandle NodeHandle = GraphHandle->GetNodeWithID(NodeID);
-	const FString& NodeName = NodeHandle->GetNodeName();
 
+	const Metasound::FVertexName& NodeName = NodeHandle->GetNodeName();
 	const FGuid VertexID = GraphHandle->GetVertexIDForInputVertex(NodeName);
 	GraphHandle->SetDefaultInput(VertexID, Literal->GetDefault());
 
@@ -249,7 +271,7 @@ void UMetasoundEditorGraphInput::UpdateDocumentInput(bool bPostTransaction)
 // 	Metasound::Editor::FGraphBuilder::RegisterGraphWithFrontend(*Metasound);
 }
 
-void UMetasoundEditorGraphInput::UpdatePreviewInstance(const Metasound::FVertexKey& InParameterName, TScriptInterface<IAudioParameterInterface>& InParamInterface) const
+void UMetasoundEditorGraphInput::UpdatePreviewInstance(const Metasound::FVertexName& InParameterName, TScriptInterface<IAudioParameterInterface>& InParamInterface) const
 {
 	if (ensure(Literal))
 	{
@@ -257,7 +279,7 @@ void UMetasoundEditorGraphInput::UpdatePreviewInstance(const Metasound::FVertexK
 	}
 }
 
-Metasound::Frontend::FNodeHandle UMetasoundEditorGraphInput::AddNodeHandle(const FText& InNodeDisplayName, FName InDataType)
+Metasound::Frontend::FNodeHandle UMetasoundEditorGraphInput::AddNodeHandle(const FName& InName, FName InDataType)
 {
 	using namespace Metasound::Editor;
 	using namespace Metasound::Frontend;
@@ -265,12 +287,11 @@ Metasound::Frontend::FNodeHandle UMetasoundEditorGraphInput::AddNodeHandle(const
 	UMetasoundEditorGraph* Graph = Cast<UMetasoundEditorGraph>(GetOuter());
 	if (!ensure(Graph))
 	{
-		return UMetasoundEditorGraphVariable::AddNodeHandle(InNodeDisplayName, InDataType);
+		return UMetasoundEditorGraphVariable::AddNodeHandle(InName, InDataType);
 	}
 
 	UObject& Metasound = Graph->GetMetasoundChecked();
-	Metasound::Frontend::FNodeHandle NewNodeHandle = FGraphBuilder::AddInputNodeHandle(Metasound, InDataType, FText::GetEmpty());
-	NewNodeHandle->SetDisplayName(InNodeDisplayName);
+	Metasound::Frontend::FNodeHandle NewNodeHandle = FGraphBuilder::AddInputNodeHandle(Metasound, InDataType, FText::GetEmpty(), nullptr, &InName);
 	return NewNodeHandle;
 }
 
@@ -313,10 +334,8 @@ void UMetasoundEditorGraphInput::OnLiteralChanged(bool bPostTransaction)
 
 		if (TScriptInterface<IAudioParameterInterface> ParamInterface = PreviewComponent)
 		{
-			// TODO: fix how identifying the parameter to update is determined. It should not be done
-			// with a "DisplayName" but rather the vertex Guid.
 			Metasound::Frontend::FConstNodeHandle NodeHandle = GetConstNodeHandle();
-			Metasound::FVertexKey VertexKey = Metasound::FVertexKey(NodeHandle->GetDisplayName().ToString());
+			Metasound::FVertexName VertexKey = NodeHandle->GetNodeName();
 			UpdatePreviewInstance(VertexKey, ParamInterface);
 		}
 	}
@@ -340,7 +359,7 @@ void UMetasoundEditorGraphInput::OnDataTypeChanged()
 	Literal = NewObject<UMetasoundEditorGraphInputLiteral>(this, InputLiteralClass, FName(), RF_Transactional);
 }
 
-Metasound::Frontend::FNodeHandle UMetasoundEditorGraphOutput::AddNodeHandle(const FText& InNodeDisplayName, FName InDataType)
+Metasound::Frontend::FNodeHandle UMetasoundEditorGraphOutput::AddNodeHandle(const FName& InName, FName InDataType)
 {
 	using namespace Metasound::Editor;
 	using namespace Metasound::Frontend;
@@ -348,12 +367,11 @@ Metasound::Frontend::FNodeHandle UMetasoundEditorGraphOutput::AddNodeHandle(cons
 	UMetasoundEditorGraph* Graph = Cast<UMetasoundEditorGraph>(GetOuter());
 	if (!ensure(Graph))
 	{
-		return UMetasoundEditorGraphVariable::AddNodeHandle(InNodeDisplayName, InDataType);
+		return UMetasoundEditorGraphVariable::AddNodeHandle(InName, InDataType);
 	}
 
 	UObject& Metasound = Graph->GetMetasoundChecked();
-	Metasound::Frontend::FNodeHandle NewNodeHandle = FGraphBuilder::AddOutputNodeHandle(Metasound, InDataType, FText::GetEmpty());
-	NewNodeHandle->SetDisplayName(InNodeDisplayName);
+	FNodeHandle NewNodeHandle = FGraphBuilder::AddOutputNodeHandle(Metasound, InDataType, FText::GetEmpty(), &InName);
 	return NewNodeHandle;
 }
 
@@ -460,13 +478,24 @@ bool UMetasoundEditorGraph::Validate(bool bInClearUpdateNotes)
 
 UMetasoundEditorGraphInput* UMetasoundEditorGraph::FindInput(FGuid InNodeID) const
 {
-	const TObjectPtr<UMetasoundEditorGraphInput>* Input = Inputs.FindByPredicate([InNodeID](const TObjectPtr<UMetasoundEditorGraphInput>& Input)
+	const TObjectPtr<UMetasoundEditorGraphInput>* Input = Inputs.FindByPredicate([InNodeID](const TObjectPtr<UMetasoundEditorGraphInput>& InInput)
+	{
+		return InInput->NodeID == InNodeID;
+	});
+	return Input ? Input->Get() : nullptr;
+}
+
+UMetasoundEditorGraphInput* UMetasoundEditorGraph::FindInput(FName InName) const
+{
+	const TObjectPtr<UMetasoundEditorGraphInput>* Input = Inputs.FindByPredicate([InName](const TObjectPtr<UMetasoundEditorGraphInput>& InInput)
 	{
-		return Input->NodeID == InNodeID;
+		const FName NodeName = InInput->GetNodeHandle()->GetNodeName();
+		return NodeName == InName;
 	});
 	return Input ? Input->Get() : nullptr;
 }
 
+
 UMetasoundEditorGraphInput* UMetasoundEditorGraph::FindOrAddInput(Metasound::Frontend::FNodeHandle InNodeHandle)
 {
 	using namespace Metasound::Editor;
@@ -515,9 +544,19 @@ UMetasoundEditorGraphInput* UMetasoundEditorGraph::FindOrAddInput(Metasound::Fro
 
 UMetasoundEditorGraphOutput* UMetasoundEditorGraph::FindOutput(FGuid InNodeID) const
 {
-	const TObjectPtr<UMetasoundEditorGraphOutput>* Output = Outputs.FindByPredicate([InNodeID](const TObjectPtr<UMetasoundEditorGraphOutput>& Output)
+	const TObjectPtr<UMetasoundEditorGraphOutput>* Output = Outputs.FindByPredicate([InNodeID](const TObjectPtr<UMetasoundEditorGraphOutput>& InOutput)
+	{
+		return InOutput->NodeID == InNodeID;
+	});
+	return Output ? Output->Get() : nullptr;
+}
+
+UMetasoundEditorGraphOutput* UMetasoundEditorGraph::FindOutput(FName InName) const
+{
+	const TObjectPtr<UMetasoundEditorGraphOutput>* Output = Outputs.FindByPredicate([InName](const TObjectPtr<UMetasoundEditorGraphOutput>& InOutput)
 	{
-		return Output->NodeID == InNodeID;
+		const FName NodeName = InOutput->GetNodeHandle()->GetNodeName();
+		return NodeName == InName;
 	});
 	return Output ? Output->Get() : nullptr;
 }
diff --git a/Engine/Plugins/Runtime/Metasound/Source/MetasoundEditor/Private/MetasoundEditorGraphBuilder.cpp b/Engine/Plugins/Runtime/Metasound/Source/MetasoundEditor/Private/MetasoundEditorGraphBuilder.cpp
--- a/Engine/Plugins/Runtime/Metasound/Source/MetasoundEditor/Private/MetasoundEditorGraphBuilder.cpp
+++ b/Engine/Plugins/Runtime/Metasound/Source/MetasoundEditor/Private/MetasoundEditorGraphBuilder.cpp
@@ -19,6 +19,7 @@
 #include "MetasoundEditorModule.h"
 #include "MetasoundEditorSettings.h"
 #include "MetasoundFrontendDataTypeRegistry.h"
+#include "MetasoundFrontendDocumentAccessPtr.h"
 #include "MetasoundFrontendQuery.h"
 #include "MetasoundFrontendQuerySteps.h"
 #include "MetasoundFrontendRegistries.h"
@@ -26,6 +27,7 @@
 #include "MetasoundFrontendTransform.h"
 #include "MetasoundLiteral.h"
 #include "MetasoundUObjectRegistry.h"
+#include "MetasoundVertex.h"
 #include "Modules/ModuleManager.h"
 #include "Templates/Tuple.h"
 #include "Toolkits/ToolkitManager.h"
@@ -72,12 +74,12 @@ namespace Metasound
 
 				// Initial graph generation is not something to be managed by the transaction stack, so don't track dirty state until
 				// after initial setup if necessary.
-				if (!MetaSoundAsset->GetGraph())
+				UMetasoundEditorGraph* Graph = Cast<UMetasoundEditorGraph>(MetaSoundAsset->GetGraph());
+				if (!Graph)
 				{
-					UMetasoundEditorGraph* Graph = NewObject<UMetasoundEditorGraph>(&InMetaSound, FName(), RF_Transactional);
+					Graph = NewObject<UMetasoundEditorGraph>(&InMetaSound, FName(), RF_Transactional);
 					Graph->Schema = UMetasoundEditorGraphSchema::StaticClass();
 					MetaSoundAsset->SetGraph(Graph);
-
 					return true;
 				}
 
@@ -215,65 +217,37 @@ namespace Metasound
 			return Categories;
 		}
 
-		FText FGraphBuilder::GenerateUniqueInputDisplayName(const UObject& InMetaSound, const FText* InBaseName)
+		FName FGraphBuilder::GenerateUniqueNameByClassType(const UObject& InMetaSound, EMetasoundFrontendClassType InClassType, const FName& InBaseName)
 		{
-			static const FText DefaultBaseName = FText::FromString(TEXT("Input"));
-			const FText& NameBase = InBaseName ? *InBaseName : DefaultBaseName;
-			return GenerateUniqueNameByFilter(InMetaSound, NameBase, [](const Frontend::FConstGraphHandle& GraphHandle, const FText& NewName)
-			{
-				bool bIsNameInvalid = false;
-				GraphHandle->IterateConstNodes([&](Frontend::FConstNodeHandle Node)
-				{
-					if (NewName.CompareToCaseIgnored(Node->GetDisplayName()) == 0)
-					{
-						bIsNameInvalid = true;
-					}
-				}, EMetasoundFrontendClassType::Input);
-				return !bIsNameInvalid;
-			});
-		}
+			const FMetasoundAssetBase* MetaSoundAsset = IMetasoundUObjectRegistry::Get().GetObjectAsAssetBase(&InMetaSound);
+			check(MetaSoundAsset);
 
-		FText FGraphBuilder::GenerateUniqueOutputDisplayName(const UObject& InMetaSound, const FText* InBaseName)
-		{
-			static const FText DefaultBaseName = FText::FromString(TEXT("Output"));
-			const FText& NameBase = InBaseName ? *InBaseName : DefaultBaseName;
-			return GenerateUniqueNameByFilter(InMetaSound, NameBase, [](const Frontend::FConstGraphHandle& GraphHandle, const FText& NewName)
+			const Frontend::FConstGraphHandle GraphHandle = MetaSoundAsset->GetRootGraphHandle();
+
+			auto InNameIsValidFilter = [&](FName NameToTest)
 			{
+				using namespace Metasound;
+
 				bool bIsNameInvalid = false;
 				GraphHandle->IterateConstNodes([&](Frontend::FConstNodeHandle Node)
 				{
-					if (NewName.CompareToCaseIgnored(Node->GetDisplayName()) == 0)
-					{
-						bIsNameInvalid = true;
-					}
-				}, EMetasoundFrontendClassType::Output);
+					bIsNameInvalid |= NameToTest == Node->GetNodeName();
+				}, InClassType);
 				return !bIsNameInvalid;
-			});
-		}
-
-		FText FGraphBuilder::GenerateUniqueNameByFilter(const UObject& InMetaSound, const FText& InBaseText, TFunctionRef<bool(const Frontend::FConstGraphHandle&, const FText&)> InNameIsValidFilter)
-		{
-			const FMetasoundAssetBase* MetaSoundAsset = IMetasoundUObjectRegistry::Get().GetObjectAsAssetBase(&InMetaSound);
-			check(MetaSoundAsset);
-
-			const Frontend::FConstGraphHandle GraphHandle = MetaSoundAsset->GetRootGraphHandle();
+			};
 
-			if (InNameIsValidFilter(GraphHandle, InBaseText))
+			if (InNameIsValidFilter(InBaseName))
 			{
-				return InBaseText;
+				return InBaseName;
 			}
 
-			FNumberFormattingOptions NumberFormatOptions;
-			NumberFormatOptions.MinimumIntegralDigits = 2;
-
 			int32 i = 1;
-			FText NewName;
+			FName NewName;
 			do
 			{
-				static const FText VariableUniqueNameFormat = LOCTEXT("GenerateUniqueVariableNameFormat", "{0} {1}");
-				NewName = FText::Format(VariableUniqueNameFormat, InBaseText, FText::AsNumber(++i, &NumberFormatOptions));
+				NewName = *FString::Format(TEXT("{0} {1}"), { *InBaseName.ToString(), ++i });
 			}
-			while (!InNameIsValidFilter(GraphHandle, NewName));
+			while (!InNameIsValidFilter(NewName));
 
 			return NewName;
 		}
@@ -358,7 +332,17 @@ namespace Metasound
 					FNodeHandle NodeHandle = EdNode->GetNodeHandle();
 					if (NodeHandle->IsValid())
 					{
-						TArray<FInputHandle> Inputs = NodeHandle->GetInputsWithVertexName(InPin->GetName());
+						TArray<FInputHandle> Inputs;
+						const EMetasoundFrontendClassType ClassType = NodeHandle->GetClassMetadata().GetType();
+						if (ClassType == EMetasoundFrontendClassType::Input || ClassType == EMetasoundFrontendClassType::Output)
+						{
+							Inputs = NodeHandle->GetInputs();
+						}
+						else
+						{
+							Inputs = NodeHandle->GetInputsWithVertexName(InPin->GetFName());
+						}
+
 						if (ensure(Inputs.Num() == 1))
 						{
 							return Inputs[0];
@@ -386,7 +370,17 @@ namespace Metasound
 					FNodeHandle NodeHandle = EdNode->GetNodeHandle();
 					if (NodeHandle->IsValid())
 					{
-						TArray<FOutputHandle> Outputs = NodeHandle->GetOutputsWithVertexName(InPin->GetName());
+						TArray<FOutputHandle> Outputs;
+						const EMetasoundFrontendClassType ClassType = NodeHandle->GetClassMetadata().GetType();
+						if (ClassType == EMetasoundFrontendClassType::Input || ClassType == EMetasoundFrontendClassType::Output)
+						{
+							Outputs = NodeHandle->GetOutputs();
+						}
+						else
+						{
+							Outputs = NodeHandle->GetOutputsWithVertexName(InPin->GetFName());
+						}
+
 						if (ensure(Outputs.Num() == 1))
 						{
 							return Outputs[0];
@@ -585,8 +579,8 @@ namespace Metasound
 					UMetasoundEditorGraphInput* Input = InputNode->Input;
 					if (ensure(Input))
 					{
-						NodeHandle = AddInputNodeHandle(InMetaSound, Input->TypeName, InGraphNode.GetTooltipText());
-						NodeHandle->SetDisplayName(FText::FromString(Pin->GetName()));
+						const FName PinName = Pin->GetFName();
+						NodeHandle = AddInputNodeHandle(InMetaSound, Input->TypeName, InGraphNode.GetTooltipText(), nullptr, &PinName);
 					}
 				}
 			}
@@ -600,8 +594,8 @@ namespace Metasound
 					UMetasoundEditorGraphOutput* Output = OutputNode->Output;
 					if (ensure(Output))
 					{
-						NodeHandle = FGraphBuilder::AddOutputNodeHandle(InMetaSound, Output->TypeName, InGraphNode.GetTooltipText());
-						NodeHandle->SetDisplayName(FText::FromString(Pin->GetName()));
+						const FName PinName = Pin->GetFName();
+						NodeHandle = FGraphBuilder::AddOutputNodeHandle(InMetaSound, Output->TypeName, InGraphNode.GetTooltipText(), &PinName);
 					}
 				}
 			}
@@ -632,22 +626,22 @@ namespace Metasound
 			return NodeHandle;
 		}
 
-		Frontend::FNodeHandle FGraphBuilder::AddInputNodeHandle(UObject& InMetaSound, const FName InTypeName, const FText& InToolTip, const FMetasoundFrontendLiteral* InDefaultValue)
+		Frontend::FNodeHandle FGraphBuilder::AddInputNodeHandle(UObject& InMetaSound, const FName InTypeName, const FText& InToolTip, const FMetasoundFrontendLiteral* InDefaultValue, const FName* InNameBase)
 		{
 			FMetasoundAssetBase* MetaSoundAsset = IMetasoundUObjectRegistry::Get().GetObjectAsAssetBase(&InMetaSound);
 			check(MetaSoundAsset);
 
-			const FText DisplayName = GenerateUniqueInputDisplayName(InMetaSound);
-			return MetaSoundAsset->GetRootGraphHandle()->AddInputNode(InTypeName, InToolTip, InDefaultValue, &DisplayName);
+			const FName NewName = GenerateUniqueNameByClassType(InMetaSound, EMetasoundFrontendClassType::Input, InNameBase ? *InNameBase : "Input");
+			return MetaSoundAsset->GetRootGraphHandle()->AddInputVertex(NewName, InTypeName, InToolTip, InDefaultValue);
 		}
 
-		Frontend::FNodeHandle FGraphBuilder::AddOutputNodeHandle(UObject& InMetaSound, const FName InTypeName, const FText& InToolTip)
+		Frontend::FNodeHandle FGraphBuilder::AddOutputNodeHandle(UObject& InMetaSound, const FName InTypeName, const FText& InToolTip, const FName* InNameBase)
 		{
 			FMetasoundAssetBase* MetaSoundAsset = IMetasoundUObjectRegistry::Get().GetObjectAsAssetBase(&InMetaSound);
 			check(MetaSoundAsset);
 
-			const FText DisplayName = GenerateUniqueOutputDisplayName(InMetaSound);
-			return MetaSoundAsset->GetRootGraphHandle()->AddOutputNode(InTypeName, InToolTip, &DisplayName);
+			const FName NewName = GenerateUniqueNameByClassType(InMetaSound, EMetasoundFrontendClassType::Output, InNameBase ? *InNameBase : "Output");
+			return MetaSoundAsset->GetRootGraphHandle()->AddOutputVertex(NewName, InTypeName, InToolTip);
 		}
 
 		UMetasoundEditorGraphNode* FGraphBuilder::AddNode(UObject& InMetaSound, Frontend::FNodeHandle InNodeHandle, FVector2D InLocation, bool bInSelectNewNode)
@@ -734,7 +728,7 @@ namespace Metasound
 			if (InPin.Direction == EGPD_Input)
 			{
 				FNodeHandle NodeHandle = CastChecked<UMetasoundEditorGraphNode>(InPin.GetOwningNode())->GetNodeHandle();
-				InputHandles = NodeHandle->GetInputsWithVertexName(InPin.GetName());
+				InputHandles = NodeHandle->GetInputsWithVertexName(InPin.GetFName());
 				InputPins.Add(&InPin);
 			}
 			else
@@ -743,7 +737,7 @@ namespace Metasound
 				for (UEdGraphPin* Pin : InPin.LinkedTo)
 				{
 					FNodeHandle NodeHandle = CastChecked<UMetasoundEditorGraphNode>(Pin->GetOwningNode())->GetNodeHandle();
-					InputHandles.Append(NodeHandle->GetInputsWithVertexName(Pin->GetName()));
+					InputHandles.Append(NodeHandle->GetInputsWithVertexName(Pin->GetFName()));
 					InputPins.Add(Pin);
 				}
 			}
@@ -870,8 +864,7 @@ namespace Metasound
 				}
 			}
 
-			const FString NodeName = NodeHandle->GetNodeName();
-			const FText NodeDisplayName = NodeHandle->GetDisplayName();
+			const Metasound::FVertexName NodeName = NodeHandle->GetNodeName();
 			FGraphHandle GraphHandle = NodeHandle->GetOwningGraph();
 			GraphHandle->RemoveNode(*NodeHandle);
 		}
@@ -1130,7 +1123,7 @@ namespace Metasound
 				PinType = EditorModule.FindDataType(DataTypeName).PinType;
 			}
 
-			UEdGraphPin* NewPin = InEditorNode.CreatePin(EGPD_Input, PinType, FName(*InInputHandle->GetName()));
+			UEdGraphPin* NewPin = InEditorNode.CreatePin(EGPD_Input, PinType, InInputHandle->GetName());
 			if (ensure(NewPin))
 			{
 				RefreshPinMetadata(*NewPin, InInputHandle->GetMetadata());
@@ -1151,7 +1144,7 @@ namespace Metasound
 				PinType = EditorModule.FindDataType(DataTypeName).PinType;
 			}
 
-			UEdGraphPin* NewPin = InEditorNode.CreatePin(EGPD_Output, PinType, FName(*InOutputHandle->GetName()));
+			UEdGraphPin* NewPin = InEditorNode.CreatePin(EGPD_Output, PinType, InOutputHandle->GetName());
 			if (ensure(NewPin))
 			{
 				NewPin->PinToolTip = InOutputHandle->GetTooltip().ToString();
@@ -1181,7 +1174,6 @@ namespace Metasound
 
 			UMetasoundEditorGraph* EditorGraph = CastChecked<UMetasoundEditorGraph>(MetaSoundAsset->GetGraph());
 
-
 			TArray<UMetasoundEditorGraphNode*> EditorNodes;
 			EditorGraph->GetNodesOfClass(EditorNodes);
 
@@ -1256,6 +1248,7 @@ namespace Metasound
 								{
 									UEdGraphPin* OutputPin = (*OutputEditorNode)[0]->FindPinChecked(OutputHandle->GetName(), EEdGraphPinDirection::EGPD_Output);
 									const FText& OwningNodeName = NodeInput->GetOwningNode()->GetDisplayName();
+
 									UE_LOG(LogMetasoundEditor, Display, TEXT("Synchronizing Node '%s' Connection: Linking Pin '%s' to '%s'"), *OwningNodeName.ToString(), *MatchingPin->GetName(), *OutputPin->GetName());
 									MatchingPin->MakeLinkTo(OutputPin);
 									bIsNodeDirty = true;
@@ -1269,8 +1262,9 @@ namespace Metasound
 						if (!MatchingPin->LinkedTo.IsEmpty())
 						{
 							MatchingPin->BreakAllPinLinks();
-							const FText& OwningNodeName = NodeInput->GetOwningNode()->GetDisplayName();
-							UE_LOG(LogMetasoundEditor, Display, TEXT("Synchronizing Node '%s' Connection: Breaking all pin links to '%s'"), *OwningNodeName.ToString(), *NodeInput->GetDisplayName().ToString());
+							const FText OwningNodeName = NodeInput->GetOwningNode()->GetDisplayName();
+							const FText InputName = NodeInput->GetDisplayName();
+							UE_LOG(LogMetasoundEditor, Display, TEXT("Synchronizing Node '%s' Connection: Breaking all pin links to '%s'"), *OwningNodeName.ToString(), *InputName.ToString());
 							bIsNodeDirty = true;
 						}
 					}
@@ -1290,19 +1284,81 @@ namespace Metasound
 
 			GraphBuilderPrivate::InitializeGraph(InMetaSound);
 
-			bool bIsEditorGraphDirty = SynchronizeVariables(InMetaSound);
+			bool bIsEditorGraphDirty = false;
 
 			FMetasoundAssetBase* MetaSoundAsset = IMetasoundUObjectRegistry::Get().GetObjectAsAssetBase(&InMetaSound);
 			check(MetaSoundAsset);
+			FGraphHandle GraphHandle = MetaSoundAsset->GetRootGraphHandle();
+			UMetasoundEditorGraph* EditorGraph = CastChecked<UMetasoundEditorGraph>(MetaSoundAsset->GetGraph());
+
+			TArray<UMetasoundEditorGraphInputNode*> InputNodes;
+			EditorGraph->GetNodesOfClassEx<UMetasoundEditorGraphInputNode>(InputNodes);
+			for (UMetasoundEditorGraphInputNode* Node : InputNodes)
+			{
+				check(Node);
+				FConstNodeHandle NodeHandle = Node->GetConstNodeHandle();
+				if (!NodeHandle->IsValid())
+				{
+					for (UEdGraphPin* Pin : Node->Pins)
+					{
+						check(Pin);
+
+						FConstClassInputAccessPtr ClassInputPtr = GraphHandle->FindClassInputWithName(Pin->PinName);
+						if (const FMetasoundFrontendClassInput* Input = ClassInputPtr.Get())
+						{
+							const FGuid& InitialID = Node->GetNodeID();
+							if (Node->GetNodeHandle()->GetID() != Input->NodeID)
+							{
+								Node->SetNodeID(Input->NodeID);
+
+								// Requery handle as the id has been fixed up
+								NodeHandle = Node->GetConstNodeHandle();
+								UE_LOG(LogMetasoundEditor, Verbose, TEXT("Editor Input Node '%s' interface versioned"), *NodeHandle->GetDisplayName().ToString());
+
+								bIsEditorGraphDirty = true;
+							}
+						}
+					}
+				}
+			}
+
+			TArray<UMetasoundEditorGraphOutputNode*> OutputNodes;
+			EditorGraph->GetNodesOfClassEx<UMetasoundEditorGraphOutputNode>(OutputNodes);
+			for (UMetasoundEditorGraphOutputNode* Node : OutputNodes)
+			{
+				FConstNodeHandle NodeHandle = Node->GetConstNodeHandle();
+				if (!NodeHandle->IsValid())
+				{
+					for (UEdGraphPin* Pin : Node->Pins)
+					{
+						check(Pin);
+						FConstClassOutputAccessPtr ClassOutputPtr = GraphHandle->FindClassOutputWithName(Pin->PinName);
+						if (const FMetasoundFrontendClassOutput* Output = ClassOutputPtr.Get())
+						{
+							const FGuid& InitialID = Node->GetNodeID();
+							if (Node->GetNodeHandle()->GetID() != Output->NodeID)
+							{
+								Node->SetNodeID(Output->NodeID);
+
+								// Requery handle as the id has been fixed up
+								NodeHandle = Node->GetConstNodeHandle();
+								UE_LOG(LogMetasoundEditor, Verbose, TEXT("Editor Output Node '%s' interface versioned"), *NodeHandle->GetDisplayName().ToString());
+
+								bIsEditorGraphDirty = true;
+							}
+						}
+					}
+				}
+			}
+
+			bIsEditorGraphDirty |= SynchronizeVariables(InMetaSound);
 
 			// Get all external nodes from Frontend graph.  Input and output references will only be added/synchronized
 			// if required when synchronizing connections (as they are not required to inhabit editor graph).
-			FGraphHandle GraphHandle = MetaSoundAsset->GetRootGraphHandle();
 			TArray<FNodeHandle> FrontendNodes = GraphHandle->GetNodes();
 
 			// Get all editor nodes from editor graph (some nodes on graph may *NOT* be metasound ed nodes,
 			// such as comment boxes, etc, so just get nodes of class UMetasoundEditorGraph).
-			UMetasoundEditorGraph* EditorGraph = CastChecked<UMetasoundEditorGraph>(MetaSoundAsset->GetGraph());
 			TArray<UMetasoundEditorGraphNode*> EditorNodes;
 			EditorGraph->GetNodesOfClass(EditorNodes);
 
@@ -1612,6 +1668,7 @@ namespace Metasound
 					Outputs.Add(Output);
 					return;
 				}
+
 				if (!ensure(NodeHandle->GetNumOutputs() == 1))
 				{
 					return;
diff --git a/Engine/Plugins/Runtime/Metasound/Source/MetasoundEditor/Private/MetasoundEditorGraphBuilder.h b/Engine/Plugins/Runtime/Metasound/Source/MetasoundEditor/Private/MetasoundEditorGraphBuilder.h
--- a/Engine/Plugins/Runtime/Metasound/Source/MetasoundEditor/Private/MetasoundEditorGraphBuilder.h
+++ b/Engine/Plugins/Runtime/Metasound/Source/MetasoundEditor/Private/MetasoundEditorGraphBuilder.h
@@ -81,7 +81,8 @@ namespace Metasound
 				UObject& InMetaSound,
 				const FName InTypeName,
 				const FText& InToolTip,
-				const FMetasoundFrontendLiteral* InDefaultValue = nullptr);
+				const FMetasoundFrontendLiteral* InDefaultValue = nullptr,
+				const FName* InNameBase = nullptr);
 
 			// Adds a corresponding UMetasoundEditorGraphExternalNode for the provided node handle.
 			static UMetasoundEditorGraphExternalNode* AddExternalNode(UObject& InMetaSound, Frontend::FNodeHandle& InNodeHandle, FVector2D InLocation, bool bInSelectNewNode = true);
@@ -97,7 +98,7 @@ namespace Metasound
 			static UMetasoundEditorGraphOutputNode* AddOutputNode(UObject& InMetaSound, Frontend::FNodeHandle& InNodeHandle, FVector2D InLocation, bool bInSelectNewNode = true);
 
 			// Generates analogous FNodeHandle for the given internal node data. Does not bind nor create EdGraph representation of given node.
-			static Frontend::FNodeHandle AddOutputNodeHandle(UObject& InMetaSound, const FName InTypeName, const FText& InToolTip);
+			static Frontend::FNodeHandle AddOutputNodeHandle(UObject& InMetaSound, const FName InTypeName, const FText& InToolTip, const FName* InNameBase = nullptr);
 
 			// Attempts to connect Frontend node counterparts together for provided pins.  Returns true if succeeded,
 			// and breaks pin link and returns false if failed.  If bConnectEdPins is set, will attempt to connect
@@ -110,14 +111,8 @@ namespace Metasound
 			// applicable post disconnection.
 			static void DisconnectPin(UEdGraphPin& InPin, bool bAddLiteralInputs = true);
 
-			// Generates a unique input display name for the given MetaSound object
-			static FText GenerateUniqueInputDisplayName(const UObject& InMetaSound, const FText* InBaseName = nullptr);
-
 			// Generates a unique output name for the given MetaSound object
-			static FText GenerateUniqueOutputDisplayName(const UObject& InMetaSound, const FText* InBaseName = nullptr);
-
-			// Generates a unique name by retrieving the provided MetaSound's associated GraphHandle and filtering by the provided function.
-			static FText GenerateUniqueNameByFilter(const UObject& InMetaSound, const FText& InBaseText, TFunctionRef<bool(const Frontend::FConstGraphHandle&, const FText&)> InNameIsValidFilter);
+			static FName GenerateUniqueNameByClassType(const UObject& InMetaSound, EMetasoundFrontendClassType InClassType, const FName& InBaseName);
 
 			static TArray<FString> GetDataTypeNameCategories(const FName& InDataTypeName);
 
diff --git a/Engine/Plugins/Runtime/Metasound/Source/MetasoundEditor/Private/MetasoundEditorGraphInputNodes.cpp b/Engine/Plugins/Runtime/Metasound/Source/MetasoundEditor/Private/MetasoundEditorGraphInputNodes.cpp
--- a/Engine/Plugins/Runtime/Metasound/Source/MetasoundEditor/Private/MetasoundEditorGraphInputNodes.cpp
+++ b/Engine/Plugins/Runtime/Metasound/Source/MetasoundEditor/Private/MetasoundEditorGraphInputNodes.cpp
@@ -56,7 +56,7 @@ FMetasoundFrontendClassName UMetasoundEditorGraphInputNode::GetClassName() const
 	return Super::GetClassName();
 }
 
-void UMetasoundEditorGraphInputNode::UpdatePreviewInstance(const Metasound::FVertexKey& InParameterName, TScriptInterface<IAudioParameterInterface>& InParameterInterface) const
+void UMetasoundEditorGraphInputNode::UpdatePreviewInstance(const Metasound::FVertexName& InParameterName, TScriptInterface<IAudioParameterInterface>& InParameterInterface) const
 {
 	if (Input)
 	{
@@ -124,9 +124,9 @@ void UMetasoundEditorGraphInputBool::SetFromLiteral(const FMetasoundFrontendLite
 	Metasound::Editor::InputPrivate::ConvertLiteral<bool>(InLiteral, Default.Value);
 }
 
-void UMetasoundEditorGraphInputBool::UpdatePreviewInstance(const Metasound::FVertexKey& InParameterName, TScriptInterface<IAudioParameterInterface>& InParameterInterface) const
+void UMetasoundEditorGraphInputBool::UpdatePreviewInstance(const Metasound::FVertexName& InParameterName, TScriptInterface<IAudioParameterInterface>& InParameterInterface) const
 	{
-		InParameterInterface->SetBoolParameter(*InParameterName, Default.Value);
+		InParameterInterface->SetBoolParameter(InParameterName, Default.Value);
 	}
 
 FMetasoundFrontendLiteral UMetasoundEditorGraphInputBoolArray::GetDefault() const
@@ -149,11 +149,11 @@ void UMetasoundEditorGraphInputBoolArray::SetFromLiteral(const FMetasoundFronten
 		Metasound::Editor::InputPrivate::ConvertLiteralToArray<bool, FMetasoundEditorGraphInputBoolRef>(InLiteral, Default);
 	}
 
-void UMetasoundEditorGraphInputBoolArray::UpdatePreviewInstance(const Metasound::FVertexKey& InParameterName, TScriptInterface<IAudioParameterInterface>& InParameterInterface) const
+void UMetasoundEditorGraphInputBoolArray::UpdatePreviewInstance(const Metasound::FVertexName& InParameterName, TScriptInterface<IAudioParameterInterface>& InParameterInterface) const
 {
 	TArray<bool> BoolArray;
 	Algo::Transform(Default, BoolArray, [](const FMetasoundEditorGraphInputBoolRef& InValue) { return InValue.Value; });
-	InParameterInterface->SetBoolArrayParameter(*InParameterName, BoolArray);
+	InParameterInterface->SetBoolArrayParameter(InParameterName, BoolArray);
 }
 
 FMetasoundFrontendLiteral UMetasoundEditorGraphInputInt::GetDefault() const
@@ -173,9 +173,9 @@ void UMetasoundEditorGraphInputInt::SetFromLiteral(const FMetasoundFrontendLiter
 	Metasound::Editor::InputPrivate::ConvertLiteral<int32>(InLiteral, Default.Value);
 }
 
-void UMetasoundEditorGraphInputInt::UpdatePreviewInstance(const Metasound::FVertexKey& InParameterName, TScriptInterface<IAudioParameterInterface>& InParameterInterface) const
+void UMetasoundEditorGraphInputInt::UpdatePreviewInstance(const Metasound::FVertexName& InParameterName, TScriptInterface<IAudioParameterInterface>& InParameterInterface) const
 {
-	InParameterInterface->SetIntParameter(*InParameterName, Default.Value);
+	InParameterInterface->SetIntParameter(InParameterName, Default.Value);
 }
 
 FMetasoundFrontendLiteral UMetasoundEditorGraphInputIntArray::GetDefault() const
@@ -199,11 +199,11 @@ void UMetasoundEditorGraphInputIntArray::SetFromLiteral(const FMetasoundFrontend
 	Metasound::Editor::InputPrivate::ConvertLiteralToArray<int32, FMetasoundEditorGraphInputIntRef>(InLiteral, Default);
 }
 
-void UMetasoundEditorGraphInputIntArray::UpdatePreviewInstance(const Metasound::FVertexKey& InParameterName, TScriptInterface<IAudioParameterInterface>& InParameterInterface) const
+void UMetasoundEditorGraphInputIntArray::UpdatePreviewInstance(const Metasound::FVertexName& InParameterName, TScriptInterface<IAudioParameterInterface>& InParameterInterface) const
 {
 	TArray<int32> IntArray;
 	Algo::Transform(Default, IntArray, [](const FMetasoundEditorGraphInputIntRef& InValue) { return InValue.Value; });
-	InParameterInterface->SetIntArrayParameter(*InParameterName, IntArray);
+	InParameterInterface->SetIntArrayParameter(InParameterName, IntArray);
 }
 
 FMetasoundFrontendLiteral UMetasoundEditorGraphInputFloat::GetDefault() const
@@ -223,9 +223,9 @@ void UMetasoundEditorGraphInputFloat::SetFromLiteral(const FMetasoundFrontendLit
 	Metasound::Editor::InputPrivate::ConvertLiteral<float>(InLiteral, Default);
 }
 
-void UMetasoundEditorGraphInputFloat::UpdatePreviewInstance(const Metasound::FVertexKey& InParameterName, TScriptInterface<IAudioParameterInterface>& InParameterInterface) const
+void UMetasoundEditorGraphInputFloat::UpdatePreviewInstance(const Metasound::FVertexName& InParameterName, TScriptInterface<IAudioParameterInterface>& InParameterInterface) const
 {
-	InParameterInterface->SetFloatParameter(*InParameterName, Default);
+	InParameterInterface->SetFloatParameter(InParameterName, Default);
 }
 
 FMetasoundFrontendLiteral UMetasoundEditorGraphInputFloatArray::GetDefault() const
@@ -245,9 +245,9 @@ void UMetasoundEditorGraphInputFloatArray::SetFromLiteral(const FMetasoundFronte
 	Metasound::Editor::InputPrivate::ConvertLiteralToArray<float>(InLiteral, Default);
 }
 
-void UMetasoundEditorGraphInputFloatArray::UpdatePreviewInstance(const Metasound::FVertexKey& InParameterName, TScriptInterface<IAudioParameterInterface>& InParameterInterface) const
+void UMetasoundEditorGraphInputFloatArray::UpdatePreviewInstance(const Metasound::FVertexName& InParameterName, TScriptInterface<IAudioParameterInterface>& InParameterInterface) const
 {
-	InParameterInterface->SetFloatArrayParameter(*InParameterName, Default);
+	InParameterInterface->SetFloatArrayParameter(InParameterName, Default);
 }
 
 FMetasoundFrontendLiteral UMetasoundEditorGraphInputString::GetDefault() const
@@ -267,9 +267,9 @@ void UMetasoundEditorGraphInputString::SetFromLiteral(const FMetasoundFrontendLi
 	Metasound::Editor::InputPrivate::ConvertLiteral<FString>(InLiteral, Default);
 }
 
-void UMetasoundEditorGraphInputString::UpdatePreviewInstance(const Metasound::FVertexKey& InParameterName, TScriptInterface<IAudioParameterInterface>& InParameterInterface) const
+void UMetasoundEditorGraphInputString::UpdatePreviewInstance(const Metasound::FVertexName& InParameterName, TScriptInterface<IAudioParameterInterface>& InParameterInterface) const
 {
-	InParameterInterface->SetStringParameter(*InParameterName, Default);
+	InParameterInterface->SetStringParameter(InParameterName, Default);
 }
 
 FMetasoundFrontendLiteral UMetasoundEditorGraphInputStringArray::GetDefault() const
@@ -289,9 +289,9 @@ void UMetasoundEditorGraphInputStringArray::SetFromLiteral(const FMetasoundFront
 	Metasound::Editor::InputPrivate::ConvertLiteralToArray<FString>(InLiteral, Default);
 }
 
-void UMetasoundEditorGraphInputStringArray::UpdatePreviewInstance(const Metasound::FVertexKey& InParameterName, TScriptInterface<IAudioParameterInterface>& InParameterInterface) const
+void UMetasoundEditorGraphInputStringArray::UpdatePreviewInstance(const Metasound::FVertexName& InParameterName, TScriptInterface<IAudioParameterInterface>& InParameterInterface) const
 {
-	InParameterInterface->SetStringArrayParameter(*InParameterName, Default);
+	InParameterInterface->SetStringArrayParameter(InParameterName, Default);
 }
 
 FMetasoundFrontendLiteral UMetasoundEditorGraphInputObject::GetDefault() const
@@ -311,7 +311,7 @@ void UMetasoundEditorGraphInputObject::SetFromLiteral(const FMetasoundFrontendLi
 	ensure(InLiteral.TryGet(Default.Object));
 }
 
-void UMetasoundEditorGraphInputObject::UpdatePreviewInstance(const Metasound::FVertexKey& InParameterName, TScriptInterface<IAudioParameterInterface>& InParameterInterface) const
+void UMetasoundEditorGraphInputObject::UpdatePreviewInstance(const Metasound::FVertexName& InParameterName, TScriptInterface<IAudioParameterInterface>& InParameterInterface) const
 {
 	// TODO. We need proxy object here safely.
 }
@@ -339,7 +339,7 @@ void UMetasoundEditorGraphInputObjectArray::SetFromLiteral(const FMetasoundFront
 	Algo::Transform(ObjectArray, Default, [](UObject* InValue) { return FMetasoundEditorGraphInputObjectRef { InValue }; });
 }
 
-void UMetasoundEditorGraphInputObjectArray::UpdatePreviewInstance(const Metasound::FVertexKey& InParameterName, TScriptInterface<IAudioParameterInterface>& InParameterInterface) const
+void UMetasoundEditorGraphInputObjectArray::UpdatePreviewInstance(const Metasound::FVertexName& InParameterName, TScriptInterface<IAudioParameterInterface>& InParameterInterface) const
 {
 	TArray<UObject*> ObjectArray;
 	Algo::Transform(Default, ObjectArray, [](const FMetasoundEditorGraphInputObjectRef& InValue) { return InValue.Object; });
diff --git a/Engine/Plugins/Runtime/Metasound/Source/MetasoundEditor/Private/MetasoundEditorGraphNode.cpp b/Engine/Plugins/Runtime/Metasound/Source/MetasoundEditor/Private/MetasoundEditorGraphNode.cpp
--- a/Engine/Plugins/Runtime/Metasound/Source/MetasoundEditor/Private/MetasoundEditorGraphNode.cpp
+++ b/Engine/Plugins/Runtime/Metasound/Source/MetasoundEditor/Private/MetasoundEditorGraphNode.cpp
@@ -534,8 +534,7 @@ bool UMetasoundEditorGraphExternalNode::RefreshPinMetadata()
 			// incurred look-up cost on subsequent tooltip update.
 			auto RemoveValidEntry = [&](const FMetasoundFrontendClassInput& Entry)
 			{
-				const FString PinNameString = Pin->PinName.ToString();
-				if (Entry.Name == PinNameString)
+				if (Entry.Name == Pin->GetFName())
 				{
 					Metasound::Editor::FGraphBuilder::RefreshPinMetadata(*Pin, Entry.Metadata);
 					return true;
@@ -552,8 +551,7 @@ bool UMetasoundEditorGraphExternalNode::RefreshPinMetadata()
 			// incurred look-up cost on subsequent tooltip update.
 			auto RemoveValidEntry = [&](const FMetasoundFrontendClassOutput& Entry)
 			{
-				const FString PinNameString = Pin->PinName.ToString();
-				if (Entry.Name == PinNameString)
+				if (Entry.Name == Pin->GetFName())
 				{
 					Metasound::Editor::FGraphBuilder::RefreshPinMetadata(*Pin, Entry.Metadata);
 					return true;
diff --git a/Engine/Plugins/Runtime/Metasound/Source/MetasoundEditor/Private/MetasoundEditorGraphSchema.cpp b/Engine/Plugins/Runtime/Metasound/Source/MetasoundEditor/Private/MetasoundEditorGraphSchema.cpp
--- a/Engine/Plugins/Runtime/Metasound/Source/MetasoundEditor/Private/MetasoundEditorGraphSchema.cpp
+++ b/Engine/Plugins/Runtime/Metasound/Source/MetasoundEditor/Private/MetasoundEditorGraphSchema.cpp
@@ -345,12 +345,10 @@ UEdGraphNode* FMetasoundGraphSchemaAction_PromoteToInput::PerformAction(UEdGraph
 	FMetasoundFrontendLiteral DefaultValue;
 	FGraphBuilder::GetPinLiteral(*FromPin, DefaultValue);
 
-	FNodeHandle NodeHandle = FGraphBuilder::AddInputNodeHandle(ParentMetasound, InputHandle->GetDataType(), FText::GetEmpty(), &DefaultValue);
+	const FName InputName = InputHandle->GetName();
+	FNodeHandle NodeHandle = FGraphBuilder::AddInputNodeHandle(ParentMetasound, InputHandle->GetDataType(), FText::GetEmpty(), &DefaultValue, &InputName);
 	if (ensure(NodeHandle->IsValid()))
 	{
-		const FText& InputName = InputHandle->GetDisplayName();
-		const FText NewNodeName = FGraphBuilder::GenerateUniqueInputDisplayName(ParentMetasound, &InputName);
-		NodeHandle->SetDisplayName(NewNodeName);
 		UMetasoundEditorGraphInput* Input = MetasoundGraph->FindOrAddInput(NodeHandle);
 		if (ensure(Input))
 		{
@@ -443,6 +441,7 @@ FMetasoundGraphSchemaAction_PromoteToOutput::FMetasoundGraphSchemaAction_Promote
 
 UEdGraphNode* FMetasoundGraphSchemaAction_PromoteToOutput::PerformAction(UEdGraph* ParentGraph, UEdGraphPin* FromPin, const FVector2D InLocation, bool bSelectNewNode /* = true */)
 {
+	using namespace Metasound;
 	using namespace Metasound::Editor;
 	using namespace Metasound::Frontend;
 
@@ -457,12 +456,11 @@ UEdGraphNode* FMetasoundGraphSchemaAction_PromoteToOutput::PerformAction(UEdGrap
 	UObject& ParentMetasound = MetasoundGraph->GetMetasoundChecked();
 	ParentMetasound.Modify();
 
-	FNodeHandle NodeHandle = FGraphBuilder::AddOutputNodeHandle(ParentMetasound, OutputHandle->GetDataType(), FText::GetEmpty());
+	const FVertexName OutputName = OutputHandle->GetName();
+	const FVertexName NewNodeName = FGraphBuilder::GenerateUniqueNameByClassType(ParentMetasound, EMetasoundFrontendClassType::Output, OutputName);
+	FNodeHandle NodeHandle = FGraphBuilder::AddOutputNodeHandle(ParentMetasound, OutputHandle->GetDataType(), FText::GetEmpty(), &NewNodeName);
 	if (ensure(NodeHandle->IsValid()))
 	{
-		const FText& OutputName = OutputHandle->GetDisplayName();
-		const FText NewNodeName = FGraphBuilder::GenerateUniqueOutputDisplayName(ParentMetasound, &OutputName);
-		NodeHandle->SetDisplayName(NewNodeName);
 		UMetasoundEditorGraphOutput* Output = MetasoundGraph->FindOrAddOutput(NodeHandle);
 		if (ensure(Output))
 		{
@@ -936,50 +934,43 @@ FText UMetasoundEditorGraphSchema::GetPinDisplayName(const UEdGraphPin* Pin) con
 	switch (ClassType)
 	{
 		case EMetasoundFrontendClassType::Input:
-		{
-			TArray<FOutputHandle> OutputHandles = NodeHandle->GetOutputsWithVertexName(Pin->GetName());
-			if (ensure(OutputHandles.Num() > 0))
-			{
-				return OutputHandles[0]->GetDisplayName();
-			}
-			else
-			{
-				return Super::GetPinDisplayName(Pin);
-			}
-		}
-
 		case EMetasoundFrontendClassType::Output:
 		{
-			TArray<FInputHandle> InputHandles = NodeHandle->GetInputsWithVertexName(Pin->GetName());
-			if (ensure(InputHandles.Num() > 0))
-			{
-				return InputHandles[0]->GetDisplayName();
-			}
-			else
-			{
-				return Super::GetPinDisplayName(Pin);
-			}
+			return NodeHandle->GetDisplayName();
 		}
 
 		case EMetasoundFrontendClassType::External:
 		{
-			const FString PinName = Pin->GetName();
 			if (Pin->Direction == EGPD_Input)
 			{
-				TArray<FInputHandle> InputHandles = NodeHandle->GetInputsWithVertexName(PinName);
-				if (ensure(InputHandles.Num() > 0))
+				TArray<FConstInputHandle> InputHandles = NodeHandle->GetConstInputsWithVertexName(Pin->GetFName());
+				if (ensure(!InputHandles.IsEmpty()))
 				{
-					return InputHandles[0]->GetDisplayName();
+					const FText& DisplayName = InputHandles[0]->GetDisplayName();
+					if (!DisplayName.IsEmptyOrWhitespace())
+					{
+						return DisplayName;
+					}
+
+					return FText::FromName(InputHandles[0]->GetName());
 				}
 			}
 			else
 			{
-				TArray<FOutputHandle> OutputHandles = NodeHandle->GetOutputsWithVertexName(PinName);
-				if (ensure(OutputHandles.Num() > 0))
+				TArray<FConstOutputHandle> OutputHandles = NodeHandle->GetConstOutputsWithVertexName(Pin->GetFName());
+				if (ensure(!OutputHandles.IsEmpty()))
 				{
-					return OutputHandles[0]->GetDisplayName();
+					const FText& DisplayName = OutputHandles[0]->GetDisplayName();
+					if (!DisplayName.IsEmptyOrWhitespace())
+					{
+						return DisplayName;
+					}
+
+					return FText::FromName(OutputHandles[0]->GetName());
 				}
 			}
+
+			return Super::GetPinDisplayName(Pin);
 		}
 
 		case EMetasoundFrontendClassType::Variable:
diff --git a/Engine/Plugins/Runtime/Metasound/Source/MetasoundEditor/Private/MetasoundNodeDetailCustomization.cpp b/Engine/Plugins/Runtime/Metasound/Source/MetasoundEditor/Private/MetasoundNodeDetailCustomization.cpp
--- a/Engine/Plugins/Runtime/Metasound/Source/MetasoundEditor/Private/MetasoundNodeDetailCustomization.cpp
+++ b/Engine/Plugins/Runtime/Metasound/Source/MetasoundEditor/Private/MetasoundNodeDetailCustomization.cpp
@@ -65,7 +65,10 @@ namespace Metasound
 			static const FText NodeTooltipText = LOCTEXT("Node_Tooltip", "Tooltip");
 
 			static const FText InputNameText = LOCTEXT("Input_Name", "Input Name");
+			static const FText InputDisplayNameText = LOCTEXT("InputDisplay_Name", "Input Display Name");
+
 			static const FText OutputNameText = LOCTEXT("Output_Name", "Output Name");
+			static const FText OutputDisplayNameText = LOCTEXT("OutputDisplay_Name", "Output Display Name");
 
 			static const FName DataTypeNameIdentifier = "DataTypeName";
 			static const FName ProxyGeneratorClassNameIdentifier = "GeneratorClass";
@@ -596,9 +599,15 @@ namespace Metasound
 
 			const bool bIsRequired = IsRequired();
 			const bool bIsGraphEditable = IsGraphEditable();
+			NameEditableTextBox = SNew(SEditableTextBox)
+				.Text(this, &FMetasoundInputDetailCustomization::GetName)
+				.OnTextChanged(this, &FMetasoundInputDetailCustomization::OnNameChanged)
+				.OnTextCommitted(this, &FMetasoundInputDetailCustomization::OnNameCommitted)
+				.IsReadOnly(bIsRequired || !bIsGraphEditable)
+				.Font(IDetailLayoutBuilder::GetDetailFont());
+
 			DisplayNameEditableTextBox = SNew(SEditableTextBox)
 				.Text(this, &FMetasoundInputDetailCustomization::GetDisplayName)
-				.OnTextChanged(this, &FMetasoundInputDetailCustomization::OnDisplayNameChanged)
 				.OnTextCommitted(this, &FMetasoundInputDetailCustomization::OnDisplayNameCommitted)
 				.IsReadOnly(bIsRequired || !bIsGraphEditable)
 				.Font(IDetailLayoutBuilder::GetDetailFont());
@@ -610,23 +619,28 @@ namespace Metasound
 				SNew(STextBlock)
 				.Font(IDetailLayoutBuilder::GetDetailFontBold())
 				.Text(VariableCustomizationPrivate::InputNameText)
-				.ToolTipText(TAttribute<FText>::Create([GraphVariable = this->GraphVariable]()
-				{
-					if (GraphVariable.IsValid())
-					{
-						FNodeHandle NodeHandle = GraphVariable->GetNodeHandle();
-						FMetasoundFrontendNodeStyle NodeStyle = NodeHandle->GetNodeStyle();
-						return NodeHandle->GetDescription();
-					}
-
-					return FText::GetEmpty();
-				}))
+				.ToolTipText(LOCTEXT("InputName_Description", "Name used by external systems to identify input. Used as DisplayName within MetaSound Graph Editor if no DisplayName is provided."))
 			]
 			.ValueContent()
 			[
-				DisplayNameEditableTextBox.ToSharedRef()
+				NameEditableTextBox.ToSharedRef()
 			];
 
+			// TODO: Enable and use proper FText property editor
+// 			CategoryBuilder.AddCustomRow(VariableCustomizationPrivate::InputDisplayNameText)
+// 			.EditCondition(!bIsRequired && bIsGraphEditable, nullptr)
+// 			.NameContent()
+// 			[
+// 				SNew(STextBlock)
+// 				.Font(IDetailLayoutBuilder::GetDetailFontBold())
+// 				.Text(VariableCustomizationPrivate::InputDisplayNameText)
+// 				.ToolTipText(LOCTEXT("InputDisplayName_Description", "Optional, localized name used within the MetaSounds editor(s) to describe the given input."))
+// 			]
+// 			.ValueContent()
+// 			[
+// 				DisplayNameEditableTextBox.ToSharedRef()
+// 			];
+
 			CategoryBuilder.AddCustomRow(VariableCustomizationPrivate::NodeTooltipText)
 			.EditCondition(!bIsRequired && bIsGraphEditable, nullptr)
 			.NameContent()
@@ -649,21 +663,6 @@ namespace Metasound
 
 			AddDataTypeSelector(DetailLayout, VariableCustomizationPrivate::DataTypeNameText, GraphVariable, !bIsRequired && bIsGraphEditable);
 
-// 			CategoryBuilder.AddCustomRow(LOCTEXT("InputPrivate", "Private"))
-// 			.Visibility(TAttribute<EVisibility>(EVisibility::Hidden))
-// 			.NameContent()
-// 			[
-// 				SNew(STextBlock)
-// 				.Text(LOCTEXT("InputPrivate", "Private"))
-// 				.Font(IDetailLayoutBuilder::GetDetailFont())
-// 			]
-// 			.ValueContent()
-// 			[
-// 				SNew(SCheckBox)
-// 				.IsChecked(this, &FMetasoundInputDetailCustomization::OnGetPrivateCheckboxState)
-// 				.OnCheckStateChanged(this, &FMetasoundInputDetailCustomization::OnPrivateChanged)
-// 			];
-
 			FNodeHandle NodeHandle = GraphVariable->GetNodeHandle();
 			const TArray<FOutputHandle>& Outputs = NodeHandle->GetOutputs();
 			if (!ensure(!Outputs.IsEmpty()))
@@ -878,9 +877,16 @@ namespace Metasound
 
 			const bool bIsRequired = IsRequired();
 			const bool bIsGraphEditable = IsGraphEditable();
+
+			NameEditableTextBox = SNew(SEditableTextBox)
+				.Text(this, &FMetasoundOutputDetailCustomization::GetName)
+				.OnTextChanged(this, &FMetasoundOutputDetailCustomization::OnNameChanged)
+				.OnTextCommitted(this, &FMetasoundOutputDetailCustomization::OnNameCommitted)
+				.IsReadOnly(bIsRequired || !bIsGraphEditable)
+				.Font(IDetailLayoutBuilder::GetDetailFont());
+
 			DisplayNameEditableTextBox = SNew(SEditableTextBox)
 				.Text(this, &FMetasoundOutputDetailCustomization::GetDisplayName)
-				.OnTextChanged(this, &FMetasoundOutputDetailCustomization::OnDisplayNameChanged)
 				.OnTextCommitted(this, &FMetasoundOutputDetailCustomization::OnDisplayNameCommitted)
 				.IsReadOnly(bIsRequired || !bIsGraphEditable)
 				.Font(IDetailLayoutBuilder::GetDetailFont());
@@ -892,22 +898,28 @@ namespace Metasound
 				SNew(STextBlock)
 				.Font(IDetailLayoutBuilder::GetDetailFontBold())
 				.Text(VariableCustomizationPrivate::OutputNameText)
-				.ToolTipText(TAttribute<FText>::Create([GraphVariable = this->GraphVariable]()
-				{
-					if (GraphVariable.IsValid())
-					{
-						FNodeHandle NodeHandle = GraphVariable->GetNodeHandle();
-						return NodeHandle->GetDescription();
-					}
-
-					return FText::GetEmpty();
-				}))
+				.ToolTipText(LOCTEXT("OutputName_Description", "Name used by external systems to identify output. Used as DisplayName within MetaSound Graph Editor if no DisplayName is provided."))
 			]
 			.ValueContent()
 			[
-				DisplayNameEditableTextBox.ToSharedRef()
+				NameEditableTextBox.ToSharedRef()
 			];
 
+			// TODO: Enable and use proper FText property editor
+// 			CategoryBuilder.AddCustomRow(VariableCustomizationPrivate::OutputDisplayNameText)
+// 			.EditCondition(!bIsRequired && bIsGraphEditable, nullptr)
+// 			.NameContent()
+// 			[
+// 				SNew(STextBlock)
+// 				.Font(IDetailLayoutBuilder::GetDetailFontBold())
+// 				.Text(VariableCustomizationPrivate::OutputDisplayNameText)
+// 				.ToolTipText(LOCTEXT("OutputDisplayName_Description", "Optional, localized name used within the MetaSounds editor(s) to describe the given output."))
+// 			]
+// 			.ValueContent()
+// 			[
+// 				DisplayNameEditableTextBox.ToSharedRef()
+// 			];
+
 			CategoryBuilder.AddCustomRow(VariableCustomizationPrivate::NodeTooltipText)
 			.EditCondition(!bIsRequired && bIsGraphEditable, nullptr)
 			.NameContent()
diff --git a/Engine/Plugins/Runtime/Metasound/Source/MetasoundEditor/Private/MetasoundNodeDetailCustomization.h b/Engine/Plugins/Runtime/Metasound/Source/MetasoundEditor/Private/MetasoundNodeDetailCustomization.h
--- a/Engine/Plugins/Runtime/Metasound/Source/MetasoundEditor/Private/MetasoundNodeDetailCustomization.h
+++ b/Engine/Plugins/Runtime/Metasound/Source/MetasoundEditor/Private/MetasoundNodeDetailCustomization.h
@@ -114,6 +114,7 @@ namespace Metasound
 			FText VariableLabel;
 
 			TWeakObjectPtr<VariableType> GraphVariable;
+			TSharedPtr<SEditableTextBox> NameEditableTextBox;
 			TSharedPtr<SEditableTextBox> DisplayNameEditableTextBox;
 			bool bIsNameInvalid = false;
 
@@ -131,7 +132,7 @@ namespace Metasound
 			}
 			// End of IDetailCustomization interface
 
-			void OnDisplayNameChanged(const FText& InNewName)
+			void OnNameChanged(const FText& InNewName)
 			{
 				using namespace Frontend;
 
@@ -163,6 +164,18 @@ namespace Metasound
 				return FText::GetEmpty();
 			}
 
+			FText GetName() const
+			{
+				using namespace Frontend;
+
+				if (GraphVariable.IsValid())
+				{
+					return FText::FromName(GraphVariable->GetConstNodeHandle()->GetNodeName());
+				}
+
+				return FText::GetEmpty();
+			}
+
 			bool IsGraphEditable() const
 			{
 				if (GraphVariable.IsValid())
@@ -219,6 +232,19 @@ namespace Metasound
 				bIsNameInvalid = false;
 			}
 
+			void OnNameCommitted(const FText& InNewName, ETextCommit::Type InTextCommit)
+			{
+				using namespace Frontend;
+
+				if (!bIsNameInvalid && GraphVariable.IsValid())
+				{
+					GraphVariable->SetNodeName(*InNewName.ToString());
+				}
+
+				DisplayNameEditableTextBox->SetError(FText::GetEmpty());
+				bIsNameInvalid = false;
+			}
+
 			ECheckBoxState OnGetPrivateCheckboxState() const
 			{
 				if (GraphVariable.IsValid())
diff --git a/Engine/Plugins/Runtime/Metasound/Source/MetasoundEditor/Private/SMetasoundEnumPin.cpp b/Engine/Plugins/Runtime/Metasound/Source/MetasoundEditor/Private/SMetasoundEnumPin.cpp
--- a/Engine/Plugins/Runtime/Metasound/Source/MetasoundEditor/Private/SMetasoundEnumPin.cpp
+++ b/Engine/Plugins/Runtime/Metasound/Source/MetasoundEditor/Private/SMetasoundEnumPin.cpp
@@ -33,7 +33,7 @@ SMetasoundEnumPin::FindEnumInterfaceFromPin(UEdGraphPin* InPin)
 {
 	auto MetasoundEditorNode = Cast<UMetasoundEditorGraphNode>(InPin->GetOwningNode());
 	Metasound::Frontend::FNodeHandle NodeHandle = MetasoundEditorNode->GetNodeHandle();
-	auto Inputs = NodeHandle->GetInputsWithVertexName(InPin->GetName());
+	auto Inputs = NodeHandle->GetInputsWithVertexName(InPin->GetFName());
 	if (Inputs.Num() > 0)
 	{
 		FName DataType = Inputs[0]->GetDataType();
diff --git a/Engine/Plugins/Runtime/Metasound/Source/MetasoundEditor/Private/SMetasoundGraphNode.cpp b/Engine/Plugins/Runtime/Metasound/Source/MetasoundEditor/Private/SMetasoundGraphNode.cpp
--- a/Engine/Plugins/Runtime/Metasound/Source/MetasoundEditor/Private/SMetasoundGraphNode.cpp
+++ b/Engine/Plugins/Runtime/Metasound/Source/MetasoundEditor/Private/SMetasoundGraphNode.cpp
@@ -71,11 +71,8 @@ void SMetasoundGraphNode::ExecuteInputTrigger(UMetasoundEditorGraphInputLiteral&
 
 	if (UAudioComponent* PreviewComponent = GEditor->GetPreviewAudioComponent())
 	{
-		// TODO: fix how identifying the parameter to update is determined. It should not be done
-		// with a "DisplayName" but rather the vertex Guid.
 		Metasound::Frontend::FConstNodeHandle NodeHandle = Input->GetConstNodeHandle();
-		Metasound::FVertexKey VertexKey = Metasound::FVertexKey(NodeHandle->GetDisplayName().ToString());
-		PreviewComponent->SetTriggerParameter(*VertexKey);
+		PreviewComponent->SetTriggerParameter(NodeHandle->GetNodeName());
 	}
 }
 
diff --git a/Engine/Plugins/Runtime/Metasound/Source/MetasoundEditor/Public/MetasoundEditorGraph.h b/Engine/Plugins/Runtime/Metasound/Source/MetasoundEditor/Public/MetasoundEditorGraph.h
--- a/Engine/Plugins/Runtime/Metasound/Source/MetasoundEditor/Public/MetasoundEditorGraph.h
+++ b/Engine/Plugins/Runtime/Metasound/Source/MetasoundEditor/Public/MetasoundEditorGraph.h
@@ -43,7 +43,7 @@ class METASOUNDEDITOR_API UMetasoundEditorGraphVariable : public UObject
 	GENERATED_BODY()
 
 protected:
-	virtual Metasound::Frontend::FNodeHandle AddNodeHandle(const FText& InNodeDisplayName, FName InDataType)
+	virtual Metasound::Frontend::FNodeHandle AddNodeHandle(const FName& InNodeName, FName InDataType)
 	{
 		return Metasound::Frontend::INodeController::GetInvalidHandle();
 	}
@@ -70,6 +70,7 @@ public:
 	void SetDataType(FName InNewType);
 	void SetDescription(const FText& InDescription);
 	void SetDisplayName(const FText& InNewName);
+	void SetNodeName(const FName& InNewName);
 
 	virtual EMetasoundFrontendClassType GetClassType() const { return EMetasoundFrontendClassType::Invalid; }
 	virtual void OnDataTypeChanged() { }
@@ -82,7 +83,7 @@ class METASOUNDEDITOR_API UMetasoundEditorGraphInputLiteral : public UObject
 	GENERATED_BODY()
 
 public:
-	virtual void UpdatePreviewInstance(const Metasound::FVertexKey& InParameterName, TScriptInterface<IAudioParameterInterface>& InParameterInterface) const
+	virtual void UpdatePreviewInstance(const Metasound::FVertexName& InParameterName, TScriptInterface<IAudioParameterInterface>& InParameterInterface) const
 	{
 	}
 
@@ -109,7 +110,7 @@ class METASOUNDEDITOR_API UMetasoundEditorGraphInput : public UMetasoundEditorGr
 	GENERATED_BODY()
 
 protected:
-	virtual Metasound::Frontend::FNodeHandle AddNodeHandle(const FText& InNodeDisplayName, FName InDataType) override;
+	virtual Metasound::Frontend::FNodeHandle AddNodeHandle(const FName& InNodeName, FName InDataType) override;
 	virtual EMetasoundFrontendClassType GetClassType() const override { return EMetasoundFrontendClassType::Input; }
 	virtual const FText& GetVariableLabel() const override;
 public:
@@ -117,7 +118,7 @@ public:
 	UMetasoundEditorGraphInputLiteral* Literal;
 
 	void UpdateDocumentInput(bool bPostTransaction = true);
-	void UpdatePreviewInstance(const Metasound::FVertexKey& InParameterName, TScriptInterface<IAudioParameterInterface>& InParameterInterface) const;
+	void UpdatePreviewInstance(const Metasound::FVertexName& InParameterName, TScriptInterface<IAudioParameterInterface>& InParameterInterface) const;
 
 	void OnDataTypeChanged() override;
 	void OnLiteralChanged(bool bPostTransaction = true);
@@ -133,7 +134,7 @@ class METASOUNDEDITOR_API UMetasoundEditorGraphOutput : public UMetasoundEditorG
 	GENERATED_BODY()
 
 protected:
-	virtual Metasound::Frontend::FNodeHandle AddNodeHandle(const FText& InNodeDisplayName, FName InDataType) override;
+	virtual Metasound::Frontend::FNodeHandle AddNodeHandle(const FName& InNodeName, FName InDataType) override;
 	virtual EMetasoundFrontendClassType GetClassType() const override { return EMetasoundFrontendClassType::Output; }
 	virtual const FText& GetVariableLabel() const override;
 };
@@ -181,9 +182,11 @@ private:
 
 public:
 	UMetasoundEditorGraphInput* FindInput(FGuid InNodeID) const;
+	UMetasoundEditorGraphInput* FindInput(FName InName) const;
 	UMetasoundEditorGraphInput* FindOrAddInput(Metasound::Frontend::FNodeHandle InNodeHandle);
 
 	UMetasoundEditorGraphOutput* FindOutput(FGuid InNodeID) const;
+	UMetasoundEditorGraphOutput* FindOutput(FName InName) const;
 	UMetasoundEditorGraphOutput* FindOrAddOutput(Metasound::Frontend::FNodeHandle InNodeHandle);
 
 	UMetasoundEditorGraphVariable* FindVariable(FGuid InNodeID) const;
diff --git a/Engine/Plugins/Runtime/Metasound/Source/MetasoundEditor/Public/MetasoundEditorGraphInputNodes.h b/Engine/Plugins/Runtime/Metasound/Source/MetasoundEditor/Public/MetasoundEditorGraphInputNodes.h
--- a/Engine/Plugins/Runtime/Metasound/Source/MetasoundEditor/Public/MetasoundEditorGraphInputNodes.h
+++ b/Engine/Plugins/Runtime/Metasound/Source/MetasoundEditor/Public/MetasoundEditorGraphInputNodes.h
@@ -16,6 +16,7 @@
 #include "MetasoundFrontendLiteral.h"
 #include "MetasoundFrontendRegistries.h"
 #include "MetasoundPrimitives.h"
+#include "MetasoundVertex.h"
 #include "Misc/Guid.h"
 #include "Sound/SoundWave.h"
 #include "UObject/ObjectMacros.h"
@@ -40,7 +41,7 @@ public:
 	TObjectPtr<UMetasoundEditorGraphInput> Input;
 
 public:
-	virtual void UpdatePreviewInstance(const Metasound::FVertexKey& InParameterName, TScriptInterface<IAudioParameterInterface>& InParameterInterface) const;
+	virtual void UpdatePreviewInstance(const Metasound::FVertexName& InParameterName, TScriptInterface<IAudioParameterInterface>& InParameterInterface) const;
 
 	virtual FMetasoundFrontendClassName GetClassName() const;
 
@@ -83,7 +84,7 @@ public:
 	virtual FMetasoundFrontendLiteral GetDefault() const override;
 	virtual EMetasoundFrontendLiteralType GetLiteralType() const override;
 	virtual void SetFromLiteral(const FMetasoundFrontendLiteral& InLiteral) override;
-	virtual void UpdatePreviewInstance(const Metasound::FVertexKey& InParameterName, TScriptInterface<IAudioParameterInterface>& InParameterInterface) const override;
+	virtual void UpdatePreviewInstance(const Metasound::FVertexName& InParameterName, TScriptInterface<IAudioParameterInterface>& InParameterInterface) const override;
 };
 
 UCLASS(MinimalAPI)
@@ -100,7 +101,7 @@ public:
 	virtual FMetasoundFrontendLiteral GetDefault() const override;
 	virtual EMetasoundFrontendLiteralType GetLiteralType() const override;
 	virtual void SetFromLiteral(const FMetasoundFrontendLiteral& InLiteral) override;
-	virtual void UpdatePreviewInstance(const Metasound::FVertexKey& InParameterName, TScriptInterface<IAudioParameterInterface>& InParameterInterface) const override;
+	virtual void UpdatePreviewInstance(const Metasound::FVertexName& InParameterName, TScriptInterface<IAudioParameterInterface>& InParameterInterface) const override;
 };
 
 // Broken out to be able to customize and swap enum behavior for basic integer literal behavior
@@ -127,7 +128,7 @@ public:
 	virtual FMetasoundFrontendLiteral GetDefault() const override;
 	virtual EMetasoundFrontendLiteralType GetLiteralType() const override;
 	virtual void SetFromLiteral(const FMetasoundFrontendLiteral& InLiteral) override;
-	virtual void UpdatePreviewInstance(const Metasound::FVertexKey& InParameterName, TScriptInterface<IAudioParameterInterface>& InParameterInterface) const override;
+	virtual void UpdatePreviewInstance(const Metasound::FVertexName& InParameterName, TScriptInterface<IAudioParameterInterface>& InParameterInterface) const override;
 };
 
 UCLASS(MinimalAPI)
@@ -144,7 +145,7 @@ public:
 	virtual FMetasoundFrontendLiteral GetDefault() const override;
 	virtual EMetasoundFrontendLiteralType GetLiteralType() const override;
 	virtual void SetFromLiteral(const FMetasoundFrontendLiteral& InLiteral) override;
-	virtual void UpdatePreviewInstance(const Metasound::FVertexKey& InParameterName, TScriptInterface<IAudioParameterInterface>& InParameterInterface) const override;
+	virtual void UpdatePreviewInstance(const Metasound::FVertexName& InParameterName, TScriptInterface<IAudioParameterInterface>& InParameterInterface) const override;
 };
 
 UCLASS(MinimalAPI)
@@ -161,7 +162,7 @@ public:
 	virtual FMetasoundFrontendLiteral GetDefault() const override;
 	virtual EMetasoundFrontendLiteralType GetLiteralType() const override;
 	virtual void SetFromLiteral(const FMetasoundFrontendLiteral& InLiteral) override;
-	virtual void UpdatePreviewInstance(const Metasound::FVertexKey& InParameterName, TScriptInterface<IAudioParameterInterface>& InParameterInterface) const override;
+	virtual void UpdatePreviewInstance(const Metasound::FVertexName& InParameterName, TScriptInterface<IAudioParameterInterface>& InParameterInterface) const override;
 };
 
 UCLASS(MinimalAPI)
@@ -178,7 +179,7 @@ public:
 	virtual FMetasoundFrontendLiteral GetDefault() const override;
 	virtual EMetasoundFrontendLiteralType GetLiteralType() const override;
 	virtual void SetFromLiteral(const FMetasoundFrontendLiteral& InLiteral) override;
-	virtual void UpdatePreviewInstance(const Metasound::FVertexKey& InParameterName, TScriptInterface<IAudioParameterInterface>& InParameterInterface) const override;
+	virtual void UpdatePreviewInstance(const Metasound::FVertexName& InParameterName, TScriptInterface<IAudioParameterInterface>& InParameterInterface) const override;
 };
 
 UCLASS(MinimalAPI)
@@ -195,7 +196,7 @@ public:
 	virtual FMetasoundFrontendLiteral GetDefault() const override;
 	virtual EMetasoundFrontendLiteralType GetLiteralType() const override;
 	virtual void SetFromLiteral(const FMetasoundFrontendLiteral& InLiteral) override;
-	virtual void UpdatePreviewInstance(const Metasound::FVertexKey& InParameterName, TScriptInterface<IAudioParameterInterface>& InParameterInterface) const override;
+	virtual void UpdatePreviewInstance(const Metasound::FVertexName& InParameterName, TScriptInterface<IAudioParameterInterface>& InParameterInterface) const override;
 };
 
 UCLASS(MinimalAPI)
@@ -212,7 +213,7 @@ public:
 	virtual FMetasoundFrontendLiteral GetDefault() const override;
 	virtual EMetasoundFrontendLiteralType GetLiteralType() const override;
 	virtual void SetFromLiteral(const FMetasoundFrontendLiteral& InLiteral) override;
-	virtual void UpdatePreviewInstance(const Metasound::FVertexKey& InParameterName, TScriptInterface<IAudioParameterInterface>& InParameterInterface) const override;
+	virtual void UpdatePreviewInstance(const Metasound::FVertexName& InParameterName, TScriptInterface<IAudioParameterInterface>& InParameterInterface) const override;
 };
 
 // Broken out to be able to customize and swap AllowedClass based on provided object proxy
@@ -239,7 +240,7 @@ public:
 	virtual FMetasoundFrontendLiteral GetDefault() const override;
 	virtual EMetasoundFrontendLiteralType GetLiteralType() const override;
 	virtual void SetFromLiteral(const FMetasoundFrontendLiteral& InLiteral) override;
-	virtual void UpdatePreviewInstance(const Metasound::FVertexKey& InParameterName, TScriptInterface<IAudioParameterInterface>& InParameterInterface) const override;
+	virtual void UpdatePreviewInstance(const Metasound::FVertexName& InParameterName, TScriptInterface<IAudioParameterInterface>& InParameterInterface) const override;
 };
 
 UCLASS(MinimalAPI)
@@ -256,5 +257,5 @@ public:
 	virtual FMetasoundFrontendLiteral GetDefault() const override;
 	virtual EMetasoundFrontendLiteralType GetLiteralType() const override;
 	virtual void SetFromLiteral(const FMetasoundFrontendLiteral& InLiteral) override;
-	virtual void UpdatePreviewInstance(const Metasound::FVertexKey& InParameterName, TScriptInterface<IAudioParameterInterface>& InParameterInterface) const override;
+	virtual void UpdatePreviewInstance(const Metasound::FVertexName& InParameterName, TScriptInterface<IAudioParameterInterface>& InParameterInterface) const override;
 };
\ No newline at end of file
diff --git a/Engine/Plugins/Runtime/Metasound/Source/MetasoundEngine/Private/MetasoundEngineArchetypes.cpp b/Engine/Plugins/Runtime/Metasound/Source/MetasoundEngine/Private/MetasoundEngineArchetypes.cpp
--- a/Engine/Plugins/Runtime/Metasound/Source/MetasoundEngine/Private/MetasoundEngineArchetypes.cpp
+++ b/Engine/Plugins/Runtime/Metasound/Source/MetasoundEngine/Private/MetasoundEngineArchetypes.cpp
@@ -75,39 +75,39 @@ namespace Metasound
 
 			}
 
-			const FString& GetOnPlayInputName()
+			const FVertexName& GetOnPlayInputName()
 			{
-				static const FString TriggerInputName = TEXT("On Play");
+				static const FVertexName TriggerInputName = TEXT("On Play");
 				return TriggerInputName;
 			}
 
-			const FString& GetAudioOutputName()
+			const FVertexName& GetAudioOutputName()
 			{
-				static const FString AudioOutputName = TEXT("Generated Audio");
+				static const FVertexName AudioOutputName = TEXT("Generated Audio");
 				return AudioOutputName;
 			}
 
-			const FString& GetIsFinishedOutputName()
+			const FVertexName& GetIsFinishedOutputName()
 			{
-				static const FString OnFinishedOutputName = TEXT("On Finished");
+				static const FVertexName OnFinishedOutputName = TEXT("On Finished");
 				return OnFinishedOutputName;
 			}
 
-			const FString& GetAudioDeviceIDVariableName()
+			const FVertexName& GetAudioDeviceIDVariableName()
 			{
-				static const FString AudioDeviceIDVarName = TEXT("AudioDeviceID");
+				static const FVertexName AudioDeviceIDVarName = TEXT("AudioDeviceID");
 				return AudioDeviceIDVarName;
 			}
 
-			const FString& GetSoundUniqueIdName()
+			const FVertexName& GetSoundUniqueIdName()
 			{
-				static const FString SoundUniqueIdVarName = TEXT("SoundUniqueId");
+				static const FVertexName SoundUniqueIdVarName = TEXT("SoundUniqueId");
 				return SoundUniqueIdVarName;
 			}
 
-			const FString& GetIsPreviewSoundName()
+			const FVertexName& GetIsPreviewSoundName()
 			{
-				static const FString SoundIsPreviewSoundName = TEXT("IsPreviewSound");
+				static const FVertexName SoundIsPreviewSoundName = TEXT("IsPreviewSound");
 				return SoundIsPreviewSoundName;
 			}
 
@@ -151,8 +151,7 @@ namespace Metasound
 				// Environment 
 				FMetasoundFrontendEnvironmentVariable AudioDeviceID;
 				AudioDeviceID.Name = GetAudioDeviceIDVariableName();
-				AudioDeviceID.Metadata.DisplayName = FText::FromString(AudioDeviceID.Name);
-				AudioDeviceID.Metadata.Description = LOCTEXT("AudioDeviceIDToolTip", "Audio device handle");
+				AudioDeviceID.Metadata.Description = LOCTEXT("AudioDeviceIDToolTip", "Audio Device ID");
 
 				Archetype.Interface.Environment.Add(AudioDeviceID);
 
@@ -168,45 +167,45 @@ namespace Metasound
 				return FMetasoundFrontendVersion{ VersionName, FMetasoundFrontendVersionNumber{1, 0} };
 			}
 
-			const FString& GetOnPlayInputName()
+			const FVertexName& GetOnPlayInputName()
 			{
-				static const FString TriggerInputName = TEXT("On Play");
+				static const FVertexName TriggerInputName = TEXT("On Play");
 				return TriggerInputName;
 			}
 
-			const FString& GetAudioOutputName()
+			const FVertexName& GetAudioOutputName()
 			{
-				static const FString AudioOutputName = TEXT("Generated Audio");
+				static const FVertexName AudioOutputName = TEXT("Generated Audio");
 				return AudioOutputName;
 			}
 
-			const FString& GetIsFinishedOutputName()
+			const FVertexName& GetIsFinishedOutputName()
 			{
-				static const FString OnFinishedOutputName = TEXT("On Finished");
+				static const FVertexName OnFinishedOutputName = TEXT("On Finished");
 				return OnFinishedOutputName;
 			}
 
-			const FString& GetAudioDeviceIDVariableName()
+			const FVertexName& GetAudioDeviceIDVariableName()
 			{
-				static const FString AudioDeviceIDVarName = TEXT("AudioDeviceID");
+				static const FVertexName AudioDeviceIDVarName = TEXT("AudioDeviceID");
 				return AudioDeviceIDVarName;
 			}
 
-			const FString& GetSoundUniqueIdName()
+			const FVertexName& GetSoundUniqueIdName()
 			{
-				static const FString SoundUniqueIdVarName = TEXT("SoundUniqueId");
+				static const FVertexName SoundUniqueIdVarName = TEXT("SoundUniqueId");
 				return SoundUniqueIdVarName;
 			}
 
-			const FString& GetIsPreviewSoundName()
+			const FVertexName& GetIsPreviewSoundName()
 			{
-				static const FString SoundIsPreviewSoundName = TEXT("IsPreviewSound");
+				static const FVertexName SoundIsPreviewSoundName = TEXT("IsPreviewSound");
 				return SoundIsPreviewSoundName;
 			}
 
-			const FString& GetGraphName()
+			const FVertexName& GetGraphName()
 			{
-				static const FString SoundGraphName = TEXT("GraphName");
+				static const FVertexName SoundGraphName = TEXT("GraphName");
 				return SoundGraphName;
 			}
 
@@ -249,8 +248,7 @@ namespace Metasound
 				// Environment
 				FMetasoundFrontendEnvironmentVariable AudioDeviceID;
 				AudioDeviceID.Name = GetAudioDeviceIDVariableName();
-				AudioDeviceID.Metadata.DisplayName = FText::FromString(AudioDeviceID.Name);
-				AudioDeviceID.Metadata.Description = LOCTEXT("AudioDeviceIDToolTip", "Audio device handle");
+				AudioDeviceID.Metadata.Description = LOCTEXT("AudioDeviceIDToolTip", "Audio Device ID");
 
 				Archetype.Interface.Environment.Add(AudioDeviceID);
 
@@ -266,51 +264,51 @@ namespace Metasound
 				return FMetasoundFrontendVersion{ VersionName, FMetasoundFrontendVersionNumber{1, 1} };
 			}
 
-			const FString& GetOnPlayInputName()
+			const FVertexName& GetOnPlayInputName()
 			{
-				static const FString TriggerInputName = TEXT("On Play");
+				static const FVertexName TriggerInputName = TEXT("On Play");
 				return TriggerInputName;
 			}
 
-			const FString& GetAudioOutputName()
+			const FVertexName& GetAudioOutputName()
 			{
-				static const FString AudioOutputName = TEXT("Audio:0");
+				static const FVertexName AudioOutputName = TEXT("Audio:0");
 				return AudioOutputName;
 			}
 
-			const FString& GetIsFinishedOutputName()
+			const FVertexName& GetIsFinishedOutputName()
 			{
-				static const FString OnFinishedOutputName = TEXT("On Finished");
+				static const FVertexName OnFinishedOutputName = TEXT("On Finished");
 				return OnFinishedOutputName;
 			}
 
-			const FString& GetAudioDeviceIDVariableName()
+			const FVertexName& GetAudioDeviceIDVariableName()
 			{
-				static const FString AudioDeviceIDVarName = TEXT("AudioDeviceID");
+				static const FVertexName AudioDeviceIDVarName = TEXT("AudioDeviceID");
 				return AudioDeviceIDVarName;
 			}
 
-			const FString& GetSoundUniqueIdName()
+			const FVertexName& GetSoundUniqueIdName()
 			{
-				static const FString SoundUniqueIdVarName = TEXT("SoundUniqueId");
+				static const FVertexName SoundUniqueIdVarName = TEXT("SoundUniqueId");
 				return SoundUniqueIdVarName;
 			}
 
-			const FString& GetIsPreviewSoundName()
+			const FVertexName& GetIsPreviewSoundName()
 			{
-				static const FString SoundIsPreviewSoundName = TEXT("IsPreviewSound");
+				static const FVertexName SoundIsPreviewSoundName = TEXT("IsPreviewSound");
 				return SoundIsPreviewSoundName;
 			}
 
-			const FString& GetGraphName()
+			const FVertexName& GetGraphName()
 			{
-				static const FString SoundGraphName = TEXT("GraphName");
+				static const FVertexName SoundGraphName = TEXT("GraphName");
 				return SoundGraphName;
 			}
 
-			const FString& GetInstanceIDName()
+			const FVertexName& GetInstanceIDName()
 			{
-				static const FString& Name = FMetaSoundParameterTransmitter::GetInstanceIDEnvironmentVariableName();
+				static const FVertexName& Name = FMetaSoundParameterTransmitter::GetInstanceIDEnvironmentVariableName();
 				return Name;
 			}
 
@@ -359,8 +357,7 @@ namespace Metasound
 				// Environment
 				FMetasoundFrontendEnvironmentVariable AudioDeviceID;
 				AudioDeviceID.Name = GetAudioDeviceIDVariableName();
-				AudioDeviceID.Metadata.DisplayName = FText::FromString(AudioDeviceID.Name);
-				AudioDeviceID.Metadata.Description = LOCTEXT("AudioDeviceIDToolTip", "Audio device handle");
+				AudioDeviceID.Metadata.Description = LOCTEXT("AudioDeviceIDToolTip", "Audio Device ID");
 
 				Archetype.Interface.Environment.Add(AudioDeviceID);
 
@@ -438,57 +435,57 @@ namespace Metasound
 				return FMetasoundFrontendVersion{ VersionName, FMetasoundFrontendVersionNumber{1, 1} };
 			}
 
-			const FString& GetOnPlayInputName()
+			const FVertexName& GetOnPlayInputName()
 			{
-				static const FString TriggerInputName = TEXT("On Play");
+				static const FVertexName TriggerInputName = TEXT("On Play");
 				return TriggerInputName;
 			}
 
-			const FString& GetLeftAudioOutputName()
+			const FVertexName& GetLeftAudioOutputName()
 			{
-				static const FString AudioOutputName = TEXT("Audio:0");
+				static const FVertexName AudioOutputName = TEXT("Audio:0");
 				return AudioOutputName;
 			}
 
-			const FString& GetRightAudioOutputName()
+			const FVertexName& GetRightAudioOutputName()
 			{
-				static const FString AudioOutputName = TEXT("Audio:1");
+				static const FVertexName AudioOutputName = TEXT("Audio:1");
 				return AudioOutputName;
 			}
 
-			const FString& GetIsFinishedOutputName()
+			const FVertexName& GetIsFinishedOutputName()
 			{
-				static const FString OnFinishedOutputName = TEXT("On Finished");
+				static const FVertexName OnFinishedOutputName = TEXT("On Finished");
 				return OnFinishedOutputName;
 			}
 
-			const FString& GetAudioDeviceIDVariableName()
+			const FVertexName& GetAudioDeviceIDVariableName()
 			{
-				static const FString AudioDeviceIDVarName = TEXT("AudioDeviceID");
+				static const FVertexName AudioDeviceIDVarName = TEXT("AudioDeviceID");
 				return AudioDeviceIDVarName;
 			}
 
-			const FString& GetSoundUniqueIdName()
+			const FVertexName& GetSoundUniqueIdName()
 			{
-				static const FString SoundUniqueIdVarName = TEXT("SoundUniqueId");
+				static const FVertexName SoundUniqueIdVarName = TEXT("SoundUniqueId");
 				return SoundUniqueIdVarName;
 			}
 
-			const FString& GetIsPreviewSoundName()
+			const FVertexName& GetIsPreviewSoundName()
 			{
-				static const FString SoundIsPreviewSoundName = TEXT("IsPreviewSound");
+				static const FVertexName SoundIsPreviewSoundName = TEXT("IsPreviewSound");
 				return SoundIsPreviewSoundName;
 			}
 
-			const FString& GetGraphName()
+			const FVertexName& GetGraphName()
 			{
-				static const FString SoundGraphName = TEXT("GraphName");
+				static const FVertexName SoundGraphName = TEXT("GraphName");
 				return SoundGraphName;
 			}
 
-			const FString& GetInstanceIDName()
+			const FVertexName& GetInstanceIDName()
 			{
-				static const FString& Name = FMetaSoundParameterTransmitter::GetInstanceIDEnvironmentVariableName();
+				static const FVertexName& Name = FMetaSoundParameterTransmitter::GetInstanceIDEnvironmentVariableName();
 				return Name;
 			}
 
@@ -552,8 +549,7 @@ namespace Metasound
 				// Environment
 				FMetasoundFrontendEnvironmentVariable AudioDeviceID;
 				AudioDeviceID.Name = GetAudioDeviceIDVariableName();
-				AudioDeviceID.Metadata.DisplayName = FText::FromString(AudioDeviceID.Name);
-				AudioDeviceID.Metadata.Description = LOCTEXT("AudioDeviceIDToolTip", "Audio device handle");
+				AudioDeviceID.Metadata.Description = LOCTEXT("AudioDeviceIDToolTip", "Audio Device ID");
 				Archetype.Interface.Environment.Add(AudioDeviceID);
 
 				return Archetype;
@@ -647,43 +643,43 @@ namespace Metasound
 			namespace CurrentMonoVersion = MetasoundSourceMonoV1_1;
 			namespace CurrentStereoVersion = MetasoundSourceStereoV1_1;
 
-			const FString& GetOnPlayInputName()
+			const FVertexName& GetOnPlayInputName()
 			{
 				check(CurrentMonoVersion::GetOnPlayInputName() == CurrentStereoVersion::GetOnPlayInputName());
 				return CurrentMonoVersion::GetOnPlayInputName();
 			}
 
-			const FString& GetIsFinishedOutputName()
+			const FVertexName& GetIsFinishedOutputName()
 			{
 				check(CurrentMonoVersion::GetIsFinishedOutputName() == CurrentStereoVersion::GetIsFinishedOutputName());
 				return CurrentMonoVersion::GetIsFinishedOutputName();
 			}
 
-			const FString& GetAudioDeviceIDVariableName()
+			const FVertexName& GetAudioDeviceIDVariableName()
 			{
 				check(CurrentMonoVersion::GetAudioDeviceIDVariableName() == CurrentStereoVersion::GetAudioDeviceIDVariableName());
 				return CurrentMonoVersion::GetAudioDeviceIDVariableName();
 			}
 
-			const FString& GetSoundUniqueIdName()
+			const FVertexName& GetSoundUniqueIdName()
 			{
 				check(CurrentMonoVersion::GetSoundUniqueIdName() == CurrentStereoVersion::GetSoundUniqueIdName());
 				return CurrentMonoVersion::GetSoundUniqueIdName();
 			}
 
-			const FString& GetIsPreviewSoundName()
+			const FVertexName& GetIsPreviewSoundName()
 			{
 				check(CurrentMonoVersion::GetIsPreviewSoundName() == CurrentStereoVersion::GetIsPreviewSoundName());
 				return CurrentMonoVersion::GetIsPreviewSoundName();
 			}
 
-			const FString& GetInstanceIDName()
+			const FVertexName& GetInstanceIDName()
 			{
 				check(CurrentMonoVersion::GetInstanceIDName() == CurrentStereoVersion::GetInstanceIDName());
 				return CurrentMonoVersion::GetInstanceIDName();
 			}
 
-			const FString& GetGraphName()
+			const FVertexName& GetGraphName()
 			{
 				check(CurrentMonoVersion::GetGraphName() == CurrentStereoVersion::GetGraphName());
 				return CurrentMonoVersion::GetGraphName();
@@ -699,7 +695,7 @@ namespace Metasound
 				return Current::GetVersion();
 			}
 
-			const FString& GetAudioOutputName()
+			const FVertexName& GetAudioOutputName()
 			{
 				return Current::GetAudioOutputName();
 			}
@@ -714,12 +710,12 @@ namespace Metasound
 				return Current::GetVersion();
 			}
 
-			const FString& GetLeftAudioOutputName()
+			const FVertexName& GetLeftAudioOutputName()
 			{
 				return Current::GetLeftAudioOutputName();
 			}
 
-			const FString& GetRightAudioOutputName()
+			const FVertexName& GetRightAudioOutputName()
 			{
 				return Current::GetRightAudioOutputName();
 			}
diff --git a/Engine/Plugins/Runtime/Metasound/Source/MetasoundEngine/Private/MetasoundEngineArchetypes.h b/Engine/Plugins/Runtime/Metasound/Source/MetasoundEngine/Private/MetasoundEngineArchetypes.h
--- a/Engine/Plugins/Runtime/Metasound/Source/MetasoundEngine/Private/MetasoundEngineArchetypes.h
+++ b/Engine/Plugins/Runtime/Metasound/Source/MetasoundEngine/Private/MetasoundEngineArchetypes.h
@@ -4,6 +4,7 @@
 
 #include "CoreMinimal.h"
 #include "MetasoundFrontendDocument.h"
+#include "MetasoundVertex.h"
 
 namespace Metasound
 {
@@ -19,81 +20,81 @@ namespace Metasound
 		namespace MetasoundSourceMonoV1_0
 		{
 			FMetasoundFrontendVersion GetVersion();
-			const FString& GetOnPlayInputName();
-			const FString& GetAudioOutputName();
-			const FString& GetIsFinishedOutputName();
-			const FString& GetAudioDeviceIDVariableName();
-			const FString& GetSoundUniqueIdName();
-			const FString& GetIsPreviewSoundName();
+			const FVertexName& GetOnPlayInputName();
+			const FVertexName& GetAudioOutputName();
+			const FVertexName& GetIsFinishedOutputName();
+			const FVertexName& GetAudioDeviceIDVariableName();
+			const FVertexName& GetSoundUniqueIdName();
+			const FVertexName& GetIsPreviewSoundName();
 		}
 
 		// V1.0 of a Metasound Stereo Source. Uses FStereoAudioFormat as output.
 		namespace MetasoundSourceStereoV1_0
 		{
 			FMetasoundFrontendVersion GetVersion();
-			const FString& GetOnPlayInputName();
-			const FString& GetAudioOutputName();
-			const FString& GetIsFinishedOutputName();
-			const FString& GetAudioDeviceIDVariableName();
-			const FString& GetSoundUniqueIdName();
-			const FString& GetIsPreviewSoundName();
-			const FString& GetGraphName();
+			const FVertexName& GetOnPlayInputName();
+			const FVertexName& GetAudioOutputName();
+			const FVertexName& GetIsFinishedOutputName();
+			const FVertexName& GetAudioDeviceIDVariableName();
+			const FVertexName& GetSoundUniqueIdName();
+			const FVertexName& GetIsPreviewSoundName();
+			const FVertexName& GetGraphName();
 		}
 
 		// V1.1 of a Metasound Mono source. Uses FAudioBuffer as output.
 		namespace MetasoundSourceMonoV1_1
 		{
 			FMetasoundFrontendVersion GetVersion();
-			const FString& GetOnPlayInputName();
-			const FString& GetAudioOutputName();
-			const FString& GetIsFinishedOutputName();
-			const FString& GetAudioDeviceIDVariableName();
-			const FString& GetSoundUniqueIdName();
-			const FString& GetIsPreviewSoundName();
-			const FString& GetInstanceIDName();
-			const FString& GetGraphName();
+			const FVertexName& GetOnPlayInputName();
+			const FVertexName& GetAudioOutputName();
+			const FVertexName& GetIsFinishedOutputName();
+			const FVertexName& GetAudioDeviceIDVariableName();
+			const FVertexName& GetSoundUniqueIdName();
+			const FVertexName& GetIsPreviewSoundName();
+			const FVertexName& GetInstanceIDName();
+			const FVertexName& GetGraphName();
 		}
 
 		// V1.1 of a Metasound Stereo source. Uses two FAudioBuffers as outputs.
 		namespace MetasoundSourceStereoV1_1
 		{
 			FMetasoundFrontendVersion GetVersion();
-			const FString& GetOnPlayInputName();
-			const FString& GetLeftAudioOutputName();
-			const FString& GetRightAudioOutputName();
-			const FString& GetIsFinishedOutputName();
-			const FString& GetAudioDeviceIDVariableName();
-			const FString& GetSoundUniqueIdName();
-			const FString& GetIsPreviewSoundName();
-			const FString& GetInstanceIDName();
-			const FString& GetGraphName();
+			const FVertexName& GetOnPlayInputName();
+			const FVertexName& GetLeftAudioOutputName();
+			const FVertexName& GetRightAudioOutputName();
+			const FVertexName& GetIsFinishedOutputName();
+			const FVertexName& GetAudioDeviceIDVariableName();
+			const FVertexName& GetSoundUniqueIdName();
+			const FVertexName& GetIsPreviewSoundName();
+			const FVertexName& GetInstanceIDName();
+			const FVertexName& GetGraphName();
 		}
 
 		// Current version of Metasound Source.
 		namespace MetasoundSource
 		{
-			const FString& GetOnPlayInputName();
-			const FString& GetIsFinishedOutputName();
-			const FString& GetAudioDeviceIDVariableName();
-			const FString& GetSoundUniqueIdName();
-			const FString& GetIsPreviewSoundName();
-			const FString& GetInstanceIDName();
-			const FString& GetGraphName();
+			const FVertexName& GetOnPlayInputName();
+			const FVertexName& GetIsFinishedOutputName();
+			const FVertexName& GetAudioDeviceIDVariableName();
+			const FVertexName& GetSoundUniqueIdName();
+			const FVertexName& GetIsPreviewSoundName();
+			const FVertexName& GetInstanceIDName();
+			const FVertexName& GetGraphName();
 		}
 
 		// Current version of MetasoundSourceMono
 		namespace MetasoundSourceMono
 		{
 			FMetasoundFrontendVersion GetVersion();
-			const FString& GetAudioOutputName();
+			const FVertexName& GetAudioOutputName();
 		}
 
 		// Current version of MetasoundSourceStereo
 		namespace MetasoundSourceStereo
 		{
 			FMetasoundFrontendVersion GetVersion();
-			const FString& GetLeftAudioOutputName();
-			const FString& GetRightAudioOutputName();
+			const FVertexName& GetLeftAudioOutputName();
+			const FVertexName& GetRightAudioOutputName();
 		}
 
 		// Register all MetasoundEngine archetypes.
diff --git a/Engine/Plugins/Runtime/Metasound/Source/MetasoundEngine/Private/MetasoundSource.cpp b/Engine/Plugins/Runtime/Metasound/Source/MetasoundEngine/Private/MetasoundSource.cpp
--- a/Engine/Plugins/Runtime/Metasound/Source/MetasoundEngine/Private/MetasoundSource.cpp
+++ b/Engine/Plugins/Runtime/Metasound/Source/MetasoundEngine/Private/MetasoundSource.cpp
@@ -34,7 +34,7 @@
 
 namespace MetaSoundSourcePrivate
 {
-	using FFormatOutputVertexKeyMap = TMap<EMetasoundSourceAudioFormat, TArray<Metasound::FVertexKey>>;
+	using FFormatOutputVertexKeyMap = TMap<EMetasoundSourceAudioFormat, TArray<Metasound::FVertexName>>;
 
 	const FFormatOutputVertexKeyMap& GetFormatOutputVertexKeys()
 	{
@@ -436,7 +436,7 @@ ISoundGeneratorPtr UMetaSoundSource::CreateSoundGenerator(const FSoundGeneratorI
 		return ISoundGeneratorPtr(nullptr);
 	}
 
-	FMetasoundGeneratorInitParams InitParams = 
+	FMetasoundGeneratorInitParams InitParams =
 	{
 		InSettings,
 		MoveTemp(MetasoundGraph),
@@ -462,8 +462,7 @@ bool UMetaSoundSource::IsParameterValid(const FAudioParameter& InParameter) cons
 		using namespace Metasound;
 		using namespace Metasound::Frontend;
 
-		// TODO: Fix this.  NodeHandles should use FNames, not display names (this is broken in other languages)
-		if (Param->ParamName == *NodeHandle->GetDisplayName().ToString())
+		if (Param->ParamName == NodeHandle->GetNodeName())
 		{
 			TArray<FConstInputHandle> InputHandles = NodeHandle->GetConstInputs();
 			if (InputHandles.IsEmpty())
@@ -684,11 +683,11 @@ Metasound::FMetasoundEnvironment UMetaSoundSource::CreateEnvironment(const Audio
 	return Environment;
 }
 
-const TArray<Metasound::FVertexKey>& UMetaSoundSource::GetAudioOutputVertexKeys() const
+const TArray<Metasound::FVertexName>& UMetaSoundSource::GetAudioOutputVertexKeys() const
 {
 	using namespace MetaSoundSourcePrivate;
 
-	if (const TArray<Metasound::FVertexKey>* ArrayKeys = GetFormatOutputVertexKeys().Find(OutputFormat))
+	if (const TArray<Metasound::FVertexName>* ArrayKeys = GetFormatOutputVertexKeys().Find(OutputFormat))
 	{
 		return *ArrayKeys;
 	}
@@ -696,7 +695,7 @@ const TArray<Metasound::FVertexKey>& UMetaSoundSource::GetAudioOutputVertexKeys(
 	{
 		// Unhandled audio format. Need to update audio output format vertex key map.
 		checkNoEntry();
-		static const TArray<Metasound::FVertexKey> Empty;
+		static const TArray<Metasound::FVertexName> Empty;
 		return Empty;
 	}
 }
diff --git a/Engine/Plugins/Runtime/Metasound/Source/MetasoundEngine/Private/MetasoundWaveInfoNode.cpp b/Engine/Plugins/Runtime/Metasound/Source/MetasoundEngine/Private/MetasoundWaveInfoNode.cpp
--- a/Engine/Plugins/Runtime/Metasound/Source/MetasoundEngine/Private/MetasoundWaveInfoNode.cpp
+++ b/Engine/Plugins/Runtime/Metasound/Source/MetasoundEngine/Private/MetasoundWaveInfoNode.cpp
@@ -14,6 +14,7 @@
 #include "MetasoundNodeRegistrationMacro.h"
 #include "MetasoundStandardNodesCategories.h"
 #include "MetasoundDataTypeRegistrationMacro.h"
+#include "MetasoundVertex.h"
 
 #define LOCTEXT_NAMESPACE "MetasoundWaveInfo"
 
@@ -164,7 +165,7 @@ namespace Metasound
 	public:
 		// public node api needs to define two conversion constructors:
 		// (1: from FString)
-		FWaveInfoNode(const FString& InInstanceName, const FGuid& InInstanceID)
+		FWaveInfoNode(const Metasound::FVertexName& InInstanceName, const FGuid& InInstanceID)
 			: FNodeFacade(InInstanceName, InInstanceID, TFacadeOperatorClass < FWaveInfoNodeOperator >())
 		{ }
 
diff --git a/Engine/Plugins/Runtime/Metasound/Source/MetasoundEngine/Private/MetasoundWavePlayerNode.cpp b/Engine/Plugins/Runtime/Metasound/Source/MetasoundEngine/Private/MetasoundWavePlayerNode.cpp
--- a/Engine/Plugins/Runtime/Metasound/Source/MetasoundEngine/Private/MetasoundWavePlayerNode.cpp
+++ b/Engine/Plugins/Runtime/Metasound/Source/MetasoundEngine/Private/MetasoundWavePlayerNode.cpp
@@ -11,9 +11,11 @@
 #include "MetasoundWave.h"
 #include "MetasoundTrigger.h"
 #include "MetasoundEngineNodesNames.h"
+#include "MetasoundVertex.h"
 
 #define LOCTEXT_NAMESPACE "MetasoundWaveNode"
 
+
 namespace Metasound
 {
 	namespace WavePlayerVertexNames
@@ -75,7 +77,7 @@ namespace Metasound
 		static FVertexInterface DeclareVertexInterface();
 		static const FNodeClassMetadata& GetNodeInfo();
 
-		FWavePlayerNode(const FString& InName, const FGuid& InInstanceID);
+		FWavePlayerNode(const FVertexName& InName, const FGuid& InInstanceID);
 		FWavePlayerNode(const FNodeInitData& InInitData);
 		virtual ~FWavePlayerNode() = default;
 
@@ -771,7 +773,7 @@ namespace Metasound
 		return Info;
 	}
 
-	FWavePlayerNode::FWavePlayerNode(const FString& InName, const FGuid& InInstanceID)
+	FWavePlayerNode::FWavePlayerNode(const FVertexName& InName, const FGuid& InInstanceID)
 		:	FNode(InName, InInstanceID, GetNodeInfo())
 		,	Factory(MakeOperatorFactoryRef<FWavePlayerNode::FOperatorFactory>())
 		,	Interface(DeclareVertexInterface())
diff --git a/Engine/Plugins/Runtime/Metasound/Source/MetasoundEngine/Public/MetasoundSource.h b/Engine/Plugins/Runtime/Metasound/Source/MetasoundEngine/Public/MetasoundSource.h
--- a/Engine/Plugins/Runtime/Metasound/Source/MetasoundEngine/Public/MetasoundSource.h
+++ b/Engine/Plugins/Runtime/Metasound/Source/MetasoundEngine/Public/MetasoundSource.h
@@ -196,5 +196,5 @@ private:
 	Metasound::FMetasoundEnvironment CreateEnvironment() const;
 	Metasound::FMetasoundEnvironment CreateEnvironment(const FSoundGeneratorInitParams& InParams) const;
 	Metasound::FMetasoundEnvironment CreateEnvironment(const Audio::FParameterTransmitterInitParams& InParams) const;
-	const TArray<Metasound::FVertexKey>& GetAudioOutputVertexKeys() const;
+	const TArray<Metasound::FVertexName>& GetAudioOutputVertexKeys() const;
 };
diff --git a/Engine/Plugins/Runtime/Metasound/Source/MetasoundEngineTest/Private/EngineTestMetasoundSource.cpp b/Engine/Plugins/Runtime/Metasound/Source/MetasoundEngineTest/Private/EngineTestMetasoundSource.cpp
--- a/Engine/Plugins/Runtime/Metasound/Source/MetasoundEngineTest/Private/EngineTestMetasoundSource.cpp
+++ b/Engine/Plugins/Runtime/Metasound/Source/MetasoundEngineTest/Private/EngineTestMetasoundSource.cpp
@@ -116,13 +116,13 @@ namespace EngineTestMetasoundSourcePrivate
 		// Make connections:
 
 		// frequency input "Frequency" -> oscillator "Frequency"
-		FOutputHandle OutputToConnect = FrequencyInputNode->GetOutputsWithVertexName(FString(TEXT("Frequency")))[0];
-		FInputHandle InputToConnect = OscNode->GetInputsWithVertexName(FString(TEXT("Frequency")))[0];
+		FOutputHandle OutputToConnect = FrequencyInputNode->GetOutputsWithVertexName("Frequency")[0];
+		FInputHandle InputToConnect = OscNode->GetInputsWithVertexName("Frequency")[0];
 		ensure(InputToConnect->Connect(*OutputToConnect));
 
 		// oscillator to output
-		OutputToConnect = OscNode->GetOutputsWithVertexName(FString(TEXT("Audio")))[0];
-		InputToConnect = OutputAudioNode->GetInputsWithVertexName(TEXT("Audio:0"))[0];
+		OutputToConnect = OscNode->GetOutputsWithVertexName("Audio")[0];
+		InputToConnect = OutputAudioNode->GetInputsWithVertexName("Audio:0")[0];
 		ensure(InputToConnect->Connect(*OutputToConnect));
 
 		return Document;
