Move FString/FText VertexName identifers to FNames. Convention is as follows:

diff --git a/Engine/Plugins/Runtime/Metasound/Source/MetasoundFrontend/Private/MetasoundArchetype.cpp b/Engine/Plugins/Runtime/Metasound/Source/MetasoundFrontend/Private/MetasoundArchetype.cpp
--- a/Engine/Plugins/Runtime/Metasound/Source/MetasoundFrontend/Private/MetasoundArchetype.cpp
+++ b/Engine/Plugins/Runtime/Metasound/Source/MetasoundFrontend/Private/MetasoundArchetype.cpp
@@ -19,7 +19,7 @@ namespace Metasound
 				{
 					return VertexA.TypeName.FastLess(VertexB.TypeName);
 				}
-				return VertexA.Name < VertexB.Name;
+				return VertexA.Name.FastLess(VertexB.Name);
 			}
 
 			bool IsLessThanEnvironmentVariable(const FMetasoundFrontendEnvironmentVariable& EnvironmentA, const FMetasoundFrontendEnvironmentVariable& EnvironmentB)
@@ -29,7 +29,7 @@ namespace Metasound
 				{
 					return EnvironmentA.TypeName.FastLess(EnvironmentB.TypeName);
 				}
-				return EnvironmentA.Name < EnvironmentB.Name;
+				return EnvironmentA.Name.FastLess(EnvironmentB.Name);
 			}
 
 			bool IsEquivalentEnvironmentVariable(const FMetasoundFrontendEnvironmentVariable& EnvironmentA, const FMetasoundFrontendEnvironmentVariable& EnvironmentB)
diff --git a/Engine/Plugins/Runtime/Metasound/Source/MetasoundFrontend/Private/MetasoundArrayNodes.cpp b/Engine/Plugins/Runtime/Metasound/Source/MetasoundFrontend/Private/MetasoundArrayNodes.cpp
--- a/Engine/Plugins/Runtime/Metasound/Source/MetasoundFrontend/Private/MetasoundArrayNodes.cpp
+++ b/Engine/Plugins/Runtime/Metasound/Source/MetasoundFrontend/Private/MetasoundArrayNodes.cpp
@@ -32,70 +32,70 @@ namespace Metasound
 	namespace ArrayNodeVertexNames
 	{
 		/* Input Vertx Names */
-		const FString& GetInputArrayName()
+		const FVertexName& GetInputArrayName()
 		{
-			static const FString Name = TEXT("Array");
+			static const FVertexName Name = TEXT("Array");
 			return Name;
 		}
 
-		const FString& GetInputLeftArrayName()
+		const FVertexName& GetInputLeftArrayName()
 		{
-			static const FString Name = TEXT("Left Array");
+			static const FVertexName Name = TEXT("Left Array");
 			return Name;
 		}
 
-		const FString& GetInputRightArrayName()
+		const FVertexName& GetInputRightArrayName()
 		{
-			static const FString Name = TEXT("Right Array");
+			static const FVertexName Name = TEXT("Right Array");
 			return Name;
 		}
 
-		const FString& GetInputTriggerName()
+		const FVertexName& GetInputTriggerName()
 		{
-			static const FString Name = TEXT("Trigger");
+			static const FVertexName Name = TEXT("Trigger");
 			return Name;
 		}
 
-		const FString& GetInputIndexName()
+		const FVertexName& GetInputIndexName()
 		{
-			static const FString Name = TEXT("Index");
+			static const FVertexName Name = TEXT("Index");
 			return Name;
 		}
 
-		const FString& GetInputStartIndexName()
+		const FVertexName& GetInputStartIndexName()
 		{
-			static const FString Name = TEXT("Start Index");
+			static const FVertexName Name = TEXT("Start Index");
 			return Name;
 		}
 
-		const FString& GetInputEndIndexName()
+		const FVertexName& GetInputEndIndexName()
 		{
-			static const FString Name = TEXT("End Index");
+			static const FVertexName Name = TEXT("End Index");
 			return Name;
 		}
 
-		const FString& GetInputValueName()
+		const FVertexName& GetInputValueName()
 		{
-			static const FString Name = TEXT("Value");
+			static const FVertexName Name = TEXT("Value");
 			return Name;
 		}
 
 		/* Output Vertex Names */
-		const FString& GetOutputNumName()
+		const FVertexName& GetOutputNumName()
 		{
-			static const FString Name = TEXT("Num");
+			static const FVertexName Name = TEXT("Num");
 			return Name;
 		}
 
-		const FString& GetOutputValueName()
+		const FVertexName& GetOutputValueName()
 		{
-			static const FString Name = TEXT("Element");
+			static const FVertexName Name = TEXT("Element");
 			return Name;
 		}
 
-		const FString& GetOutputArrayName()
+		const FVertexName& GetOutputArrayName()
 		{
-			static const FString Name = TEXT("Array");
+			static const FVertexName Name = TEXT("Array");
 			return Name;
 		}
 	}
diff --git a/Engine/Plugins/Runtime/Metasound/Source/MetasoundFrontend/Private/MetasoundArrayShuffleNode.cpp b/Engine/Plugins/Runtime/Metasound/Source/MetasoundFrontend/Private/MetasoundArrayShuffleNode.cpp
--- a/Engine/Plugins/Runtime/Metasound/Source/MetasoundFrontend/Private/MetasoundArrayShuffleNode.cpp
+++ b/Engine/Plugins/Runtime/Metasound/Source/MetasoundFrontend/Private/MetasoundArrayShuffleNode.cpp
@@ -9,69 +9,69 @@ namespace Metasound
 {
 	namespace ArrayNodeShuffleVertexNames
 	{
-		const FString& GetInputTriggerNextName()
+		const FVertexName& GetInputTriggerNextName()
 		{
-			static const FString Name = TEXT("Next");
+			static const FVertexName Name = TEXT("Next");
 			return Name;
 		}
 
-		const FString& GetInputTriggerShuffleName()
+		const FVertexName& GetInputTriggerShuffleName()
 		{
-			static const FString Name = TEXT("Shuffle");
+			static const FVertexName Name = TEXT("Shuffle");
 			return Name;
 		}
 
-		const FString& GetInputTriggerResetName()
+		const FVertexName& GetInputTriggerResetName()
 		{
-			static const FString Name = TEXT("Reset Seed");
+			static const FVertexName Name = TEXT("Reset Seed");
 			return Name;
 		}
 
-		const FString& GetInputShuffleArrayName()
+		const FVertexName& GetInputShuffleArrayName()
 		{
-			static const FString Name = TEXT("In Array");
+			static const FVertexName Name = TEXT("In Array");
 			return Name;
 		}
 
-		const FString& GetInputSeedName()
+		const FVertexName& GetInputSeedName()
 		{
-			static const FString Name = TEXT("Seed");
+			static const FVertexName Name = TEXT("Seed");
 			return Name;
 		}
 
-		const FString& GetInputAutoShuffleName()
+		const FVertexName& GetInputAutoShuffleName()
 		{
-			static const FString Name = TEXT("Auto Shuffle");
+			static const FVertexName Name = TEXT("Auto Shuffle");
 			return Name;
 		}
 
-		const FString& GetInputEnableSharedStateName()
+		const FVertexName& GetInputEnableSharedStateName()
 		{
-			static const FString Name = TEXT("Enable Shared State");
+			static const FVertexName Name = TEXT("Enable Shared State");
 			return Name;
 		}
 
-		const FString& GetOutputTriggerOnNextName()
+		const FVertexName& GetOutputTriggerOnNextName()
 		{
-			static const FString Name = TEXT("On Next");
+			static const FVertexName Name = TEXT("On Next");
 			return Name;
 		}
 
-		const FString& GetOutputTriggerOnShuffleName()
+		const FVertexName& GetOutputTriggerOnShuffleName()
 		{
-			static const FString Name = TEXT("On Shuffle");
+			static const FVertexName Name = TEXT("On Shuffle");
 			return Name;
 		}
 
-		const FString& GetOutputTriggerOnResetName()
+		const FVertexName& GetOutputTriggerOnResetName()
 		{
-			static const FString Name = TEXT("On Reset Seed");
+			static const FVertexName Name = TEXT("On Reset Seed");
 			return Name;
 		}
 
-		const FString& GetOutputValueName()
+		const FVertexName& GetOutputValueName()
 		{
-			static const FString Name = TEXT("Value");
+			static const FVertexName Name = TEXT("Value");
 			return Name;
 		}
 	}
@@ -171,31 +171,30 @@ namespace Metasound
 
 	namespace ArrayNodeGetGlobalArrayKeyVertexNames
 	{
-		const FString& GetInputNamespaceName()
+		const FVertexName& GetInputNamespaceName()
 		{
-			static const FString Name = TEXT("Namespace");
-			return Name;		
+			static const FVertexName Name = TEXT("Namespace");
+			return Name;
 		}
 
-		const FString& GetInputArraySizeName()
+		const FVertexName& GetInputArraySizeName()
 		{
-			static const FString Name = TEXT("Array Size");
+			static const FVertexName Name = TEXT("Array Size");
 			return Name;
 		}
 
-		const FString& GetInputSeedName()
+		const FVertexName& GetInputSeedName()
 		{
-			static const FString Name = TEXT("Seed");
+			static const FVertexName Name = TEXT("Seed");
 			return Name;
 		}
 
-		const FString& GetOutputArrayKeyName()
+		const FVertexName& GetOutputArrayKeyName()
 		{
-			static const FString Name = TEXT("Global Array Key");
+			static const FVertexName Name = TEXT("Global Array Key");
 			return Name;
 		}
 	}
-
 }
 
 #undef LOCTEXT_NAMESPACE
diff --git a/Engine/Plugins/Runtime/Metasound/Source/MetasoundFrontend/Private/MetasoundAssetBase.cpp b/Engine/Plugins/Runtime/Metasound/Source/MetasoundFrontend/Private/MetasoundAssetBase.cpp
--- a/Engine/Plugins/Runtime/Metasound/Source/MetasoundFrontend/Private/MetasoundAssetBase.cpp
+++ b/Engine/Plugins/Runtime/Metasound/Source/MetasoundFrontend/Private/MetasoundAssetBase.cpp
@@ -21,6 +21,7 @@
 #include "MetasoundLog.h"
 #include "MetasoundParameterTransmitter.h"
 #include "MetasoundTrace.h"
+#include "MetasoundVertex.h"
 #include "StructSerializer.h"
 #include "UObject/MetaData.h"
 
@@ -446,7 +447,7 @@ TUniquePtr<Metasound::IGraph> FMetasoundAssetBase::BuildMetasoundDocument() cons
 	TUniquePtr<FFrontendGraph> FrontendGraph = FFrontendGraphBuilder::CreateGraph(*Doc);
 	if (FrontendGraph.IsValid())
 	{
-		TSet<FVertexKey> VerticesToSkip = GetNonTransmittableInputVertices(*Doc);
+		TSet<FVertexName> VerticesToSkip = GetNonTransmittableInputVertices(*Doc);
 		bool bSuccessfullyInjectedReceiveNodes = InjectReceiveNodes(*FrontendGraph, FMetaSoundParameterTransmitter::CreateSendAddressFromEnvironment, VerticesToSkip);
 		if (!bSuccessfullyInjectedReceiveNodes)
 		{
@@ -527,7 +528,7 @@ bool FMetasoundAssetBase::AddingReferenceCausesLoop(const FSoftObjectPath& InRef
 	return bCausesLoop;
 }
 
-Metasound::FSendAddress FMetasoundAssetBase::CreateSendAddress(uint64 InInstanceID, const FString& InVertexName, const FName& InDataTypeName) const
+Metasound::FSendAddress FMetasoundAssetBase::CreateSendAddress(uint64 InInstanceID, const Metasound::FVertexName& InVertexName, const FName& InDataTypeName) const
 
 {
 	using namespace Metasound;
@@ -535,7 +536,7 @@ Metasound::FSendAddress FMetasoundAssetBase::CreateSendAddress(uint64 InInstance
 	FSendAddress Address;
 
 	Address.Subsystem = GetSubsystemNameForSendScope(ETransmissionScope::Global);
-	Address.ChannelName = FName(FString::Printf(TEXT("%d:%s:%s"), InInstanceID, *InVertexName, *InDataTypeName.ToString()));
+	Address.ChannelName = FName(FString::Printf(TEXT("%d:%s:%s"), InInstanceID, *InVertexName.ToString(), *InDataTypeName.ToString()));
 
 	return Address;
 }
@@ -563,8 +564,8 @@ TArray<FMetasoundAssetBase::FSendInfoAndVertexName> FMetasoundAssetBase::GetSend
 
 	FConstGraphHandle RootGraph = GetRootGraphHandle();
 
-	TArray<FString> SendVertices = GetTransmittableInputVertexNames();
-	for (const FString& VertexName : SendVertices)
+	const TArray<FVertexName> SendVertices = GetTransmittableInputVertexNames();
+	for (const FVertexName& VertexName : SendVertices)
 	{
 		FConstNodeHandle InputNode = RootGraph->GetInputNodeWithName(VertexName);
 		for (FConstInputHandle InputHandle : InputNode->GetConstInputs())
@@ -573,9 +574,8 @@ TArray<FMetasoundAssetBase::FSendInfoAndVertexName> FMetasoundAssetBase::GetSend
 
 			// TODO: incorporate VertexID into address. But need to ensure that VertexID
 			// will be maintained after injecting Receive nodes. 
-			Info.SendInfo.Address = FMetaSoundParameterTransmitter::CreateSendAddressFromInstanceID(InInstanceID, InputHandle->GetName(), InputHandle->GetDataType());
-			Info.SendInfo.ParameterName = FName(InputHandle->GetDisplayName().ToString()); // TODO: display name hack. Need to have naming consistent in editor for inputs
-			//Info.SendInfo.ParameterName = FName(*InputHandle->GetName()); // TODO: this is the expected parameter name.
+			Info.SendInfo.Address = FMetaSoundParameterTransmitter::CreateSendAddressFromInstanceID(InInstanceID, VertexName, InputHandle->GetDataType());
+			Info.SendInfo.ParameterName = VertexName;
 			Info.SendInfo.TypeName = InputHandle->GetDataType();
 			Info.VertexName = VertexName;
 
@@ -586,14 +586,14 @@ TArray<FMetasoundAssetBase::FSendInfoAndVertexName> FMetasoundAssetBase::GetSend
 	return SendInfos;
 }
 
-TSet<Metasound::FVertexKey> FMetasoundAssetBase::GetNonTransmittableInputVertices(const FMetasoundFrontendDocument& InDoc) const
+TSet<Metasound::FVertexName> FMetasoundAssetBase::GetNonTransmittableInputVertices(const FMetasoundFrontendDocument& InDoc) const
 {
 	using namespace Metasound;
 	using namespace Metasound::Frontend;
 
 	check(IsInGameThread() || IsInAudioThread());
 
-	TSet<FVertexKey> NonTransmittableInputVertices;
+	TSet<FVertexName> NonTransmittableInputVertices;
 
 	auto GetVertexKey = [](const FMetasoundFrontendClassVertex& InVertex) { return InVertex.Name; };
 
@@ -623,7 +623,7 @@ TSet<Metasound::FVertexKey> FMetasoundAssetBase::GetNonTransmittableInputVertice
 	return NonTransmittableInputVertices;
 }
 
-TArray<FString> FMetasoundAssetBase::GetTransmittableInputVertexNames() const
+TArray<Metasound::FVertexName> FMetasoundAssetBase::GetTransmittableInputVertexNames() const
 {
 	using namespace Metasound;
 	using namespace Metasound::Frontend;
@@ -641,19 +641,19 @@ TArray<FString> FMetasoundAssetBase::GetTransmittableInputVertexNames() const
 	}
 
 	// Unused inputs are all input vertices that are not in the archetype.
-	TArray<FString> ArchetypeInputVertexNames;
+	TArray<FVertexName> ArchetypeInputVertexNames;
 	for (const FMetasoundFrontendClassVertex& Vertex : Archetype.Interface.Inputs)
 	{
 		ArchetypeInputVertexNames.Add(Vertex.Name);
 	}
 
 	Frontend::FConstGraphHandle RootGraph = Document->GetRootGraph();
-	TArray<FString> GraphInputVertexNames = RootGraph->GetInputVertexNames();
+	TArray<FVertexName> GraphInputVertexNames = RootGraph->GetInputVertexNames();
 
 	// Filter graph inputs by archetype inputs.
-	GraphInputVertexNames = GraphInputVertexNames.FilterByPredicate([&](const FString& InName) { return !ArchetypeInputVertexNames.Contains(InName); });
+	GraphInputVertexNames = GraphInputVertexNames.FilterByPredicate([&](const FVertexName& InName) { return !ArchetypeInputVertexNames.Contains(InName); });
 
-	auto IsDataTypeTransmittable = [&](const FString& InVertexName)
+	auto IsDataTypeTransmittable = [&](const FVertexName& InVertexName)
 	{
 		using namespace Metasound::Frontend;
 
@@ -686,7 +686,7 @@ Metasound::Frontend::FNodeHandle FMetasoundAssetBase::AddInputPinForSendAddress(
 	FMetasoundFrontendClassInput Description;
 	FGuid VertexID = FGuid::NewGuid();
 
-	Description.Name = InSendInfo.Address.ChannelName.ToString();
+	Description.Name = InSendInfo.Address.ChannelName;
 	Description.TypeName = Metasound::GetMetasoundDataTypeName<Metasound::FSendAddress>();
 	Description.Metadata.Description = FText::GetEmpty();
 	Description.VertexID = VertexID;
diff --git a/Engine/Plugins/Runtime/Metasound/Source/MetasoundFrontend/Private/MetasoundFrontend.cpp b/Engine/Plugins/Runtime/Metasound/Source/MetasoundFrontend/Private/MetasoundFrontend.cpp
--- a/Engine/Plugins/Runtime/Metasound/Source/MetasoundFrontend/Private/MetasoundFrontend.cpp
+++ b/Engine/Plugins/Runtime/Metasound/Source/MetasoundFrontend/Private/MetasoundFrontend.cpp
@@ -48,20 +48,20 @@ namespace Metasound
 
 			const FInputVertexInterface& InputInterface = InNodeMetadata.DefaultInterface.GetInputInterface();
 			FMetasoundFrontendInterfaceStyle InputStyle;
-			for (auto& InputTuple : InputInterface)
+			for (const TPair<FVertexName, FInputDataVertex>& InputTuple : InputInterface)
 			{
 				FMetasoundFrontendClassInput ClassInput;
 
-				ClassInput.Name = InputTuple.Value.GetVertexName();
-				ClassInput.TypeName = InputTuple.Value.GetDataTypeName();
+				const FInputDataVertex& InputVertex = InputTuple.Value;
+				ClassInput.Name = InputVertex.GetVertexName();
+				ClassInput.TypeName = InputVertex.GetDataTypeName();
 				ClassInput.VertexID = FGuid::NewGuid();
-				ClassInput.Metadata.DisplayName = FText::FromString(InputTuple.Value.GetVertexName());
 
-				const FDataVertexMetadata& VertexMetadata = InputTuple.Value.GetMetadata();
+				const FDataVertexMetadata& VertexMetadata = InputVertex.GetMetadata();
 				ClassInput.Metadata.Description = VertexMetadata.Description;
 				ClassInput.Metadata.bIsAdvancedDisplay = VertexMetadata.bIsAdvancedDisplay;
 
-				FLiteral DefaultLiteral = InputTuple.Value.GetDefaultLiteral();
+				FLiteral DefaultLiteral = InputVertex.GetDefaultLiteral();
 				if (DefaultLiteral.GetType() != ELiteralType::Invalid)
 				{
 					ClassInput.DefaultLiteral.SetFromLiteral(DefaultLiteral);
@@ -81,14 +81,13 @@ namespace Metasound
 
 			const FOutputVertexInterface& OutputInterface = InNodeMetadata.DefaultInterface.GetOutputInterface();
 			FMetasoundFrontendInterfaceStyle OutputStyle;
-			for (auto& OutputTuple : OutputInterface)
+			for (const TPair<FVertexName, FOutputDataVertex>& OutputTuple : OutputInterface)
 			{
 				FMetasoundFrontendClassOutput ClassOutput;
 
 				ClassOutput.Name = OutputTuple.Value.GetVertexName();
 				ClassOutput.TypeName = OutputTuple.Value.GetDataTypeName();
 				ClassOutput.VertexID = FGuid::NewGuid();
-				ClassOutput.Metadata.DisplayName = FText::FromString(OutputTuple.Value.GetVertexName());
 
 				const FDataVertexMetadata& VertexMetadata = OutputTuple.Value.GetMetadata();
 				ClassOutput.Metadata.Description = VertexMetadata.Description;
@@ -111,7 +110,6 @@ namespace Metasound
 				FMetasoundFrontendClassEnvironmentVariable EnvVar;
 
 				EnvVar.Name = EnvTuple.Value.GetVertexName();
-				EnvVar.Metadata.DisplayName = FText::FromString(EnvTuple.Value.GetVertexName());
 				EnvVar.Metadata.Description = EnvTuple.Value.GetDescription();
 				EnvVar.bIsRequired = true;
 
diff --git a/Engine/Plugins/Runtime/Metasound/Source/MetasoundFrontend/Private/MetasoundFrontendDocumentAccessPtr.cpp b/Engine/Plugins/Runtime/Metasound/Source/MetasoundFrontend/Private/MetasoundFrontendDocumentAccessPtr.cpp
--- a/Engine/Plugins/Runtime/Metasound/Source/MetasoundFrontend/Private/MetasoundFrontendDocumentAccessPtr.cpp
+++ b/Engine/Plugins/Runtime/Metasound/Source/MetasoundFrontend/Private/MetasoundFrontendDocumentAccessPtr.cpp
@@ -236,14 +236,18 @@ namespace Metasound
 			/** A callable object for testing whether a vertex contains a matching name. */
 			struct FIsVertexWithName
 			{
-				FIsVertexWithName(const FString& InName)
+				FIsVertexWithName(const FVertexName& InName)
 				: Name(InName)
 				{
 				}
 
-				bool operator() (const FMetasoundFrontendVertex& InVertex) const { return InVertex.Name == Name; }
+				bool operator() (const FMetasoundFrontendVertex& InVertex) const
+				{
+					return InVertex.Name == Name;
+				}
+
 			private:
-				FString Name;
+				FVertexName Name;
 			};
 
 			/** A callable object for testing whether a vertex contains a matching ID. */
@@ -254,7 +258,11 @@ namespace Metasound
 				{
 				}
 
-				bool operator() (const FMetasoundFrontendVertex& InVertex) const { return InVertex.VertexID == ID; }
+				bool operator() (const FMetasoundFrontendVertex& InVertex) const
+				{
+					return InVertex.VertexID == ID;
+				}
+
 			private:
 				FGuid ID;
 			};
@@ -280,7 +288,11 @@ namespace Metasound
 				{
 				}
 
-				bool operator() (const FMetasoundFrontendClassOutput& InClassOutput) const { return InClassOutput.NodeID == NodeID; }
+				bool operator() (const FMetasoundFrontendClassOutput& InClassOutput) const
+				{
+					return InClassOutput.NodeID == NodeID;
+				}
+
 			private:
 				FGuid NodeID;
 			};
@@ -293,7 +305,11 @@ namespace Metasound
 				{
 				}
 
-				bool operator() (const FMetasoundFrontendNode& InNode) const { return InNode.ID == NodeID; }
+				bool operator() (const FMetasoundFrontendNode& InNode) const
+				{
+					return InNode.GetID() == NodeID;
+				}
+
 			private:
 				FGuid NodeID;
 			};
@@ -306,7 +322,11 @@ namespace Metasound
 				{
 				}
 
-				bool operator() (const FMetasoundFrontendClass& InClass) const { return InClass.ID == ClassID; }
+				bool operator() (const FMetasoundFrontendClass& InClass) const
+				{
+					return InClass.ID == ClassID;
+				}
+
 			private:
 				FGuid ClassID;
 			};
@@ -338,7 +358,7 @@ namespace Metasound
 				}
 
 				bool operator()(const FMetasoundFrontendClass& InClass) const
-				{ 
+				{
 					return NodeRegistryKey::IsEqual(NodeInfo, InClass.Metadata);
 				};
 			private:
@@ -355,7 +375,7 @@ namespace Metasound
 				}
 
 				bool operator()(const FMetasoundFrontendClass& InClass) const
-				{ 
+				{
 					return NodeRegistryKey::IsEqual(RegistryKey, NodeRegistryKey::CreateKey(InClass.Metadata));
 				};
 			private:
@@ -363,7 +383,7 @@ namespace Metasound
 			};
 		}
 
-		FVertexAccessPtr FNodeAccessPtr::GetInputWithName(const FString& InName)
+		FVertexAccessPtr FNodeAccessPtr::GetInputWithName(const FVertexName& InName)
 		{
 			using namespace MetasoundFrontendDocumentAccessPtrPrivate;
 
@@ -383,7 +403,7 @@ namespace Metasound
 			return GetMemberAccessPtr<FVertexAccessPtr>(Finder);
 		}
 
-		FVertexAccessPtr FNodeAccessPtr::GetOutputWithName(const FString& InName)
+		FVertexAccessPtr FNodeAccessPtr::GetOutputWithName(const FVertexName& InName)
 		{
 			using namespace MetasoundFrontendDocumentAccessPtrPrivate;
 
@@ -403,7 +423,7 @@ namespace Metasound
 			return GetMemberAccessPtr<FVertexAccessPtr>(Finder);
 		}
 
-		FConstVertexAccessPtr FNodeAccessPtr::GetInputWithName(const FString& InName) const
+		FConstVertexAccessPtr FNodeAccessPtr::GetInputWithName(const FVertexName& InName) const
 		{
 			using namespace MetasoundFrontendDocumentAccessPtrPrivate;
 
@@ -423,7 +443,7 @@ namespace Metasound
 			return GetMemberAccessPtr<FConstVertexAccessPtr>(Finder);
 		}
 
-		FConstVertexAccessPtr FNodeAccessPtr::GetOutputWithName(const FString& InName) const
+		FConstVertexAccessPtr FNodeAccessPtr::GetOutputWithName(const FVertexName& InName) const
 		{
 			using namespace MetasoundFrontendDocumentAccessPtrPrivate;
 
@@ -443,7 +463,7 @@ namespace Metasound
 			return GetMemberAccessPtr<FConstVertexAccessPtr>(Finder);
 		}
 
-		FConstVertexAccessPtr FConstNodeAccessPtr::GetInputWithName(const FString& InName) const
+		FConstVertexAccessPtr FConstNodeAccessPtr::GetInputWithName(const FVertexName& InName) const
 		{
 			using namespace MetasoundFrontendDocumentAccessPtrPrivate;
 
@@ -463,7 +483,7 @@ namespace Metasound
 			return GetMemberAccessPtr<FConstVertexAccessPtr>(Finder);
 		}
 
-		FConstVertexAccessPtr FConstNodeAccessPtr::GetOutputWithName(const FString& InName) const
+		FConstVertexAccessPtr FConstNodeAccessPtr::GetOutputWithName(const FVertexName& InName) const
 		{
 			using namespace MetasoundFrontendDocumentAccessPtrPrivate;
 
@@ -484,7 +504,7 @@ namespace Metasound
 		}
 
 		// FConstClassAccessPtr
-		FClassInputAccessPtr FClassAccessPtr::GetInputWithName(const FString& InName)
+		FClassInputAccessPtr FClassAccessPtr::GetInputWithName(const FVertexName& InName)
 		{
 			using namespace MetasoundFrontendDocumentAccessPtrPrivate;
 
@@ -494,17 +514,17 @@ namespace Metasound
 			return GetMemberAccessPtr<FClassInputAccessPtr>(Finder);
 		}
 
-		FClassOutputAccessPtr FClassAccessPtr::GetOutputWithName(const FString& InName)
+		FClassOutputAccessPtr FClassAccessPtr::GetOutputWithName(const FVertexName& InName)
 		{
 			using namespace MetasoundFrontendDocumentAccessPtrPrivate;
 
 			using FFinder= TFindInArray<FMetasoundFrontendClass, FMetasoundFrontendClassOutput, FGetClassOutputArray, FIsVertexWithName>;
-			FFinder Finder{FGetClassOutputArray(), FIsVertexWithName(InName)};
+			FFinder Finder { FGetClassOutputArray(), FIsVertexWithName(InName) };
 
 			return GetMemberAccessPtr<FClassOutputAccessPtr>(Finder);
 		}
 
-		FConstClassInputAccessPtr FClassAccessPtr::GetInputWithName(const FString& InName) const
+		FConstClassInputAccessPtr FClassAccessPtr::GetInputWithName(const FVertexName& InName) const
 		{
 			using namespace MetasoundFrontendDocumentAccessPtrPrivate;
 
@@ -514,7 +534,7 @@ namespace Metasound
 			return GetMemberAccessPtr<FConstClassInputAccessPtr>(Finder);
 		}
 
-		FConstClassOutputAccessPtr FClassAccessPtr::GetOutputWithName(const FString& InName) const
+		FConstClassOutputAccessPtr FClassAccessPtr::GetOutputWithName(const FVertexName& InName) const
 		{
 			using namespace MetasoundFrontendDocumentAccessPtrPrivate;
 
@@ -525,7 +545,7 @@ namespace Metasound
 		}
 
 		// FConstClassAccessPtr
-		FConstClassInputAccessPtr FConstClassAccessPtr::GetInputWithName(const FString& InName) const
+		FConstClassInputAccessPtr FConstClassAccessPtr::GetInputWithName(const FVertexName& InName) const
 		{
 			using namespace MetasoundFrontendDocumentAccessPtrPrivate;
 
@@ -535,7 +555,7 @@ namespace Metasound
 			return GetMemberAccessPtr<FConstClassInputAccessPtr>(Finder);
 		}
 
-		FConstClassOutputAccessPtr FConstClassAccessPtr::GetOutputWithName(const FString& InName) const
+		FConstClassOutputAccessPtr FConstClassAccessPtr::GetOutputWithName(const FVertexName& InName) const
 		{
 			using namespace MetasoundFrontendDocumentAccessPtrPrivate;
 
@@ -546,7 +566,7 @@ namespace Metasound
 		}
 
 		// FGraphClassAccessPtr
-		FClassInputAccessPtr FGraphClassAccessPtr::GetInputWithName(const FString& InName)
+		FClassInputAccessPtr FGraphClassAccessPtr::GetInputWithName(const FVertexName& InName)
 		{
 			using namespace MetasoundFrontendDocumentAccessPtrPrivate;
 
@@ -566,7 +586,7 @@ namespace Metasound
 			return GetMemberAccessPtr<FClassInputAccessPtr>(Finder);
 		}
 
-		FClassOutputAccessPtr FGraphClassAccessPtr::GetOutputWithName(const FString& InName)
+		FClassOutputAccessPtr FGraphClassAccessPtr::GetOutputWithName(const FVertexName& InName)
 		{
 			using namespace MetasoundFrontendDocumentAccessPtrPrivate;
 
@@ -607,7 +627,7 @@ namespace Metasound
 		}
 
 
-		FConstClassInputAccessPtr FGraphClassAccessPtr::GetInputWithName(const FString& InName) const
+		FConstClassInputAccessPtr FGraphClassAccessPtr::GetInputWithName(const FVertexName& InName) const
 		{
 			using namespace MetasoundFrontendDocumentAccessPtrPrivate;
 
@@ -628,7 +648,7 @@ namespace Metasound
 			return GetMemberAccessPtr<FConstClassInputAccessPtr>(Finder);
 		}
 
-		FConstClassOutputAccessPtr FGraphClassAccessPtr::GetOutputWithName(const FString& InName) const
+		FConstClassOutputAccessPtr FGraphClassAccessPtr::GetOutputWithName(const FVertexName& InName) const
 		{
 			using namespace MetasoundFrontendDocumentAccessPtrPrivate;
 
@@ -669,7 +689,7 @@ namespace Metasound
 		}
 
 		// FConstGraphClassAccessPtr
-		FConstClassInputAccessPtr FConstGraphClassAccessPtr::GetInputWithName(const FString& InName) const
+		FConstClassInputAccessPtr FConstGraphClassAccessPtr::GetInputWithName(const FVertexName& InName) const
 		{
 			using namespace MetasoundFrontendDocumentAccessPtrPrivate;
 
@@ -690,7 +710,7 @@ namespace Metasound
 			return GetMemberAccessPtr<FConstClassInputAccessPtr>(Finder);
 		}
 
-		FConstClassOutputAccessPtr FConstGraphClassAccessPtr::GetOutputWithName(const FString& InName) const
+		FConstClassOutputAccessPtr FConstGraphClassAccessPtr::GetOutputWithName(const FVertexName& InName) const
 		{
 			using namespace MetasoundFrontendDocumentAccessPtrPrivate;
 
diff --git a/Engine/Plugins/Runtime/Metasound/Source/MetasoundFrontend/Private/MetasoundFrontendGraph.cpp b/Engine/Plugins/Runtime/Metasound/Source/MetasoundFrontend/Private/MetasoundFrontendGraph.cpp
--- a/Engine/Plugins/Runtime/Metasound/Source/MetasoundFrontend/Private/MetasoundFrontendGraph.cpp
+++ b/Engine/Plugins/Runtime/Metasound/Source/MetasoundFrontend/Private/MetasoundFrontendGraph.cpp
@@ -22,11 +22,8 @@ namespace Metasound
 		{
 			FNodeInitData InitData;
 
-			InitData.InstanceName.Append(InNode.Name);
-			InitData.InstanceName.AppendChar('_');
-			InitData.InstanceName.Append(InNode.ID.ToString());
-
-			InitData.InstanceID = InNode.ID;
+			InitData.InstanceName = *FString::Format(TEXT("{0}_{1}"), { *InNode.Name.ToString(), *InNode.GetID().ToString() });
+			InitData.InstanceID = InNode.GetID();
 
 			return InitData;
 		}
@@ -37,7 +34,7 @@ namespace Metasound
 	{
 	}
 
-	void FFrontendGraph::AddInputNode(FGuid InDependencyId, int32 InIndex, const FVertexKey& InVertexKey, TSharedPtr<const INode> InNode)
+	void FFrontendGraph::AddInputNode(FGuid InDependencyId, int32 InIndex, const FVertexName& InVertexName, TSharedPtr<const INode> InNode)
 	{
 		if (InNode.IsValid())
 		{
@@ -46,12 +43,12 @@ namespace Metasound
 
 			// Input nodes need an extra Index value to keep track of their position in the graph's inputs.
 			InputNodes.Add(InIndex, InNode.Get());
-			AddInputDataDestination(*InNode, InVertexKey);
+			AddInputDataDestination(*InNode, InVertexName);
 			AddNode(InDependencyId, InNode);
 		}
 	}
 
-	void FFrontendGraph::AddOutputNode(FGuid InNodeID, int32 InIndex, const FVertexKey& InVertexKey, TSharedPtr<const INode> InNode)
+	void FFrontendGraph::AddOutputNode(FGuid InNodeID, int32 InIndex, const FVertexName& InVertexName, TSharedPtr<const INode> InNode)
 	{
 		if (InNode.IsValid())
 		{
@@ -60,7 +57,7 @@ namespace Metasound
 
 			// Output nodes need an extra Index value to keep track of their position in the graph's inputs.
 			OutputNodes.Add(InIndex, InNode.Get());
-			AddOutputDataSource(*InNode, InVertexKey);
+			AddOutputDataSource(*InNode, InVertexName);
 			AddNode(InNodeID, InNode);
 		}
 	}
@@ -182,7 +179,7 @@ namespace Metasound
 					FInputNodeConstructorParams InitParams =
 					{
 						InNode.Name,
-						InNode.ID,
+						InNode.GetID(),
 						InputVertex.Name,
 						MoveTemp(Literal)
 					};
@@ -191,13 +188,13 @@ namespace Metasound
 				}
 				else
 				{
-					UE_LOG(LogMetaSound, Error, TEXT("Cannot create input node [NodeID:%s]. [Vertex:%s] cannot be constructed with the provided literal type."), *InNode.ID.ToString(), *InputVertex.Name);
+					UE_LOG(LogMetaSound, Error, TEXT("Cannot create input node [NodeID:%s]. [Vertex:%s] cannot be constructed with the provided literal type."), *InNode.GetID().ToString(), *InputVertex.Name.ToString());
 				}
 			}
 		}
 		else
 		{
-			UE_LOG(LogMetaSound, Error, TEXT("Cannot create input node [NodeID:%s]. No default literal set for input node."), *InNode.ID.ToString());
+			UE_LOG(LogMetaSound, Error, TEXT("Cannot create input node [NodeID:%s]. No default literal set for input node."), *InNode.GetID().ToString());
 		}
 
 		return TUniquePtr<INode>(nullptr);
@@ -217,7 +214,7 @@ namespace Metasound
 			FOutputNodeConstructorParams InitParams =
 			{
 				InNode.Name,
-				InNode.ID,
+				InNode.GetID(),
 				OutputVertex.Name
 			};
 
@@ -226,7 +223,7 @@ namespace Metasound
 				TArray<FDefaultLiteralData> DefaultLiteralData = GetInputDefaultLiteralData(InNode, InitData);
 				for (FDefaultLiteralData& Data : DefaultLiteralData)
 				{
-					InGraphContext.DefaultInputs.Emplace(FNodeIDVertexID { InNode.ID, Data.DestinationVertexID }, MoveTemp(Data));
+					InGraphContext.DefaultInputs.Emplace(FNodeIDVertexID { InNode.GetID(), Data.DestinationVertexID }, MoveTemp(Data));
 				}
 			}
 
@@ -245,7 +242,7 @@ namespace Metasound
 			TArray<FDefaultLiteralData> DefaultLiteralData = GetInputDefaultLiteralData(InNode, InitData);
 			for (FDefaultLiteralData& Data : DefaultLiteralData)
 			{
-				InGraphContext.DefaultInputs.Emplace(FNodeIDVertexID { InNode.ID, Data.DestinationVertexID }, MoveTemp(Data));
+				InGraphContext.DefaultInputs.Emplace(FNodeIDVertexID { InNode.GetID(), Data.DestinationVertexID }, MoveTemp(Data));
 			}
 		}
 
@@ -268,7 +265,7 @@ namespace Metasound
 
 			auto IsMatchingInput = [&](const FMetasoundFrontendClassInput& GraphInput)
 			{
-				return (InInputNode.ID == GraphInput.NodeID);
+				return (InInputNode.GetID() == GraphInput.NodeID);
 			};
 
 			OutClassInputIndex = InOwningGraph.Interface.Inputs.IndexOfByPredicate(IsMatchingInput);
@@ -291,7 +288,7 @@ namespace Metasound
 
 			auto IsMatchingOutput = [&](const FMetasoundFrontendClassOutput& GraphOutput)
 			{
-				return (InOutputNode.ID == GraphOutput.NodeID);
+				return (InOutputNode.GetID() == GraphOutput.NodeID);
 			};
 
 			OutClassOutputIndex = InOwningGraph.Interface.Outputs.IndexOfByPredicate(IsMatchingOutput);
@@ -374,12 +371,12 @@ namespace Metasound
 						if ((nullptr != ClassInput) && (INDEX_NONE != InputIndex))
 						{
 							TSharedPtr<const INode> InputNode(CreateInputNode(Node, *NodeClass, *ClassInput).Release());
-							InGraphContext.Graph->AddInputNode(Node.ID, InputIndex, ClassInput->Name, InputNode);
+							InGraphContext.Graph->AddInputNode(Node.GetID(), InputIndex, ClassInput->Name, InputNode);
 						}
 						else
 						{
 							const FString GraphClassIDString = InGraphContext.GraphClass.ID.ToString();
-							UE_LOG(LogMetaSound, Error, TEXT("Failed to match input node [NodeID:%s, NodeName:%s] to owning graph [ClassID:%s] output."), *Node.ID.ToString(), *Node.Name, *GraphClassIDString);
+							UE_LOG(LogMetaSound, Error, TEXT("Failed to match input node [NodeID:%s, NodeName:%s] to owning graph [ClassID:%s] output."), *Node.GetID().ToString(), *Node.Name.ToString(), *GraphClassIDString);
 						}
 					}
 					break;
@@ -391,12 +388,12 @@ namespace Metasound
 						if ((nullptr != ClassOutput) && (INDEX_NONE != OutputIndex))
 						{
 							TSharedPtr<const INode> OutputNode(CreateOutputNode(Node, *NodeClass, InGraphContext).Release());
-							InGraphContext.Graph->AddOutputNode(Node.ID, OutputIndex, ClassOutput->Name, OutputNode);
+							InGraphContext.Graph->AddOutputNode(Node.GetID(), OutputIndex, ClassOutput->Name, OutputNode);
 						}
 						else
 						{
 							const FString GraphClassIDString = InGraphContext.GraphClass.ID.ToString();
-							UE_LOG(LogMetaSound, Error, TEXT("Failed to match output node [NodeID:%s, NodeName:%s] to owning graph [ClassID:%s] output."), *Node.ID.ToString(), *Node.Name, *GraphClassIDString);
+							UE_LOG(LogMetaSound, Error, TEXT("Failed to match output node [NodeID:%s, NodeName:%s] to owning graph [ClassID:%s] output."), *Node.GetID().ToString(), *Node.Name.ToString(), *GraphClassIDString);
 						}
 					}
 					break;
@@ -414,11 +411,11 @@ namespace Metasound
 
 						if (SubgraphPtr.IsValid())
 						{
-							InGraphContext.Graph->AddNode(Node.ID, SubgraphPtr);
+							InGraphContext.Graph->AddNode(Node.GetID(), SubgraphPtr);
 						}
 						else
 						{
-							UE_LOG(LogMetaSound, Error, TEXT("Failed to find subgraph for node [NodeID:%s, NodeName:%s, ClassID:%s]"), *Node.ID.ToString(), *Node.Name, *Node.ClassID.ToString());
+							UE_LOG(LogMetaSound, Error, TEXT("Failed to find subgraph for node [NodeID:%s, NodeName:%s, ClassID:%s]"), *Node.GetID().ToString(), *Node.Name.ToString(), *Node.ClassID.ToString());
 						}
 					}
 					break;
@@ -427,7 +424,7 @@ namespace Metasound
 					default:
 					{
 						TSharedPtr<const INode> ExternalNode(CreateExternalNode(Node, *NodeClass, InGraphContext).Release());
-						InGraphContext.Graph->AddNode(Node.ID, ExternalNode);
+						InGraphContext.Graph->AddNode(Node.GetID(), ExternalNode);
 					}
 					break;
 				}
@@ -451,21 +448,21 @@ namespace Metasound
 		// Add nodes to NodeID/VertexID map
 		for (const FMetasoundFrontendNode& Node : InGraphContext.GraphClass.Graph.Nodes)
 		{
-			const INode* CoreNode = InGraphContext.Graph->FindNode(Node.ID);
+			const INode* CoreNode = InGraphContext.Graph->FindNode(Node.GetID());
 			if (nullptr == CoreNode)
 			{
-				UE_LOG(LogMetaSound, Warning, TEXT("Could not find referenced node [Name:%s, NodeID:%s]"), *Node.Name, *Node.ID.ToString());
+				UE_LOG(LogMetaSound, Warning, TEXT("Could not find referenced node [Name:%s, NodeID:%s]"), *Node.Name.ToString(), *Node.GetID().ToString());
 				continue;
 			}
 
 			for (const FMetasoundFrontendVertex& Vertex : Node.Interface.Inputs)
 			{
-				NodeDestinationsByID.Add(FNodeIDVertexID(Node.ID, Vertex.VertexID), FCoreNodeAndFrontendVertex({CoreNode, &Vertex}));
+				NodeDestinationsByID.Add(FNodeIDVertexID(Node.GetID(), Vertex.VertexID), FCoreNodeAndFrontendVertex({CoreNode, &Vertex}));
 			}
 
 			for (const FMetasoundFrontendVertex& Vertex : Node.Interface.Outputs)
 			{
-				NodeSourcesByID.Add(FNodeIDVertexID(Node.ID, Vertex.VertexID), FCoreNodeAndFrontendVertex({CoreNode, &Vertex}));
+				NodeSourcesByID.Add(FNodeIDVertexID(Node.GetID(), Vertex.VertexID), FCoreNodeAndFrontendVertex({CoreNode, &Vertex}));
 			}
 		};
 
@@ -505,10 +502,10 @@ namespace Metasound
 			}
 
 			const INode* FromNode = SourceNodeAndVertex->Node;
-			const FVertexKey FromVertexKey = SourceNodeAndVertex->Vertex->Name;
+			const FVertexName FromVertexKey = SourceNodeAndVertex->Vertex->Name;
 
 			const INode* ToNode = DestinationNodeAndVertex->Node;
-			const FVertexKey ToVertexKey = DestinationNodeAndVertex->Vertex->Name;
+			const FVertexName ToVertexKey = DestinationNodeAndVertex->Vertex->Name;
 
 			bool bSuccess = InGraphContext.Graph->AddDataEdge(*FromNode, FromVertexKey,  *ToNode, ToVertexKey);
 
@@ -548,19 +545,19 @@ namespace Metasound
 			{
 				continue;
 			}
-			const FVertexKey& FromVertexKey = LiteralNodeNames::GetOutputDataName();
+			const FVertexName& FromVertexName = LiteralNodeNames::GetOutputDataName();
 
 			const INode* ToNode = InGraphContext.Graph->FindNode(LiteralData.DestinationNodeID);
 			if (!ensure(ToNode))
 			{
 				continue;
 			}
-			const FVertexKey& ToVertexKey = LiteralData.DestinationVertexKey;
+			const FVertexName& ToVertexName = LiteralData.DestinationVertexKey;
 
-			bool bSuccess = InGraphContext.Graph->AddDataEdge(*FromNode, FromVertexKey, *ToNode, ToVertexKey);
+			bool bSuccess = InGraphContext.Graph->AddDataEdge(*FromNode, FromVertexName, *ToNode, ToVertexName);
 			if (!bSuccess)
 			{
-				UE_LOG(LogMetaSound, Error, TEXT("Failed to connect default variable edge: from '%s' to '%s'"), *FromVertexKey, *ToVertexKey);
+				UE_LOG(LogMetaSound, Error, TEXT("Failed to connect default variable edge: from '%s' to '%s'"), *FromVertexName.ToString(), *ToVertexName.ToString());
 			}
 		}
 
@@ -579,17 +576,17 @@ namespace Metasound
 			FName TypeName;
 
 			FGuid VertexID = Literal.VertexID;
-			FVertexKey DestinationVertexKey;
+			FVertexName DestinationVertexName;
 			auto RemoveAndBuildParams = [&](const FMetasoundFrontendVertex& Vertex)
 			{
 				if (Vertex.VertexID == VertexID)
 				{
 					InitParams.Literal = Literal.Value.ToLiteral(Vertex.TypeName);
 					InitParams.InstanceID = FGuid::NewGuid();
-					InitParams.NodeName = InInitData.InstanceName + Vertex.Name + InitParams.InstanceID.ToString();
+					InitParams.NodeName = "Literal";
 					TypeName = Vertex.TypeName;
 
-					DestinationVertexKey = Vertex.Name;
+					DestinationVertexName = Vertex.Name;
 
 					return true;
 				}
@@ -602,9 +599,9 @@ namespace Metasound
 			{
 				DefaultLiteralData.Emplace(FDefaultLiteralData
 				{
-					InNode.ID,
+					InNode.GetID(),
 					VertexID,
-					DestinationVertexKey,
+					DestinationVertexName,
 					TypeName,
 					MoveTemp(InitParams)
 				});
diff --git a/Engine/Plugins/Runtime/Metasound/Source/MetasoundFrontend/Private/MetasoundFrontendInjectReceiveNodes.cpp b/Engine/Plugins/Runtime/Metasound/Source/MetasoundFrontend/Private/MetasoundFrontendInjectReceiveNodes.cpp
--- a/Engine/Plugins/Runtime/Metasound/Source/MetasoundFrontend/Private/MetasoundFrontendInjectReceiveNodes.cpp
+++ b/Engine/Plugins/Runtime/Metasound/Source/MetasoundFrontend/Private/MetasoundFrontendInjectReceiveNodes.cpp
@@ -45,9 +45,9 @@ namespace Metasound
 			class FAddressNode : public FNodeFacade
 			{
 			public:
-				FAddressNode(const FGuid& InID, const FVertexKey& InVertexKey, const FName& InTypeName, FReceiveNodeAddressFunction InAddressFunction)
-				: FNodeFacade(FString::Format(TEXT("ReceiveAddressInject_{0}"), { InVertexKey }), InID, TFacadeOperatorClass<FAddressOperator>())
-				, VertexKey(InVertexKey)
+				FAddressNode(const FGuid& InID, const FVertexName& InVertexName, const FName& InTypeName, FReceiveNodeAddressFunction InAddressFunction)
+				: FNodeFacade(*FString::Format(TEXT("ReceiveAddressInject_{0}"), { InVertexName.ToString() }), InID, TFacadeOperatorClass<FAddressOperator>())
+				, VertexKey(InVertexName)
 				, TypeName(InTypeName)
 				, AddressFunction(InAddressFunction)
 				{
@@ -59,7 +59,7 @@ namespace Metasound
 				}
 
 			private:
-				FVertexKey VertexKey;
+				FVertexName VertexKey;
 				FName TypeName;
 				FReceiveNodeAddressFunction AddressFunction;
 			};
@@ -113,11 +113,11 @@ namespace Metasound
 				return Metadata;
 			}
 
-			TUniquePtr<INode> CreateReceiveNodeForDataType(const FGuid& InID, const FVertexKey& InVertexKey, const FName& InDataType)
+			TUniquePtr<INode> CreateReceiveNodeForDataType(const FGuid& InID, const FVertexName& InVertexName, const FName& InDataType)
 			{
 				FNodeInitData ReceiveNodeInitData 
 				{ 
-					FString::Format(TEXT("ReceiveInject_{0}"), { InVertexKey }),
+					*FString::Format(TEXT("ReceiveInject_{0}"), { InVertexName.ToString() }),
 					InID
 				};
 				return IDataTypeRegistry::Get().CreateReceiveNode(InDataType, ReceiveNodeInitData);
@@ -130,7 +130,7 @@ namespace Metasound
 			// should never contain null nodes for input destination.
 			check(InputDestination.Node != nullptr);
 
-			const FVertexKey& VertexKey = InputDestination.Vertex.GetVertexName();
+			const FVertexName& VertexKey = InputDestination.Vertex.GetVertexName();
 			const FName& DataType = InputDestination.Vertex.GetDataTypeName();
 
 			// Create a receive node.
@@ -139,7 +139,7 @@ namespace Metasound
 
 			if (!ReceiveNode.IsValid())
 			{
-				UE_LOG(LogMetaSound, Warning, TEXT("Failed to create receive node while injecting receive node for graph input [VertexKey:%s, VertexDescription:%s, DataTypeName:%s]"), *VertexKey, *InputDestination.Vertex.GetMetadata().Description.ToString(), *DataType.ToString());
+				UE_LOG(LogMetaSound, Warning, TEXT("Failed to create receive node while injecting receive node for graph input [VertexKey:%s, VertexDescription:%s, DataTypeName:%s]"), *VertexKey.ToString(), *InputDestination.Vertex.GetMetadata().Description.ToString(), *DataType.ToString());
 				return false;
 			}
 
@@ -149,7 +149,7 @@ namespace Metasound
 
 			if (!AddressNode.IsValid())
 			{
-				UE_LOG(LogMetaSound, Warning, TEXT("Failed to create address node while injecting receive node for graph input [VertexKey:%s, VertexDescription:%s, DataTypeName:%s]"), *VertexKey, *InputDestination.Vertex.GetMetadata().Description.ToString(), *DataType.ToString());
+				UE_LOG(LogMetaSound, Warning, TEXT("Failed to create address node while injecting receive node for graph input [VertexKey:%s, VertexDescription:%s, DataTypeName:%s]"), *VertexKey.ToString(), *InputDestination.Vertex.GetMetadata().Description.ToString(), *DataType.ToString());
 				return false;
 			}
 
@@ -185,7 +185,7 @@ namespace Metasound
 			return true;
 		}
 
-		bool InjectReceiveNodes(FFrontendGraph& InGraph, const FReceiveNodeAddressFunction& InAddressFunction, const TSet<FVertexKey>& InInputVerticesToSkip)
+		bool InjectReceiveNodes(FFrontendGraph& InGraph, const FReceiveNodeAddressFunction& InAddressFunction, const TSet<FVertexName>& InInputVerticesToSkip)
 		{
 			using namespace Metasound::Frontend;
 			bool bSuccess = true;
@@ -202,7 +202,7 @@ namespace Metasound
 					continue;
 				}
 
-				const FVertexKey& VertexKey = InputDestination.Vertex.GetVertexName();
+				const FVertexName& VertexKey = InputDestination.Vertex.GetVertexName();
 				if (InInputVerticesToSkip.Contains(VertexKey))
 				{
 					// Skip this key as requested by caller.
diff --git a/Engine/Plugins/Runtime/Metasound/Source/MetasoundFrontend/Private/MetasoundFrontendInvalidController.h b/Engine/Plugins/Runtime/Metasound/Source/MetasoundFrontend/Private/MetasoundFrontendInvalidController.h
--- a/Engine/Plugins/Runtime/Metasound/Source/MetasoundFrontend/Private/MetasoundFrontendInvalidController.h
+++ b/Engine/Plugins/Runtime/Metasound/Source/MetasoundFrontend/Private/MetasoundFrontendInvalidController.h
@@ -9,6 +9,7 @@
 #include "MetasoundFrontendController.h"
 #include "MetasoundFrontendDocument.h"
 #include "MetasoundGraph.h"
+#include "MetasoundVertex.h"
 
 
 namespace Metasound
@@ -47,14 +48,14 @@ namespace Metasound
 			bool IsValid() const override { return false; }
 			FGuid GetID() const override { return Metasound::FrontendInvalidID; }
 			const FName& GetDataType() const override { return MetasoundFrontendInvalidControllerPrivate::GetInvalid<FName>(); }
-			const FString& GetName() const override { return MetasoundFrontendInvalidControllerPrivate::GetInvalid<FString>(); }
-			const FText& GetDisplayName() const override { return FText::GetEmpty(); }
+			const FVertexName& GetName() const override { return MetasoundFrontendInvalidControllerPrivate::GetInvalid<FVertexName>(); }
+			FText GetDisplayName() const override { return FText::GetEmpty(); }
 			const FText& GetTooltip() const override { return FText::GetEmpty(); }
 			const FMetasoundFrontendVertexMetadata& GetMetadata() const override { return MetasoundFrontendInvalidControllerPrivate::GetInvalid<FMetasoundFrontendVertexMetadata>(); }
 			FGuid GetOwningNodeID() const override { return Metasound::FrontendInvalidID; }
 			TSharedRef<INodeController> GetOwningNode() override;
 			TSharedRef<const INodeController> GetOwningNode() const override;
-
+			void SetName(const FVertexName& InName) override { }
 			bool IsConnected() const override { return false; }
 			TArray<FInputHandle> GetConnectedInputs() override { return TArray<FInputHandle>(); }
 			TArray<FConstInputHandle> GetConstConnectedInputs() const override { return TArray<FConstInputHandle>(); }
@@ -91,8 +92,8 @@ namespace Metasound
 			FGuid GetID() const override { return Metasound::FrontendInvalidID; }
 			bool IsConnected() const override { return false; }
 			const FName& GetDataType() const override { return MetasoundFrontendInvalidControllerPrivate::GetInvalid<FName>(); }
-			const FString& GetName() const override { return MetasoundFrontendInvalidControllerPrivate::GetInvalid<FString>(); }
-			const FText& GetDisplayName() const override { return MetasoundFrontendInvalidControllerPrivate::GetInvalid<FText>(); }
+			const FVertexName& GetName() const override { return MetasoundFrontendInvalidControllerPrivate::GetInvalid<FVertexName>(); }
+			FText GetDisplayName() const override { return MetasoundFrontendInvalidControllerPrivate::GetInvalid<FText>(); }
 			const FText& GetTooltip() const override { return MetasoundFrontendInvalidControllerPrivate::GetInvalid<FText>(); }
 			const FMetasoundFrontendVertexMetadata& GetMetadata() const override { return MetasoundFrontendInvalidControllerPrivate::GetInvalid<FMetasoundFrontendVertexMetadata>(); }
 			const FMetasoundFrontendLiteral* GetLiteral() const override { return nullptr; }
@@ -106,6 +107,8 @@ namespace Metasound
 			virtual TSharedRef<const IOutputController> GetConnectedOutput() const override { return FInvalidOutputController::GetInvalid(); }
 			bool Disconnect() override { return false; }
 
+			void SetName(const FVertexName& InName) override { }
+
 			virtual FConnectability CanConnectTo(const IOutputController& InController) const override { return FConnectability(); }
 			virtual bool Connect(IOutputController& InController) override { return false; }
 			virtual bool ConnectWithConverterNode(IOutputController& InController, const FConverterNodeInfo& InNodeClassName) override { return false; }
@@ -139,17 +142,18 @@ namespace Metasound
 			TArray<TSharedRef<const IInputController>> GetConstInputs() const override { return TArray<TSharedRef<const IInputController>>(); }
 			TArray<TSharedRef<const IOutputController>> GetConstOutputs() const override { return TArray<TSharedRef<const IOutputController>>(); }
 
-			TArray<FInputHandle> GetInputsWithVertexName(const FString& InName) override { return TArray<FInputHandle>(); }
-			TArray<FConstInputHandle> GetConstInputsWithVertexName(const FString& InName) const override { return TArray<FConstInputHandle>(); }
-			TArray<FOutputHandle> GetOutputsWithVertexName(const FString& InName) override { return TArray<FOutputHandle>(); }
-			TArray<FConstOutputHandle> GetConstOutputsWithVertexName(const FString& InName) const override { return TArray<FConstOutputHandle>(); }
+			TArray<FInputHandle> GetInputsWithVertexName(const FVertexName& InName) override { return TArray<FInputHandle>(); }
+			TArray<FConstInputHandle> GetConstInputsWithVertexName(const FVertexName& InName) const override { return TArray<FConstInputHandle>(); }
+			TArray<FOutputHandle> GetOutputsWithVertexName(const FVertexName& InName) override { return TArray<FOutputHandle>(); }
+			TArray<FConstOutputHandle> GetConstOutputsWithVertexName(const FVertexName& InName) const override { return TArray<FConstOutputHandle>(); }
 			TSharedRef<IInputController> GetInputWithID(FGuid InVertexID) override { return FInvalidInputController::GetInvalid(); }
 			TSharedRef<IOutputController> GetOutputWithID(FGuid InVertexID) override { return FInvalidOutputController::GetInvalid(); }
 			TSharedRef<const IInputController> GetInputWithID(FGuid InVertexID) const override { return FInvalidInputController::GetInvalid(); }
 			TSharedRef<const IOutputController> GetOutputWithID(FGuid InVertexID) const override { return FInvalidOutputController::GetInvalid(); }
 
 			const FMetasoundFrontendNodeStyle& GetNodeStyle() const override { static const FMetasoundFrontendNodeStyle Invalid; return Invalid; }
-			void SetNodeStyle(const FMetasoundFrontendNodeStyle& InNodeStyle) { };
+			void SetNodeStyle(const FMetasoundFrontendNodeStyle& InNodeStyle) { }
+			void SetNodeName(const FVertexName& InName) override { }
 
 			FNodeHandle ReplaceWithVersion(const FMetasoundFrontendVersionNumber& InNewVersion) override { return FInvalidNodeController::GetInvalid(); }
 
@@ -199,8 +203,8 @@ namespace Metasound
 			int32 GetNumInputs() const override { return 0; }
 			int32 GetNumOutputs() const override { return 0; }
 
-			const FString& GetNodeName() const override { return MetasoundFrontendInvalidControllerPrivate::GetInvalid<FString>(); }
-			const FText&  GetDisplayName() const override { return MetasoundFrontendInvalidControllerPrivate::GetInvalid<FText>(); }
+			const FVertexName& GetNodeName() const override { return MetasoundFrontendInvalidControllerPrivate::GetInvalid<FVertexName>(); }
+			FText  GetDisplayName() const override { return MetasoundFrontendInvalidControllerPrivate::GetInvalid<FText>(); }
 			const FText& GetDisplayTitle() const override { return MetasoundFrontendInvalidControllerPrivate::GetInvalid<FText>(); }
 			void SetDescription(const FText& InDescription) override { }
 			void SetDisplayName(const FText& InText) override { }
@@ -228,10 +232,10 @@ namespace Metasound
 
 			bool IsValid() const override { return false; }
 			FGuid GetClassID() const override { return Metasound::FrontendInvalidID; }
-			const FText& GetDisplayName() const override { return MetasoundFrontendInvalidControllerPrivate::GetInvalid<FText>(); }
+			FText GetDisplayName() const override { return MetasoundFrontendInvalidControllerPrivate::GetInvalid<FText>(); }
 
-			TArray<FString> GetInputVertexNames() const override { return TArray<FString>(); }
-			TArray<FString> GetOutputVertexNames() const override { return TArray<FString>(); }
+			TArray<FVertexName> GetInputVertexNames() const override { return TArray<FVertexName>(); }
+			TArray<FVertexName> GetOutputVertexNames() const override { return TArray<FVertexName>(); }
 
 			TArray<TSharedRef<INodeController>> GetNodes() override { return TArray<TSharedRef<INodeController>>(); }
 			TArray<TSharedRef<const INodeController>> GetConstNodes() const override { return TArray<TSharedRef<const INodeController>>(); }
@@ -249,35 +253,37 @@ namespace Metasound
 
 			void ClearGraph() override { };
 
-			void IterateConstNodes(TUniqueFunction<void(FConstNodeHandle)> InFunction, EMetasoundFrontendClassType InClassType /* = EMetasoundFrontendClassType::Invalid */) const override { }
-			void IterateNodes(TUniqueFunction<void(FNodeHandle)> InFunction, EMetasoundFrontendClassType InClassType /* = EMetasoundFrontendClassType::Invalid */) override { }
+			void IterateConstNodes(TFunctionRef<void(FConstNodeHandle)> InFunction, EMetasoundFrontendClassType InClassType /* = EMetasoundFrontendClassType::Invalid */) const override { }
+			void IterateNodes(TFunctionRef<void(FNodeHandle)> InFunction, EMetasoundFrontendClassType InClassType /* = EMetasoundFrontendClassType::Invalid */) override { }
 
-			bool ContainsOutputVertexWithName(const FString& InName) const override { return false; }
-			bool ContainsInputVertexWithName(const FString& InName) const override { return false; }
+			bool ContainsOutputVertexWithName(const FVertexName& InName) const override { return false; }
+			bool ContainsInputVertexWithName(const FVertexName& InName) const override { return false; }
 
-			TSharedRef<const INodeController> GetOutputNodeWithName(const FString& InName) const override { return FInvalidNodeController::GetInvalid(); }
-			TSharedRef<const INodeController> GetInputNodeWithName(const FString& InName) const override { return FInvalidNodeController::GetInvalid(); }
-			TSharedRef<INodeController> GetOutputNodeWithName(const FString& InName) override { return FInvalidNodeController::GetInvalid(); }
-			TSharedRef<INodeController> GetInputNodeWithName(const FString& InName) override { return FInvalidNodeController::GetInvalid(); }
+			TSharedRef<const INodeController> GetOutputNodeWithName(const FVertexName& InName) const override { return FInvalidNodeController::GetInvalid(); }
+			TSharedRef<const INodeController> GetInputNodeWithName(const FVertexName& InName) const override { return FInvalidNodeController::GetInvalid(); }
+			TSharedRef<INodeController> GetOutputNodeWithName(const FVertexName& InName) override { return FInvalidNodeController::GetInvalid(); }
+			TSharedRef<INodeController> GetInputNodeWithName(const FVertexName& InName) override { return FInvalidNodeController::GetInvalid(); }
 
-			FConstClassInputAccessPtr FindClassInputWithName(const FString& InName) const override { return FConstClassInputAccessPtr(); }
-			FConstClassOutputAccessPtr FindClassOutputWithName(const FString& InName) const override { return FConstClassOutputAccessPtr(); }
+			FConstClassInputAccessPtr FindClassInputWithName(const FVertexName& InName) const override { return FConstClassInputAccessPtr(); }
+			FConstClassOutputAccessPtr FindClassOutputWithName(const FVertexName& InName) const override { return FConstClassOutputAccessPtr(); }
 
 			TSharedRef<INodeController> AddInputVertex(const FMetasoundFrontendClassInput& InDescription) override { return FInvalidNodeController::GetInvalid(); }
-			bool RemoveInputVertex(const FString& InputName) override { return false; }
+			TSharedRef<INodeController> AddInputVertex(const FVertexName& InName, const FName InTypeName, const FText& InToolTip, const FMetasoundFrontendLiteral* InDefaultValue) override { return INodeController::GetInvalidHandle(); }
+			bool RemoveInputVertex(const FVertexName& InputName) override { return false; }
 
 			TSharedRef<INodeController> AddOutputVertex(const FMetasoundFrontendClassOutput& InDescription) override { return FInvalidNodeController::GetInvalid(); }
-			bool RemoveOutputVertex(const FString& OutputName) override { return false; }
+			TSharedRef<INodeController> AddOutputVertex(const FVertexName& InName, const FName InTypeName, const FText& InToolTip) override { return INodeController::GetInvalidHandle(); }
+			bool RemoveOutputVertex(const FVertexName& OutputName) override { return false; }
 
 			// This can be used to determine what kind of property editor we should use for the data type of a given input.
 			// Will return Invalid if the input couldn't be found, or if the input doesn't support any kind of literals.
-			ELiteralType GetPreferredLiteralTypeForInputVertex(const FString& InInputName) const override { return ELiteralType::Invalid; }
+			ELiteralType GetPreferredLiteralTypeForInputVertex(const FVertexName& InInputName) const override { return ELiteralType::Invalid; }
 
 			// For inputs whose preferred literal type is UObject or UObjectArray, this can be used to determine the UObject corresponding to that input's datatype.
-			UClass* GetSupportedClassForInputVertex(const FString& InInputName) override { return nullptr; }
+			UClass* GetSupportedClassForInputVertex(const FVertexName& InInputName) override { return nullptr; }
 
-			FGuid GetVertexIDForInputVertex(const FString& InInputName) const { return Metasound::FrontendInvalidID; }
-			FGuid GetVertexIDForOutputVertex(const FString& InInputName) const { return Metasound::FrontendInvalidID; }
+			FGuid GetVertexIDForInputVertex(const FVertexName& InInputName) const { return Metasound::FrontendInvalidID; }
+			FGuid GetVertexIDForOutputVertex(const FVertexName& InOutputName) const { return Metasound::FrontendInvalidID; }
 			FMetasoundFrontendLiteral GetDefaultInput(const FGuid& InVertexID) const override { return FMetasoundFrontendLiteral{}; }
 
 			// These can be used to set the default value for a given input on this graph.
@@ -285,17 +291,17 @@ namespace Metasound
 			bool SetDefaultInput(const FGuid& InVertexID, const FMetasoundFrontendLiteral& InLiteral) override { return false; }
 			bool SetDefaultInputToDefaultLiteralOfType(const FGuid& InVertexID) override { return false; }
 
-			const FText& GetInputDescription(const FString& InName) const override { return FText::GetEmpty(); }
-			const FText& GetOutputDescription(const FString& InName) const override { return FText::GetEmpty(); }
+			const FText& GetInputDescription(const FVertexName& InName) const override { return FText::GetEmpty(); }
+			const FText& GetOutputDescription(const FVertexName& InName) const override { return FText::GetEmpty(); }
 
-			void SetInputDescription(const FString& InName, const FText& InDescription) override { }
-			void SetOutputDescription(const FString& InName, const FText& InDescription) override { }
-			void SetInputDisplayName(const FString& InName, const FText& InDisplayName) override { }
-			void SetOutputDisplayName(const FString& InName, const FText& InDisplayName) override { }
+			void SetInputDescription(const FVertexName& InName, const FText& InDescription) override { }
+			void SetOutputDescription(const FVertexName& InName, const FText& InDescription) override { }
+			void SetInputDisplayName(const FVertexName& InName, const FText& InDisplayName) override { }
+			void SetOutputDisplayName(const FVertexName& InName, const FText& InDisplayName) override { }
 
 			// This can be used to clear the current literal for a given input.
 			// @returns false if the input name couldn't be found.
-			bool ClearLiteralForInput(const FString& InInputName, FGuid InVertexID) override { return false; }
+			bool ClearLiteralForInput(const FVertexName& InInputName, FGuid InVertexID) override { return false; }
 
 			TSharedRef<INodeController> AddNode(const FNodeRegistryKey& InNodeClass) override { return FInvalidNodeController::GetInvalid(); }
 			TSharedRef<INodeController> AddNode(const FMetasoundFrontendClassMetadata& InNodeClass) override { return FInvalidNodeController::GetInvalid(); }
@@ -322,11 +328,6 @@ namespace Metasound
 
 			void UpdateInterfaceChangeID() override { }
 
-			FNodeHandle AddInputNode(const FMetasoundFrontendClassInput& InClassInput, const FMetasoundFrontendLiteral* InDefaultValue, const FText* InDisplayName) override { return INodeController::GetInvalidHandle(); }
-			FNodeHandle AddInputNode(const FName InTypeName, const FText& InToolTip, const FMetasoundFrontendLiteral* InDefaultValue, const FText* InDisplayName) override { return INodeController::GetInvalidHandle(); }
-			FNodeHandle AddOutputNode(const FMetasoundFrontendClassOutput& InClassOutput, const FText* InDisplayName) override { return INodeController::GetInvalidHandle(); }
-			FNodeHandle AddOutputNode(const FName InTypeName, const FText& InToolTip, const FText* InDisplayName) override { return INodeController::GetInvalidHandle(); }
-
 		protected:
 			FDocumentAccess ShareAccess() override { return FDocumentAccess(); }
 			FConstDocumentAccess ShareAccess() const override { return FConstDocumentAccess(); }
