Move FString/FText VertexName identifers to FNames. Convention is as follows:

diff --git a/Engine/Plugins/Runtime/Metasound/Source/MetasoundFrontend/Private/MetasoundFrontendStandardController.cpp b/Engine/Plugins/Runtime/Metasound/Source/MetasoundFrontend/Private/MetasoundFrontendStandardController.cpp
--- a/Engine/Plugins/Runtime/Metasound/Source/MetasoundFrontend/Private/MetasoundFrontendStandardController.cpp
+++ b/Engine/Plugins/Runtime/Metasound/Source/MetasoundFrontend/Private/MetasoundFrontendStandardController.cpp
@@ -21,6 +21,7 @@
 #include "MetasoundJsonBackend.h"
 #include "MetasoundOperatorBuilder.h"
 #include "MetasoundTrace.h"
+#include "MetasoundVertex.h"
 #include "StructDeserializer.h"
 #include "StructSerializer.h"
 #include "Serialization/MemoryReader.h"
@@ -119,13 +120,13 @@ namespace Metasound
 			return FrontendControllerIntrinsics::GetInvalidValueConstRef<FName>();
 		}
 
-		const FString& FBaseOutputController::GetName() const
+		const FVertexName& FBaseOutputController::GetName() const
 		{
 			if (const FMetasoundFrontendVertex* Vertex = NodeVertexPtr.Get())
 			{
 				return Vertex->Name;
 			}
-			return FrontendControllerIntrinsics::GetInvalidValueConstRef<FString>();
+			return FrontendControllerIntrinsics::GetInvalidValueConstRef<FVertexName>();
 		}
 
 		FGuid FBaseOutputController::GetOwningNodeID() const
@@ -241,11 +242,16 @@ namespace Metasound
 			return TArray<FMetasoundFrontendEdge>();
 		}
 
-		const FText& FBaseOutputController::GetDisplayName() const
+		FText FBaseOutputController::GetDisplayName() const
 		{
 			if (const FMetasoundFrontendClassOutput* Output = ClassOutputPtr.Get())
 			{
-				return Output->Metadata.DisplayName;
+				if (!Output->Metadata.DisplayName.IsEmptyOrWhitespace())
+				{
+					return Output->Metadata.DisplayName;
+				}
+
+				return FText::FromName(GetName());
 			}
 
 			return FrontendControllerIntrinsics::GetInvalidValueConstRef<FText>();
@@ -308,7 +314,7 @@ namespace Metasound
 			return FBaseOutputController::IsValid() && (nullptr != OwningGraphClassInputPtr.Get());
 		}
 
-		const FText& FInputNodeOutputController::GetDisplayName() const
+		FText FInputNodeOutputController::GetDisplayName() const
 		{
 			if (const FMetasoundFrontendClassInput* OwningInput = OwningGraphClassInputPtr.Get())
 			{
@@ -347,6 +353,14 @@ namespace Metasound
 			return FrontendControllerIntrinsics::GetInvalidValueConstRef<FMetasoundFrontendVertexMetadata>();
 		}
 
+		void FInputNodeOutputController::SetName(const FVertexName& InName)
+		{
+			if (FMetasoundFrontendVertex* Vertex = ConstCastAccessPtr<FVertexAccessPtr>(NodeVertexPtr).Get())
+			{
+				Vertex->Name = InName;
+			}
+		}
+
 		FDocumentAccess FInputNodeOutputController::ShareAccess() 
 		{
 			FDocumentAccess Access = FBaseOutputController::ShareAccess();
@@ -379,11 +393,16 @@ namespace Metasound
 			return FBaseOutputController::IsValid() && (nullptr != OwningGraphClassOutputPtr.Get());
 		}
 
-		const FText& FOutputNodeOutputController::GetDisplayName() const
+		FText FOutputNodeOutputController::GetDisplayName() const
 		{
 			if (const FMetasoundFrontendClassOutput* OwningOutput = OwningGraphClassOutputPtr.Get())
 			{
-				return OwningOutput->Metadata.DisplayName;
+				if (!OwningOutput->Metadata.DisplayName.IsEmptyOrWhitespace())
+				{
+					return OwningOutput->Metadata.DisplayName;
+				}
+
+				return FText::FromName(GetName());
 			}
 
 			return FrontendControllerIntrinsics::GetInvalidValueConstRef<FText>();
@@ -409,6 +428,14 @@ namespace Metasound
 			return FrontendControllerIntrinsics::GetInvalidValueConstRef<FMetasoundFrontendVertexMetadata>();
 		}
 
+		void FOutputNodeOutputController::SetName(const FVertexName& InName)
+		{
+			if (FMetasoundFrontendVertex* Vertex = ConstCastAccessPtr<FVertexAccessPtr>(NodeVertexPtr).Get())
+			{
+				Vertex->Name = InName;
+			}
+		}
+
 		FConnectability FOutputNodeOutputController::CanConnectTo(const IInputController& InController) const 
 		{
 			// Cannot connect to a graph's output.
@@ -459,21 +486,26 @@ namespace Metasound
 			return FrontendControllerIntrinsics::GetInvalidValueConstRef<FName>();
 		}
 
-		const FString& FBaseInputController::GetName() const
+		const FVertexName& FBaseInputController::GetName() const
 		{
 			if (const FMetasoundFrontendVertex* Vertex = NodeVertexPtr.Get())
 			{
 				return Vertex->Name;
 			}
 			
-			return FrontendControllerIntrinsics::GetInvalidValueConstRef<FString>();
+			return FrontendControllerIntrinsics::GetInvalidValueConstRef<FVertexName>();
 		}
 
-		const FText& FBaseInputController::GetDisplayName() const
+		FText FBaseInputController::GetDisplayName() const
 		{
 			if (const FMetasoundFrontendClassInput* ClassInput = ClassInputPtr.Get())
 			{
-				return ClassInput->Metadata.DisplayName;
+				if (!ClassInput->Metadata.DisplayName.IsEmptyOrWhitespace())
+				{
+					return ClassInput->Metadata.DisplayName;
+				}
+
+				return FText::FromName(GetName());
 			}
 			
 			return FrontendControllerIntrinsics::GetInvalidValueConstRef<FText>();
@@ -661,13 +693,13 @@ namespace Metasound
 
 			if (ConverterInputs.Num() < 1)
 			{
-				UE_LOG(LogMetaSound, Warning, TEXT("Converter node [Name: %s] does not support preferred input vertex [Vertex: %s]"), *ConverterNode->GetNodeName(), *InConverterInfo.PreferredConverterInputPin);
+				UE_LOG(LogMetaSound, Warning, TEXT("Converter node [Name: %s] does not support preferred input vertex [Vertex: %s]"), *ConverterNode->GetNodeName().ToString(), *InConverterInfo.PreferredConverterInputPin.ToString());
 				return false;
 			}
 
 			if (ConverterOutputs.Num() < 1)
 			{
-				UE_LOG(LogMetaSound, Warning, TEXT("Converter node [Name: %s] does not support preferred output vertex [Vertex: %s]"), *ConverterNode->GetNodeName(), *InConverterInfo.PreferredConverterOutputPin);
+				UE_LOG(LogMetaSound, Warning, TEXT("Converter node [Name: %s] does not support preferred output vertex [Vertex: %s]"), *ConverterNode->GetNodeName().ToString(), *InConverterInfo.PreferredConverterOutputPin.ToString());
 				return false;
 			}
 
@@ -797,7 +829,7 @@ namespace Metasound
 			return FBaseInputController::IsValid() && (nullptr != OwningGraphClassOutputPtr.Get());
 		}
 
-		const FText& FOutputNodeInputController::GetDisplayName() const
+		FText FOutputNodeInputController::GetDisplayName() const
 		{
 			if (const FMetasoundFrontendClassOutput* OwningOutput = OwningGraphClassOutputPtr.Get())
 			{
@@ -837,6 +869,14 @@ namespace Metasound
 			return FrontendControllerIntrinsics::GetInvalidValueConstRef<FMetasoundFrontendVertexMetadata>();
 		}
 
+		void FOutputNodeInputController::SetName(const FVertexName& InName)
+		{
+			if (FMetasoundFrontendVertex* Vertex = ConstCastAccessPtr<FVertexAccessPtr>(NodeVertexPtr).Get())
+			{
+				Vertex->Name = InName;
+			}
+		}
+
 		FDocumentAccess FOutputNodeInputController::ShareAccess() 
 		{
 			FDocumentAccess Access = FBaseInputController::ShareAccess();
@@ -870,11 +910,16 @@ namespace Metasound
 			return FBaseInputController::IsValid() && (nullptr != OwningGraphClassInputPtr.Get());
 		}
 
-		const FText& FInputNodeInputController::GetDisplayName() const
+		FText FInputNodeInputController::GetDisplayName() const
 		{
 			if (const FMetasoundFrontendClassInput* OwningInput = OwningGraphClassInputPtr.Get())
 			{
-				return OwningInput->Metadata.DisplayName;
+				if (!OwningInput->Metadata.DisplayName.IsEmptyOrWhitespace())
+				{
+					return OwningInput->Metadata.DisplayName;
+				}
+
+				return FText::FromName(GetName());
 			}
 
 			return FrontendControllerIntrinsics::GetInvalidValueConstRef<FText>();
@@ -900,6 +945,14 @@ namespace Metasound
 			return FrontendControllerIntrinsics::GetInvalidValueConstRef<FMetasoundFrontendVertexMetadata>();
 		}
 
+		void FInputNodeInputController::SetName(const FVertexName& InName)
+		{
+			if (FMetasoundFrontendVertex* Vertex = ConstCastAccessPtr<FVertexAccessPtr>(NodeVertexPtr).Get())
+			{
+				Vertex->Name = InName;
+			}
+		}
+
 		FConnectability FInputNodeInputController::CanConnectTo(const IOutputController& InController) const 
 		{
 			static const FConnectability Connectability = {FConnectability::EConnectable::No};
@@ -961,7 +1014,7 @@ namespace Metasound
 		{
 			if (const FMetasoundFrontendNode* Node = NodePtr.Get())
 			{
-				return Node->ID;
+				return Node->GetID();
 			}
 			return Metasound::FrontendInvalidID;
 		}
@@ -1114,13 +1167,13 @@ namespace Metasound
 			return FrontendControllerIntrinsics::GetInvalidValueConstRef<FText>();
 		}
 
-		const FString& FBaseNodeController::GetNodeName() const
+		const FVertexName& FBaseNodeController::GetNodeName() const
 		{
 			if (const FMetasoundFrontendNode* Node = NodePtr.Get())
 			{
 				return Node->Name;
 			}
-			return FrontendControllerIntrinsics::GetInvalidValueConstRef<FString>();
+			return FrontendControllerIntrinsics::GetInvalidValueConstRef<FVertexName>();
 		}
 
 		bool FBaseNodeController::CanAddInput(const FString& InVertexName) const
@@ -1279,11 +1332,17 @@ namespace Metasound
 			return FrontendControllerIntrinsics::GetInvalidValueConstRef<FText>();
 		}
 
-		const FText& FBaseNodeController::GetDisplayName() const
+		FText FBaseNodeController::GetDisplayName() const
 		{
 			if (const FMetasoundFrontendClass* Class = ClassPtr.Get())
 			{
-				return Class->Metadata.GetDisplayName();
+				const FText& DisplayName = Class->Metadata.GetDisplayName();
+				if (!DisplayName.IsEmptyOrWhitespace())
+				{
+					return Class->Metadata.GetDisplayName();
+				}
+
+				return FText::FromString(GetNodeName().ToString());
 			}
 
 			return FrontendControllerIntrinsics::GetInvalidValueConstRef<FText>();
@@ -1338,7 +1397,7 @@ namespace Metasound
 			}
 		}
 
-		TArray<FInputHandle> FBaseNodeController::GetInputsWithVertexName(const FString& InName) 
+		TArray<FInputHandle> FBaseNodeController::GetInputsWithVertexName(const FVertexName& InName)
 		{
 			TArray<FInputHandle> Inputs;
 
@@ -1356,7 +1415,7 @@ namespace Metasound
 			return Inputs;
 		}
 
-		TArray<FConstInputHandle> FBaseNodeController::GetConstInputsWithVertexName(const FString& InName) const 
+		TArray<FConstInputHandle> FBaseNodeController::GetConstInputsWithVertexName(const FVertexName& InName) const
 		{
 			TArray<FConstInputHandle> Inputs;
 
@@ -1375,7 +1434,7 @@ namespace Metasound
 			return Inputs;
 		}
 
-		TArray<FOutputHandle> FBaseNodeController::GetOutputsWithVertexName(const FString& InName)
+		TArray<FOutputHandle> FBaseNodeController::GetOutputsWithVertexName(const FVertexName& InName)
 		{
 			TArray<FOutputHandle> Outputs;
 
@@ -1393,7 +1452,7 @@ namespace Metasound
 			return Outputs;
 		}
 
-		TArray<FConstOutputHandle> FBaseNodeController::GetConstOutputsWithVertexName(const FString& InName) const
+		TArray<FConstOutputHandle> FBaseNodeController::GetConstOutputsWithVertexName(const FVertexName& InName) const
 		{
 			TArray<FConstOutputHandle> Outputs;
 
@@ -1497,7 +1556,7 @@ namespace Metasound
 			{
 				for (const FMetasoundFrontendVertex& NodeOutputVertex : Node->Interface.Outputs)
 				{
-					const FString& VertexName = NodeOutputVertex.Name;
+					const FVertexName& VertexName = NodeOutputVertex.Name;
 
 					FConstVertexAccessPtr NodeVertexPtr = NodePtr.GetOutputWithName(VertexName);
 					FConstClassOutputAccessPtr ClassOutputPtr = ClassPtr.GetOutputWithName(VertexName);
@@ -1509,7 +1568,7 @@ namespace Metasound
 			return Outputs;
 		}
 
-		TArray<FBaseNodeController::FInputControllerParams> FBaseNodeController::GetInputControllerParamsWithVertexName(const FString& InName) const
+		TArray<FBaseNodeController::FInputControllerParams> FBaseNodeController::GetInputControllerParamsWithVertexName(const FVertexName& InName) const
 		{
 			TArray<FBaseNodeController::FInputControllerParams> Inputs;
 
@@ -1525,7 +1584,7 @@ namespace Metasound
 			return Inputs;
 		}
 
-		TArray<FBaseNodeController::FOutputControllerParams> FBaseNodeController::GetOutputControllerParamsWithVertexName(const FString& InName) const
+		TArray<FBaseNodeController::FOutputControllerParams> FBaseNodeController::GetOutputControllerParamsWithVertexName(const FVertexName& InName) const
 		{
 			TArray<FBaseNodeController::FOutputControllerParams> Outputs;
 
@@ -1651,7 +1710,7 @@ namespace Metasound
 
 			FMetasoundFrontendNodeStyle Style = GetNodeStyle();
 
-			using FConnectionKey = TPair<FString, FName>;
+			using FConnectionKey = TPair<FVertexName, FName>;
 
 			struct FInputConnectionInfo
 			{
@@ -1940,7 +1999,7 @@ namespace Metasound
 					}
 					else
 					{
-						UE_LOG(LogMetaSound, Warning, TEXT("Frontend Node [NodeID:%s, ClassID:%s] is not of expected class class [ClassID:%s]"), *Node->ID.ToString(), *Node->ClassID.ToString(), *Class->ID.ToString());
+						UE_LOG(LogMetaSound, Warning, TEXT("Frontend Node [NodeID:%s, ClassID:%s] is not of expected class class [ClassID:%s]"), *Node->GetID().ToString(), *Node->ClassID.ToString(), *Class->ID.ToString());
 					}
 				}
 			}
@@ -1960,7 +2019,7 @@ namespace Metasound
 					}
 					else
 					{
-						UE_LOG(LogMetaSound, Warning, TEXT("Frontend Node [NodeID:%s, ClassID:%s] is not of expected class class [ClassID:%s]"), *Node->ID.ToString(), *Node->ClassID.ToString(), *Class->ID.ToString());
+						UE_LOG(LogMetaSound, Warning, TEXT("Frontend Node [NodeID:%s, ClassID:%s] is not of expected class class [ClassID:%s]"), *Node->GetID().ToString(), *Node->ClassID.ToString(), *Class->ID.ToString());
 					}
 				}
 			}
@@ -2026,7 +2085,7 @@ namespace Metasound
 						}
 						else
 						{
-							UE_LOG(LogMetaSound, Warning, TEXT("Frontend Node [NodeID:%s, ClassID:%s] is not of expected class class [ClassID:%s]"), *Node->ID.ToString(), *Node->ClassID.ToString(), *Class->ID.ToString());
+							UE_LOG(LogMetaSound, Warning, TEXT("Frontend Node [NodeID:%s, ClassID:%s] is not of expected class class [ClassID:%s]"), *Node->GetID().ToString(), *Node->ClassID.ToString(), *Class->ID.ToString());
 						}
 					}
 					else
@@ -2049,11 +2108,16 @@ namespace Metasound
 			return FrontendControllerIntrinsics::GetInvalidValueConstRef<FText>();
 		}
 
-		const FText& FOutputNodeController::GetDisplayName() const
+		FText FOutputNodeController::GetDisplayName() const
 		{
 			if (const FMetasoundFrontendClassOutput* OwningOutput = OwningGraphClassOutputPtr.Get())
 			{
-				return OwningOutput->Metadata.DisplayName;
+				if (!OwningOutput->Metadata.DisplayName.IsEmptyOrWhitespace())
+				{
+					return OwningOutput->Metadata.DisplayName;
+				}
+
+				return FText::FromName(GetNodeName());
 			}
 
 			return FrontendControllerIntrinsics::GetInvalidValueConstRef<FText>();
@@ -2069,6 +2133,21 @@ namespace Metasound
 			}
 		}
 
+		void FOutputNodeController::SetNodeName(const FVertexName& InName)
+		{
+			if (FMetasoundFrontendNode* Node = NodePtr.Get())
+			{
+				Node->Name = InName;
+			}
+
+			// TODO: can we remove the const cast by constructing output nodes with a non-const access to class outputs?
+			if (FMetasoundFrontendClassOutput* ClassOutput = ConstCastAccessPtr<FClassOutputAccessPtr>(OwningGraphClassOutputPtr).Get())
+			{
+				ClassOutput->Name = InName;
+				OwningGraph->UpdateInterfaceChangeID();
+			}
+		}
+
 		void FOutputNodeController::SetDisplayName(const FText& InDisplayName)
 		{
 			// TODO: can we remove the const cast by constructing output nodes with a non-const access to class outputs?
@@ -2093,7 +2172,7 @@ namespace Metasound
 						}
 						else
 						{
-							UE_LOG(LogMetaSound, Warning, TEXT("Frontend Node [NodeID:%s, ClassID:%s] is not of expected class class [ClassID:%s]"), *Node->ID.ToString(), *Node->ClassID.ToString(), *Class->ID.ToString());
+							UE_LOG(LogMetaSound, Warning, TEXT("Frontend Node [NodeID:%s, ClassID:%s] is not of expected class class [ClassID:%s]"), *Node->GetID().ToString(), *Node->ClassID.ToString(), *Class->ID.ToString());
 						}
 					}
 					else
@@ -2123,7 +2202,7 @@ namespace Metasound
 			if (bIsNodeOnRootGraph)
 			{
 				// If the node is on the root graph, test if it is in the archetypes
-				// required inputs or outputs. 
+				// required inputs or outputs.
 				FMetasoundFrontendArchetype Archetype;
 
 				FArchetypeRegistryKey ArchetypeKey = GetArchetypeRegistryKey(OwningDocument->GetArchetypeVersion());
@@ -2133,7 +2212,7 @@ namespace Metasound
 				{
 					if (const FMetasoundFrontendNode* Node = NodePtr.Get())
 					{
-						const FString& Name = Node->Name;
+						const FVertexName& Name = Node->Name;
 						auto IsVertexWithSameName = [&Name](const FMetasoundFrontendClassVertex& InVertex)
 						{
 							return InVertex.Name == Name;
@@ -2211,7 +2290,7 @@ namespace Metasound
 						}
 						else
 						{
-							UE_LOG(LogMetaSound, Warning, TEXT("Frontend Node [NodeID:%s, ClassID:%s] is not of expected class class [ClassID:%s]"), *Node->ID.ToString(), *Node->ClassID.ToString(), *Class->ID.ToString());
+							UE_LOG(LogMetaSound, Warning, TEXT("Frontend Node [NodeID:%s, ClassID:%s] is not of expected class class [ClassID:%s]"), *Node->GetID().ToString(), *Node->ClassID.ToString(), *Class->ID.ToString());
 						}
 					}
 					else
@@ -2238,7 +2317,7 @@ namespace Metasound
 						}
 						else
 						{
-							UE_LOG(LogMetaSound, Warning, TEXT("Frontend Node [NodeID:%s, ClassID:%s] is not of expected class class [ClassID:%s]"), *Node->ID.ToString(), *Node->ClassID.ToString(), *Class->ID.ToString());
+							UE_LOG(LogMetaSound, Warning, TEXT("Frontend Node [NodeID:%s, ClassID:%s] is not of expected class class [ClassID:%s]"), *Node->GetID().ToString(), *Node->ClassID.ToString(), *Class->ID.ToString());
 						}
 					}
 					else
@@ -2276,11 +2355,16 @@ namespace Metasound
 			return FrontendControllerIntrinsics::GetInvalidValueConstRef<FText>();
 		}
 
-		const FText& FInputNodeController::GetDisplayName() const
+		FText FInputNodeController::GetDisplayName() const
 		{
 			if (const FMetasoundFrontendClassInput* OwningInput = OwningGraphClassInputPtr.Get())
 			{
-				return OwningInput->Metadata.DisplayName;
+				if (!OwningInput->Metadata.DisplayName.IsEmptyOrWhitespace())
+				{
+					return OwningInput->Metadata.DisplayName;
+				}
+
+				return FText::FromName(GetNodeName());
 			}
 
 			return FrontendControllerIntrinsics::GetInvalidValueConstRef<FText>();
@@ -2313,7 +2397,7 @@ namespace Metasound
 				{
 					if (const FMetasoundFrontendNode* Node = NodePtr.Get())
 					{
-						const FString& Name = Node->Name;
+						const FVertexName& Name = Node->Name;
 						auto IsVertexWithSameName = [&Name](const FMetasoundFrontendClassVertex& InVertex)
 						{
 							return InVertex.Name == Name;
@@ -2340,6 +2424,19 @@ namespace Metasound
 			}
 		}
 
+		void FInputNodeController::SetNodeName(const FVertexName& InName)
+		{
+			if (FMetasoundFrontendNode* Node = NodePtr.Get())
+			{
+				Node->Name = InName;
+			}
+
+			if (FMetasoundFrontendClassInput* ClassInput = ConstCastAccessPtr<FClassInputAccessPtr>(OwningGraphClassInputPtr).Get())
+			{
+				ClassInput->Name = InName;
+			}
+		}
+
 		void FInputNodeController::SetDisplayName(const FText& InDisplayName)
 		{
 			// TODO: can we remove these const casts by constructing FINputNodeController with non-const access to the class input?
@@ -2427,7 +2524,7 @@ namespace Metasound
 			return Metasound::FrontendInvalidID;
 		}
 
-		const FText& FGraphController::GetDisplayName() const 
+		FText FGraphController::GetDisplayName() const
 		{
 			if (const FMetasoundFrontendGraphClass* GraphClass = GraphClassPtr.Get())
 			{
@@ -2437,9 +2534,9 @@ namespace Metasound
 			return FrontendControllerIntrinsics::GetInvalidValueConstRef<FText>();
 		}
 
-		TArray<FString> FGraphController::GetInputVertexNames() const 
+		TArray<FVertexName> FGraphController::GetInputVertexNames() const
 		{
-			TArray<FString> Names;
+			TArray<FVertexName> Names;
 
 			if (const FMetasoundFrontendGraphClass* GraphClass = GraphClassPtr.Get())
 			{
@@ -2452,9 +2549,9 @@ namespace Metasound
 			return Names;
 		}
 
-		TArray<FString> FGraphController::GetOutputVertexNames() const
+		TArray<FVertexName> FGraphController::GetOutputVertexNames() const
 		{
-			TArray<FString> Names;
+			TArray<FVertexName> Names;
 
 			if (const FMetasoundFrontendGraphClass* GraphClass = GraphClassPtr.Get())
 			{
@@ -2467,17 +2564,17 @@ namespace Metasound
 			return Names;
 		}
 
-		FConstClassInputAccessPtr FGraphController::FindClassInputWithName(const FString& InName) const
+		FConstClassInputAccessPtr FGraphController::FindClassInputWithName(const FVertexName& InName) const
 		{
 			return GraphClassPtr.GetInputWithName(InName);
 		}
 
-		FConstClassOutputAccessPtr FGraphController::FindClassOutputWithName(const FString& InName) const
+		FConstClassOutputAccessPtr FGraphController::FindClassOutputWithName(const FVertexName& InName) const
 		{
 			return GraphClassPtr.GetOutputWithName(InName);
 		}
 
-		FGuid FGraphController::GetVertexIDForInputVertex(const FString& InInputName) const
+		FGuid FGraphController::GetVertexIDForInputVertex(const FVertexName& InInputName) const
 		{
 			if (const FMetasoundFrontendClassInput* Input = FindClassInputWithName(InInputName).Get())
 			{
@@ -2486,7 +2583,7 @@ namespace Metasound
 			return Metasound::FrontendInvalidID;
 		}
 
-		FGuid FGraphController::GetVertexIDForOutputVertex(const FString& InOutputName) const
+		FGuid FGraphController::GetVertexIDForOutputVertex(const FVertexName& InOutputName) const
 		{
 			if (const FMetasoundFrontendClassOutput* Output = FindClassOutputWithName(InOutputName).Get())
 			{
@@ -2509,7 +2606,7 @@ namespace Metasound
 		{
 			auto IsNodeWithSameID = [&](const FMetasoundFrontendClass& NodeClas, const FMetasoundFrontendNode& Node) 
 			{ 
-				return Node.ID == InNodeID; 
+				return Node.GetID() == InNodeID;
 			};
 			return GetNodeByPredicate(IsNodeWithSameID);
 		}
@@ -2518,7 +2615,7 @@ namespace Metasound
 		{
 			auto IsNodeWithSameID = [&](const FMetasoundFrontendClass& NodeClas, const FMetasoundFrontendNode& Node) 
 			{ 
-				return Node.ID == InNodeID; 
+				return Node.GetID() == InNodeID;
 			};
 			return GetNodeByPredicate(IsNodeWithSameID);
 		}
@@ -2571,7 +2668,7 @@ namespace Metasound
 			}
 		}
 
-		void FGraphController::IterateNodes(TUniqueFunction<void(FNodeHandle)> InFunction, EMetasoundFrontendClassType InClassType)
+		void FGraphController::IterateNodes(TFunctionRef<void(FNodeHandle)> InFunction, EMetasoundFrontendClassType InClassType)
 		{
 			if (FMetasoundFrontendGraphClass* GraphClass = GraphClassPtr.Get())
 			{
@@ -2582,20 +2679,20 @@ namespace Metasound
 					{
 						if (InClassType == EMetasoundFrontendClassType::Invalid || NodeClass->Metadata.GetType() == InClassType)
 						{
-							FNodeAccessPtr NodePtr = GraphClassPtr.GetNodeWithNodeID(Node.ID);
+							FNodeAccessPtr NodePtr = GraphClassPtr.GetNodeWithNodeID(Node.GetID());
 							FNodeHandle NodeHandle = GetNodeHandle(FGraphController::FNodeAndClass{ NodePtr, NodeClassPtr });
 							InFunction(NodeHandle);
 						}
 					}
 					else
 					{
-						UE_LOG(LogMetaSound, Warning, TEXT("Failed to find class for node [NodeID:%s, ClassID:%s]"), *Node.ID.ToString(), *Node.ClassID.ToString());
+						UE_LOG(LogMetaSound, Warning, TEXT("Failed to find class for node [NodeID:%s, ClassID:%s]"), *Node.GetID().ToString(), *Node.ClassID.ToString());
 					}
 				}
 			}
 		}
 
-		void FGraphController::IterateConstNodes(TUniqueFunction<void(FConstNodeHandle)> InFunction, EMetasoundFrontendClassType InClassType) const
+		void FGraphController::IterateConstNodes(TFunctionRef<void(FConstNodeHandle)> InFunction, EMetasoundFrontendClassType InClassType) const
 		{
 			if (FMetasoundFrontendGraphClass* GraphClass = GraphClassPtr.Get())
 			{
@@ -2606,14 +2703,14 @@ namespace Metasound
 					{
 						if (InClassType == EMetasoundFrontendClassType::Invalid || NodeClass->Metadata.GetType() == InClassType)
 						{
-							FConstNodeAccessPtr NodePtr = GraphClassPtr.GetNodeWithNodeID(Node.ID);
+							FConstNodeAccessPtr NodePtr = GraphClassPtr.GetNodeWithNodeID(Node.GetID());
 							FConstNodeHandle NodeHandle = GetNodeHandle(FGraphController::FConstNodeAndClass{ NodePtr, NodeClassPtr });
 							InFunction(NodeHandle);
 						}
 					}
 					else
 					{
-						UE_LOG(LogMetaSound, Warning, TEXT("Failed to find class for node [NodeID:%s, ClassID:%s]"), *Node.ID.ToString(), *Node.ClassID.ToString());
+						UE_LOG(LogMetaSound, Warning, TEXT("Failed to find class for node [NodeID:%s, ClassID:%s]"), *Node.GetID().ToString(), *Node.ClassID.ToString());
 					}
 				}
 			}
@@ -2637,7 +2734,7 @@ namespace Metasound
 			return GetNodesByPredicate(IsInputNode);
 		}
 
-		bool FGraphController::ContainsOutputVertexWithName(const FString& InName) const
+		bool FGraphController::ContainsOutputVertexWithName(const FVertexName& InName) const
 		{
 			if (const FMetasoundFrontendGraphClass* GraphClass = GraphClassPtr.Get())
 			{
@@ -2650,11 +2747,11 @@ namespace Metasound
 			return false;
 		}
 
-		bool FGraphController::ContainsInputVertexWithName(const FString& InName) const
+		bool FGraphController::ContainsInputVertexWithName(const FVertexName& InName) const
 		{
 			if (const FMetasoundFrontendGraphClass* GraphClass = GraphClassPtr.Get())
 			{
-				auto IsInputVertexWithSameName = [&](const FMetasoundFrontendClassInput& ClassInput) 
+				auto IsInputVertexWithSameName = [&](const FMetasoundFrontendClassInput& ClassInput)
 				{ 
 					return ClassInput.Name == InName;
 				};
@@ -2663,34 +2760,34 @@ namespace Metasound
 			return false;
 		}
 
-		FConstNodeHandle FGraphController::GetOutputNodeWithName(const FString& InName) const
+		FConstNodeHandle FGraphController::GetOutputNodeWithName(const FVertexName& InName) const
 		{
-			auto IsOutputNodeWithSameName = [&](const FMetasoundFrontendClass& NodeClass, const FMetasoundFrontendNode& Node) 
-			{ 
+			auto IsOutputNodeWithSameName = [&](const FMetasoundFrontendClass& NodeClass, const FMetasoundFrontendNode& Node)
+			{
 				return (NodeClass.Metadata.GetType() == EMetasoundFrontendClassType::Output) && (Node.Name == InName);
 			};
 			return GetNodeByPredicate(IsOutputNodeWithSameName);
 		}
 
-		FConstNodeHandle FGraphController::GetInputNodeWithName(const FString& InName) const
+		FConstNodeHandle FGraphController::GetInputNodeWithName(const FVertexName& InName) const
 		{
-			auto IsInputNodeWithSameName = [&](const FMetasoundFrontendClass& NodeClass, const FMetasoundFrontendNode& Node) 
-			{ 
+			auto IsInputNodeWithSameName = [&](const FMetasoundFrontendClass& NodeClass, const FMetasoundFrontendNode& Node)
+			{
 				return (NodeClass.Metadata.GetType() == EMetasoundFrontendClassType::Input) && (Node.Name == InName);
 			};
 			return GetNodeByPredicate(IsInputNodeWithSameName);
 		}
 
-		FNodeHandle FGraphController::GetOutputNodeWithName(const FString& InName)
+		FNodeHandle FGraphController::GetOutputNodeWithName(const FVertexName& InName)
 		{
-			auto IsOutputNodeWithSameName = [&](const FMetasoundFrontendClass& NodeClass, const FMetasoundFrontendNode& Node) 
-			{ 
+			auto IsOutputNodeWithSameName = [&](const FMetasoundFrontendClass& NodeClass, const FMetasoundFrontendNode& Node)
+			{
 				return (NodeClass.Metadata.GetType() == EMetasoundFrontendClassType::Output) && (Node.Name == InName);
 			};
 			return GetNodeByPredicate(IsOutputNodeWithSameName);
 		}
 
-		FNodeHandle FGraphController::GetInputNodeWithName(const FString& InName)
+		FNodeHandle FGraphController::GetInputNodeWithName(const FVertexName& InName)
 		{
 			auto IsInputNodeWithSameName = [&](const FMetasoundFrontendClass& NodeClass, const FMetasoundFrontendNode& Node)
 			{
@@ -2715,12 +2812,14 @@ namespace Metasound
 							FMetasoundFrontendNode& Node = GraphClass->Graph.Nodes.Emplace_GetRef(*InputClass);
 
 							Node.Name = InClassInput.Name;
-							Node.ID = FGuid::NewGuid();
+							Node.UpdateID();
+
+							const FVertexName NewName = InClassInput.Name;
 
 							auto IsVertexWithTypeName = [&](FMetasoundFrontendVertex& Vertex) { return Vertex.TypeName == InClassInput.TypeName; };
 							if (FMetasoundFrontendVertex* InputVertex = Node.Interface.Inputs.FindByPredicate(IsVertexWithTypeName))
-							{ 
-								InputVertex->Name = InClassInput.Name;
+							{
+								InputVertex->Name = NewName;
 							}
 							else
 							{
@@ -2729,117 +2828,90 @@ namespace Metasound
 
 							if (Node.Interface.Outputs.Num() == 1)
 							{
-								Node.Interface.Outputs[0].Name = InClassInput.Name;
+								Node.Interface.Outputs[0].Name = NewName;
 							}
 							else if (FMetasoundFrontendVertex* OutputVertex = Node.Interface.Outputs.FindByPredicate(IsVertexWithTypeName))
-							{ 
-								OutputVertex->Name = InClassInput.Name;
+							{
+								OutputVertex->Name = NewName;
 							}
 
 							FMetasoundFrontendClassInput& NewInput = GraphClass->Interface.Inputs.Add_GetRef(InClassInput);
-							NewInput.NodeID = Node.ID;
+							NewInput.NodeID = Node.GetID();
 
 							GraphClass->Interface.UpdateChangeID();
 
-							FNodeAccessPtr NodePtr = GraphClassPtr.GetNodeWithNodeID(Node.ID);
-							return GetNodeHandle(FGraphController::FNodeAndClass{NodePtr, InputClassPtr});
+							FNodeAccessPtr NodePtr = GraphClassPtr.GetNodeWithNodeID(Node.GetID());
+							return GetNodeHandle(FGraphController::FNodeAndClass { NodePtr, InputClassPtr });
 						}
 					}
-					else 
+					else
 					{
 						UE_LOG(LogMetaSound, Display, TEXT("Failed to add input. No input node registered for data type [TypeName:%s]"), *InClassInput.TypeName.ToString());
 					}
 				}
 				else
 				{
-					UE_LOG(LogMetaSound, Display, TEXT("Failed to add input. Input with same name \"%s\" exists in class [ClassID:%s]"), *InClassInput.Name, *GraphClass->ID.ToString());
+					UE_LOG(LogMetaSound, Display, TEXT("Failed to add input. Input with same name \"%s\" exists in class [ClassID:%s]"), *InClassInput.Name.ToString(), *GraphClass->ID.ToString());
 				}
 			}
 			return FInvalidNodeController::GetInvalid();
 		}
 
-		FNodeHandle FGraphController::AddInputNode(const FMetasoundFrontendClassInput& InClassInput, const FMetasoundFrontendLiteral* InDefaultValue, const FText* InDisplayName)
-		{
-			FNodeHandle NodeHandle = AddInputVertex(InClassInput);
-			if (ensure(NodeHandle->IsValid()))
-			{
-				if (InDisplayName)
-				{
-					SetInputDisplayName(InClassInput.Name, *InDisplayName);
-				}
-
-				if (InDefaultValue)
-				{
-					SetDefaultInput(InClassInput.VertexID, *InDefaultValue);
-				}
-				else
-				{
-					SetDefaultInputToDefaultLiteralOfType(InClassInput.VertexID);
-				}
-			}
-
-			return NodeHandle;
-		}
-
-		FNodeHandle FGraphController::AddInputNode(const FName InTypeName, const FText& InToolTip, const FMetasoundFrontendLiteral* InDefaultValue, const FText* InDisplayName)
+		FNodeHandle FGraphController::AddInputVertex(const FVertexName& InName, const FName InTypeName, const FText& InToolTip, const FMetasoundFrontendLiteral* InDefaultValue)
 		{
 			const FGuid VertexID = FGuid::NewGuid();
 
 			FMetasoundFrontendClassInput Description;
 
-			Description.Name = VertexID.ToString();
+			Description.Name = InName;
 			Description.TypeName = InTypeName;
 			Description.Metadata.Description = InToolTip;
 			Description.VertexID = VertexID;
 
-			return AddInputNode(Description, InDefaultValue, InDisplayName);
+			if (InDefaultValue)
+			{
+				Description.DefaultLiteral = *InDefaultValue;
+			}
+			else
+			{
+				Metasound::FLiteral Literal = IDataTypeRegistry::Get().CreateDefaultLiteral(InTypeName);
+				Description.DefaultLiteral.SetFromLiteral(Literal);
+			}
+
+			return AddInputVertex(Description);
 		}
 
-		FNodeHandle FGraphController::AddOutputNode(const FMetasoundFrontendClassOutput& InClassOutput, const FText* InDisplayName)
+		bool FGraphController::RemoveInputVertex(const FVertexName& InName)
 		{
-			FNodeHandle NewOutput = AddOutputVertex(InClassOutput);
-			if (ensure(NewOutput->IsValid()))
+			auto IsInputNodeWithSameName = [&](const FMetasoundFrontendClass& InClass, const FMetasoundFrontendNode& InNode)
+			{
+				return InClass.Metadata.GetType() == EMetasoundFrontendClassType::Input && InNode.Name == InName;
+			};
+			const TArray<FGraphController::FNodeAndClass> NodeAndClassPairs = GetNodesAndClassesByPredicate(IsInputNodeWithSameName);
+
+			for (const FNodeAndClass& NodeAndClass : NodeAndClassPairs)
 			{
-				if (InDisplayName)
+				if (const FMetasoundFrontendNode* Node = NodeAndClass.Node.Get())
 				{
-					SetOutputDisplayName(InClassOutput.Name, *InDisplayName);
+					return RemoveInput(*Node);
 				}
 			}
 
-			return NewOutput;
+			return false;
 		}
 
-		FNodeHandle FGraphController::AddOutputNode(const FName InTypeName, const FText& InToolTip, const FText* InDisplayName)
+		FNodeHandle FGraphController::AddOutputVertex(const FVertexName& InName, const FName InTypeName, const FText& InToolTip)
 		{
 			const FGuid VertexID = FGuid::NewGuid();
 
 			FMetasoundFrontendClassOutput Description;
 
-			Description.Name = VertexID.ToString();
+			Description.Name = InName;
 			Description.TypeName = InTypeName;
 			Description.Metadata.Description = InToolTip;
 			Description.VertexID = VertexID;
 
-			return AddOutputNode(Description, InDisplayName);
-		}
-
-
-		bool FGraphController::RemoveInputVertex(const FString& InName)
-		{
-			auto IsInputNodeWithSameName = [&](const FMetasoundFrontendClass& InClass, const FMetasoundFrontendNode& InNode)
-			{
-				return (InClass.Metadata.GetType() == EMetasoundFrontendClassType::Input) && (InNode.Name == InName);
-			};
-
-			for (const FNodeAndClass& NodeAndClass : GetNodesAndClassesByPredicate(IsInputNodeWithSameName))
-			{
-				if (const FMetasoundFrontendNode* Node = NodeAndClass.Node.Get())
-				{
-					return RemoveInput(*Node);
-				}
-			}
-
-			return false;
+			return AddOutputVertex(Description);
 		}
 
 		FNodeHandle FGraphController::AddOutputVertex(const FMetasoundFrontendClassOutput& InClassOutput)
@@ -2855,25 +2927,26 @@ namespace Metasound
 						FConstClassAccessPtr OutputClassPtr = OwningDocument->FindOrAddClass(Key);
 						if (const FMetasoundFrontendClass* OutputClass = OutputClassPtr.Get())
 						{
-							FMetasoundFrontendNode& Node = GraphClass->Graph.Nodes.Emplace_GetRef(*OutputClass);
-
+							FMetasoundFrontendNode& Node = GraphClass->Graph.Nodes.Add_GetRef(*OutputClass);
 							Node.Name = InClassOutput.Name;
-							Node.ID = FGuid::NewGuid();
+							Node.UpdateID();
+
+							const FVertexName NewName = InClassOutput.Name;
 
 							// TODO: have something that checks if input node has valid interface.
 							auto IsVertexWithTypeName = [&](FMetasoundFrontendVertex& Vertex) { return Vertex.TypeName == InClassOutput.TypeName; };
 							if (Node.Interface.Inputs.Num() == 1)
 							{
-								Node.Interface.Inputs[0].Name = InClassOutput.Name;
+								Node.Interface.Inputs[0].Name = NewName;
 							}
 							else if (FMetasoundFrontendVertex* InputVertex = Node.Interface.Inputs.FindByPredicate(IsVertexWithTypeName))
 							{ 
-								InputVertex->Name = InClassOutput.Name;
+								InputVertex->Name = NewName;
 							}
 
 							if (FMetasoundFrontendVertex* OutputVertex = Node.Interface.Outputs.FindByPredicate(IsVertexWithTypeName))
 							{ 
-								OutputVertex->Name = InClassOutput.Name;
+								OutputVertex->Name = NewName;
 							}
 							else
 							{
@@ -2881,11 +2954,11 @@ namespace Metasound
 							}
 
 							FMetasoundFrontendClassOutput& NewOutput = GraphClass->Interface.Outputs.Add_GetRef(InClassOutput);
-							NewOutput.NodeID = Node.ID;
+							NewOutput.NodeID = Node.GetID();
 
 							GraphClass->Interface.UpdateChangeID();
 
-							FNodeAccessPtr NodePtr = GraphClassPtr.GetNodeWithNodeID(Node.ID);
+							FNodeAccessPtr NodePtr = GraphClassPtr.GetNodeWithNodeID(Node.GetID());
 							return GetNodeHandle(FGraphController::FNodeAndClass { NodePtr, OutputClassPtr });
 						}
 					}
@@ -2896,21 +2969,23 @@ namespace Metasound
 				}
 				else
 				{
-					UE_LOG(LogMetaSound, Display, TEXT("Failed to add output. Output with same name \"%s\" exists in class [ClassID:%s]"), *InClassOutput.Name, *GraphClass->ID.ToString());
+					UE_LOG(LogMetaSound, Display, TEXT("Failed to add output. Output with same name \"%s\" exists in class [ClassID:%s]"), *InClassOutput.Name.ToString(), *GraphClass->ID.ToString());
 				}
 			}
 
 			return FInvalidNodeController::GetInvalid();
 		}
 
-		bool FGraphController::RemoveOutputVertex(const FString& InName)
+		bool FGraphController::RemoveOutputVertex(const FVertexName& InName)
 		{
 			auto IsOutputNodeWithSameName = [&](const FMetasoundFrontendClass& InClass, const FMetasoundFrontendNode& InNode)
 			{
-				return (InClass.Metadata.GetType() == EMetasoundFrontendClassType::Output) && (InNode.Name == InName);
+				return InClass.Metadata.GetType() == EMetasoundFrontendClassType::Output
+					&& InNode.Name == InName;
 			};
+			const TArray<FGraphController::FNodeAndClass> NodeAndClassPairs = GetNodesAndClassesByPredicate(IsOutputNodeWithSameName);
 
-			for (const FNodeAndClass& NodeAndClass : GetNodesAndClassesByPredicate(IsOutputNodeWithSameName))
+			for (const FNodeAndClass& NodeAndClass : NodeAndClassPairs)
 			{
 				if (const FMetasoundFrontendNode* Node = NodeAndClass.Node.Get())
 				{
@@ -2923,7 +2998,7 @@ namespace Metasound
 
 		// This can be used to determine what kind of property editor we should use for the data type of a given input.
 		// Will return Invalid if the input couldn't be found, or if the input doesn't support any kind of literals.
-		ELiteralType FGraphController::GetPreferredLiteralTypeForInputVertex(const FString& InInputName) const
+		ELiteralType FGraphController::GetPreferredLiteralTypeForInputVertex(const FVertexName& InInputName) const
 		{
 			if (const FMetasoundFrontendClassInput* Desc = FindInputDescriptionWithName(InInputName))
 			{
@@ -2933,7 +3008,7 @@ namespace Metasound
 		}
 
 		// For inputs whose preferred literal type is UObject or UObjectArray, this can be used to determine the UObject corresponding to that input's datatype.
-		UClass* FGraphController::GetSupportedClassForInputVertex(const FString& InInputName)
+		UClass* FGraphController::GetSupportedClassForInputVertex(const FVertexName& InInputName)
 		{
 			if (const FMetasoundFrontendClassInput* Desc = FindInputDescriptionWithName(InInputName))
 			{
@@ -2981,7 +3056,7 @@ namespace Metasound
 			return false;
 		}
 
-		const FText& FGraphController::GetInputDescription(const FString& InName) const
+		const FText& FGraphController::GetInputDescription(const FVertexName& InName) const
 		{
 			if (const FMetasoundFrontendClassInput* Desc = FindInputDescriptionWithName(InName))
 			{
@@ -2991,7 +3066,7 @@ namespace Metasound
 			return FText::GetEmpty();
 		}
 
-		const FText& FGraphController::GetOutputDescription(const FString& InName) const
+		const FText& FGraphController::GetOutputDescription(const FVertexName& InName) const
 		{
 			if (const FMetasoundFrontendClassOutput* Desc = FindOutputDescriptionWithName(InName))
 			{
@@ -3001,7 +3076,7 @@ namespace Metasound
 			return FText::GetEmpty();
 		}
 
-		void FGraphController::SetInputDisplayName(const FString& InName, const FText& InDisplayName)
+		void FGraphController::SetInputDisplayName(const FVertexName& InName, const FText& InDisplayName)
 		{
 			if (FMetasoundFrontendClassInput* Desc = FindInputDescriptionWithName(InName))
 			{
@@ -3009,7 +3084,7 @@ namespace Metasound
 			}
 		}
 
-		void FGraphController::SetOutputDisplayName(const FString& InName, const FText& InDisplayName)
+		void FGraphController::SetOutputDisplayName(const FVertexName& InName, const FText& InDisplayName)
 		{
 			if (FMetasoundFrontendClassOutput* Desc = FindOutputDescriptionWithName(InName))
 			{
@@ -3017,7 +3092,7 @@ namespace Metasound
 			}
 		}
 
-		void FGraphController::SetInputDescription(const FString& InName, const FText& InDescription)
+		void FGraphController::SetInputDescription(const FVertexName& InName, const FText& InDescription)
 		{
 			if (FMetasoundFrontendClassInput* Desc = FindInputDescriptionWithName(InName))
 			{
@@ -3025,7 +3100,7 @@ namespace Metasound
 			}
 		}
 
-		void FGraphController::SetOutputDescription(const FString& InName, const FText& InDescription)
+		void FGraphController::SetOutputDescription(const FVertexName& InName, const FText& InDescription)
 		{
 			if (FMetasoundFrontendClassOutput* Desc = FindOutputDescriptionWithName(InName))
 			{
@@ -3035,7 +3110,7 @@ namespace Metasound
 
 		// This can be used to clear the current literal for a given input.
 		// @returns false if the input name couldn't be found.
-		bool FGraphController::ClearLiteralForInput(const FString& InInputName, FGuid InVertexID)
+		bool FGraphController::ClearLiteralForInput(const FVertexName& InInputName, FGuid InVertexID)
 		{
 			if (FMetasoundFrontendClassInput* Desc = FindInputDescriptionWithName(InInputName))
 			{
@@ -3100,7 +3175,7 @@ namespace Metasound
 			if (FMetasoundFrontendGraphClass* GraphClass = GraphClassPtr.Get())
 			{
 				FGuid NodeID = InNode.GetID();
-				auto IsNodeWithSameID = [&](const FMetasoundFrontendNode& InDesc) { return InDesc.ID == NodeID; };
+				auto IsNodeWithSameID = [&](const FMetasoundFrontendNode& InDesc) { return InDesc.GetID() == NodeID; };
 				if (const FMetasoundFrontendNode* Desc = GraphClass->Graph.Nodes.FindByPredicate(IsNodeWithSameID))
 				{
 					switch(InNode.GetClassMetadata().GetType())
@@ -3220,9 +3295,9 @@ namespace Metasound
 				{
 					FMetasoundFrontendNode& Node = GraphClass->Graph.Nodes.Emplace_GetRef(*NodeClass);
 
-					Node.ID = FGuid::NewGuid();
+					Node.UpdateID();
 
-					FNodeAccessPtr NodePtr = GraphClassPtr.GetNodeWithNodeID(Node.ID);
+					FNodeAccessPtr NodePtr = GraphClassPtr.GetNodeWithNodeID(Node.GetID());
 					return GetNodeHandle(FGraphController::FNodeAndClass{NodePtr, InExistingDependency});
 				}
 			}
@@ -3234,12 +3309,12 @@ namespace Metasound
 		{
 			if (FMetasoundFrontendGraphClass* GraphClass = GraphClassPtr.Get())
 			{
-				auto IsEdgeForThisNode = [&](const FMetasoundFrontendEdge& ConDesc) { return (ConDesc.FromNodeID == InDesc.ID) || (ConDesc.ToNodeID == InDesc.ID); };
+				auto IsEdgeForThisNode = [&](const FMetasoundFrontendEdge& ConDesc) { return (ConDesc.FromNodeID == InDesc.GetID()) || (ConDesc.ToNodeID == InDesc.GetID()); };
 
 				// Remove any reference connections
 				int32 NumRemoved = GraphClass->Graph.Edges.RemoveAll(IsEdgeForThisNode);
 
-				auto IsNodeWithID = [&](const FMetasoundFrontendNode& Desc) { return InDesc.ID == Desc.ID; };
+				auto IsNodeWithID = [&](const FMetasoundFrontendNode& Desc) { return InDesc.GetID() == Desc.GetID(); };
 
 				NumRemoved += GraphClass->Graph.Nodes.RemoveAll(IsNodeWithID);
 
@@ -3254,7 +3329,7 @@ namespace Metasound
 		{
 			if (FMetasoundFrontendGraphClass* GraphClass = GraphClassPtr.Get())
 			{
-				auto IsInputWithSameNodeID = [&](const FMetasoundFrontendClassInput& ClassInput) { return ClassInput.NodeID == InNode.ID; };
+				auto IsInputWithSameNodeID = [&](const FMetasoundFrontendClassInput& ClassInput) { return ClassInput.NodeID == InNode.GetID(); };
 
 				int32 NumInputsRemoved = GraphClass->Interface.Inputs.RemoveAll(IsInputWithSameNodeID);
 
@@ -3275,7 +3350,7 @@ namespace Metasound
 		{
 			if (FMetasoundFrontendGraphClass* GraphClass = GraphClassPtr.Get())
 			{
-				auto IsOutputWithSameNodeID = [&](const FMetasoundFrontendClassOutput& ClassOutput) { return ClassOutput.NodeID == InNode.ID; };
+				auto IsOutputWithSameNodeID = [&](const FMetasoundFrontendClassOutput& ClassOutput) { return ClassOutput.NodeID == InNode.GetID(); };
 
 				int32 NumOutputsRemoved = GraphClass->Interface.Outputs.RemoveAll(IsOutputWithSameNodeID);
 
@@ -3316,7 +3391,7 @@ namespace Metasound
 					}
 					else
 					{
-						UE_LOG(LogMetaSound, Warning, TEXT("Failed to find class for node [NodeID:%s, ClassID:%s]"), *Node.ID.ToString(), *Node.ClassID.ToString());
+						UE_LOG(LogMetaSound, Warning, TEXT("Failed to find class for node [NodeID:%s, ClassID:%s]"), *Node.GetID().ToString(), *Node.ClassID.ToString());
 					}
 				}
 			}
@@ -3332,7 +3407,7 @@ namespace Metasound
 			{
 				for (FMetasoundFrontendNode& Node : GraphClass->Graph.Nodes)
 				{
-					FNodeAccessPtr NodePtr = GraphClassPtr.GetNodeWithNodeID(Node.ID);
+					FNodeAccessPtr NodePtr = GraphClassPtr.GetNodeWithNodeID(Node.GetID());
 					FConstClassAccessPtr NodeClassPtr = OwningDocument->FindClassWithID(Node.ClassID);
 
 					const bool bIsValidNodePtr = (nullptr != NodePtr.Get());
@@ -3344,7 +3419,7 @@ namespace Metasound
 					}
 					else
 					{
-						UE_LOG(LogMetaSound, Warning, TEXT("Failed to find class for node [NodeID:%s, ClassID:%s]"), *Node.ID.ToString(), *Node.ClassID.ToString());
+						UE_LOG(LogMetaSound, Warning, TEXT("Failed to find class for node [NodeID:%s, ClassID:%s]"), *Node.GetID().ToString(), *Node.ClassID.ToString());
 					}
 				}
 			}
@@ -3360,7 +3435,7 @@ namespace Metasound
 			{
 				for (const FMetasoundFrontendNode& Node : GraphClass->Graph.Nodes)
 				{
-					FConstNodeAccessPtr NodePtr = GraphClassPtr.GetNodeWithNodeID(Node.ID);
+					FConstNodeAccessPtr NodePtr = GraphClassPtr.GetNodeWithNodeID(Node.GetID());
 					FConstClassAccessPtr NodeClassPtr = OwningDocument->FindClassWithID(Node.ClassID);
 
 					const bool bIsValidNodePtr = (nullptr != NodePtr.Get());
@@ -3372,7 +3447,7 @@ namespace Metasound
 					}
 					else
 					{
-						UE_LOG(LogMetaSound, Warning, TEXT("Failed to find class for node [NodeID:%s, ClassID:%s]"), *Node.ID.ToString(), *Node.ClassID.ToString());
+						UE_LOG(LogMetaSound, Warning, TEXT("Failed to find class for node [NodeID:%s, ClassID:%s]"), *Node.GetID().ToString(), *Node.ClassID.ToString());
 					}
 				}
 			}
@@ -3393,13 +3468,13 @@ namespace Metasound
 					{
 						if (InPredicate(*NodeClass, Node))
 						{
-							FNodeAccessPtr NodePtr = GraphClassPtr.GetNodeWithNodeID(Node.ID);
+							FNodeAccessPtr NodePtr = GraphClassPtr.GetNodeWithNodeID(Node.GetID());
 							NodesAndClasses.Add(FGraphController::FNodeAndClass{ NodePtr, NodeClassPtr });
 						}
 					}
 					else
 					{
-						UE_LOG(LogMetaSound, Warning, TEXT("Failed to find class for node [NodeID:%s, ClassID:%s]"), *Node.ID.ToString(), *Node.ClassID.ToString());
+						UE_LOG(LogMetaSound, Warning, TEXT("Failed to find class for node [NodeID:%s, ClassID:%s]"), *Node.GetID().ToString(), *Node.ClassID.ToString());
 					}
 				}
 			}
@@ -3420,13 +3495,13 @@ namespace Metasound
 					{
 						if (InPredicate(*NodeClass, Node))
 						{
-							FConstNodeAccessPtr NodePtr = GraphClassPtr.GetNodeWithNodeID(Node.ID);
+							FConstNodeAccessPtr NodePtr = GraphClassPtr.GetNodeWithNodeID(Node.GetID());
 							NodesAndClasses.Add(FGraphController::FConstNodeAndClass{ NodePtr, NodeClassPtr });
 						}
 					}
 					else
 					{
-						UE_LOG(LogMetaSound, Warning, TEXT("Failed to find class for node [NodeID:%s, ClassID:%s]"), *Node.ID.ToString(), *Node.ClassID.ToString());
+						UE_LOG(LogMetaSound, Warning, TEXT("Failed to find class for node [NodeID:%s, ClassID:%s]"), *Node.GetID().ToString(), *Node.ClassID.ToString());
 					}
 				}
 			}
@@ -3512,7 +3587,7 @@ namespace Metasound
 				{
 					case EMetasoundFrontendClassType::Input:
 						{
-							FClassInputAccessPtr OwningGraphClassInputPtr = FindInputDescriptionWithNodeID(Node->ID);
+							FClassInputAccessPtr OwningGraphClassInputPtr = FindInputDescriptionWithNodeID(Node->GetID());
 							if (nullptr != OwningGraphClassInputPtr.Get())
 							{
 								FInputNodeController::FInitParams InitParams
@@ -3543,7 +3618,7 @@ namespace Metasound
 
 					case EMetasoundFrontendClassType::Output:
 						{
-							FClassOutputAccessPtr OwningGraphClassOutputPtr = FindOutputDescriptionWithNodeID(Node->ID);
+							FClassOutputAccessPtr OwningGraphClassOutputPtr = FindOutputDescriptionWithNodeID(Node->GetID());
 							if (nullptr != OwningGraphClassOutputPtr.Get())
 							{
 								FOutputNodeController::FInitParams InitParams
@@ -3620,7 +3695,7 @@ namespace Metasound
 				{
 					case EMetasoundFrontendClassType::Input:
 						{
-							FConstClassInputAccessPtr OwningGraphClassInputPtr = FindInputDescriptionWithNodeID(Node->ID);
+							FConstClassInputAccessPtr OwningGraphClassInputPtr = FindInputDescriptionWithNodeID(Node->GetID());
 							if (nullptr != OwningGraphClassInputPtr.Get())
 							{
 								FInputNodeController::FInitParams InitParams
@@ -3638,7 +3713,7 @@ namespace Metasound
 
 					case EMetasoundFrontendClassType::Output:
 						{
-							FConstClassOutputAccessPtr OwningGraphClassOutputPtr = FindOutputDescriptionWithNodeID(Node->ID);
+							FConstClassOutputAccessPtr OwningGraphClassOutputPtr = FindOutputDescriptionWithNodeID(Node->GetID());
 							if (nullptr != OwningGraphClassOutputPtr.Get())
 							{
 								FOutputNodeController::FInitParams InitParams
@@ -3676,7 +3751,7 @@ namespace Metasound
 			return FInvalidNodeController::GetInvalid();
 		}
 
-		FMetasoundFrontendClassInput* FGraphController::FindInputDescriptionWithName(const FString& InName)
+		FMetasoundFrontendClassInput* FGraphController::FindInputDescriptionWithName(const FVertexName& InName)
 		{
 			if (FMetasoundFrontendGraphClass* GraphClass = GraphClassPtr.Get())
 			{
@@ -3692,7 +3767,7 @@ namespace Metasound
 			return nullptr;
 		}
 
-		const FMetasoundFrontendClassInput* FGraphController::FindInputDescriptionWithName(const FString& InName) const
+		const FMetasoundFrontendClassInput* FGraphController::FindInputDescriptionWithName(const FVertexName& InName) const
 		{
 			if (const FMetasoundFrontendGraphClass* GraphClass = GraphClassPtr.Get())
 			{
@@ -3701,7 +3776,7 @@ namespace Metasound
 			return nullptr;
 		}
 
-		FMetasoundFrontendClassOutput* FGraphController::FindOutputDescriptionWithName(const FString& InName)
+		FMetasoundFrontendClassOutput* FGraphController::FindOutputDescriptionWithName(const FVertexName& InName)
 		{
 			if (FMetasoundFrontendGraphClass* GraphClass = GraphClassPtr.Get())
 			{
@@ -3717,7 +3792,7 @@ namespace Metasound
 			return nullptr;
 		}
 
-		const FMetasoundFrontendClassOutput* FGraphController::FindOutputDescriptionWithName(const FString& InName) const
+		const FMetasoundFrontendClassOutput* FGraphController::FindOutputDescriptionWithName(const FVertexName& InName) const
 		{
 			if (const FMetasoundFrontendGraphClass* GraphClass = GraphClassPtr.Get())
 			{
@@ -3903,7 +3978,7 @@ namespace Metasound
 					else
 					{
 						bSuccess = false;
-						UE_LOG(LogMetaSound, Error, TEXT("Failed to copy subgraph. Subgraph document is missing dependency info for node [Node:%s, NodeID:%s]"), *Node.Name, *Node.ID.ToString());
+						UE_LOG(LogMetaSound, Error, TEXT("Failed to copy subgraph. Subgraph document is missing dependency info for node [Node:%s, NodeID:%s]"), *Node.Name.ToString(), *Node.GetID().ToString());
 					}
 				}
 
diff --git a/Engine/Plugins/Runtime/Metasound/Source/MetasoundFrontend/Private/MetasoundFrontendStandardController.h b/Engine/Plugins/Runtime/Metasound/Source/MetasoundFrontend/Private/MetasoundFrontendStandardController.h
--- a/Engine/Plugins/Runtime/Metasound/Source/MetasoundFrontend/Private/MetasoundFrontendStandardController.h
+++ b/Engine/Plugins/Runtime/Metasound/Source/MetasoundFrontend/Private/MetasoundFrontendStandardController.h
@@ -8,6 +8,7 @@
 #include "MetasoundFrontendDocument.h"
 #include "MetasoundFrontendInvalidController.h"
 #include "MetasoundLog.h"
+#include "MetasoundVertex.h"
 #include "UObject/Object.h"
 
 
@@ -46,18 +47,20 @@ namespace Metasound
 
 			FGuid GetID() const override;
 			const FName& GetDataType() const override;
-			const FString& GetName() const override;
+			const FVertexName& GetName() const override;
 
 			// Output metadata
-			const FText& GetDisplayName() const override;
+			FText GetDisplayName() const override;
 			const FText& GetTooltip() const override;
 			const FMetasoundFrontendVertexMetadata& GetMetadata() const override;
 
-			// Return info on containing node. 
+			// Return info on containing node.
 			FGuid GetOwningNodeID() const override;
 			FNodeHandle GetOwningNode() override;
 			FConstNodeHandle GetOwningNode() const override;
 
+			void SetName(const FVertexName& InName) override { }
+
 			bool IsConnected() const override;
 			TArray<FInputHandle> GetConnectedInputs() override;
 			TArray<FConstInputHandle> GetConstConnectedInputs() const override;
@@ -117,10 +120,12 @@ namespace Metasound
 			bool IsValid() const override;
 
 			// Input metadata
-			const FText& GetDisplayName() const override;
+			FText GetDisplayName() const override;
 			const FText& GetTooltip() const override;
 			const FMetasoundFrontendVertexMetadata& GetMetadata() const override;
 
+			void SetName(const FVertexName& InName) override;
+
 		protected:
 			FDocumentAccess ShareAccess() override;
 			FConstDocumentAccess ShareAccess() const override;
@@ -162,10 +167,12 @@ namespace Metasound
 			bool IsValid() const override;
 
 			// Output metadata
-			const FText& GetDisplayName() const override;
+			FText GetDisplayName() const override;
 			const FText& GetTooltip() const override;
 			const FMetasoundFrontendVertexMetadata& GetMetadata() const override;
 
+			void SetName(const FVertexName& InName) override;
+
 			FConnectability CanConnectTo(const IInputController& InController) const override;
 			bool Connect(IInputController& InController) override;
 			bool ConnectWithConverterNode(IInputController& InController, const FConverterNodeInfo& InNodeClassName) override;
@@ -201,15 +208,19 @@ namespace Metasound
 
 			FGuid GetID() const override;
 			const FName& GetDataType() const override;
-			const FString& GetName() const override;
+			const FVertexName& GetName() const override;
 
 			const FMetasoundFrontendLiteral* GetLiteral() const override;
 			void SetLiteral(const FMetasoundFrontendLiteral& InLiteral) override;
 
 			const FMetasoundFrontendLiteral* GetClassDefaultLiteral() const override;
 
+			// This only exists to allow for transform fix-ups to easily cleanup input/output
+			// vertex names & should not be used for typical edit or runtime callsites.
+			void SetName(const FVertexName& InName) override { checkNoEntry(); }
+
 			// Input metadata
-			const FText& GetDisplayName() const override;
+			FText GetDisplayName() const override;
 			const FText& GetTooltip() const override;
 			const FMetasoundFrontendVertexMetadata& GetMetadata() const override;
 
@@ -245,8 +256,6 @@ namespace Metasound
 			FConstClassInputAccessPtr ClassInputPtr;
 			FGraphAccessPtr GraphPtr;
 			FNodeHandle OwningNode;
-
-
 		};
 
 		/** FOutputNodeInputController represents the input vertex of an output 
@@ -275,9 +284,10 @@ namespace Metasound
 			bool IsValid() const override;
 
 			// Input metadata
-			const FText& GetDisplayName() const override;
+			FText GetDisplayName() const override;
 			const FText& GetTooltip() const override;
 			const FMetasoundFrontendVertexMetadata& GetMetadata() const override;
+			void SetName(const FVertexName& InName) override;
 
 		protected:
 
@@ -317,10 +327,12 @@ namespace Metasound
 			bool IsValid() const override;
 
 			// Input metadata
-			const FText& GetDisplayName() const override;
+			FText GetDisplayName() const override;
 			const FText& GetTooltip() const override;
 			const FMetasoundFrontendVertexMetadata& GetMetadata() const override;
 
+			void SetName(const FVertexName& InName) override;
+
 			FConnectability CanConnectTo(const IOutputController& InController) const override;
 			bool Connect(IOutputController& InController) override;
 
@@ -384,10 +396,15 @@ namespace Metasound
 			FMetasoundFrontendVersionNumber FindHighestVersionInRegistry() const override;
 			FMetasoundFrontendVersionNumber FindHighestMinorVersionInRegistry() const override;
 
-			const FString& GetNodeName() const override;
+			const FVertexName& GetNodeName() const override;
+
+			// This only exists to allow for transform fix-ups to easily cleanup input/output
+			// node names & should not be used for standard edit or runtime callsites.
+			void SetNodeName(const FVertexName& InName) override { checkNoEntry(); }
 
-			/** Returns the human-readable display name of the given node. */
-			const FText& GetDisplayName() const override;
+			/** Returns the readable display name of the given node (Used only within MetaSound
+			  * Editor context, and not guaranteed to be a unique identifier). */
+			FText GetDisplayName() const override;
 
 			/** Sets the description of the node. */
 			void SetDescription(const FText& InDescription) override { }
@@ -421,8 +438,8 @@ namespace Metasound
 			int32 GetNumInputs() const override;
 			int32 GetNumOutputs() const override;
 
-			TArray<FInputHandle> GetInputsWithVertexName(const FString& InName) override;
-			TArray<FConstInputHandle> GetConstInputsWithVertexName(const FString& InName) const override;
+			TArray<FInputHandle> GetInputsWithVertexName(const FVertexName& InName) override;
+			TArray<FConstInputHandle> GetConstInputsWithVertexName(const FVertexName& InName) const override;
 
 			/** Returns all node outputs. */
 			TArray<FOutputHandle> GetOutputs() override;
@@ -430,8 +447,8 @@ namespace Metasound
 			/** Returns all node outputs. */
 			TArray<FConstOutputHandle> GetConstOutputs() const override;
 
-			TArray<FOutputHandle> GetOutputsWithVertexName(const FString& InName) override;
-			TArray<FConstOutputHandle> GetConstOutputsWithVertexName(const FString& InName) const override;
+			TArray<FOutputHandle> GetOutputsWithVertexName(const FVertexName& InName) override;
+			TArray<FConstOutputHandle> GetConstOutputsWithVertexName(const FVertexName& InName) const override;
 
 			bool IsRequired() const override;
 
@@ -488,8 +505,8 @@ namespace Metasound
 			virtual TArray<FInputControllerParams> GetInputControllerParams() const;
 			virtual TArray<FOutputControllerParams> GetOutputControllerParams() const;
 
-			virtual TArray<FInputControllerParams> GetInputControllerParamsWithVertexName(const FString& InName) const;
-			virtual TArray<FOutputControllerParams> GetOutputControllerParamsWithVertexName(const FString& InName) const;
+			virtual TArray<FInputControllerParams> GetInputControllerParamsWithVertexName(const FVertexName& InName) const;
+			virtual TArray<FOutputControllerParams> GetOutputControllerParamsWithVertexName(const FVertexName& InName) const;
 
 			virtual bool FindInputControllerParamsWithID(FGuid InVertexID, FInputControllerParams& OutParams) const;
 			virtual bool FindOutputControllerParamsWithID(FGuid InVertexID, FOutputControllerParams& OutParams) const;
@@ -594,10 +611,11 @@ namespace Metasound
 
 			bool IsValid() const override;
 			const FText& GetDescription() const override;
-			const FText& GetDisplayName() const override;
+			FText GetDisplayName() const override;
+			const FText& GetDisplayTitle() const override;
 			void SetDescription(const FText& InDescription) override;
 			void SetDisplayName(const FText& InText) override;
-			const FText& GetDisplayTitle() const override;
+			void SetNodeName(const FVertexName& InName) override;
 			bool IsRequired() const override;
 
 		protected:
@@ -655,12 +673,13 @@ namespace Metasound
 			virtual ~FInputNodeController() = default;
 
 			const FText& GetDescription() const override;
-			const FText& GetDisplayName() const override;
+			FText GetDisplayName() const override;
 			const FText& GetDisplayTitle() const override;
 			bool IsRequired() const override;
 			bool IsValid() const override;
 			void SetDescription(const FText& InDescription) override;
 			void SetDisplayName(const FText& InText) override;
+			void SetNodeName(const FVertexName& InName) override;
 
 			// No-ops as inputs do not handle literals the same way as other nodes
 			bool ClearInputLiteral(FGuid InVertexID) override { return false; }
@@ -722,14 +741,14 @@ namespace Metasound
 			bool IsValid() const override;
 
 			FGuid GetClassID() const override;
-			const FText& GetDisplayName() const override;
+			FText GetDisplayName() const override;
 
-			TArray<FString> GetInputVertexNames() const override;
-			TArray<FString> GetOutputVertexNames() const override;
-			FConstClassInputAccessPtr FindClassInputWithName(const FString& InName) const override;
-			FConstClassOutputAccessPtr FindClassOutputWithName(const FString& InName) const override;
-			FGuid GetVertexIDForInputVertex(const FString& InInputName) const override;
-			FGuid GetVertexIDForOutputVertex(const FString& InOutputName) const override;
+			TArray<FVertexName> GetInputVertexNames() const override;
+			TArray<FVertexName> GetOutputVertexNames() const override;
+			FConstClassInputAccessPtr FindClassInputWithName(const FVertexName& InName) const override;
+			FConstClassOutputAccessPtr FindClassOutputWithName(const FVertexName& InName) const override;
+			FGuid GetVertexIDForInputVertex(const FVertexName& InInputName) const override;
+			FGuid GetVertexIDForOutputVertex(const FVertexName& InOutputName) const override;
 
 			TArray<FNodeHandle> GetNodes() override;
 			TArray<FConstNodeHandle> GetConstNodes() const override;
@@ -742,8 +761,8 @@ namespace Metasound
 
 			void ClearGraph() override;
 
-			void IterateNodes(TUniqueFunction<void(FNodeHandle)> InFunction, EMetasoundFrontendClassType InClassType) override;
-			void IterateConstNodes(TUniqueFunction<void(FConstNodeHandle)> InFunction, EMetasoundFrontendClassType InClassType) const override;
+			void IterateNodes(TFunctionRef<void(FNodeHandle)> InFunction, EMetasoundFrontendClassType InClassType) override;
+			void IterateConstNodes(TFunctionRef<void(FConstNodeHandle)> InFunction, EMetasoundFrontendClassType InClassType) const override;
 
 			TArray<FNodeHandle> GetOutputNodes() override;
 			TArray<FConstNodeHandle> GetConstOutputNodes() const override;
@@ -751,29 +770,31 @@ namespace Metasound
 			const FMetasoundFrontendGraphStyle& GetGraphStyle() const override;
 			void SetGraphStyle(const FMetasoundFrontendGraphStyle& InStyle) override;
 
-			bool ContainsInputVertexWithName(const FString& InName) const override;
-			bool ContainsOutputVertexWithName(const FString& InName) const override;
+			bool ContainsInputVertexWithName(const FVertexName& InName) const override;
+			bool ContainsOutputVertexWithName(const FVertexName& InName) const override;
 
-			FConstNodeHandle GetInputNodeWithName(const FString& InName) const override;
-			FConstNodeHandle GetOutputNodeWithName(const FString& InName) const override;
+			FConstNodeHandle GetInputNodeWithName(const FVertexName& InName) const override;
+			FConstNodeHandle GetOutputNodeWithName(const FVertexName& InName) const override;
 
-			FNodeHandle GetInputNodeWithName(const FString& InName) override;
-			FNodeHandle GetOutputNodeWithName(const FString& InName) override;
+			FNodeHandle GetInputNodeWithName(const FVertexName& InName) override;
+			FNodeHandle GetOutputNodeWithName(const FVertexName& InName) override;
 
 			FNodeHandle AddInputVertex(const FMetasoundFrontendClassInput& InDescription) override;
-			bool RemoveInputVertex(const FString& InName) override;
+			FNodeHandle AddInputVertex(const FVertexName& InName, const FName InTypeName, const FText& InToolTip, const FMetasoundFrontendLiteral* InDefaultValue) override;
+			bool RemoveInputVertex(const FVertexName& InName) override;
 
 			FNodeHandle AddOutputVertex(const FMetasoundFrontendClassOutput& InDescription) override;
-			bool RemoveOutputVertex(const FString& InName) override;
+			FNodeHandle AddOutputVertex(const FVertexName& InName, const FName InTypeName, const FText& InToolTip) override;
+			bool RemoveOutputVertex(const FVertexName& InName) override;
 
 			void UpdateInterfaceChangeID() override;
 
 			// This can be used to determine what kind of property editor we should use for the data type of a given input.
 			// Will return Invalid if the input couldn't be found, or if the input doesn't support any kind of literals.
-			ELiteralType GetPreferredLiteralTypeForInputVertex(const FString& InInputName) const override;
+			ELiteralType GetPreferredLiteralTypeForInputVertex(const FVertexName& InInputName) const override;
 
 			// For inputs whose preferred literal type is UObject or UObjectArray, this can be used to determine the UObject corresponding to that input's datatype.
-			UClass* GetSupportedClassForInputVertex(const FString& InInputName) override;
+			UClass* GetSupportedClassForInputVertex(const FVertexName& InInputName) override;
 
 			FMetasoundFrontendLiteral GetDefaultInput(const FGuid& InVertexID) const override;
 
@@ -783,26 +804,26 @@ namespace Metasound
 			bool SetDefaultInputToDefaultLiteralOfType(const FGuid& InVertexID) override;
 
 			// Get the description for the input with the given name.
-			const FText& GetInputDescription(const FString& InName) const override;
+			const FText& GetInputDescription(const FVertexName& InName) const override;
 
 			// Get the description for the output with the given name.
-			const FText& GetOutputDescription(const FString& InName) const override;
+			const FText& GetOutputDescription(const FVertexName& InName) const override;
 
 			// Set the description for the input with the given name
-			void SetInputDescription(const FString& InName, const FText& InDescription) override;
+			void SetInputDescription(const FVertexName& InName, const FText& InDescription) override;
 
 			// Set the description for the input with the given name
-			void SetOutputDescription(const FString& InName, const FText& InDescription) override;
+			void SetOutputDescription(const FVertexName& InName, const FText& InDescription) override;
 
 			// Set the display name for the input with the given name
-			void SetInputDisplayName(const FString& InName, const FText& InDisplayName) override;
+			void SetInputDisplayName(const FVertexName& InName, const FText& InDisplayName) override;
 
 			// Set the display name for the output with the given name
-			void SetOutputDisplayName(const FString& InName, const FText& InDisplayName) override;
+			void SetOutputDisplayName(const FVertexName& InName, const FText& InDisplayName) override;
 
 			// This can be used to clear the current literal for a given input.
 			// @returns false if the input name couldn't be found.
-			bool ClearLiteralForInput(const FString& InInputName, FGuid InVertexID) override;
+			bool ClearLiteralForInput(const FVertexName& InInputName, FGuid InVertexID) override;
 
 			FNodeHandle AddNode(const FNodeRegistryKey& InNodeClass) override;
 			FNodeHandle AddNode(const FMetasoundFrontendClassMetadata& InClassMetadata) override;
@@ -824,11 +845,6 @@ namespace Metasound
 			FDocumentHandle GetOwningDocument() override;
 			FConstDocumentHandle GetOwningDocument() const override;
 
-			FNodeHandle AddInputNode(const FMetasoundFrontendClassInput& InClassInput, const FMetasoundFrontendLiteral* InDefaultValue, const FText* InDisplayName) override;
-			FNodeHandle AddInputNode(const FName InTypeName, const FText& InToolTip, const FMetasoundFrontendLiteral* InDefaultValue, const FText* InDisplayName) override;
-			FNodeHandle AddOutputNode(const FMetasoundFrontendClassOutput& InClassOutput, const FText* InDisplayName) override;
-			FNodeHandle AddOutputNode(const FName InTypeName, const FText& InToolTip, const FText* InDisplayName) override;
-
 		protected:
 
 			FDocumentAccess ShareAccess() override;
@@ -874,14 +890,14 @@ namespace Metasound
 			FNodeHandle GetNodeHandle(const FNodeAndClass& InNodeAndClass);
 			FConstNodeHandle GetNodeHandle(const FConstNodeAndClass& InNodeAndClass) const;
 
-			FMetasoundFrontendClassInput* FindInputDescriptionWithName(const FString& InName);
-			const FMetasoundFrontendClassInput* FindInputDescriptionWithName(const FString& InName) const;
+			FMetasoundFrontendClassInput* FindInputDescriptionWithName(const FVertexName& InName);
+			const FMetasoundFrontendClassInput* FindInputDescriptionWithName(const FVertexName& InName) const;
 
 			FMetasoundFrontendClassInput* FindInputDescriptionWithVertexID(const FGuid& InVertexID);
 			const FMetasoundFrontendClassInput* FindInputDescriptionWithVertexID(const FGuid& InVertexID) const;
 
-			FMetasoundFrontendClassOutput* FindOutputDescriptionWithName(const FString& InName);
-			const FMetasoundFrontendClassOutput* FindOutputDescriptionWithName(const FString& InName) const;
+			FMetasoundFrontendClassOutput* FindOutputDescriptionWithName(const FVertexName& InName);
+			const FMetasoundFrontendClassOutput* FindOutputDescriptionWithName(const FVertexName& InName) const;
 
 			FMetasoundFrontendClassOutput* FindOutputDescriptionWithVertexID(const FGuid& InVertexID);
 			const FMetasoundFrontendClassOutput* FindOutputDescriptionWithVertexID(const FGuid& InVertexID) const;
@@ -893,7 +909,7 @@ namespace Metasound
 			FConstClassOutputAccessPtr FindOutputDescriptionWithNodeID(FGuid InNodeID) const;
 
 			FGraphClassAccessPtr GraphClassPtr;
-			FDocumentHandle OwningDocument; 
+			FDocumentHandle OwningDocument;
 		};
 
 		/** FDocumentController represents an entire Metasound document. */
diff --git a/Engine/Plugins/Runtime/Metasound/Source/MetasoundFrontend/Private/MetasoundFrontendSubgraphNodeController.cpp b/Engine/Plugins/Runtime/Metasound/Source/MetasoundFrontend/Private/MetasoundFrontendSubgraphNodeController.cpp
--- a/Engine/Plugins/Runtime/Metasound/Source/MetasoundFrontend/Private/MetasoundFrontendSubgraphNodeController.cpp
+++ b/Engine/Plugins/Runtime/Metasound/Source/MetasoundFrontend/Private/MetasoundFrontendSubgraphNodeController.cpp
@@ -17,7 +17,7 @@ namespace Metasound
 				// Reverse iterate to allow removal of elements from SetA within loop.
 				for (int32 SetAIndex = (SetA.Num() - 1); SetAIndex >= 0; SetAIndex--) 
 				{
-					const FString& SetAName = SetA[SetAIndex]->Name;
+					const FVertexName& SetAName = SetA[SetAIndex]->Name;
 					int32 SetBIndex = SetB.IndexOfByPredicate([&](const SetBType* InVertex) { return InVertex->Name == SetAName; });
 					if (INDEX_NONE != SetBIndex)
 					{
@@ -84,7 +84,7 @@ namespace Metasound
 					}
 					else
 					{
-						UE_LOG(LogMetaSound, Warning, TEXT("Frontend Node [NodeID:%s, ClassID:%s] is not of expected class class [ClassID:%s]"), *Node->ID.ToString(), *Node->ClassID.ToString(), *Class->ID.ToString());
+						UE_LOG(LogMetaSound, Warning, TEXT("Frontend Node [NodeID:%s, ClassID:%s] is not of expected class class [ClassID:%s]"), *Node->GetID().ToString(), *Node->ClassID.ToString(), *Class->ID.ToString());
 					}
 				}
 			}
@@ -104,7 +104,7 @@ namespace Metasound
 					}
 					else
 					{
-						UE_LOG(LogMetaSound, Warning, TEXT("Frontend Node [NodeID:%s, ClassID:%s] is not of expected class class [ClassID:%s]"), *Node->ID.ToString(), *Node->ClassID.ToString(), *Class->ID.ToString());
+						UE_LOG(LogMetaSound, Warning, TEXT("Frontend Node [NodeID:%s, ClassID:%s] is not of expected class class [ClassID:%s]"), *Node->GetID().ToString(), *Node->ClassID.ToString(), *Class->ID.ToString());
 					}
 				}
 			}
@@ -152,7 +152,7 @@ namespace Metasound
 			return Super::GetOutputControllerParams();
 		}
 
-		TArray<FBaseNodeController::FInputControllerParams> FSubgraphNodeController::GetInputControllerParamsWithVertexName(const FString& InName) const
+		TArray<FBaseNodeController::FInputControllerParams> FSubgraphNodeController::GetInputControllerParamsWithVertexName(const FVertexName& InName) const
 		{
 			// TODO: Trigger ConformNodeInterfaceToClassInterface() on an
 			// event callback rather than on every call to this function.
@@ -161,7 +161,7 @@ namespace Metasound
 			return Super::GetInputControllerParamsWithVertexName(InName);
 		}
 
-		TArray<FBaseNodeController::FOutputControllerParams> FSubgraphNodeController::GetOutputControllerParamsWithVertexName(const FString& InName) const
+		TArray<FBaseNodeController::FOutputControllerParams> FSubgraphNodeController::GetOutputControllerParamsWithVertexName(const FVertexName& InName) const
 		{
 			TArray<FBaseNodeController::FOutputControllerParams> Outputs;
 
diff --git a/Engine/Plugins/Runtime/Metasound/Source/MetasoundFrontend/Private/MetasoundFrontendSubgraphNodeController.h b/Engine/Plugins/Runtime/Metasound/Source/MetasoundFrontend/Private/MetasoundFrontendSubgraphNodeController.h
--- a/Engine/Plugins/Runtime/Metasound/Source/MetasoundFrontend/Private/MetasoundFrontendSubgraphNodeController.h
+++ b/Engine/Plugins/Runtime/Metasound/Source/MetasoundFrontend/Private/MetasoundFrontendSubgraphNodeController.h
@@ -61,8 +61,8 @@ namespace Metasound
 			TArray<FBaseNodeController::FInputControllerParams> GetInputControllerParams() const override;
 			TArray<FBaseNodeController::FOutputControllerParams> GetOutputControllerParams() const override;
 
-			TArray<FBaseNodeController::FInputControllerParams> GetInputControllerParamsWithVertexName(const FString& InName) const override;
-			TArray<FBaseNodeController::FOutputControllerParams> GetOutputControllerParamsWithVertexName(const FString& InName) const override;
+			TArray<FBaseNodeController::FInputControllerParams> GetInputControllerParamsWithVertexName(const FVertexName& InName) const override;
+			TArray<FBaseNodeController::FOutputControllerParams> GetOutputControllerParamsWithVertexName(const FVertexName& InName) const override;
 
 			bool FindInputControllerParamsWithID(FGuid InVertexID, FInputControllerParams& OutParams) const override;
 			bool FindOutputControllerParamsWithID(FGuid InVertexID, FOutputControllerParams& OutParams) const override;
diff --git a/Engine/Plugins/Runtime/Metasound/Source/MetasoundFrontend/Private/MetasoundFrontendTransform.cpp b/Engine/Plugins/Runtime/Metasound/Source/MetasoundFrontend/Private/MetasoundFrontendTransform.cpp
--- a/Engine/Plugins/Runtime/Metasound/Source/MetasoundFrontend/Private/MetasoundFrontendTransform.cpp
+++ b/Engine/Plugins/Runtime/Metasound/Source/MetasoundFrontend/Private/MetasoundFrontendTransform.cpp
@@ -192,7 +192,7 @@ namespace Metasound
 
 				FNodeHandle NewInputNode = InGraph->AddInputVertex(NewVertex);
 				// Copy prior node locations
-				if (Locations.Num() > 0)
+				if (!Locations.IsEmpty())
 				{
 					// TODO: copy entire style.
 					FMetasoundFrontendNodeStyle Style = NewInputNode->GetNodeStyle();
@@ -475,22 +475,22 @@ namespace Metasound
 			const FMetasoundFrontendVersion& RefArchetypeVersion = ReferencedDocument->GetArchetypeVersion();
 			FMatchRootGraphToArchetype(RefArchetypeVersion).Transform(InDocument);
 
-			using FDefaultKey = TPair<FString, FName>;
+			using FDefaultKey = TPair<FVertexName, FName>;
 			struct FDefaultInputValue
 			{
 				FGuid VertexID;
-				FString VertexName;
+				FVertexName VertexName;
 				FMetasoundFrontendLiteral Literal;
 			};
 			struct FDefaultOutputValue
 			{
 				FGuid VertexID;
-				FString VertexName;
+				FVertexName VertexName;
 			};
 			TMap<FDefaultKey, FDefaultInputValue> InitInputDefaults;
 			TMap<FDefaultKey, FDefaultOutputValue> InitOutputDefaults;
 
-			// 1. Prep by caching existing output guids, input literal values, clearing graph & ensuring ref archetype is supported
+			// 1. Prep by caching existing output guids, input literal values, & clearing graph
 			RootGraphHandle->IterateConstNodes([&](FConstNodeHandle InputNode)
 			{
 				FGuid InputID = InputNode->GetID();
@@ -500,8 +500,9 @@ namespace Metasound
 					{
 						const FGuid InputVertexID = RootGraphHandle->GetVertexIDForInputVertex(InputNode->GetNodeName());
 						FMetasoundFrontendLiteral Literal = RootGraphHandle->GetDefaultInput(InputVertexID);
-						FDefaultKey Key(InputNode->GetDisplayName().ToString(), Input->GetDataType());
-						FDefaultInputValue Value = { MoveTemp(InputID), InputNode->GetNodeName(), MoveTemp(Literal), };
+
+						FDefaultKey Key(InputNode->GetNodeName(), Input->GetDataType());
+						FDefaultInputValue Value = { MoveTemp(InputID), InputNode->GetNodeName(), MoveTemp(Literal) };
 						InitInputDefaults.Emplace(MoveTemp(Key), MoveTemp(Value));
 					});
 				}
@@ -514,7 +515,7 @@ namespace Metasound
 				{
 					OutputNode->IterateConstOutputs([&](FConstOutputHandle Output)
 					{
-						FDefaultKey Key(OutputNode->GetDisplayName().ToString(), Output->GetDataType());
+						FDefaultKey Key(OutputNode->GetNodeName(), Output->GetDataType());
 						FDefaultOutputValue Value = { MoveTemp(OutputID), OutputNode->GetNodeName() };
 						InitOutputDefaults.Emplace(MoveTemp(Key), MoveTemp(Value));
 					});
@@ -533,13 +534,12 @@ namespace Metasound
 			const FVector2D ReferenceNodeLocation = InputNodeLocation + FVector2D(DisplayStyle::NodeLayout::DefaultOffsetX.X, ReferenceNodeYOffset);
 			FVector2D OutputNodeLocation = InputNodeLocation + (2 * DisplayStyle::NodeLayout::DefaultOffsetX);
 
-			TMap<FString, FString> PresetInputToReferenceInputMap;
-			TMap<FString, FString> PresetOutputToReferenceOutputMap;
+			TMap<FVertexName, FVertexName> PresetInputToReferenceInputMap;
+			TMap<FVertexName, FVertexName> PresetOutputToReferenceOutputMap;
 
 			// 2. Generate preset's inputs & outputs
 			for (const FConstNodeHandle& RefGraphInputNode : ReferencedInputs)
 			{
-				const FText& DisplayName = RefGraphInputNode->GetDisplayName();
 				const FText& Description = RefGraphInputNode->GetDescription();
 
 				// Should only ever be one
@@ -547,9 +547,11 @@ namespace Metasound
 				RefGraphInputNode->IterateConstInputs([&](FConstInputHandle RefGraphInputHandle)
 				{
 					using namespace Metasound::Frontend;
+
 					FNodeHandle NodeHandle = INodeController::GetInvalidHandle();
 					const FName DataTypeName = RefGraphInputHandle->GetDataType();
-					const FDefaultKey Key(DisplayName.ToString(), DataTypeName);
+
+					const FDefaultKey Key(RefGraphInputNode->GetNodeName(), DataTypeName);
 					if (const FDefaultInputValue* InputValue = InitInputDefaults.Find(Key))
 					{
 						FMetasoundFrontendClassInput ClassInput;
@@ -557,14 +559,15 @@ namespace Metasound
 						ClassInput.TypeName = DataTypeName;
 						ClassInput.Metadata.Description = Description;
 						ClassInput.VertexID = InputValue->VertexID;
+						ClassInput.DefaultLiteral = InputValue->Literal;
 
-						NodeHandle = RootGraphHandle->AddInputNode(ClassInput, &InputValue->Literal, &DisplayName);
+						NodeHandle = RootGraphHandle->AddInputVertex(ClassInput);
 					}
 					else
 					{
 						FGuid InputID = ReferencedGraphHandle->GetVertexIDForInputVertex(RefGraphInputHandle->GetName());
 						const FMetasoundFrontendLiteral DefaultLiteral = ReferencedGraphHandle->GetDefaultInput(InputID);
-						NodeHandle = RootGraphHandle->AddInputNode(DataTypeName, Description, &DefaultLiteral, &DisplayName);
+						NodeHandle = RootGraphHandle->AddInputVertex(RefGraphInputNode->GetNodeName(), DataTypeName, Description, &DefaultLiteral);
 					}
 
 					if (ensure(NodeHandle->IsValid()))
@@ -581,7 +584,6 @@ namespace Metasound
 
 			for (const FConstNodeHandle& RefGraphOutputNode : ReferencedOutputs)
 			{
-				const FText& DisplayName = RefGraphOutputNode->GetDisplayName();
 				const FText& Description = RefGraphOutputNode->GetDescription();
 
 				// Should only ever be one
@@ -592,7 +594,9 @@ namespace Metasound
 
 					FNodeHandle NodeHandle = INodeController::GetInvalidHandle();
 					const FName DataTypeName = RefGraphOutputHandle->GetDataType();
-					const FDefaultKey Key(DisplayName.ToString(), DataTypeName);
+					
+					const FVertexName OutputName = RefGraphOutputHandle->GetName();
+					const FDefaultKey Key(OutputName, DataTypeName);
 					if (FDefaultOutputValue* OutputValue = InitOutputDefaults.Find(Key))
 					{
 						FMetasoundFrontendClassOutput ClassOutput;
@@ -601,11 +605,11 @@ namespace Metasound
 						ClassOutput.Metadata.Description = Description;
 						ClassOutput.VertexID = OutputValue->VertexID;
 
-						NodeHandle = RootGraphHandle->AddOutputNode(ClassOutput, &DisplayName);
+						NodeHandle = RootGraphHandle->AddOutputVertex(ClassOutput);
 					}
 					else
 					{
-						NodeHandle = RootGraphHandle->AddOutputNode(DataTypeName, Description, &DisplayName);
+						NodeHandle = RootGraphHandle->AddOutputVertex(OutputName, DataTypeName, Description);
 					}
 
 					if (ensure(NodeHandle->IsValid()))
@@ -635,14 +639,13 @@ namespace Metasound
 			// 4. Connect preset's respective inputs & outputs to generated referencing node.
 			RootGraphHandle->IterateNodes([&](FNodeHandle InputNode)
 			{
-				const FString& Name = InputNode->GetNodeName();
+				const FVertexName& Name = InputNode->GetNodeName();
 
 				// Should only ever be one
 				ensure(InputNode->GetNumOutputs() == 1);
 				InputNode->IterateOutputs([&](FOutputHandle OutputHandle)
 				{
-					const FString& InputName = PresetInputToReferenceInputMap.FindChecked(Name);
-					const TArray<FInputHandle> ReferenceInputs = ReferencedNodeHandle->GetInputsWithVertexName(InputName);
+					const TArray<FInputHandle> ReferenceInputs = ReferencedNodeHandle->GetInputs();
 					if (ensure(ReferenceInputs.Num() == 1))
 					{
 						ensure(ReferenceInputs[0]->Connect(*OutputHandle));
@@ -653,13 +656,13 @@ namespace Metasound
 
 			RootGraphHandle->IterateNodes([&](FNodeHandle OutputNode)
 			{
-				const FString& Name = OutputNode->GetNodeName();
+				const FVertexName& Name = OutputNode->GetNodeName();
 
 				// Should only ever be one
 				ensure(OutputNode->GetNumInputs() == 1);
 				OutputNode->IterateInputs([&](FInputHandle InputHandle)
 				{
-					const FString& OutputName = PresetOutputToReferenceOutputMap.FindChecked(Name);
+					const FVertexName& OutputName = PresetOutputToReferenceOutputMap.FindChecked(Name);
 					const TArray<FOutputHandle> ReferenceInputs = ReferencedNodeHandle->GetOutputsWithVertexName(OutputName);
 					if (ensure(ReferenceInputs.Num() == 1))
 					{
@@ -927,6 +930,48 @@ namespace Metasound
 			}
 		};
 
+		/** Versions document from 1.6 to 1.7. */
+		class FVersionDocument_1_7 : public FVersionDocumentTransform
+		{
+		public:
+			FVersionDocument_1_7() = default;
+
+			FMetasoundFrontendVersionNumber GetTargetVersion() const override
+			{
+				return { 1, 7 };
+			}
+
+			void TransformInternal(FDocumentHandle InDocument) const override
+			{
+				auto RenameTransform = [](FNodeHandle NodeHandle)
+				{
+					// Required nodes are all (at the point of this transform) providing
+					// unique names and customized display names (ex. 'Audio' for both mono &
+					// L/R output, On Play, & 'On Finished'), so do not replace them by nulling
+					// out the guid as a name and using the converted FName of the FText DisplayName.
+					if (!NodeHandle->IsRequired())
+					{
+						const FName NewNodeName = *NodeHandle->GetDisplayName().ToString();
+						NodeHandle->IterateInputs([&](FInputHandle InputHandle)
+						{
+							InputHandle->SetName(NewNodeName);
+						});
+
+						NodeHandle->IterateOutputs([&](FOutputHandle OutputHandle)
+						{
+							OutputHandle->SetName(NewNodeName);
+						});
+
+						NodeHandle->SetDisplayName(FText());
+						NodeHandle->SetNodeName(NewNodeName);
+					}
+				};
+
+				InDocument->GetRootGraph()->IterateNodes(RenameTransform, EMetasoundFrontendClassType::Input);
+				InDocument->GetRootGraph()->IterateNodes(RenameTransform, EMetasoundFrontendClassType::Output);
+			}
+		};
+
 		FVersionDocument::FVersionDocument(FName InName, const FString& InPath)
 			: Name(InName)
 			, Path(InPath)
@@ -951,6 +996,7 @@ namespace Metasound
 			bWasUpdated |= FVersionDocument_1_4().Transform(InDocument);
 			bWasUpdated |= FVersionDocument_1_5(Name).Transform(InDocument);
 			bWasUpdated |= FVersionDocument_1_6().Transform(InDocument);
+			bWasUpdated |= FVersionDocument_1_7().Transform(InDocument);
 
 			if (bWasUpdated)
 			{
diff --git a/Engine/Plugins/Runtime/Metasound/Source/MetasoundFrontend/Private/MetasoundParameterTransmitter.cpp b/Engine/Plugins/Runtime/Metasound/Source/MetasoundFrontend/Private/MetasoundParameterTransmitter.cpp
--- a/Engine/Plugins/Runtime/Metasound/Source/MetasoundFrontend/Private/MetasoundParameterTransmitter.cpp
+++ b/Engine/Plugins/Runtime/Metasound/Source/MetasoundFrontend/Private/MetasoundParameterTransmitter.cpp
@@ -7,14 +7,15 @@
 
 namespace Metasound
 {
-	FString FMetaSoundParameterTransmitter::GetInstanceIDEnvironmentVariableName()
+	const FVertexName& FMetaSoundParameterTransmitter::GetInstanceIDEnvironmentVariableName()
 	{
-		return TEXT("TransmitterInstanceID");
+		static const FVertexName VarName = "TransmitterInstanceID";
+		return VarName;
 	}
 
-	FSendAddress FMetaSoundParameterTransmitter::CreateSendAddressFromEnvironment(const FMetasoundEnvironment& InEnvironment, const FVertexKey& InVertexKey, const FName& InTypeName)
+	FSendAddress FMetaSoundParameterTransmitter::CreateSendAddressFromEnvironment(const FMetasoundEnvironment& InEnvironment, const FVertexName& InVertexName, const FName& InTypeName)
 	{
-		const FString IDVarName = GetInstanceIDEnvironmentVariableName();
+		const FVertexName IDVarName = GetInstanceIDEnvironmentVariableName();
 		uint64 InstanceID = -1;
 
 		if (ensure(InEnvironment.Contains<uint64>(IDVarName)))
@@ -22,15 +23,15 @@ namespace Metasound
 			InstanceID = InEnvironment.GetValue<uint64>(IDVarName);
 		}
 
-		return CreateSendAddressFromInstanceID(InstanceID, InVertexKey, InTypeName);
+		return CreateSendAddressFromInstanceID(InstanceID, InVertexName, InTypeName);
 	}
 
-	FSendAddress FMetaSoundParameterTransmitter::CreateSendAddressFromInstanceID(uint64 InInstanceID, const FVertexKey& InVertexKey, const FName& InTypeName)
+	FSendAddress FMetaSoundParameterTransmitter::CreateSendAddressFromInstanceID(uint64 InInstanceID, const FVertexName& InVertexName, const FName& InTypeName)
 	{
 		FSendAddress Address;
 
 		Address.Subsystem = GetSubsystemNameForSendScope(ETransmissionScope::Global);
-		Address.ChannelName = FName(FString::Format(TEXT("{0}:{1}:{2}"), { InInstanceID, InVertexKey, InTypeName.ToString() }));
+		Address.ChannelName = *FString::Format(TEXT("{0}:{1}:{2}"), { InInstanceID, InVertexName.ToString(), InTypeName.ToString() });
 
 		return Address;
 	}
diff --git a/Engine/Plugins/Runtime/Metasound/Source/MetasoundFrontend/Private/MetasoundReceiveNode.cpp b/Engine/Plugins/Runtime/Metasound/Source/MetasoundFrontend/Private/MetasoundReceiveNode.cpp
--- a/Engine/Plugins/Runtime/Metasound/Source/MetasoundFrontend/Private/MetasoundReceiveNode.cpp
+++ b/Engine/Plugins/Runtime/Metasound/Source/MetasoundFrontend/Private/MetasoundReceiveNode.cpp
@@ -6,27 +6,27 @@ namespace Metasound
 {
 	namespace ReceiveNodeInfo
 	{
-		const FString& GetAddressInputName()
+		const FVertexName& GetAddressInputName()
 		{
-			static const FString InputName(TEXT("Address"));
+			static const FVertexName InputName(TEXT("Address"));
 			return InputName;
 		}
 
-		const FString& GetDefaultDataInputName()
+		const FVertexName& GetDefaultDataInputName()
 		{
-			static const FString DefaultDataName(TEXT("Default"));
+			static const FVertexName DefaultDataName(TEXT("Default"));
 			return DefaultDataName;
 		}
 
-		const FString& GetOutputName()
+		const FVertexName& GetOutputName()
 		{
-			static const FString OutputName(TEXT("Out"));
+			static const FVertexName OutputName(TEXT("Out"));
 			return OutputName;
 		}
 
 		FNodeClassName GetClassNameForDataType(const FName& InDataTypeName)
 		{
-			return FNodeClassName{ "Receive", InDataTypeName, FName() };
+			return FNodeClassName { "Receive", InDataTypeName, FName() };
 		}
 
 		int32 GetCurrentMajorVersion() { return 1; }
diff --git a/Engine/Plugins/Runtime/Metasound/Source/MetasoundFrontend/Public/MetasoundArrayNodes.h b/Engine/Plugins/Runtime/Metasound/Source/MetasoundFrontend/Public/MetasoundArrayNodes.h
--- a/Engine/Plugins/Runtime/Metasound/Source/MetasoundFrontend/Public/MetasoundArrayNodes.h
+++ b/Engine/Plugins/Runtime/Metasound/Source/MetasoundFrontend/Public/MetasoundArrayNodes.h
@@ -1,5 +1,4 @@
 // Copyright Epic Games, Inc. All Rights Reserved.
-
 #pragma once
 
 #include "CoreMinimal.h"
@@ -13,11 +12,13 @@
 #include "MetasoundOperatorInterface.h"
 #include "MetasoundPrimitives.h"
 #include "MetasoundTrigger.h"
+#include "MetasoundVertex.h"
 
 #include <type_traits>
 
 #define LOCTEXT_NAMESPACE "MetasoundFrontend"
 
+
 namespace Metasound
 {
 	namespace MetasoundArrayNodesPrivate
@@ -43,19 +44,19 @@ namespace Metasound
 	namespace ArrayNodeVertexNames
 	{
 		/* Input Vertex Names */
-		METASOUNDFRONTEND_API const FString& GetInputArrayName();
-		METASOUNDFRONTEND_API const FString& GetInputLeftArrayName();
-		METASOUNDFRONTEND_API const FString& GetInputRightArrayName();
-		METASOUNDFRONTEND_API const FString& GetInputTriggerName();
-		METASOUNDFRONTEND_API const FString& GetInputStartIndexName();
-		METASOUNDFRONTEND_API const FString& GetInputEndIndexName();
-		METASOUNDFRONTEND_API const FString& GetInputIndexName();
-		METASOUNDFRONTEND_API const FString& GetInputValueName();
+		METASOUNDFRONTEND_API const FVertexName& GetInputArrayName();
+		METASOUNDFRONTEND_API const FVertexName& GetInputLeftArrayName();
+		METASOUNDFRONTEND_API const FVertexName& GetInputRightArrayName();
+		METASOUNDFRONTEND_API const FVertexName& GetInputTriggerName();
+		METASOUNDFRONTEND_API const FVertexName& GetInputStartIndexName();
+		METASOUNDFRONTEND_API const FVertexName& GetInputEndIndexName();
+		METASOUNDFRONTEND_API const FVertexName& GetInputIndexName();
+		METASOUNDFRONTEND_API const FVertexName& GetInputValueName();
 
 		/* Output Vertex Names */
-		METASOUNDFRONTEND_API const FString& GetOutputNumName();
-		METASOUNDFRONTEND_API const FString& GetOutputValueName();
-		METASOUNDFRONTEND_API const FString& GetOutputArrayName();
+		METASOUNDFRONTEND_API const FVertexName& GetOutputNumName();
+		METASOUNDFRONTEND_API const FVertexName& GetOutputValueName();
+		METASOUNDFRONTEND_API const FVertexName& GetOutputArrayName();
 	};
 
 	/** TArrayNumOperator gets the number of elements in an Array. The operator
@@ -714,6 +715,5 @@ namespace Metasound
 
 		virtual ~TArraySubsetNode() = default;
 	};
-}
-
+} // namespace Metasound
 #undef LOCTEXT_NAMESPACE
diff --git a/Engine/Plugins/Runtime/Metasound/Source/MetasoundFrontend/Public/MetasoundArrayShuffleNode.h b/Engine/Plugins/Runtime/Metasound/Source/MetasoundFrontend/Public/MetasoundArrayShuffleNode.h
--- a/Engine/Plugins/Runtime/Metasound/Source/MetasoundFrontend/Public/MetasoundArrayShuffleNode.h
+++ b/Engine/Plugins/Runtime/Metasound/Source/MetasoundFrontend/Public/MetasoundArrayShuffleNode.h
@@ -4,6 +4,7 @@
 
 #include "CoreMinimal.h"
 #include "Internationalization/Text.h"
+#include "MetasoundArrayNodes.h"
 #include "MetasoundBuilderInterface.h"
 #include "MetasoundDataFactory.h"
 #include "MetasoundExecutableOperator.h"
@@ -13,7 +14,7 @@
 #include "MetasoundOperatorInterface.h"
 #include "MetasoundPrimitives.h"
 #include "MetasoundTrigger.h"
-#include "MetasoundArrayNodes.h"
+#include "MetasoundVertex.h"
 #include "Misc/ScopeLock.h"
 
 #include <type_traits>
@@ -27,18 +28,18 @@ namespace Metasound
 	namespace ArrayNodeShuffleVertexNames
 	{
 		/** Input Vertex Names */
-		METASOUNDFRONTEND_API const FString& GetInputTriggerNextName();
-		METASOUNDFRONTEND_API const FString& GetInputTriggerShuffleName();
-		METASOUNDFRONTEND_API const FString& GetInputTriggerResetName();
-		METASOUNDFRONTEND_API const FString& GetInputShuffleArrayName();
-		METASOUNDFRONTEND_API const FString& GetInputSeedName();
-		METASOUNDFRONTEND_API const FString& GetInputAutoShuffleName();
-		METASOUNDFRONTEND_API const FString& GetInputEnableSharedStateName();
-
-		METASOUNDFRONTEND_API const FString& GetOutputTriggerOnNextName();
-		METASOUNDFRONTEND_API const FString& GetOutputTriggerOnShuffleName();
-		METASOUNDFRONTEND_API const FString& GetOutputTriggerOnResetName();
-		METASOUNDFRONTEND_API const FString& GetOutputValueName();
+		METASOUNDFRONTEND_API const FVertexName& GetInputTriggerNextName();
+		METASOUNDFRONTEND_API const FVertexName& GetInputTriggerShuffleName();
+		METASOUNDFRONTEND_API const FVertexName& GetInputTriggerResetName();
+		METASOUNDFRONTEND_API const FVertexName& GetInputShuffleArrayName();
+		METASOUNDFRONTEND_API const FVertexName& GetInputSeedName();
+		METASOUNDFRONTEND_API const FVertexName& GetInputAutoShuffleName();
+		METASOUNDFRONTEND_API const FVertexName& GetInputEnableSharedStateName();
+
+		METASOUNDFRONTEND_API const FVertexName& GetOutputTriggerOnNextName();
+		METASOUNDFRONTEND_API const FVertexName& GetOutputTriggerOnShuffleName();
+		METASOUNDFRONTEND_API const FVertexName& GetOutputTriggerOnResetName();
+		METASOUNDFRONTEND_API const FVertexName& GetOutputValueName();
 	}
 
 	class METASOUNDFRONTEND_API FArrayIndexShuffler
diff --git a/Engine/Plugins/Runtime/Metasound/Source/MetasoundFrontend/Public/MetasoundAssetBase.h b/Engine/Plugins/Runtime/Metasound/Source/MetasoundFrontend/Public/MetasoundAssetBase.h
--- a/Engine/Plugins/Runtime/Metasound/Source/MetasoundFrontend/Public/MetasoundAssetBase.h
+++ b/Engine/Plugins/Runtime/Metasound/Source/MetasoundFrontend/Public/MetasoundAssetBase.h
@@ -11,6 +11,7 @@
 #include "MetasoundGraph.h"
 #include "MetasoundLog.h"
 #include "MetasoundParameterTransmitter.h"
+#include "MetasoundVertex.h"
 #include "UObject/SoftObjectPath.h"
 #include "UObject/WeakObjectPtrTemplates.h"
 
@@ -197,7 +198,7 @@ public:
 	struct FSendInfoAndVertexName
 	{
 		Metasound::FMetaSoundParameterTransmitter::FSendInfo SendInfo;
-		FString VertexName;
+		Metasound::FVertexName VertexName;
 	};
 
 	// Builds the Metasound Document returned by `GetDocument() const`.
@@ -230,8 +231,8 @@ protected:
 private:
 	Metasound::Frontend::FNodeRegistryKey RegistryKey;
 
-	TSet<Metasound::FVertexKey> GetNonTransmittableInputVertices(const FMetasoundFrontendDocument& InDoc) const;
-	TArray<FString> GetTransmittableInputVertexNames() const;
-	Metasound::FSendAddress CreateSendAddress(uint64 InInstanceID, const FString& InVertexName, const FName& InDataTypeName) const;
+	TSet<Metasound::FVertexName> GetNonTransmittableInputVertices(const FMetasoundFrontendDocument& InDoc) const;
+	TArray<Metasound::FVertexName> GetTransmittableInputVertexNames() const;
+	Metasound::FSendAddress CreateSendAddress(uint64 InInstanceID, const Metasound::FVertexName& InVertexName, const FName& InDataTypeName) const;
 	Metasound::Frontend::FNodeHandle AddInputPinForSendAddress(const Metasound::FMetaSoundParameterTransmitter::FSendInfo& InSendInfo, Metasound::Frontend::FGraphHandle InGraph) const;
 };
diff --git a/Engine/Plugins/Runtime/Metasound/Source/MetasoundFrontend/Public/MetasoundAutoConverterNode.h b/Engine/Plugins/Runtime/Metasound/Source/MetasoundFrontend/Public/MetasoundAutoConverterNode.h
--- a/Engine/Plugins/Runtime/Metasound/Source/MetasoundFrontend/Public/MetasoundAutoConverterNode.h
+++ b/Engine/Plugins/Runtime/Metasound/Source/MetasoundFrontend/Public/MetasoundAutoConverterNode.h
@@ -12,11 +12,13 @@
 #include "MetasoundExecutableOperator.h"
 #include "MetasoundFrontend.h"
 #include "MetasoundFrontendNodesCategories.h"
+#include "MetasoundVertex.h"
 
 #include <type_traits>
 
 #define LOCTEXT_NAMESPACE "MetasoundFrontend"
 
+
 namespace Metasound
 {
 	// This convenience node can be registered and will invoke static_cast<ToDataType>(FromDataType) every time it is executed.
@@ -26,15 +28,15 @@ namespace Metasound
 		static_assert(std::is_convertible<FromDataType, ToDataType>::value, "Tried to create an auto converter node between two types we can't static_cast between.");
 
 	public:
-		static const FString& GetInputName()
+		static const FVertexName& GetInputName()
 		{
-			static const FString InputName = GetMetasoundDataTypeString<FromDataType>();
+			static const FVertexName InputName = GetMetasoundDataTypeName<FromDataType>();
 			return InputName;
 		}
 
-		static const FString& GetOutputName()
+		static const FVertexName& GetOutputName()
 		{
-			static const FString OutputName = GetMetasoundDataTypeString<ToDataType>();
+			static const FVertexName OutputName = GetMetasoundDataTypeName<ToDataType>();
 			return OutputName;
 		}
 
@@ -54,9 +56,6 @@ namespace Metasound
 		{
 			auto InitNodeInfo = []() -> FNodeClassMetadata
 			{
-				const FString& InputDisplayName = GetInputName();
-				const FString& OutputDisplayName = GetOutputName();
-
 				FNodeDisplayStyle DisplayStyle;
 				DisplayStyle.bShowName = false;
 				DisplayStyle.bShowInputNames = false;
@@ -141,7 +140,7 @@ namespace Metasound
 				{
 					TDataWriteReference<ToDataType> WriteReference = TDataWriteReferenceFactory<ToDataType>::CreateAny(InParams.OperatorSettings);
 
-					const FString& InputName = GetInputName();
+					const FVertexName& InputName = GetInputName();
 					const bool bContainsRef = InParams.InputDataReferences.ContainsDataReadReference<FromDataType>(InputName);
 					if (bContainsRef)
 					{
@@ -201,5 +200,5 @@ namespace Metasound
 			FVertexInterface Interface;
 			FOperatorFactorySharedRef Factory;
 	};
-}
+} // namespace Metasound
 #undef LOCTEXT_NAMESPACE
diff --git a/Engine/Plugins/Runtime/Metasound/Source/MetasoundFrontend/Public/MetasoundConverterNodeRegistrationMacro.h b/Engine/Plugins/Runtime/Metasound/Source/MetasoundFrontend/Public/MetasoundConverterNodeRegistrationMacro.h
--- a/Engine/Plugins/Runtime/Metasound/Source/MetasoundFrontend/Public/MetasoundConverterNodeRegistrationMacro.h
+++ b/Engine/Plugins/Runtime/Metasound/Source/MetasoundFrontend/Public/MetasoundConverterNodeRegistrationMacro.h
@@ -11,9 +11,10 @@
 #include "MetasoundLog.h"
 #include "MetasoundOperatorInterface.h"
 #include "MetasoundOutputNode.h"
+#include "MetasoundVertex.h"
 
 template<typename TNodeType, typename TFromType, typename TToType>
-bool RegisterConversionNode(const FString& FromPin, const FString& ToPin, const Metasound::FNodeClassMetadata& InNodeMetadata)
+bool RegisterConversionNode(const Metasound::FVertexName& FromPin, const Metasound::FVertexName& ToPin, const Metasound::FNodeClassMetadata& InNodeMetadata)
 {
 	using FConverterNodeRegistryKey = ::Metasound::Frontend::FConverterNodeRegistryKey;
 	using FConverterNodeInfo = ::Metasound::Frontend::FConverterNodeInfo;
diff --git a/Engine/Plugins/Runtime/Metasound/Source/MetasoundFrontend/Public/MetasoundFrontend.h b/Engine/Plugins/Runtime/Metasound/Source/MetasoundFrontend/Public/MetasoundFrontend.h
--- a/Engine/Plugins/Runtime/Metasound/Source/MetasoundFrontend/Public/MetasoundFrontend.h
+++ b/Engine/Plugins/Runtime/Metasound/Source/MetasoundFrontend/Public/MetasoundFrontend.h
@@ -74,7 +74,7 @@ namespace Metasound
 		FMetasoundFrontendClass GenerateClassDescription()
 		{
 			FNodeInitData InitData;
-			InitData.InstanceName = FString(TEXT("GenerateClassDescriptionForNode"));
+			InitData.InstanceName = TEXT("GenerateClassDescriptionForNode");
 
 			return GenerateClassDescription<NodeType>(InitData);
 		}
