Move FString/FText VertexName identifers to FNames. Convention is as follows:

diff --git a/Engine/Plugins/Runtime/Metasound/Source/MetasoundFrontend/Public/MetasoundFrontendController.h b/Engine/Plugins/Runtime/Metasound/Source/MetasoundFrontend/Public/MetasoundFrontendController.h
--- a/Engine/Plugins/Runtime/Metasound/Source/MetasoundFrontend/Public/MetasoundFrontendController.h
+++ b/Engine/Plugins/Runtime/Metasound/Source/MetasoundFrontend/Public/MetasoundFrontendController.h
@@ -10,6 +10,7 @@
 #include "MetasoundFrontendDocumentAccessPtr.h"
 #include "MetasoundFrontendRegistries.h"
 #include "MetasoundGraph.h"
+#include "MetasoundVertex.h"
 
 class IMetaSoundAssetManager;
 
@@ -172,10 +173,10 @@ namespace Metasound
 			virtual const FName& GetDataType() const = 0;
 			
 			/** Returns the name associated with this output. */
-			virtual const FString& GetName() const = 0;
+			virtual const FVertexName& GetName() const = 0;
 			
 			/** Returns the human readable name associated with this output. */
-			virtual const FText& GetDisplayName() const = 0;
+			virtual FText GetDisplayName() const = 0;
 			
 			/** Returns the tooltip associated with this output. */
 			virtual const FText& GetTooltip() const = 0;
@@ -192,6 +193,9 @@ namespace Metasound
 			/** Returns a FConstNodeHandle to the node which owns this output. */
 			virtual FConstNodeHandle GetOwningNode() const = 0;
 
+			/** This should only be used as a means of fixing up vertex names for document model versioning transform(s). */
+			virtual void SetName(const FVertexName& InName) = 0;
+
 			/** Return true if the input is connect to an output. */
 			virtual bool IsConnected() const = 0;
 
@@ -246,10 +250,10 @@ namespace Metasound
 			virtual const FName& GetDataType() const = 0;
 
 			/** Returns the data type name associated with this input. */
-			virtual const FString& GetName() const = 0;
+			virtual const FVertexName& GetName() const = 0;
 
 			/** Returns the data type name associated with this input. */
-			virtual const FText& GetDisplayName() const = 0;
+			virtual FText GetDisplayName() const = 0;
 
 			/** Returns the value for the given input instance if set. */
 			virtual const FMetasoundFrontendLiteral* GetLiteral() const = 0;
@@ -283,6 +287,9 @@ namespace Metasound
 			 * connected, the returned handle will be invalid. */
 			virtual FConstOutputHandle GetConnectedOutput() const = 0;
 
+			/** This should only be used as a means of fixing up vertex names for document model versioning transform(s). */
+			virtual void SetName(const FVertexName& InName) = 0;
+
 			/** Returns information describing connectability between this input and the supplied output. */
 			virtual FConnectability CanConnectTo(const IOutputController& InController) const = 0;
 
@@ -327,11 +334,13 @@ namespace Metasound
 			virtual TArray<FConstInputHandle> GetConstInputs() const = 0;
 
 			/** Returns the display name of the given node (what to distinguish and label in visual arrays, such as context menus). */
-			virtual const FText& GetDisplayName() const = 0;
+			virtual FText GetDisplayName() const = 0;
 
 			/** Sets the description of the node. */
 			virtual void SetDescription(const FText& InDescription) = 0;
 
+			virtual void SetNodeName(const FVertexName& InName) = 0;
+
 			/** Sets the display name of the node. */
 			virtual void SetDisplayName(const FText& InDisplayName) = 0;
 
@@ -352,8 +361,8 @@ namespace Metasound
 			/** Returns number of node outputs. */
 			virtual int32 GetNumOutputs() const = 0;
 
-			virtual TArray<FInputHandle> GetInputsWithVertexName(const FString& InName) = 0;
-			virtual TArray<FConstInputHandle> GetConstInputsWithVertexName(const FString& InName) const = 0;
+			virtual TArray<FInputHandle> GetInputsWithVertexName(const FVertexName& InName) = 0;
+			virtual TArray<FConstInputHandle> GetConstInputsWithVertexName(const FVertexName& InName) const = 0;
 
 			/** Returns all node outputs. */
 			virtual TArray<FOutputHandle> GetOutputs() = 0;
@@ -361,8 +370,8 @@ namespace Metasound
 			/** Returns all node outputs. */
 			virtual TArray<FConstOutputHandle> GetConstOutputs() const = 0;
 
-			virtual TArray<FOutputHandle> GetOutputsWithVertexName(const FString& InName) = 0;
-			virtual TArray<FConstOutputHandle> GetConstOutputsWithVertexName(const FString& InName) const = 0;
+			virtual TArray<FOutputHandle> GetOutputsWithVertexName(const FVertexName& InName) = 0;
+			virtual TArray<FConstOutputHandle> GetConstOutputsWithVertexName(const FVertexName& InName) const = 0;
 
 			/** Returns true if node is required to satisfy the document archetype. */
 			virtual bool IsRequired() const = 0;
@@ -456,7 +465,7 @@ namespace Metasound
 			virtual FConstGraphHandle AsGraph() const = 0;
 
 			/** Returns the name of this node. */
-			virtual const FString& GetNodeName() const = 0;
+			virtual const FVertexName& GetNodeName() const = 0;
 
 			virtual const FMetasoundFrontendNodeStyle& GetNodeStyle() const = 0;
 			virtual void SetNodeStyle(const FMetasoundFrontendNodeStyle& InStyle) = 0;
@@ -499,10 +508,10 @@ namespace Metasound
 			virtual void SetGraphMetadata(const FMetasoundFrontendClassMetadata& InMetadata) = 0;
 
 			/** Return the display name of the graph. */
-			virtual const FText& GetDisplayName() const = 0;
+			virtual FText GetDisplayName() const = 0;
 
-			virtual TArray<FString> GetInputVertexNames() const = 0;
-			virtual TArray<FString> GetOutputVertexNames() const = 0;
+			virtual TArray<FVertexName> GetInputVertexNames() const = 0;
+			virtual TArray<FVertexName> GetOutputVertexNames() const = 0;
 
 			/** Returns all nodes in the graph. */
 			virtual TArray<FNodeHandle> GetNodes() = 0;
@@ -538,24 +547,24 @@ namespace Metasound
 			virtual void ClearGraph() = 0;
 
 			/** Iterate over all input nodes with the given function. If ClassType is specified, only iterate over given type. */
-			virtual void IterateNodes(TUniqueFunction<void(FNodeHandle)> InFunction, EMetasoundFrontendClassType InClassType = EMetasoundFrontendClassType::Invalid) = 0;
+			virtual void IterateNodes(TFunctionRef<void(FNodeHandle)> InFunction, EMetasoundFrontendClassType InClassType = EMetasoundFrontendClassType::Invalid) = 0;
 
 			/** Iterate over all nodes with the given function. If ClassType is specified, only iterate over given type. */
-			virtual void IterateConstNodes(TUniqueFunction<void(FConstNodeHandle)> InFunction, EMetasoundFrontendClassType InClassType = EMetasoundFrontendClassType::Invalid) const = 0;
+			virtual void IterateConstNodes(TFunctionRef<void(FConstNodeHandle)> InFunction, EMetasoundFrontendClassType InClassType = EMetasoundFrontendClassType::Invalid) const = 0;
 
 			/** Returns true if an output vertex with the given Name exists.
 			 *
 			 * @param InName - Name of vertex.
 			 * @return True if the vertex exists, false otherwise. 
 			 */
-			virtual bool ContainsOutputVertexWithName(const FString& InName) const = 0;
+			virtual bool ContainsOutputVertexWithName(const FVertexName& InName) const = 0;
 
 			/** Returns true if an input vertex with the given Name exists.
 			 *
 			 * @param InName - Name of vertex.
 			 * @return True if the vertex exists, false otherwise. 
 			 */
-			virtual bool ContainsInputVertexWithName(const FString& InName) const = 0;
+			virtual bool ContainsInputVertexWithName(const FVertexName& InName) const = 0;
 
 			/** Returns a handle to an existing output node for the given graph output name.
 			 * If no output exists with the given name, an invalid node handle is returned. 
@@ -563,7 +572,7 @@ namespace Metasound
 			 * @param InName - Name of graph output.. 
 			 * @return The node handle for the output node. If the output does not exist, an invalid handle is returned.
 			 */
-			virtual FNodeHandle GetOutputNodeWithName(const FString& InName) = 0;
+			virtual FNodeHandle GetOutputNodeWithName(const FVertexName& InName) = 0;
 
 			/** Returns a handle to an existing output node for the given graph output name.
 			 * If no output exists with the given name, an invalid node handle is returned. 
@@ -571,7 +580,7 @@ namespace Metasound
 			 * @param InName - Name of graph output.
 			 * @return The node handle for the output node. If the node does not exist, an invalid handle is returned.
 			 */
-			virtual FConstNodeHandle GetOutputNodeWithName(const FString& InName) const = 0;
+			virtual FConstNodeHandle GetOutputNodeWithName(const FVertexName& InName) const = 0;
 
 			/** Returns a handle to an existing input node for the given graph input name.
 			 * If no input exists with the given name, an invalid node handle is returned. 
@@ -579,7 +588,7 @@ namespace Metasound
 			 * @param InName - Name of graph input. 
 			 * @return The node handle for the input node. If the node does not exist, an invalid handle is returned.
 			 */
-			virtual FNodeHandle GetInputNodeWithName(const FString& InName) = 0;
+			virtual FNodeHandle GetInputNodeWithName(const FVertexName & InName) = 0;
 
 			/** Returns a handle to an existing input node for the given graph input name.
 			 * If no input exists with the given name, an invalid node handle is returned. 
@@ -587,10 +596,10 @@ namespace Metasound
 			 * @param InName - Name of graph input. 
 			 * @return The node handle for the input node. If the node does not exist, an invalid handle is returned.
 			 */
-			virtual FConstNodeHandle GetInputNodeWithName(const FString& InName) const = 0;
+			virtual FConstNodeHandle GetInputNodeWithName(const FVertexName& InName) const = 0;
 
-			virtual FConstClassInputAccessPtr FindClassInputWithName(const FString& InName) const = 0;
-			virtual FConstClassOutputAccessPtr FindClassOutputWithName(const FString& InName) const = 0;
+			virtual FConstClassInputAccessPtr FindClassInputWithName(const FVertexName& InName) const = 0;
+			virtual FConstClassOutputAccessPtr FindClassOutputWithName(const FVertexName& InName) const = 0;
 
 			/** Add a new input node using the input description. 
 			 *
@@ -598,9 +607,10 @@ namespace Metasound
 			 * @return On success, a valid input node handle. On failure, an invalid node handle.
 			 */
 			virtual FNodeHandle AddInputVertex(const FMetasoundFrontendClassInput& InDescription) = 0;
+			virtual FNodeHandle AddInputVertex(const FVertexName& InName, const FName InTypeName, const FText& InToolTip, const FMetasoundFrontendLiteral* InDefaultValue) = 0;
 
 			/** Remove the input with the given name. Returns true if successfully removed, false otherwise. */
-			virtual bool RemoveInputVertex(const FString& InputName) = 0;
+			virtual bool RemoveInputVertex(const FVertexName& InputName) = 0;
 
 			/** Add a new output node using the output description. 
 			 *
@@ -608,9 +618,10 @@ namespace Metasound
 			 * @return On success, a valid output node handle. On failure, an invalid node handle.
 			 */
 			virtual FNodeHandle AddOutputVertex(const FMetasoundFrontendClassOutput& InDescription) = 0;
+			virtual FNodeHandle AddOutputVertex(const FVertexName& InName, const FName InTypeName, const FText& InToolTip) = 0;
 
 			/** Remove the output with the given name. Returns true if successfully removed, false otherwise. */
-			virtual bool RemoveOutputVertex(const FString& OutputName) = 0;
+			virtual bool RemoveOutputVertex(const FVertexName& OutputName) = 0;
 
 			/** Returns the preferred literal argument type for a given input.
 			 * Returns ELiteralType::Invalid if the input couldn't be found, 
@@ -619,7 +630,7 @@ namespace Metasound
 			 * @param InInputName - Name of graph input.
 			 * @return The preferred literal argument type. 
 			 */
-			virtual ELiteralType GetPreferredLiteralTypeForInputVertex(const FString& InInputName) const = 0;
+			virtual ELiteralType GetPreferredLiteralTypeForInputVertex(const FVertexName& InInputName) const = 0;
 
 			/** Return the UObject class corresponding an input. Meaningful for inputs whose preferred 
 			 * literal type is UObject or UObjectArray.
@@ -628,10 +639,10 @@ namespace Metasound
 			 *
 			 * @return The UClass* for the literal argument input. nullptr on error or if UObject argument is not supported.
 			 */
-			virtual UClass* GetSupportedClassForInputVertex(const FString& InInputName) = 0;
+			virtual UClass* GetSupportedClassForInputVertex(const FVertexName& InInputName) = 0;
 
-			virtual FGuid GetVertexIDForInputVertex(const FString& InInputName) const = 0;
-			virtual FGuid GetVertexIDForOutputVertex(const FString& InOutputName) const = 0;
+			virtual FGuid GetVertexIDForInputVertex(const FVertexName& InInputName) const = 0;
+			virtual FGuid GetVertexIDForOutputVertex(const FVertexName& InOutputName) const = 0;
 
 			virtual FMetasoundFrontendLiteral GetDefaultInput(const FGuid& InVertexID) const = 0;
 			virtual bool SetDefaultInput(const FGuid& InVertexID, const FMetasoundFrontendLiteral& InLiteral) = 0;
@@ -647,22 +658,22 @@ namespace Metasound
 			virtual bool SetDefaultInputToDefaultLiteralOfType(const FGuid& InVertexID) = 0;
 
 			/** Set the display name for the input with the given name. */
-			virtual void SetInputDisplayName(const FString& InName, const FText& InDisplayName) = 0;
+			virtual void SetInputDisplayName(const FVertexName& InName, const FText& InDisplayName) = 0;
 
 			/** Set the display name for the output with the given name. */
-			virtual void SetOutputDisplayName(const FString& InName, const FText& InDisplayName) = 0;
+			virtual void SetOutputDisplayName(const FVertexName& InName, const FText& InDisplayName) = 0;
 
 			/** Get the description for the input with the given name. */
-			virtual const FText& GetInputDescription(const FString& InName) const = 0;
+			virtual const FText& GetInputDescription(const FVertexName& InName) const = 0;
 
 			/** Get the description for the output with the given name. */
-			virtual const FText& GetOutputDescription(const FString& InName) const = 0;
+			virtual const FText& GetOutputDescription(const FVertexName& InName) const = 0;
 
 			/** Set the description for the input with the given name. */
-			virtual void SetInputDescription(const FString& InName, const FText& InDescription) = 0;
+			virtual void SetInputDescription(const FVertexName& InName, const FText& InDescription) = 0;
 
 			/** Set the description for the output with the given name. */
-			virtual void SetOutputDescription(const FString& InName, const FText& InDescription) = 0;
+			virtual void SetOutputDescription(const FVertexName& InName, const FText& InDescription) = 0;
 
 			/**
 			  * Updates the ChangeID for the class interface, which signals AutoUpdate to
@@ -675,7 +686,7 @@ namespace Metasound
 			 *
 			 * @return True on success, false on failure.
 			 */
-			virtual bool ClearLiteralForInput(const FString& InInputName, FGuid InVertexID) = 0;
+			virtual bool ClearLiteralForInput(const FVertexName& InInputName, FGuid InVertexID) = 0;
 
 			/** Add a new node to this graph from the node registry.
 			 *
@@ -734,18 +745,6 @@ namespace Metasound
 
 			/** Returns a handle to the document owning this graph. */
 			virtual FConstDocumentHandle GetOwningDocument() const = 0;
-
-			/** Creates an input node */
-			virtual FNodeHandle AddInputNode(const FMetasoundFrontendClassInput& InClassInput, const FMetasoundFrontendLiteral* InDefaultValue, const FText* InDisplayName) = 0;
-
-			/** Creates an input node */
-			virtual FNodeHandle AddInputNode(const FName InTypeName, const FText& InToolTip, const FMetasoundFrontendLiteral* InDefaultValue, const FText* InDisplayName) = 0;
-
-			/** Creates an output node */
-			virtual FNodeHandle AddOutputNode(const FMetasoundFrontendClassOutput& InClassOutput, const FText* InDisplayName) = 0;
-
-			/** Creates an output node */
-			virtual FNodeHandle AddOutputNode(const FName InTypeName, const FText& InToolTip, const FText* InDisplayName) = 0;
 		};
 
 		/* An IDocumentController provides methods for querying and manipulating a Metasound document. */
diff --git a/Engine/Plugins/Runtime/Metasound/Source/MetasoundFrontend/Public/MetasoundFrontendDocument.h b/Engine/Plugins/Runtime/Metasound/Source/MetasoundFrontend/Public/MetasoundFrontendDocument.h
--- a/Engine/Plugins/Runtime/Metasound/Source/MetasoundFrontend/Public/MetasoundFrontendDocument.h
+++ b/Engine/Plugins/Runtime/Metasound/Source/MetasoundFrontend/Public/MetasoundFrontendDocument.h
@@ -7,6 +7,7 @@
 #include "MetasoundAccessPtr.h"
 #include "MetasoundFrontendLiteral.h"
 #include "MetasoundNodeInterface.h"
+#include "MetasoundVertex.h"
 #include "Misc/Guid.h"
 
 #include "MetasoundFrontendDocument.generated.h"
@@ -219,14 +220,14 @@ struct METASOUNDFRONTEND_API FMetasoundFrontendVertex
 
 	// Name of the vertex. Unique amongst other vertices on the same interface.
 	UPROPERTY(VisibleAnywhere, Category = CustomView)
-	FString Name;
+	FName Name;
 
 	// Data type name of the vertex.
 	UPROPERTY(EditAnywhere, Category = Parameters)
 	FName TypeName;
 
 	// ID of vertex
-	UPROPERTY()	
+	UPROPERTY()
 	FGuid VertexID;
 
 	// Returns true if vertices have equal name & type.
@@ -324,17 +325,19 @@ struct METASOUNDFRONTEND_API FMetasoundFrontendNode
 	// Construct node to satisfy class. 
 	FMetasoundFrontendNode(const FMetasoundFrontendClass& InClass);
 
+private:
 	// Unique ID of this node.
 	UPROPERTY()
 	FGuid ID = Metasound::FrontendInvalidID;
 
+public:
 	// ID of FMetasoundFrontendClass corresponding to this node.
 	UPROPERTY()
 	FGuid ClassID = Metasound::FrontendInvalidID;
 
-	// Name of node instance. 
+	// Name of node instance.
 	UPROPERTY()
-	FString Name;
+	FName Name;
 
 	// Interface of node instance.
 	UPROPERTY()
@@ -347,6 +350,16 @@ struct METASOUNDFRONTEND_API FMetasoundFrontendNode
 	// Style info related to a node.
 	UPROPERTY()
 	FMetasoundFrontendNodeStyle Style;
+
+	const FGuid& GetID() const
+	{
+		return ID;
+	}
+
+	void UpdateID()
+	{
+		ID = FGuid::NewGuid();
+	}
 };
 
 
@@ -588,7 +601,7 @@ struct METASOUNDFRONTEND_API FMetasoundFrontendEnvironmentVariable
 
 	// Name of environment variable.
 	UPROPERTY()
-	FString Name;
+	FName Name;
 
 	// Type of environment variable.
 	UPROPERTY()
diff --git a/Engine/Plugins/Runtime/Metasound/Source/MetasoundFrontend/Public/MetasoundFrontendDocumentAccessPtr.h b/Engine/Plugins/Runtime/Metasound/Source/MetasoundFrontend/Public/MetasoundFrontendDocumentAccessPtr.h
--- a/Engine/Plugins/Runtime/Metasound/Source/MetasoundFrontend/Public/MetasoundFrontendDocumentAccessPtr.h
+++ b/Engine/Plugins/Runtime/Metasound/Source/MetasoundFrontend/Public/MetasoundFrontendDocumentAccessPtr.h
@@ -35,14 +35,14 @@ namespace Metasound
 			// Inherit constructors from base class
 			using Super::Super;
 
-			FVertexAccessPtr GetInputWithName(const FString& InName);
+			FVertexAccessPtr GetInputWithName(const FVertexName& InName);
 			FVertexAccessPtr GetInputWithVertexID(const FGuid& InID);
-			FVertexAccessPtr GetOutputWithName(const FString& InName);
+			FVertexAccessPtr GetOutputWithName(const FVertexName& InName);
 			FVertexAccessPtr GetOutputWithVertexID(const FGuid& InID);
 
-			FConstVertexAccessPtr GetInputWithName(const FString& InName) const;
+			FConstVertexAccessPtr GetInputWithName(const FVertexName& InName) const;
 			FConstVertexAccessPtr GetInputWithVertexID(const FGuid& InID) const;
-			FConstVertexAccessPtr GetOutputWithName(const FString& InName) const;
+			FConstVertexAccessPtr GetOutputWithName(const FVertexName& InName) const;
 			FConstVertexAccessPtr GetOutputWithVertexID(const FGuid& InID) const;
 		};
 
@@ -54,9 +54,9 @@ namespace Metasound
 			// Inherit constructors from base class
 			using Super::Super;
 
-			FConstVertexAccessPtr GetInputWithName(const FString& InName) const;
+			FConstVertexAccessPtr GetInputWithName(const FVertexName& InName) const;
 			FConstVertexAccessPtr GetInputWithVertexID(const FGuid& InID) const;
-			FConstVertexAccessPtr GetOutputWithName(const FString& InName) const;
+			FConstVertexAccessPtr GetOutputWithName(const FVertexName& InName) const;
 			FConstVertexAccessPtr GetOutputWithVertexID(const FGuid& InID) const;
 		};
 
@@ -71,10 +71,10 @@ namespace Metasound
 			// Inherit constructors from base class
 			using Super::Super;
 
-			FClassInputAccessPtr GetInputWithName(const FString& InName);
-			FClassOutputAccessPtr GetOutputWithName(const FString& InName);
-			FConstClassInputAccessPtr GetInputWithName(const FString& InName) const;
-			FConstClassOutputAccessPtr GetOutputWithName(const FString& InName) const;
+			FClassInputAccessPtr GetInputWithName(const FVertexName& InName);
+			FClassOutputAccessPtr GetOutputWithName(const FVertexName& InName);
+			FConstClassInputAccessPtr GetInputWithName(const FVertexName& InName) const;
+			FConstClassOutputAccessPtr GetOutputWithName(const FVertexName& InName) const;
 		};
 
 		class METASOUNDFRONTEND_API FConstClassAccessPtr : public TAccessPtr<const FMetasoundFrontendClass>
@@ -85,8 +85,8 @@ namespace Metasound
 			// Inherit constructors from base class
 			using Super::Super;
 
-			FConstClassInputAccessPtr GetInputWithName(const FString& InName) const;
-			FConstClassOutputAccessPtr GetOutputWithName(const FString& InName) const;
+			FConstClassInputAccessPtr GetInputWithName(const FVertexName& InName) const;
+			FConstClassOutputAccessPtr GetOutputWithName(const FVertexName& InName) const;
 		};
 
 		class METASOUNDFRONTEND_API FGraphClassAccessPtr : public TAccessPtr<FMetasoundFrontendGraphClass>
@@ -97,17 +97,17 @@ namespace Metasound
 			// Inherit constructors from base class
 			using Super::Super;
 
-			FClassInputAccessPtr GetInputWithName(const FString& InName);
+			FClassInputAccessPtr GetInputWithName(const FVertexName& InName);
 			FClassInputAccessPtr GetInputWithNodeID(const FGuid& InNodeID);
-			FClassOutputAccessPtr GetOutputWithName(const FString& InName);
+			FClassOutputAccessPtr GetOutputWithName(const FVertexName& InName);
 			FClassOutputAccessPtr GetOutputWithNodeID(const FGuid& InNodeID);
 
 			FNodeAccessPtr GetNodeWithNodeID(const FGuid& InNodeID);
 			FGraphAccessPtr GetGraph();
 
-			FConstClassInputAccessPtr GetInputWithName(const FString& InName) const;
+			FConstClassInputAccessPtr GetInputWithName(const FVertexName& InName) const;
 			FConstClassInputAccessPtr GetInputWithNodeID(const FGuid& InNodeID) const;
-			FConstClassOutputAccessPtr GetOutputWithName(const FString& InName) const;
+			FConstClassOutputAccessPtr GetOutputWithName(const FVertexName& InName) const;
 			FConstClassOutputAccessPtr GetOutputWithNodeID(const FGuid& InNodeID) const;
 
 			FConstNodeAccessPtr GetNodeWithNodeID(const FGuid& InNodeID) const;
@@ -121,9 +121,9 @@ namespace Metasound
 			// Inherit constructors from base class
 			using Super::Super;
 
-			FConstClassInputAccessPtr GetInputWithName(const FString& InName) const;
+			FConstClassInputAccessPtr GetInputWithName(const FVertexName& InName) const;
 			FConstClassInputAccessPtr GetInputWithNodeID(const FGuid& InNodeID) const;
-			FConstClassOutputAccessPtr GetOutputWithName(const FString& InName) const;
+			FConstClassOutputAccessPtr GetOutputWithName(const FVertexName& InName) const;
 			FConstClassOutputAccessPtr GetOutputWithNodeID(const FGuid& InNodeID) const;
 			FConstNodeAccessPtr GetNodeWithNodeID(const FGuid& InNodeID) const;
 			FConstGraphAccessPtr GetGraph() const;
diff --git a/Engine/Plugins/Runtime/Metasound/Source/MetasoundFrontend/Public/MetasoundFrontendGraph.h b/Engine/Plugins/Runtime/Metasound/Source/MetasoundFrontend/Public/MetasoundFrontendGraph.h
--- a/Engine/Plugins/Runtime/Metasound/Source/MetasoundFrontend/Public/MetasoundFrontendGraph.h
+++ b/Engine/Plugins/Runtime/Metasound/Source/MetasoundFrontend/Public/MetasoundFrontendGraph.h
@@ -30,19 +30,19 @@ namespace Metasound
 			 *
 			 * @param InNodeID - The NodeID related to the parent FMetasoundFrontendClass.
 			 * @param InIndex - The positional index for the input.
-			 * @param InVertexKey - The key for the graph input vertex.
+			 * @param InVertexName - The key for the graph input vertex.
 			 * @param InNode - A shared pointer to an input node. 
 			 */
-			void AddInputNode(FGuid InNodeID, int32 InIndex, const FVertexKey& InVertexKey, TSharedPtr<const INode> InNode);
+			void AddInputNode(FGuid InNodeID, int32 InIndex, const FVertexName& InVertexName, TSharedPtr<const INode> InNode);
 
 			/** Add an output node to this graph.
 			 *
 			 * @param InNodeID - The NodeID related to the parent FMetasoundFrontendClass.
 			 * @param InIndex - The positional index for the output.
-			 * @param InVertexKey - The key for the graph output vertex.
+			 * @param InVertexName - The key for the graph output vertex.
 			 * @param InNode - A shared pointer to an output node. 
 			 */
-			void AddOutputNode(FGuid InNodeID, int32 InIndex, const FVertexKey& InVertexKey, TSharedPtr<const INode> InNode);
+			void AddOutputNode(FGuid InNodeID, int32 InIndex, const FVertexName& InVertexName, TSharedPtr<const INode> InNode);
 
 			/** Store a node on this graph. 
 			 *
@@ -119,7 +119,7 @@ namespace Metasound
 		{
 			FGuid DestinationNodeID;
 			FGuid DestinationVertexID;
-			FVertexKey DestinationVertexKey;
+			FVertexName DestinationVertexKey;
 			FName TypeName;
 			FLiteralNodeConstructorParams InitParams;
 		};
diff --git a/Engine/Plugins/Runtime/Metasound/Source/MetasoundFrontend/Public/MetasoundFrontendInjectReceiveNodes.h b/Engine/Plugins/Runtime/Metasound/Source/MetasoundFrontend/Public/MetasoundFrontendInjectReceiveNodes.h
--- a/Engine/Plugins/Runtime/Metasound/Source/MetasoundFrontend/Public/MetasoundFrontendInjectReceiveNodes.h
+++ b/Engine/Plugins/Runtime/Metasound/Source/MetasoundFrontend/Public/MetasoundFrontendInjectReceiveNodes.h
@@ -15,13 +15,13 @@ namespace Metasound
 		 * in the Metasound Environment to create a unique address.
 		 *
 		 * @param InEnv - The current environment of the graph hosting the receive node.
-		 * @param InVertexKey - The input vertex key of the graph input which will have a 
+		 * @param InVertexName - The input vertex key of the graph input which will have a 
 		 *                      receive node injected. 
 		 * @param InTypeName - The data type of the receive node. 
 		 *
 		 * @return FSendAddress The resulting transmission address for the receive node. 
 		 */
-		using FReceiveNodeAddressFunction = TFunction<FSendAddress(const FMetasoundEnvironment& InEnv, const FVertexKey& InVertexKey, const FName& InTypeName)>;
+		using FReceiveNodeAddressFunction = TFunction<FSendAddress(const FMetasoundEnvironment& InEnv, const FVertexName& InVertexName, const FName& InTypeName)>;
 
 		/** Injects a receive node between an graph input and the connected internal nodes.
 		 *
@@ -41,7 +41,7 @@ namespace Metasound
 		 *
 		 * @return True on success, false on failure. 
 		 */
-		METASOUNDFRONTEND_API bool InjectReceiveNodes(FFrontendGraph& InGraph, const FReceiveNodeAddressFunction& InAddressPolicy, const TSet<FVertexKey>& InInputVerticesToSkip);
+		METASOUNDFRONTEND_API bool InjectReceiveNodes(FFrontendGraph& InGraph, const FReceiveNodeAddressFunction& InAddressPolicy, const TSet<FVertexName>& InInputVerticesToSkip);
 	}
 }
 
diff --git a/Engine/Plugins/Runtime/Metasound/Source/MetasoundFrontend/Public/MetasoundFrontendRegistries.h b/Engine/Plugins/Runtime/Metasound/Source/MetasoundFrontend/Public/MetasoundFrontendRegistries.h
--- a/Engine/Plugins/Runtime/Metasound/Source/MetasoundFrontend/Public/MetasoundFrontendRegistries.h
+++ b/Engine/Plugins/Runtime/Metasound/Source/MetasoundFrontend/Public/MetasoundFrontendRegistries.h
@@ -10,6 +10,7 @@
 #include "MetasoundNodeInterface.h"
 #include "MetasoundOperatorInterface.h"
 #include "MetasoundRouter.h"
+#include "MetasoundVertex.h"
 #include "Templates/Function.h"
 
 
@@ -19,23 +20,23 @@ namespace Metasound
 	struct FDefaultNamedVertexNodeConstructorParams
 	{
 		// the instance name and name of the specific connection that should be used.
-		FString NodeName;
+		FVertexName NodeName;
 		FGuid InstanceID;
-		FString VertexName;
+		FVertexName VertexName;
 	};
 
 	struct FDefaultNamedVertexWithLiteralNodeConstructorParams
 	{
 		// the instance name and name of the specific connection that should be used.
-		FString NodeName;
+		FVertexName NodeName;
 		FGuid InstanceID;
-		FString VertexName;
+		FVertexName VertexName;
 		FLiteral InitParam = FLiteral::CreateInvalid();
 	};
 
 	struct FDefaultLiteralNodeConstructorParams
 	{
-		FString NodeName;
+		FVertexName NodeName;
 		FGuid InstanceID;
 		FLiteral Literal = FLiteral::CreateInvalid();
 	};
@@ -183,10 +184,10 @@ namespace Metasound
 		struct METASOUNDFRONTEND_API FConverterNodeInfo
 		{
 			// If this node has multiple input pins, we use this to designate which pin should be used.
-			FVertexKey PreferredConverterInputPin;
+			FVertexName PreferredConverterInputPin;
 
 			// If this node has multiple output pins, we use this to designate which pin should be used.
-			FVertexKey PreferredConverterOutputPin;
+			FVertexName PreferredConverterOutputPin;
 
 			// The key for this node in the node registry.
 			FNodeRegistryKey NodeKey;
diff --git a/Engine/Plugins/Runtime/Metasound/Source/MetasoundFrontend/Public/MetasoundFrontendTransform.h b/Engine/Plugins/Runtime/Metasound/Source/MetasoundFrontend/Public/MetasoundFrontendTransform.h
--- a/Engine/Plugins/Runtime/Metasound/Source/MetasoundFrontend/Public/MetasoundFrontendTransform.h
+++ b/Engine/Plugins/Runtime/Metasound/Source/MetasoundFrontend/Public/MetasoundFrontendTransform.h
@@ -5,6 +5,7 @@
 #include "MetasoundFrontendArchetypeRegistry.h"
 #include "MetasoundFrontendController.h"
 #include "MetasoundFrontendDocument.h"
+#include "MetasoundVertex.h"
 
 
 namespace Metasound
@@ -123,7 +124,7 @@ namespace Metasound
 		public:
 			static FMetasoundFrontendVersionNumber GetMaxVersion()
 			{
-				return FMetasoundFrontendVersionNumber { 1, 6 };
+				return FMetasoundFrontendVersionNumber { 1, 7 };
 			}
 
 			FVersionDocument(FName InName, const FString& InPath);
diff --git a/Engine/Plugins/Runtime/Metasound/Source/MetasoundFrontend/Public/MetasoundInputNode.h b/Engine/Plugins/Runtime/Metasound/Source/MetasoundFrontend/Public/MetasoundInputNode.h
--- a/Engine/Plugins/Runtime/Metasound/Source/MetasoundFrontend/Public/MetasoundInputNode.h
+++ b/Engine/Plugins/Runtime/Metasound/Source/MetasoundFrontend/Public/MetasoundInputNode.h
@@ -20,7 +20,7 @@ namespace Metasound
 		public:
 			using FDataWriteReference = TDataWriteReference<DataType>;
 
-			TInputOperator(const FVertexKey& InDataReferenceName, FDataWriteReference InDataReference)
+			TInputOperator(const FVertexName& InDataReferenceName, FDataWriteReference InDataReference)
 				: DataReferenceName(InDataReferenceName)
 				// Executable DataTypes require a copy of the output to operate on whereas non-executable
 				// types do not. Avoid copy by assigning to reference for non-executable types.
@@ -89,7 +89,7 @@ namespace Metasound
 			}
 
 		protected:
-			FVertexKey DataReferenceName;
+			FVertexName DataReferenceName;
 
 			FDataWriteReference InputValue;
 			FDataWriteReference OutputValue;
@@ -107,7 +107,7 @@ namespace Metasound
 		using FDataReadReference = TDataReadReference<DataType>;
 		using Super = TInputOperator<DataType>;
 
-		TPassThroughOperator(const FVertexKey& InDataReferenceName, FDataReadReference InDataReference)
+		TPassThroughOperator(const FVertexName& InDataReferenceName, FDataReadReference InDataReference)
 		: TInputOperator<DataType>(InDataReferenceName, WriteCast(InDataReference)) // Write cast is safe because `GetExecuteFunction() and GetInputs() are overridden, ensuring that data is not written.
 		, DataReferenceName(InDataReferenceName)
 		{
@@ -178,7 +178,7 @@ namespace Metasound
 		}
 
 	private:
-		FVertexKey DataReferenceName;
+		FVertexName DataReferenceName;
 	};
 
 
@@ -261,7 +261,7 @@ namespace Metasound
 			// If true, this node can be instantiated by the FrontEnd.
 			static constexpr bool bCanRegister = FCreateDataReferenceWithLiteral<DataType>::bCanCreateWithLiteral;
 
-			static FVertexInterface DeclareVertexInterface(const FVertexKey& InVertexName)
+			static FVertexInterface DeclareVertexInterface(const FVertexName& InVertexName)
 			{
 				return FVertexInterface(
 					FInputVertexInterface(
@@ -273,7 +273,7 @@ namespace Metasound
 				);
 			}
 
-			static FNodeClassMetadata GetNodeInfo(const FVertexKey& InVertexName)
+			static FNodeClassMetadata GetNodeInfo(const FVertexName& InVertexName)
 			{
 				FNodeClassMetadata Info;
 
@@ -310,8 +310,8 @@ namespace Metasound
 			/* Construct a TInputNode using the TInputOperatorFactory<> and forwarding 
 			 * Args to the TInputOperatorFactory constructor.*/
 			template<typename... ArgTypes>
-			TInputNode(const FString& InNodeDescription, const FGuid& InInstanceID, const FVertexKey& InVertexName, ArgTypes&&... Args)
-			:	FNode(InNodeDescription, InInstanceID, GetNodeInfo(InVertexName))
+			TInputNode(const FVertexName& InInstanceName, const FGuid& InInstanceID, const FVertexName& InVertexName, ArgTypes&&... Args)
+			:	FNode(InInstanceName, InInstanceID, GetNodeInfo(InVertexName))
 			,	VertexName(InVertexName)
 			,	Interface(DeclareVertexInterface(InVertexName))
 			,	Factory(CreateOperatorFactoryWithArgs(Forward<ArgTypes>(Args)...))
@@ -320,15 +320,15 @@ namespace Metasound
 
 			/* Construct a TInputNode using the TInputOperatorLiteralFactory<> and moving
 			 * InParam to the TInputOperatorLiteralFactory constructor.*/
-			explicit TInputNode(const FString& InNodeDescription, const FGuid& InInstanceID, const FVertexKey& InVertexName, FLiteral&& InParam)
-			:	FNode(InNodeDescription, InInstanceID, GetNodeInfo(InVertexName))
+			explicit TInputNode(const FVertexName& InNodeName, const FGuid& InInstanceID, const FVertexName& InVertexName, FLiteral&& InParam)
+			:	FNode(InNodeName, InInstanceID, GetNodeInfo(InVertexName))
 			,	VertexName(InVertexName)
 			,	Interface(DeclareVertexInterface(InVertexName))
 			, 	Factory(CreateOperatorFactoryWithLiteral(MoveTemp(InParam)))
 			{
 			}
 
-			const FVertexKey& GetVertexName() const
+			const FVertexName& GetVertexName() const
 			{
 				return VertexName;
 			}
@@ -354,7 +354,7 @@ namespace Metasound
 			}
 
 		private:
-			FVertexKey VertexName;
+			FVertexName VertexName;
 
 			FVertexInterface Interface;
 			FOperatorFactorySharedRef Factory;
@@ -367,7 +367,7 @@ namespace Metasound
 		using FInputNodeType = TInputNode<DataType>;
 
 		const FInputNodeType& InputNode = static_cast<const FInputNodeType&>(InParams.Node);
-		const FVertexKey& VertexKey = InputNode.GetVertexName();
+		const FVertexName& VertexKey = InputNode.GetVertexName();
 
 		if (InParams.InputDataReferences.ContainsDataWriteReference<DataType>(VertexKey))
 		{
diff --git a/Engine/Plugins/Runtime/Metasound/Source/MetasoundFrontend/Public/MetasoundOutputNode.h b/Engine/Plugins/Runtime/Metasound/Source/MetasoundFrontend/Public/MetasoundOutputNode.h
--- a/Engine/Plugins/Runtime/Metasound/Source/MetasoundFrontend/Public/MetasoundOutputNode.h
+++ b/Engine/Plugins/Runtime/Metasound/Source/MetasoundFrontend/Public/MetasoundOutputNode.h
@@ -7,6 +7,7 @@
 #include "MetasoundNodeInterface.h"
 #include "MetasoundOperatorInterface.h"
 #include "MetasoundNode.h"
+#include "MetasoundVertex.h"
 
 #define LOCTEXT_NAMESPACE "MetasoundGraphCore"
 
@@ -40,7 +41,7 @@ namespace Metasound
 			public:
 				using FDataReadReference = TDataReadReference<DataType>;
 
-				FOutputOperator(const FString& InDataReferenceName, FDataReadReference InDataReference)
+				FOutputOperator(const FVertexName& InDataReferenceName, FDataReadReference InDataReference)
 				{
 					Outputs.AddDataReadReference<DataType>(InDataReferenceName, InDataReference);
 				}
@@ -69,7 +70,7 @@ namespace Metasound
 		class FOutputOperatorFactory : public IOperatorFactory
 		{
 			public:
-				FOutputOperatorFactory(const FString& InDataReferenceName)
+				FOutputOperatorFactory(const FVertexName& InDataReferenceName)
 				:	DataReferenceName(InDataReferenceName)
 				{
 				}
@@ -95,22 +96,24 @@ namespace Metasound
 				}
 
 			private:
-				FString DataReferenceName;
+				FVertexName DataReferenceName;
 		};
 
-		static FVertexInterface GetVertexInterface(const FString& InVertexName)
+		static FVertexInterface GetVertexInterface(const FVertexName& InVertexName)
 		{
+			static const FText VertexDescription = LOCTEXT("Metasound_OutputVertexDescription", "Output from the parent Metasound graph.");
+
 			return FVertexInterface(
 				FInputVertexInterface(
-					TInputDataVertexModel<DataType>(InVertexName, LOCTEXT("Metasound_OutputVertexDescription", "Output from the parent Metasound graph."))
+					TInputDataVertexModel<DataType>(InVertexName, VertexDescription)
 				),
 				FOutputVertexInterface(
-					TOutputDataVertexModel<DataType>(InVertexName, LOCTEXT("Metasound_OutputVertexDescription", "Output from the parent Metasound graph."))
+					TOutputDataVertexModel<DataType>(InVertexName, VertexDescription)
 				)
 			);
 		}
 
-		static FNodeClassMetadata GetNodeInfo(const FString& InVertexName)
+		static FNodeClassMetadata GetNodeInfo(const FVertexName& InVertexName)
 		{
 			FNodeClassMetadata Info;
 
@@ -129,7 +132,7 @@ namespace Metasound
 
 
 		public:
-			TOutputNode(const FString& InInstanceName, const FGuid& InInstanceID, const FString& InVertexName)
+			TOutputNode(const FVertexName& InInstanceName, const FGuid& InInstanceID, const FVertexName& InVertexName)
 			:	FNode(InInstanceName, InInstanceID, GetNodeInfo(InVertexName))
 			,	VertexInterface(GetVertexInterface(InVertexName))
 			,	Factory(MakeShared<FOutputOperatorFactory, ESPMode::ThreadSafe>(InVertexName))
diff --git a/Engine/Plugins/Runtime/Metasound/Source/MetasoundFrontend/Public/MetasoundParameterTransmitter.h b/Engine/Plugins/Runtime/Metasound/Source/MetasoundFrontend/Public/MetasoundParameterTransmitter.h
--- a/Engine/Plugins/Runtime/Metasound/Source/MetasoundFrontend/Public/MetasoundParameterTransmitter.h
+++ b/Engine/Plugins/Runtime/Metasound/Source/MetasoundFrontend/Public/MetasoundParameterTransmitter.h
@@ -63,13 +63,13 @@ namespace Metasound
 		};
 
 		/** Returns the MetaSound environment variable name which contains the instance ID. */
-		static FString GetInstanceIDEnvironmentVariableName();
+		static const FVertexName& GetInstanceIDEnvironmentVariableName();
 
 		/** Creates a unique send address using the given MetaSound environment. */
-		static FSendAddress CreateSendAddressFromEnvironment(const FMetasoundEnvironment& InEnvironment, const FVertexKey& InVertexkey, const FName& InTypeName);
+		static FSendAddress CreateSendAddressFromEnvironment(const FMetasoundEnvironment& InEnvironment, const FVertexName& InVertexName, const FName& InTypeName);
 		
 		/** Creates a unique send address using the given InstanceID. */
-		static FSendAddress CreateSendAddressFromInstanceID(uint64 InInstanceID, const FVertexKey& InVertexKey, const FName& InTypeName);
+		static FSendAddress CreateSendAddressFromInstanceID(uint64 InInstanceID, const FVertexName& InVertexName, const FName& InTypeName);
 
 		FMetaSoundParameterTransmitter(const FMetaSoundParameterTransmitter::FInitParams& InInitParams);
 		virtual ~FMetaSoundParameterTransmitter() = default;
diff --git a/Engine/Plugins/Runtime/Metasound/Source/MetasoundFrontend/Public/MetasoundReceiveNode.h b/Engine/Plugins/Runtime/Metasound/Source/MetasoundFrontend/Public/MetasoundReceiveNode.h
--- a/Engine/Plugins/Runtime/Metasound/Source/MetasoundFrontend/Public/MetasoundReceiveNode.h
+++ b/Engine/Plugins/Runtime/Metasound/Source/MetasoundFrontend/Public/MetasoundReceiveNode.h
@@ -20,9 +20,9 @@ namespace Metasound
 {
 	namespace ReceiveNodeInfo
 	{
-		METASOUNDFRONTEND_API const FString& GetAddressInputName();
-		METASOUNDFRONTEND_API const FString& GetDefaultDataInputName();
-		METASOUNDFRONTEND_API const FString& GetOutputName();
+		METASOUNDFRONTEND_API const FVertexName& GetAddressInputName();
+		METASOUNDFRONTEND_API const FVertexName& GetDefaultDataInputName();
+		METASOUNDFRONTEND_API const FVertexName& GetOutputName();
 		METASOUNDFRONTEND_API FNodeClassName GetClassNameForDataType(const FName& InDataTypeName);
 		METASOUNDFRONTEND_API int32 GetCurrentMajorVersion();
 		METASOUNDFRONTEND_API int32 GetCurrentMinorVersion();
diff --git a/Engine/Plugins/Runtime/Metasound/Source/MetasoundFrontend/Public/MetasoundSendNode.h b/Engine/Plugins/Runtime/Metasound/Source/MetasoundFrontend/Public/MetasoundSendNode.h
--- a/Engine/Plugins/Runtime/Metasound/Source/MetasoundFrontend/Public/MetasoundSendNode.h
+++ b/Engine/Plugins/Runtime/Metasound/Source/MetasoundFrontend/Public/MetasoundSendNode.h
@@ -9,6 +9,7 @@
 #include "MetasoundDataReference.h"
 #include "MetasoundExecutableOperator.h"
 #include "MetasoundRouter.h"
+#include "MetasoundVertex.h"
 
 #include <type_traits>
 
@@ -21,15 +22,15 @@ namespace Metasound
 	class TSendNode : public FNode
 	{
 	public:
-		static const FString& GetAddressInputName()
+		static const FVertexName& GetAddressInputName()
 		{
-			static const FString InputName = TEXT("Address");
+			static const FVertexName InputName = TEXT("Address");
 			return InputName;
 		}
 
-		static const FString& GetSendInputName()
+		static const FVertexName& GetSendInputName()
 		{
-			static const FString& SendInput = GetMetasoundDataTypeString<TDataType>();
+			static const FVertexName& SendInput = GetMetasoundDataTypeName<TDataType>();
 			return SendInput;
 		}
 
@@ -49,7 +50,7 @@ namespace Metasound
 		{
 			auto InitNodeInfo = []() -> FNodeClassMetadata
 			{
-				const FString& InputName = GetSendInputName();
+				const FVertexName& InputName = GetSendInputName();
 				FNodeClassMetadata Info;
 
 				Info.ClassName = { "Send", GetMetasoundDataTypeName<TDataType>(), FName() };
@@ -144,7 +145,7 @@ namespace Metasound
 						return MakeUnique<TSendOperator>(InParams.InputDataReferences.GetDataReadReference<TDataType>(GetSendInputName()),
 							InParams.InputDataReferences.GetDataReadReferenceOrConstruct<FSendAddress>(GetAddressInputName()),
 							InParams.OperatorSettings
-							);
+						);
 					}
 					else
 					{
diff --git a/Engine/Plugins/Runtime/Metasound/Source/MetasoundGenerator/Private/MetasoundGenerator.cpp b/Engine/Plugins/Runtime/Metasound/Source/MetasoundGenerator/Private/MetasoundGenerator.cpp
--- a/Engine/Plugins/Runtime/Metasound/Source/MetasoundGenerator/Private/MetasoundGenerator.cpp
+++ b/Engine/Plugins/Runtime/Metasound/Source/MetasoundGenerator/Private/MetasoundGenerator.cpp
@@ -72,11 +72,11 @@ namespace Metasound
 			FDataReferenceCollection Outputs = GraphOperator->GetOutputs();
 
 			// Get output audio buffers.
-			for (const FString& AudioOutputName : InitParams.AudioOutputNames)
+			for (const FVertexName& AudioOutputName : InitParams.AudioOutputNames)
 			{
 				if (!Outputs.ContainsDataReadReference<FAudioBuffer>(AudioOutputName))
 				{
-					UE_LOG(LogMetaSound, Warning, TEXT("MetasoundSource [%s] does not contain audio output [%s] in output"), *InitParams.MetaSoundName, *AudioOutputName);
+					UE_LOG(LogMetaSound, Warning, TEXT("MetasoundSource [%s] does not contain audio output [%s] in output"), *InitParams.MetaSoundName, *AudioOutputName.ToString());
 				}
 				OutputBuffers.Add(Outputs.GetDataReadReferenceOrConstruct<FAudioBuffer>(AudioOutputName, InitParams.OperatorSettings));
 			}
diff --git a/Engine/Plugins/Runtime/Metasound/Source/MetasoundGenerator/Public/MetasoundGenerator.h b/Engine/Plugins/Runtime/Metasound/Source/MetasoundGenerator/Public/MetasoundGenerator.h
--- a/Engine/Plugins/Runtime/Metasound/Source/MetasoundGenerator/Public/MetasoundGenerator.h
+++ b/Engine/Plugins/Runtime/Metasound/Source/MetasoundGenerator/Public/MetasoundGenerator.h
@@ -1,18 +1,20 @@
 // Copyright Epic Games, Inc. All Rights Reserved.
 #pragma once
 
+#include "Async/Async.h"
 #include "CoreMinimal.h"
 #include "DSP/Dsp.h"
-#include "MetasoundOperatorInterface.h"
-#include "MetasoundExecutableOperator.h"
 #include "MetasoundAudioFormats.h"
+#include "MetasoundExecutableOperator.h"
+#include "MetasoundFrontendController.h"
+#include "MetasoundOperatorInterface.h"
 #include "MetasoundPrimitives.h"
 #include "MetasoundTrigger.h"
+#include "MetasoundVertex.h"
 #include "Sound/SoundGenerator.h"
-#include "MetasoundFrontendController.h"
-#include "Async/Async.h"
 #include "Tickable.h"
 
+
 namespace Metasound
 {
 	// Struct needed for building the metasound graph
@@ -22,9 +24,9 @@ namespace Metasound
 		TUniquePtr<IGraph> Graph;
 		FMetasoundEnvironment Environment;
 		FString MetaSoundName;
-		TArray<FString> AudioOutputNames;
-		FString OnPlayInputName;
-		FString IsFinishedOutputName;
+		TArray<FVertexName> AudioOutputNames;
+		FVertexName OnPlayInputName;
+		FVertexName IsFinishedOutputName;
 
 		void Release();
 	};
diff --git a/Engine/Plugins/Runtime/Metasound/Source/MetasoundGraphCore/Private/MetasoundDataReferenceCollection.cpp b/Engine/Plugins/Runtime/Metasound/Source/MetasoundGraphCore/Private/MetasoundDataReferenceCollection.cpp
--- a/Engine/Plugins/Runtime/Metasound/Source/MetasoundGraphCore/Private/MetasoundDataReferenceCollection.cpp
+++ b/Engine/Plugins/Runtime/Metasound/Source/MetasoundGraphCore/Private/MetasoundDataReferenceCollection.cpp
@@ -8,35 +8,33 @@
 
 namespace Metasound
 {
-	bool FDataReferenceCollection::ContainsDataReadReference(const FString& InName, const FName& InTypeName) const
+	bool FDataReferenceCollection::ContainsDataReadReference(const FVertexName& InName, const FName& InTypeName) const
 	{
-		if (DataReadRefMap.Contains(InName))
+		if (const FDataRefWrapper* Wrapper = DataReadRefMap.Find(InName))
 		{
-			const FDataRefWrapper& Wrapper = DataReadRefMap[InName];
-			if (Wrapper.IsValid())
+			if (const IDataReference* Ref = Wrapper->GetDataReference())
 			{
-				return InTypeName == Wrapper.GetDataReference()->GetDataTypeName();
+				return InTypeName == Ref->GetDataTypeName();
 			}
 		}
 
 		return false;
 	}
 
-	bool FDataReferenceCollection::ContainsDataWriteReference(const FString& InName, const FName& InTypeName) const
+	bool FDataReferenceCollection::ContainsDataWriteReference(const FVertexName& InName, const FName& InTypeName) const
 	{
-		if (DataWriteRefMap.Contains(InName))
+		if (const FDataRefWrapper* Wrapper = DataWriteRefMap.Find(InName))
 		{
-			const FDataRefWrapper& Wrapper = DataWriteRefMap[InName];
-			if (Wrapper.IsValid())
+			if (const IDataReference* Ref = Wrapper->GetDataReference())
 			{
-				return InTypeName == Wrapper.GetDataReference()->GetDataTypeName();
+				return InTypeName == Ref->GetDataTypeName();
 			}
 		}
 
 		return false;
 	}
 
-	bool FDataReferenceCollection::AddDataReadReferenceFrom(const FString& InName, const FDataReferenceCollection& OtherCollection, const FString& OtherName, const FName& OtherTypeName)
+	bool FDataReferenceCollection::AddDataReadReferenceFrom(const FVertexName& InName, const FDataReferenceCollection& OtherCollection, const FVertexName& OtherName, const FName& OtherTypeName)
 	{
 		if (OtherCollection.ContainsDataReadReference(OtherName, OtherTypeName))
 		{
@@ -54,7 +52,7 @@ namespace Metasound
 		return false;
 	}
 
-	bool FDataReferenceCollection::AddDataWriteReferenceFrom(const FString& InName, const FDataReferenceCollection& OtherCollection, const FString& OtherName, const FName& OtherTypeName)
+	bool FDataReferenceCollection::AddDataWriteReferenceFrom(const FVertexName& InName, const FDataReferenceCollection& OtherCollection, const FVertexName& OtherName, const FName& OtherTypeName)
 	{
 		if (OtherCollection.ContainsDataWriteReference(OtherName, OtherTypeName))
 		{
@@ -66,17 +64,17 @@ namespace Metasound
 		return false;
 	}
 
-	const IDataReference* FDataReferenceCollection::GetDataReference(const FDataReferenceMap& InMap, const FString& InName) const
+	const IDataReference* FDataReferenceCollection::GetDataReference(const FDataReferenceMap& InMap, const FVertexName& InName) const
 	{
-		if (InMap.Contains(InName))
+		if (const FDataRefWrapper* RefWrapper = InMap.Find(InName))
 		{
-			return InMap[InName].GetDataReference();
+			return RefWrapper->GetDataReference();
 		}
 
 		return nullptr;
 	}
 
-	void FDataReferenceCollection::AddDataReference(FDataReferenceMap& InMap, const FString& InName, FDataRefWrapper&& InDataRef)
+	void FDataReferenceCollection::AddDataReference(FDataReferenceMap& InMap, const FVertexName& InName, FDataRefWrapper&& InDataRef)
 	{
 		InMap.Add(InName, InDataRef);
 	}
diff --git a/Engine/Plugins/Runtime/Metasound/Source/MetasoundGraphCore/Private/MetasoundGraph.cpp b/Engine/Plugins/Runtime/Metasound/Source/MetasoundGraphCore/Private/MetasoundGraph.cpp
--- a/Engine/Plugins/Runtime/Metasound/Source/MetasoundGraphCore/Private/MetasoundGraph.cpp
+++ b/Engine/Plugins/Runtime/Metasound/Source/MetasoundGraphCore/Private/MetasoundGraph.cpp
@@ -12,7 +12,7 @@ namespace Metasound
 	{
 	}
 
-	const FString& FGraph::GetInstanceName() const
+	const FVertexName& FGraph::GetInstanceName() const
 	{
 		return InstanceName;
 	}
@@ -27,14 +27,14 @@ namespace Metasound
 		return Metadata;
 	}
 
-	bool FGraph::AddInputDataDestination(const INode& InNode, const FVertexKey& InVertexKey)
+	bool FGraph::AddInputDataDestination(const INode& InNode, const FVertexName& InVertexName)
 	{
-		if (!InNode.GetVertexInterface().ContainsInputVertex(InVertexKey))
+		if (!InNode.GetVertexInterface().ContainsInputVertex(InVertexName))
 		{
 			return false;
 		}
 
-		FInputDataDestination Destination(InNode, InNode.GetVertexInterface().GetInputVertex(InVertexKey));
+		FInputDataDestination Destination(InNode, InNode.GetVertexInterface().GetInputVertex(InVertexName));
 
 		AddInputDataDestination(Destination);
 		
@@ -52,14 +52,14 @@ namespace Metasound
 		return InputDestinations;
 	}
 
-	bool FGraph::AddOutputDataSource(const INode& InNode, const FVertexKey& InVertexKey)
+	bool FGraph::AddOutputDataSource(const INode& InNode, const FVertexName& InVertexName)
 	{
-		if (!InNode.GetVertexInterface().ContainsOutputVertex(InVertexKey))
+		if (!InNode.GetVertexInterface().ContainsOutputVertex(InVertexName))
 		{
 			return false;
 		}
 
-		FOutputDataSource Source(InNode, InNode.GetVertexInterface().GetOutputVertex(InVertexKey));
+		FOutputDataSource Source(InNode, InNode.GetVertexInterface().GetOutputVertex(InVertexName));
 
 		AddOutputDataSource(Source);
 
@@ -82,7 +82,7 @@ namespace Metasound
 		Edges.Add(InEdge);
 	}
 
-	bool FGraph::AddDataEdge(const INode& FromNode, const FVertexKey& FromKey, const INode& ToNode, const FVertexKey& ToKey)
+	bool FGraph::AddDataEdge(const INode& FromNode, const FVertexName& FromKey, const INode& ToNode, const FVertexName& ToKey)
 	{
 		const FVertexInterface& FromVertexInterface = FromNode.GetVertexInterface();
 		const FVertexInterface& ToVertexInterface = ToNode.GetVertexInterface();
diff --git a/Engine/Plugins/Runtime/Metasound/Source/MetasoundGraphCore/Private/MetasoundGraphLinter.cpp b/Engine/Plugins/Runtime/Metasound/Source/MetasoundGraphCore/Private/MetasoundGraphLinter.cpp
--- a/Engine/Plugins/Runtime/Metasound/Source/MetasoundGraphCore/Private/MetasoundGraphLinter.cpp
+++ b/Engine/Plugins/Runtime/Metasound/Source/MetasoundGraphCore/Private/MetasoundGraphLinter.cpp
@@ -19,52 +19,31 @@ namespace Metasound
 		// run validation on input destination
 		bool IsInputDestinationValid(const FInputDataDestination& InDest)
 		{
-			if (InDest.Node == nullptr)
+			if (InDest.Node)
 			{
-				return false;
-			}
-
-
-			const FInputVertexInterface& InputInterface = InDest.Node->GetVertexInterface().GetInputInterface();
-
-			FVertexKey VertexKey = MakeDataVertexKey(InDest.Vertex);
-
-			if (!InputInterface.Contains(VertexKey))
-			{
-				return false;
-			}
-
-			if (InputInterface[VertexKey] != InDest.Vertex)
-			{
-				return false;
+				const FInputVertexInterface& Interface = InDest.Node->GetVertexInterface().GetInputInterface();
+				if (const FInputDataVertex* Vertex = Interface.Find(InDest.Vertex.GetVertexName()))
+				{
+					return *Vertex == InDest.Vertex;
+				}
 			}
 
-			return true;
+			return false;
 		}
 		
 		// run validation on output source
 		bool IsOutputSourceValid(const FOutputDataSource& InSource)
 		{
-			if (InSource.Node == nullptr)
+			if (InSource.Node)
 			{
-				return false;
-			}
-
-			const FOutputVertexInterface& OutputInterface = InSource.Node->GetVertexInterface().GetOutputInterface();
-
-			FVertexKey VertexKey = MakeDataVertexKey(InSource.Vertex);
-
-			if (!OutputInterface.Contains(VertexKey))
-			{
-				return false;
-			}
-
-			if (OutputInterface[VertexKey] != InSource.Vertex)
-			{
-				return false;
+				const FOutputVertexInterface& Interface = InSource.Node->GetVertexInterface().GetOutputInterface();
+				if (const FOutputDataVertex* Vertex = Interface.Find(InSource.Vertex.GetVertexName()))
+				{
+					return *Vertex == InSource.Vertex;
+				}
 			}
 
-			return true;
+			return false;
 		};
 
 		const FDataEdge* GetEdgePointer(const FDataEdge& InEdge) 
diff --git a/Engine/Plugins/Runtime/Metasound/Source/MetasoundGraphCore/Private/MetasoundNode.cpp b/Engine/Plugins/Runtime/Metasound/Source/MetasoundGraphCore/Private/MetasoundNode.cpp
--- a/Engine/Plugins/Runtime/Metasound/Source/MetasoundGraphCore/Private/MetasoundNode.cpp
+++ b/Engine/Plugins/Runtime/Metasound/Source/MetasoundGraphCore/Private/MetasoundNode.cpp
@@ -6,7 +6,7 @@
 
 namespace Metasound
 {
-	FNode::FNode(const FString& InInstanceName, const FGuid& InInstanceID, const FNodeClassMetadata& InInfo)
+	FNode::FNode(const FVertexName& InInstanceName, const FGuid& InInstanceID, const FNodeClassMetadata& InInfo)
 	: InstanceName(InInstanceName)
 	, InstanceID(InInstanceID)
 	, Info(InInfo)
@@ -14,7 +14,7 @@ namespace Metasound
 	}
 
 	/** Return the name of this specific instance of the node class. */
-	const FString& FNode::GetInstanceName() const
+	const FVertexName& FNode::GetInstanceName() const
 	{
 		return InstanceName;
 	}
diff --git a/Engine/Plugins/Runtime/Metasound/Source/MetasoundGraphCore/Private/MetasoundVariableNodes.cpp b/Engine/Plugins/Runtime/Metasound/Source/MetasoundGraphCore/Private/MetasoundVariableNodes.cpp
--- a/Engine/Plugins/Runtime/Metasound/Source/MetasoundGraphCore/Private/MetasoundVariableNodes.cpp
+++ b/Engine/Plugins/Runtime/Metasound/Source/MetasoundGraphCore/Private/MetasoundVariableNodes.cpp
@@ -7,27 +7,27 @@ namespace Metasound
 {
 	namespace VariableNodeVertexNames
 	{
-		const FString& GetInputDataName()
+		const FVertexName& GetInputDataName()
 		{
-			static const FString Name = TEXT("Value");
+			static const FVertexName Name = TEXT("Value");
 			return Name;
 		}
 
-		const FString& GetOutputDataName()
+		const FVertexName& GetOutputDataName()
 		{
-			static const FString Name = TEXT("Value");
+			static const FVertexName Name = TEXT("Value");
 			return Name;
 		}
 
-		const FString& GetInputVariableName()
+		const FVertexName& GetInputVariableName()
 		{
-			static const FString Name = TEXT("Variable");
+			static const FVertexName Name = TEXT("Variable");
 			return Name;
 		}
 
-		const FString& GetOutputVariableName()
+		const FVertexName& GetOutputVariableName()
 		{
-			static const FString Name = TEXT("Variable");
+			static const FVertexName Name = TEXT("Variable");
 			return Name;
 		}
 	}
diff --git a/Engine/Plugins/Runtime/Metasound/Source/MetasoundGraphCore/Private/MetasoundVertex.cpp b/Engine/Plugins/Runtime/Metasound/Source/MetasoundGraphCore/Private/MetasoundVertex.cpp
--- a/Engine/Plugins/Runtime/Metasound/Source/MetasoundGraphCore/Private/MetasoundVertex.cpp
+++ b/Engine/Plugins/Runtime/Metasound/Source/MetasoundGraphCore/Private/MetasoundVertex.cpp
@@ -43,7 +43,7 @@ namespace Metasound
 		return *this;
 	}
 
-	const FString& FInputDataVertex::GetVertexName() const
+	const FVertexName& FInputDataVertex::GetVertexName() const
 	{
 		return VertexModel->VertexName;
 	}
@@ -96,7 +96,7 @@ namespace Metasound
 		}
 		else
 		{
-			return InLHS.GetVertexName() < InRHS.GetVertexName();
+			return InLHS.GetVertexName().FastLess(InRHS.GetVertexName());
 		}
 	}
 
@@ -137,7 +137,7 @@ namespace Metasound
 		return *this;
 	}
 
-	const FString& FOutputDataVertex::GetVertexName() const
+	FVertexName FOutputDataVertex::GetVertexName() const
 	{
 		return VertexModel->VertexName;
 	}
@@ -186,7 +186,7 @@ namespace Metasound
 		}
 		else
 		{
-			return InLHS.GetVertexName() < InRHS.GetVertexName();
+			return InLHS.GetVertexName().FastLess(InRHS.GetVertexName());
 		}
 	}
 
@@ -226,7 +226,7 @@ namespace Metasound
 		return *this;
 	}
 
-	const FString& FEnvironmentVertex::GetVertexName() const
+	const FVertexName& FEnvironmentVertex::GetVertexName() const
 	{
 		return VertexModel->VertexName;
 	}
@@ -258,7 +258,7 @@ namespace Metasound
 			return false;
 		}
 
-		return InLHS.GetVertexName() < InRHS.GetVertexName();
+		return InLHS.GetVertexName().FastLess(InRHS.GetVertexName());
 	}
 
 	FVertexInterface::FVertexInterface(const FInputVertexInterface& InInputs, const FOutputVertexInterface& InOutputs)
@@ -284,12 +284,12 @@ namespace Metasound
 		return InputInterface;
 	}
 
-	const FInputDataVertex& FVertexInterface::GetInputVertex(const FVertexKey& InKey) const
+	const FInputDataVertex& FVertexInterface::GetInputVertex(const FVertexName& InKey) const
 	{
 		return InputInterface[InKey];
 	}
 
-	bool FVertexInterface::ContainsInputVertex(const FVertexKey& InKey) const
+	bool FVertexInterface::ContainsInputVertex(const FVertexName& InKey) const
 	{
 		return InputInterface.Contains(InKey);
 	}
@@ -304,14 +304,14 @@ namespace Metasound
 		return OutputInterface;
 	}
 
-	const FOutputDataVertex& FVertexInterface::GetOutputVertex(const FVertexKey& InKey) const
+	const FOutputDataVertex& FVertexInterface::GetOutputVertex(const FVertexName& InName) const
 	{
-		return OutputInterface[InKey];
+		return OutputInterface[InName];
 	}
 
-	bool FVertexInterface::ContainsOutputVertex(const FVertexKey& InKey) const
+	bool FVertexInterface::ContainsOutputVertex(const FVertexName& InName) const
 	{
-		return OutputInterface.Contains(InKey);
+		return OutputInterface.Contains(InName);
 	}
 
 	const FEnvironmentVertexInterface& FVertexInterface::GetEnvironmentInterface() const
@@ -324,12 +324,12 @@ namespace Metasound
 		return EnvironmentInterface;
 	}
 
-	const FEnvironmentVertex& FVertexInterface::GetEnvironmentVertex(const FVertexKey& InKey) const
+	const FEnvironmentVertex& FVertexInterface::GetEnvironmentVertex(const FVertexName& InKey) const
 	{
 		return EnvironmentInterface[InKey];
 	}
 
-	bool FVertexInterface::ContainsEnvironmentVertex(const FVertexKey& InKey) const
+	bool FVertexInterface::ContainsEnvironmentVertex(const FVertexName& InKey) const
 	{
 		return EnvironmentInterface.Contains(InKey);
 	}
diff --git a/Engine/Plugins/Runtime/Metasound/Source/MetasoundGraphCore/Public/MetasoundDataReferenceCollection.h b/Engine/Plugins/Runtime/Metasound/Source/MetasoundGraphCore/Public/MetasoundDataReferenceCollection.h
--- a/Engine/Plugins/Runtime/Metasound/Source/MetasoundGraphCore/Public/MetasoundDataReferenceCollection.h
+++ b/Engine/Plugins/Runtime/Metasound/Source/MetasoundGraphCore/Public/MetasoundDataReferenceCollection.h
@@ -67,7 +67,7 @@ namespace Metasound
 		};
 
 
-		using FDataReferenceMap = TMap<FString, FDataRefWrapper>;
+		using FDataReferenceMap = TMap<FVertexName, FDataRefWrapper>;
 
 		public:
 
@@ -77,7 +77,7 @@ namespace Metasound
 			 * @InDataReference - Readable data reference. 
 			 */
 			template<typename DataType>
-			void AddDataReadReference(const FString& InName, const TDataReadReference<DataType>& InDataReference)
+			void AddDataReadReference(const FVertexName& InName, const TDataReadReference<DataType>& InDataReference)
 			{
 				FDataRefWrapper Wrapper(InDataReference);
 
@@ -90,7 +90,7 @@ namespace Metasound
 			 * @InDataReference - Writable data reference. 
 			 */
 			template<typename DataType>
-			void AddDataReadReference(const FString& InName, const TDataWriteReference<DataType>& InDataReference)
+			void AddDataReadReference(const FVertexName& InName, const TDataWriteReference<DataType>& InDataReference)
 			{
 				TDataReadReference<DataType> Readable(InDataReference);
 
@@ -103,7 +103,7 @@ namespace Metasound
 			 * @InDataReference - Writable data reference. 
 			 */
 			template<typename DataType>
-			void AddDataWriteReference(const FString& InName, const TDataWriteReference<DataType>& InDataReference)
+			void AddDataWriteReference(const FVertexName& InName, const TDataWriteReference<DataType>& InDataReference)
 			{
 				FDataRefWrapper Wrapper(InDataReference);
 
@@ -119,7 +119,7 @@ namespace Metasound
 			 *
 			 * @return True if the data was successfully added. False otherwise. 
 			 */
-			bool AddDataReadReferenceFrom(const FString& InLocalName, const FDataReferenceCollection& InOtherCollection, const FString& InOtherName, const FName& InOtherTypeName);
+			bool AddDataReadReferenceFrom(const FVertexName& InLocalName, const FDataReferenceCollection& InOtherCollection, const FVertexName& InOtherName, const FName& InOtherTypeName);
 
 			/** Add a writable data reference from another collection.
 			 *
@@ -130,7 +130,7 @@ namespace Metasound
 			 *
 			 * @return True if the data was successfully added. False otherwise. 
 			 */
-			bool AddDataWriteReferenceFrom(const FString& InLocalName, const FDataReferenceCollection& InOtherCollection, const FString& InOtherName, const FName& InOtherTypeName);
+			bool AddDataWriteReferenceFrom(const FVertexName& InLocalName, const FDataReferenceCollection& InOtherCollection, const FVertexName& InOtherName, const FName& InOtherTypeName);
 
 			/** Query whether a readable data reference is within the collection. 
 			 *
@@ -139,7 +139,7 @@ namespace Metasound
 			 *
 			 * @return True if the data exists in the collection. False otherwise. 
 			 */
-			bool ContainsDataReadReference(const FString& InName, const FName& InTypeName) const;
+			bool ContainsDataReadReference(const FVertexName& InName, const FName& InTypeName) const;
 
 			/** Query whether a writable data reference is within the collection. 
 			 *
@@ -148,7 +148,7 @@ namespace Metasound
 			 *
 			 * @return True if the data exists in the collection. False otherwise. 
 			 */
-			bool ContainsDataWriteReference(const FString& InName, const FName& InTypeName) const;
+			bool ContainsDataWriteReference(const FVertexName& InName, const FName& InTypeName) const;
 
 
 			/** Query whether a readable data reference is within the collection. 
@@ -158,7 +158,7 @@ namespace Metasound
 			 * @return True if the data exists in the collection. False otherwise. 
 			 */
 			template<typename DataType>
-			bool ContainsDataReadReference(const FString& InName) const
+			bool ContainsDataReadReference(const FVertexName& InName) const
 			{
 				return Contains<DataType>(DataReadRefMap, InName);
 			}
@@ -170,7 +170,7 @@ namespace Metasound
 			 * @return True if the data exists in the collection. False otherwise. 
 			 */
 			template<typename DataType>
-			bool ContainsDataWriteReference(const FString& InName) const
+			bool ContainsDataWriteReference(const FVertexName& InName) const
 			{
 				return Contains<DataType>(DataWriteRefMap, InName);
 			}
@@ -182,7 +182,7 @@ namespace Metasound
 			 * @return A readable data reference.
 			 */
 			template<typename DataType>
-			TDataReadReference<DataType> GetDataReadReference(const FString& InName) const
+			TDataReadReference<DataType> GetDataReadReference(const FVertexName& InName) const
 			{
 				typedef TDataReadReference<DataType> FDataRefType;
 
@@ -204,7 +204,7 @@ namespace Metasound
 			 * @return A readable data reference.
 			 */
 			template<typename DataType, typename... ConstructorArgTypes>
-			TDataReadReference<DataType> GetDataReadReferenceOrConstruct(const FString& InName, ConstructorArgTypes&&... ConstructorArgs) const
+			TDataReadReference<DataType> GetDataReadReferenceOrConstruct(FVertexName InName, ConstructorArgTypes&&... ConstructorArgs) const
 			{
 				typedef TDataReadReference<DataType> FDataRefType;
 
@@ -228,7 +228,7 @@ namespace Metasound
 			 * @return A readable data reference.
 			 */
 			template<typename DataType>
-			TDataReadReference<DataType> GetDataReadReferenceOrConstructWithVertexDefault(const FInputVertexInterface& InputVertices, const FString& InName, const FOperatorSettings& InSettings) const
+			TDataReadReference<DataType> GetDataReadReferenceOrConstructWithVertexDefault(const FInputVertexInterface& InputVertices, FVertexName InName, const FOperatorSettings& InSettings) const
 			{
 				using FDataFactory = TDataReadReferenceLiteralFactory<DataType>;
 
@@ -256,7 +256,7 @@ namespace Metasound
 			 * @return A readable data reference.
 			 */
 			template<typename DataType>
-			TDataWriteReference<DataType> GetDataWriteReference(const FString& InName) const
+			TDataWriteReference<DataType> GetDataWriteReference(const FVertexName& InName) const
 			{
 				typedef TDataWriteReference<DataType> FDataRefType;
 
@@ -278,7 +278,7 @@ namespace Metasound
 			 * @return A readable data reference.
 			 */
 			template<typename DataType, typename... ConstructorArgTypes>
-			TDataWriteReference<DataType> GetDataWriteReferenceOrConstruct(const FString& InName, ConstructorArgTypes&&... ConstructorArgs) const
+			TDataWriteReference<DataType> GetDataWriteReferenceOrConstruct(const FVertexName& InName, ConstructorArgTypes&&... ConstructorArgs) const
 			{
 				typedef TDataWriteReference<DataType> FDataRefType;
 
@@ -295,7 +295,7 @@ namespace Metasound
 		private:
 
 			template<typename DataType>
-			bool Contains(const FDataReferenceMap& InMap, const FString& InName) const
+			bool Contains(const FDataReferenceMap& InMap, const FVertexName& InName) const
 			{
 				if (InMap.Contains(InName))
 				{
@@ -310,9 +310,9 @@ namespace Metasound
 				return false;
 			}
 
-			const IDataReference* GetDataReference(const FDataReferenceMap& InMap, const FString& InName) const;
+			const IDataReference* GetDataReference(const FDataReferenceMap& InMap, const FVertexName& InName) const;
 
-			void AddDataReference(FDataReferenceMap& InMap, const FString& InName, FDataRefWrapper&& InDataReference);
+			void AddDataReference(FDataReferenceMap& InMap, const FVertexName& InName, FDataRefWrapper&& InDataReference);
 
 			FDataReferenceMap DataReadRefMap;
 			FDataReferenceMap DataWriteRefMap;
diff --git a/Engine/Plugins/Runtime/Metasound/Source/MetasoundGraphCore/Public/MetasoundEnvironment.h b/Engine/Plugins/Runtime/Metasound/Source/MetasoundGraphCore/Public/MetasoundEnvironment.h
--- a/Engine/Plugins/Runtime/Metasound/Source/MetasoundGraphCore/Public/MetasoundEnvironment.h
+++ b/Engine/Plugins/Runtime/Metasound/Source/MetasoundGraphCore/Public/MetasoundEnvironment.h
@@ -81,7 +81,7 @@ namespace Metasound
 		virtual ~IMetasoundEnvironmentVariable() = default;
 
 		/** The name of this environment variable. */
-		virtual const FString& GetName() const = 0;
+		virtual const FName& GetName() const = 0;
 
 		/** The type id of this environment variable. */
 		virtual FMetasoundEnvironmentVariableTypeId GetTypeId() const = 0;
@@ -108,7 +108,7 @@ namespace Metasound
 		 * @param ...Args - Parameter pack to be forwarded to underlying types constructor.
 		 */
 		template<typename... ArgTypes>
-		TMetasoundEnvironmentVariable(const FString& InName, ArgTypes&&... Args)
+		TMetasoundEnvironmentVariable(const FName& InName, ArgTypes&&... Args)
 		:	Name(InName)
 		,	Value(Forward<ArgTypes>(Args)...)
 		{
@@ -127,7 +127,7 @@ namespace Metasound
 		}
 
 		/** Get the name of this environment variable. */
-		const FString& GetName() const override
+		const FName& GetName() const override
 		{
 			return Name;
 		}
@@ -145,7 +145,7 @@ namespace Metasound
 		}
 
 	private:
-		FString Name;
+		FName Name;
 		Type Value;
 	};
 
@@ -197,7 +197,7 @@ namespace Metasound
 		 * @tparam VarType - Type of the underlying data stored in the environment variable. 
 		 */
 		template<typename VarType>
-		bool IsType(const FString& InVariableName) const
+		bool IsType(const FName& InVariableName) const
 		{
 			return IsEnvironmentVariableOfType<VarType>(*Variables[InVariableName]);
 		}
@@ -209,7 +209,7 @@ namespace Metasound
 		 * @tparam VarType - Type of the underlying data stored in the environment variable. 
 		 */
 		template<typename VarType>
-		bool Contains(const FString& InVariableName) const
+		bool Contains(const FName& InVariableName) const
 		{
 			if (Variables.Contains(InVariableName))
 			{
@@ -226,7 +226,7 @@ namespace Metasound
 		 * @return copy of the underlying data.
 		 */
 		template<typename VarType>
-		VarType GetValue(const FString& InVariableName) const
+		VarType GetValue(const FName& InVariableName) const
 		{
 			if (ensure(Contains<VarType>(InVariableName)))
 			{
@@ -245,14 +245,13 @@ namespace Metasound
 		 * @return const ref to the underlying data.
 		 */
 		template<typename VarType>
-		void SetValue(const FString& InVariableName, const VarType& InValue)
+		void SetValue(const FName& InVariableName, const VarType& InValue)
 		{
 			Variables.Add(InVariableName, MakeUnique<TMetasoundEnvironmentVariable<VarType>>(InVariableName, InValue));
 		}
 
-
 	private:
-		TSortedMap<FString, TUniquePtr<IMetasoundEnvironmentVariable>> Variables;
+		TSortedMap<FName, TUniquePtr<IMetasoundEnvironmentVariable>, FDefaultAllocator, FNameFastLess> Variables;
 	};
 }
 
diff --git a/Engine/Plugins/Runtime/Metasound/Source/MetasoundGraphCore/Public/MetasoundFacade.h b/Engine/Plugins/Runtime/Metasound/Source/MetasoundGraphCore/Public/MetasoundFacade.h
--- a/Engine/Plugins/Runtime/Metasound/Source/MetasoundGraphCore/Public/MetasoundFacade.h
+++ b/Engine/Plugins/Runtime/Metasound/Source/MetasoundGraphCore/Public/MetasoundFacade.h
@@ -118,7 +118,7 @@ namespace Metasound
 			 * 						  IOperator which is created by this node.
 			 */
 			template<typename OperatorType>
-			FNodeFacade(const FString& InInstanceName, const FGuid& InInstanceID, TFacadeOperatorClass<OperatorType> OperatorClass)
+			FNodeFacade(const FVertexName& InInstanceName, const FGuid& InInstanceID, TFacadeOperatorClass<OperatorType> OperatorClass)
 			:	FNode(InInstanceName, InInstanceID, OperatorType::GetNodeInfo())
 			, 	Factory(MakeShared<FFactory, ESPMode::ThreadSafe>(&OperatorType::CreateOperator))
 			{
diff --git a/Engine/Plugins/Runtime/Metasound/Source/MetasoundGraphCore/Public/MetasoundGraph.h b/Engine/Plugins/Runtime/Metasound/Source/MetasoundGraphCore/Public/MetasoundGraph.h
--- a/Engine/Plugins/Runtime/Metasound/Source/MetasoundGraphCore/Public/MetasoundGraph.h
+++ b/Engine/Plugins/Runtime/Metasound/Source/MetasoundGraphCore/Public/MetasoundGraph.h
@@ -21,7 +21,7 @@ namespace Metasound
 			virtual ~FGraph() = default;
 
 			/** Return the name of this specific instance of the node class. */
-			const FString& GetInstanceName() const override;
+			const FVertexName& GetInstanceName() const override;
 
 			/** Return the ID of this specific instance of the node class. */
 			const FGuid& GetInstanceID() const override;
@@ -70,13 +70,13 @@ namespace Metasound
 			 *
 			 * @return True if the edge was successfully added. False otherwise.
 			 */
-			bool AddDataEdge(const INode& FromNode, const FVertexKey& FromVertexKey, const INode& ToNode, const FVertexKey& ToVertexKey);
+			bool AddDataEdge(const INode& FromNode, const FVertexName& FromVertexKey, const INode& ToNode, const FVertexName& ToVertexKey);
 
 			/** Remove the given data edge. 
 			 *
 			 * @return True on success, false on failure.
 			 */
-			bool RemoveDataEdge(const INode& FromNode, const FVertexKey& FromVertexKey, const INode& ToNode, const FVertexKey& ToVertexKey);
+			bool RemoveDataEdge(const INode& FromNode, const FVertexName& FromVertexKey, const INode& ToNode, const FVertexName& ToVertexKey);
 
 			/** Removes all edges for which that predicate returns true.
 			 *
@@ -93,12 +93,12 @@ namespace Metasound
 			 * outside this graph should be routed internally.
 			 *
 			 * @param InNode - Node which receives the data.
-			 * @param InVertexKey - Key for input vertex on InNode.
+			 * @param InVertexName - Key for input vertex on InNode.
 			 *
 			 * @return True if the destination was successfully added. False 
 			 * otherwise.
 			 */
-			bool AddInputDataDestination(const INode& InNode, const FVertexKey& InVertexKey);
+			bool AddInputDataDestination(const INode& InNode, const FVertexName& InVertexName);
 
 
 			/** Add an input data destination to describe how data provided 
@@ -110,12 +110,12 @@ namespace Metasound
 			 * owned this graph and exposed externally.
 			 *
 			 * @param InNode - Node which produces the data.
-			 * @param InVertexKey - Key for output vertex on InNode.
+			 * @param InVertexName - Key for output vertex on InNode.
 			 *
 			 * @return True if the source was successfully added. False 
 			 * otherwise.
 			 */
-			bool AddOutputDataSource(const INode& InNode, const FVertexKey& InVertexKey);
+			bool AddOutputDataSource(const INode& InNode, const FVertexName& InVertexName);
 
 			/** Add an output data source which describes routing of data which is 
 			 * owned this graph and exposed externally.
@@ -136,7 +136,7 @@ namespace Metasound
 				virtual TUniquePtr<IOperator> CreateOperator(const FCreateOperatorParams& InParams, FBuildErrorArray& OutErrors) override;
 			};
 
-			FString InstanceName;
+			FVertexName InstanceName;
 			FGuid InstanceID;
 
 			FNodeClassMetadata Metadata;
diff --git a/Engine/Plugins/Runtime/Metasound/Source/MetasoundGraphCore/Public/MetasoundLiteralNode.h b/Engine/Plugins/Runtime/Metasound/Source/MetasoundGraphCore/Public/MetasoundLiteralNode.h
--- a/Engine/Plugins/Runtime/Metasound/Source/MetasoundGraphCore/Public/MetasoundLiteralNode.h
+++ b/Engine/Plugins/Runtime/Metasound/Source/MetasoundGraphCore/Public/MetasoundLiteralNode.h
@@ -127,8 +127,8 @@ namespace Metasound
 
 		/* Construct a TLiteralNode using the TLiteralOperatorLiteralFactory<> and moving
 		 * InParam to the TLiteralOperatorLiteralFactory constructor.*/
-		explicit TLiteralNode(const FString& InNodeDescription, const FGuid& InInstanceID, FLiteral&& InParam)
-		:	FNode(InNodeDescription, InInstanceID, GetNodeInfo())
+		explicit TLiteralNode(const FVertexName& InInstanceName, const FGuid& InInstanceID, FLiteral&& InParam)
+		:	FNode(InInstanceName, InInstanceID, GetNodeInfo())
 		,	Interface(DeclareVertexInterface())
 		, 	Factory(MakeOperatorFactoryRef<TLiteralOperatorLiteralFactory<DataType>>(MoveTemp(InParam)))
 		{
diff --git a/Engine/Plugins/Runtime/Metasound/Source/MetasoundGraphCore/Public/MetasoundNode.h b/Engine/Plugins/Runtime/Metasound/Source/MetasoundGraphCore/Public/MetasoundNode.h
--- a/Engine/Plugins/Runtime/Metasound/Source/MetasoundGraphCore/Public/MetasoundNode.h
+++ b/Engine/Plugins/Runtime/Metasound/Source/MetasoundGraphCore/Public/MetasoundNode.h
@@ -1,21 +1,22 @@
 // Copyright Epic Games, Inc. All Rights Reserved.
-
 #pragma once
 
 #include "CoreMinimal.h"
 #include "MetasoundNodeInterface.h"
+#include "MetasoundVertex.h"
+
 
 namespace Metasound
 {
 	class METASOUNDGRAPHCORE_API FNode : public INode
 	{
 		public:
-			FNode(const FString& InInstanceName, const FGuid& InInstanceID, const FNodeClassMetadata& InInfo);
+			FNode(const FVertexName& InInstanceName, const FGuid& InInstanceID, const FNodeClassMetadata& InInfo);
 
 			virtual ~FNode() = default;
 
 			/** Return the name of this specific instance of the node class. */
-			const FString& GetInstanceName() const override;
+			const FVertexName& GetInstanceName() const override;
 
 			/** Return the ID of this specific instance of the node class. */
 			const FGuid& GetInstanceID() const override;
@@ -25,7 +26,7 @@ namespace Metasound
 
 		private:
 
-			FString InstanceName;
+			FVertexName InstanceName;
 			FGuid InstanceID;
 			FNodeClassMetadata Info;
 	};
diff --git a/Engine/Plugins/Runtime/Metasound/Source/MetasoundGraphCore/Public/MetasoundNodeInterface.h b/Engine/Plugins/Runtime/Metasound/Source/MetasoundGraphCore/Public/MetasoundNodeInterface.h
--- a/Engine/Plugins/Runtime/Metasound/Source/MetasoundGraphCore/Public/MetasoundNodeInterface.h
+++ b/Engine/Plugins/Runtime/Metasound/Source/MetasoundGraphCore/Public/MetasoundNodeInterface.h
@@ -130,10 +130,10 @@ namespace Metasound
 	class INodeBase
 	{
 		public:
-			virtual ~INodeBase() {}
+			virtual ~INodeBase() = default;
 
 			/** Return the name of this specific instance of the node class. */
-			virtual const FString& GetInstanceName() const = 0;
+			virtual const FName& GetInstanceName() const = 0;
 
 			/** Return the ID of this node instance. */
 			virtual const FGuid& GetInstanceID() const = 0;
@@ -250,7 +250,7 @@ namespace Metasound
 	};
 
 	/** Key type for an FOutputDataSource or FInputDataDestination. */
-	typedef TTuple<const INode*, FVertexKey> FNodeDataVertexKey;
+	typedef TTuple<const INode*, FVertexName> FNodeDataVertexKey;
 
 	/** FOutputDataSourceCollection contains multiple FOutputDataSources mapped 
 	 * by FNodeDataVertexKeys.
diff --git a/Engine/Plugins/Runtime/Metasound/Source/MetasoundGraphCore/Public/MetasoundVariableNodes.h b/Engine/Plugins/Runtime/Metasound/Source/MetasoundGraphCore/Public/MetasoundVariableNodes.h
--- a/Engine/Plugins/Runtime/Metasound/Source/MetasoundGraphCore/Public/MetasoundVariableNodes.h
+++ b/Engine/Plugins/Runtime/Metasound/Source/MetasoundGraphCore/Public/MetasoundVariableNodes.h
@@ -1,5 +1,4 @@
 // Copyright Epic Games, Inc. All Rights Reserved.
-
 #pragma once
 
 #include "CoreMinimal.h"
@@ -11,24 +10,26 @@
 #include "MetasoundNode.h"
 #include "MetasoundNodeInterface.h"
 #include "MetasoundVariable.h"
+#include "MetasoundVertex.h"
 
 #define LOCTEXT_NAMESPACE "MetasoundGraphCore"
 
+
 namespace Metasound
 {
 	namespace VariableNodeVertexNames
 	{
 		/** Input vertex name for data */
-		METASOUNDGRAPHCORE_API const FString& GetInputDataName();
+		METASOUNDGRAPHCORE_API const FVertexName& GetInputDataName();
 
 		/** Output vertex name for data */
-		METASOUNDGRAPHCORE_API const FString& GetOutputDataName();
+		METASOUNDGRAPHCORE_API const FVertexName& GetOutputDataName();
 
 		/** Input vertex name for variables */
-		METASOUNDGRAPHCORE_API const FString& GetInputVariableName();
+		METASOUNDGRAPHCORE_API const FVertexName& GetInputVariableName();
 		
 		/** Output vertex name for variables */
-		METASOUNDGRAPHCORE_API const FString& GetOutputVariableName();
+		METASOUNDGRAPHCORE_API const FVertexName& GetOutputVariableName();
 	}
 
 	/** Init Variable nodes initialize variable values. The output of a InitVariableNode
@@ -125,8 +126,8 @@ namespace Metasound
 
 	public:
 
-		TInitVariableNode(const FString& InNodeDescription, const FGuid& InInstanceID, FLiteral&& InLiteral)
-		: FNode(InNodeDescription, InInstanceID, GetNodeMetadata())
+		TInitVariableNode(const FVertexName& InNodeName, const FGuid& InInstanceID, FLiteral&& InLiteral)
+		: FNode(InNodeName, InInstanceID, GetNodeMetadata())
 		, Interface(DeclareVertexInterface())
 		, Factory(MakeShared<FFactory>(MoveTemp(InLiteral)))
 		{
@@ -193,8 +194,8 @@ namespace Metasound
 		};
 
 	public:
-		TSetVariableNode(const FString& InNodeDescription, const FGuid& InInstanceID)
-		: FNodeFacade(InNodeDescription, InInstanceID, TFacadeOperatorClass<FOperator>())
+		TSetVariableNode(const FVertexName& InNodeName, const FGuid& InInstanceID)
+		: FNodeFacade(InNodeName, InInstanceID, TFacadeOperatorClass<FOperator>())
 		{
 		}
 
@@ -316,8 +317,8 @@ namespace Metasound
 
 	public:
 
-		TGetDelayedVariableNode(const FString& InNodeDescription, const FGuid& InInstanceID)
-		: FNodeFacade(InNodeDescription, InInstanceID, TFacadeOperatorClass<FOperator>())
+		TGetDelayedVariableNode(const FVertexName& InNodeName, const FGuid& InInstanceID)
+		: FNodeFacade(InNodeName, InInstanceID, TFacadeOperatorClass<FOperator>())
 		{
 		}
 
@@ -417,8 +418,8 @@ namespace Metasound
 		};
 	public:
 
-		TGetVariableNode(const FString& InNodeDescription, const FGuid& InInstanceID)
-		: FNodeFacade(InNodeDescription, InInstanceID, TFacadeOperatorClass<FOperator>())
+		TGetVariableNode(const FVertexName& InNodeName, const FGuid& InInstanceID)
+		: FNodeFacade(InNodeName, InInstanceID, TFacadeOperatorClass<FOperator>())
 		{
 		}
 
diff --git a/Engine/Plugins/Runtime/Metasound/Source/MetasoundGraphCore/Public/MetasoundVertex.h b/Engine/Plugins/Runtime/Metasound/Source/MetasoundGraphCore/Public/MetasoundVertex.h
--- a/Engine/Plugins/Runtime/Metasound/Source/MetasoundGraphCore/Public/MetasoundVertex.h
+++ b/Engine/Plugins/Runtime/Metasound/Source/MetasoundGraphCore/Public/MetasoundVertex.h
@@ -13,10 +13,9 @@
 
 namespace Metasound
 {
-	/** Key type for an FInputDataVertexColletion or 
-	 * FOutputDataVertexCollection.
+	/** Name of a given vertex.  Only unique for a given node interface.
 	 */
-	using FVertexKey = FString;
+	using FVertexName = FName;
 
 	// Forward declarations
 	class FInputDataVertex;
@@ -51,7 +50,7 @@ namespace Metasound
 		 * @InDataTypeName - Name of data type.
 		 * @InMetadata - Metadata pertaining to given vertex.
 		 */
-		FDataVertexModel(const FString& InVertexName, const FName& InDataTypeName, const FDataVertexMetadata& InMetadata)
+		FDataVertexModel(const FVertexName& InVertexName, const FName& InDataTypeName, const FDataVertexMetadata& InMetadata)
 		:	VertexName(InVertexName)
 		,	DataTypeName(InDataTypeName)
 		,	VertexMetadata(InMetadata)
@@ -61,7 +60,7 @@ namespace Metasound
 		virtual ~FDataVertexModel() = default;
 
 		/** Name of vertex. */
-		const FString VertexName;
+		const FVertexName VertexName;
 
 		/** Type name of data. */
 		const FName DataTypeName;
@@ -125,27 +124,27 @@ namespace Metasound
 
 	public:
 
-		FInputDataVertexModel(const FString& InVertexName, const FName& InDataTypeName, const FText& InDescription)
+		FInputDataVertexModel(const FVertexName& InVertexName, const FName& InDataTypeName, const FText& InDescription)
 			: FDataVertexModel(InVertexName, InDataTypeName, { InDescription })
 			, LiteralFactory(MakeUnique<TLiteralFactory<FLiteral::FNone>>(FLiteral::FNone{}))
 		{
 		}
 
-		FInputDataVertexModel(const FString& InVertexName, const FName& InDataTypeName, const FDataVertexMetadata& InMetadata)
+		FInputDataVertexModel(const FVertexName& InVertexName, const FName& InDataTypeName, const FDataVertexMetadata& InMetadata)
 			: FDataVertexModel(InVertexName, InDataTypeName, InMetadata)
 			, LiteralFactory(MakeUnique<TLiteralFactory<FLiteral::FNone>>(FLiteral::FNone{}))
 		{
 		}
 
 		template<typename LiteralValueType>
-		FInputDataVertexModel(const FString& InVertexName, const FName& InDataTypeName, const FText& InDescription, const LiteralValueType& InLiteralValue)
+		FInputDataVertexModel(const FVertexName& InVertexName, const FName& InDataTypeName, const FText& InDescription, const LiteralValueType& InLiteralValue)
 			: FDataVertexModel(InVertexName, InDataTypeName, { InDescription })
 			, LiteralFactory(MakeUnique<TLiteralFactory<LiteralValueType>>(InLiteralValue))
 		{
 		}
 
 		template<typename LiteralValueType>
-		FInputDataVertexModel(const FString& InVertexName, const FName& InDataTypeName, const FDataVertexMetadata& InMetadata, const LiteralValueType& InLiteralValue)
+		FInputDataVertexModel(const FVertexName& InVertexName, const FName& InDataTypeName, const FDataVertexMetadata& InMetadata, const LiteralValueType& InLiteralValue)
 			: FDataVertexModel(InVertexName, InDataTypeName, InMetadata)
 			, LiteralFactory(MakeUnique<TLiteralFactory<LiteralValueType>>(InLiteralValue))
 		{
@@ -224,7 +223,7 @@ namespace Metasound
 		 * @InDescription - Human readable vertex description.
 		 */
 		template<typename... ModelArgTypes>
-		TBaseVertexModel(const FString& InVertexName, const FText& InDescription, ModelArgTypes&&... ModelArgs)
+		TBaseVertexModel(const FVertexName& InVertexName, const FText& InDescription, ModelArgTypes&&... ModelArgs)
 		:	VertexModelType(InVertexName, GetMetasoundDataTypeName<DataType>(), { InDescription }, Forward<ModelArgTypes>(ModelArgs)...)
 		{
 		}
@@ -235,7 +234,7 @@ namespace Metasound
 		 * @InVertexMetadata - Vertex metadata, used primarily for debugging or live edit.
 		 */
 		template<typename... ModelArgTypes>
-		TBaseVertexModel(const FString& InVertexName, const FDataVertexMetadata& InMetadata, ModelArgTypes&&... ModelArgs)
+		TBaseVertexModel(const FVertexName& InVertexName, const FDataVertexMetadata& InMetadata, ModelArgTypes&&... ModelArgs)
 		:	VertexModelType(InVertexName, GetMetasoundDataTypeName<DataType>(), InMetadata, Forward<ModelArgTypes>(ModelArgs)...)
 		{
 		}
@@ -281,7 +280,7 @@ namespace Metasound
 		 * @InVertexName - Name of vertex.
 		 * @InDescription - Human readable vertex description.
 		 */
-		TInputDataVertexModel(const FString& InVertexName, const FText& InDescription)
+		TInputDataVertexModel(const FVertexName& InVertexName, const FText& InDescription)
 		:	TBaseVertexModel<DataType, FInputDataVertexModel>(InVertexName, InDescription)
 		{
 		}
@@ -293,7 +292,7 @@ namespace Metasound
 		 * @InDefaultValue - Default Value of vertex
 		 */
 		template<typename LiteralValueType>
-		TInputDataVertexModel(const FString& InVertexName, const FText& InDescription, const LiteralValueType& InDefaultValue)
+		TInputDataVertexModel(const FVertexName& InVertexName, const FText& InDescription, const LiteralValueType& InDefaultValue)
 		:	TBaseVertexModel<DataType, FInputDataVertexModel>(InVertexName, InDescription, InDefaultValue)
 		{
 		}
@@ -305,7 +304,7 @@ namespace Metasound
 		 * @InDefaultValue - Default Value of vertex
 		 */
 		template<typename LiteralValueType>
-		TInputDataVertexModel(const FString& InVertexName, const FDataVertexMetadata& InMetadata)
+		TInputDataVertexModel(const FVertexName& InVertexName, const FDataVertexMetadata& InMetadata)
 		:	TBaseVertexModel<DataType, FInputDataVertexModel>(InVertexName, InMetadata)
 		{
 		}
@@ -367,7 +366,7 @@ namespace Metasound
 			FInputDataVertex& operator=(const FInputDataVertex& InOther);
 
 			/** Name of vertex. */
-			const FString& GetVertexName() const;
+			const FVertexName& GetVertexName() const;
 
 			/** Type name of data reference. */
 			const FName& GetDataTypeName() const;
@@ -461,7 +460,7 @@ namespace Metasound
 			FOutputDataVertex& operator=(const FOutputDataVertex& InOther);
 
 			/** Name of vertex. */
-			const FString& GetVertexName() const;
+			FVertexName GetVertexName() const;
 
 			/** Type name of data reference. */
 			const FName& GetDataTypeName() const;
@@ -508,7 +507,7 @@ namespace Metasound
 		 * @InVertexName - Name of vertex.
 		 * @InDescription - Human readable vertex description.
 		 */
-		FEnvironmentVertexModel(const FString& InVertexName, const FText& InDescription)
+		FEnvironmentVertexModel(FVertexName InVertexName, const FText& InDescription)
 		:	VertexName(InVertexName)
 		,	Description(InDescription)
 		{
@@ -517,7 +516,7 @@ namespace Metasound
 		virtual ~FEnvironmentVertexModel() = default;
 
 		/** Name of vertex. */
-		const FString VertexName;
+		const FVertexName VertexName;
 
 		/** Description of the vertex. */
 		const FText Description;
@@ -622,7 +621,7 @@ namespace Metasound
 			FEnvironmentVertex& operator=(const FEnvironmentVertex& InOther);
 
 			/** Name of vertex. */
-			const FString& GetVertexName() const;
+			const FVertexName& GetVertexName() const;
 
 			/** Description of the vertex. */
 			const FText& GetDescription() const;
@@ -644,27 +643,8 @@ namespace Metasound
 			TUniquePtr<FEnvironmentVertexModel> VertexModel;
 	};
 
-
-	/** Create an FVertexKey from an FOutDataVertex. */
-	FORCEINLINE FVertexKey MakeDataVertexKey(const FInputDataVertex& InVertex)
-	{
-		return InVertex.GetVertexName();
-	}
-
-	/** Create an FVertexKey from an FOutputDataVertex. */
-	FORCEINLINE FVertexKey MakeDataVertexKey(const FOutputDataVertex& InVertex)
-	{
-		return InVertex.GetVertexName();
-	}
-
-	/** Create an FVertexKey from an FEnvironmentVertex. */
-	FORCEINLINE FVertexKey MakeEnvironmentVertexKey(const FEnvironmentVertex& InVertex)
-	{
-		return InVertex.GetVertexName();
-	}
-
 	/** TVertexInterfaceGroups encapsulates multiple related data vertices. It 
-	 * requires that each vertex in the group have a unique FVertexKey.
+	 * requires that each vertex in the group have a unique FVertexName.
 	 */
 	template<typename VertexType>
 	class TVertexInterfaceGroup
@@ -676,8 +656,8 @@ namespace Metasound
 
 			static_assert(bIsSupportedVertexType, "VertexType must be derived from FInputDataVertex, FOutputDataVertex, or FEnvironmentVertex");
 
-			using FContainerType = TSortedMap<FVertexKey, VertexType>;
-			using FOrderContainerType = TArray<FVertexKey>;
+			using FContainerType = TSortedMap<FVertexName, VertexType, FDefaultAllocator, FNameFastLess>;
+			using FOrderContainerType = TArray<FVertexName>;
 
 			// Required for end of recursion.
 			static void CopyInputs(FContainerType& InStorage, FOrderContainerType& InOrder)
@@ -696,9 +676,8 @@ namespace Metasound
 				// Create vertex out of vertex model
 				VertexType Vertex(InInput);
 
-				const FVertexKey VertexKey = MakeDataVertexKey(Vertex);
-				InStorage.Add(VertexKey, Vertex);
-				InKeyOrder.Add(VertexKey);
+				InStorage.Add(Vertex.GetVertexName(), Vertex);
+				InKeyOrder.Add(Vertex.GetVertexName());
 				CopyInputs(InStorage, InKeyOrder, InRemainingInputs...);
 			}
 
@@ -726,13 +705,12 @@ namespace Metasound
 			/** Add a vertex to the group. */
 			void Add(const VertexType& InVertex)
 			{
-				const FVertexKey VertexKey = MakeDataVertexKey(InVertex);
-				Vertices.Add(VertexKey, InVertex);
-				OrderedKeys.Add(VertexKey);
+				Vertices.Add(InVertex.GetVertexName(), InVertex);
+				OrderedKeys.Add(InVertex.GetVertexName());
 			}
 
 			/** Remove a vertex by key. */
-			bool Remove(const FVertexKey& InKey)
+			bool Remove(const FVertexName& InKey)
 			{
 				int32 NumRemoved = Vertices.Remove(InKey);
 				OrderedKeys.Remove(InKey);
@@ -740,12 +718,17 @@ namespace Metasound
 			}
 
 			/** Returns true if the group contains a vertex with a matching key. */
-			bool Contains(const FVertexKey& InKey) const
+			bool Contains(const FVertexName& InKey) const
 			{
 				return Vertices.Contains(InKey);
 			}
 
-			int32 GetOrderIndex(const FVertexKey& InKey) const
+			const VertexType* Find(const FVertexName& InKey) const
+			{
+				return Vertices.Find(InKey);
+			}
+
+			int32 GetOrderIndex(const FVertexName& InKey) const
 			{
 				int32 OutIndex = INDEX_NONE;
 				OrderedKeys.Find(InKey, OutIndex);
@@ -753,9 +736,9 @@ namespace Metasound
 			}
 
 			/** Return the vertex for a given vertex key. */
-			const VertexType& operator[](const FVertexKey& InKey) const
+			const VertexType& operator[](const FVertexName& InName) const
 			{
-				return Vertices[InKey];
+				return Vertices[InName];
 			}
 
 			/** Iterator for ranged for loops. */
@@ -827,10 +810,10 @@ namespace Metasound
 			FInputVertexInterface& GetInputInterface();
 
 			/** Return an input vertex. */
-			const FInputDataVertex& GetInputVertex(const FVertexKey& InKey) const;
+			const FInputDataVertex& GetInputVertex(const FVertexName& InKey) const;
 
 			/** Returns true if an input vertex with the given key exists. */
-			bool ContainsInputVertex(const FVertexKey& InKey) const;
+			bool ContainsInputVertex(const FVertexName& InKey) const;
 
 			/** Return the output interface. */
 			const FOutputVertexInterface& GetOutputInterface() const;
@@ -839,10 +822,10 @@ namespace Metasound
 			FOutputVertexInterface& GetOutputInterface();
 
 			/** Return an output vertex. */
-			const FOutputDataVertex& GetOutputVertex(const FVertexKey& InKey) const;
+			const FOutputDataVertex& GetOutputVertex(const FVertexName& InName) const;
 
-			/** Returns true if an output vertex with the given key exists. */
-			bool ContainsOutputVertex(const FVertexKey& InKey) const;
+			/** Returns true if an output vertex with the given name exists. */
+			bool ContainsOutputVertex(const FVertexName& InName) const;
 
 			/** Return the output interface. */
 			const FEnvironmentVertexInterface& GetEnvironmentInterface() const;
@@ -851,10 +834,10 @@ namespace Metasound
 			FEnvironmentVertexInterface& GetEnvironmentInterface();
 
 			/** Return an output vertex. */
-			const FEnvironmentVertex& GetEnvironmentVertex(const FVertexKey& InKey) const;
+			const FEnvironmentVertex& GetEnvironmentVertex(const FVertexName& InKey) const;
 
 			/** Returns true if an output vertex with the given key exists. */
-			bool ContainsEnvironmentVertex(const FVertexKey& InKey) const;
+			bool ContainsEnvironmentVertex(const FVertexName& InKey) const;
 
 			/** Test for equality between two interfaces. */
 			friend bool METASOUNDGRAPHCORE_API operator==(const FVertexInterface& InLHS, const FVertexInterface& InRHS);
@@ -875,7 +858,7 @@ namespace Metasound
 	 */
 	struct FNodeInitData
 	{
-		FString InstanceName;
+		FVertexName InstanceName;
 		FGuid InstanceID;
 	};
 }
