Move FString/FText VertexName identifers to FNames. Convention is as follows:

diff --git a/Engine/Plugins/Runtime/Metasound/Source/MetasoundStandardNodes/Private/MetasoundAudioFormats.cpp b/Engine/Plugins/Runtime/Metasound/Source/MetasoundStandardNodes/Private/MetasoundAudioFormats.cpp
--- a/Engine/Plugins/Runtime/Metasound/Source/MetasoundStandardNodes/Private/MetasoundAudioFormats.cpp
+++ b/Engine/Plugins/Runtime/Metasound/Source/MetasoundStandardNodes/Private/MetasoundAudioFormats.cpp
@@ -9,6 +9,7 @@
 #include "MetasoundDataTypeRegistrationMacro.h"
 #include "MetasoundLiteral.h"
 #include "MetasoundOutputNode.h"
+#include "MetasoundVertex.h"
 
 #define LOCTEXT_NAMESPACE "MetasoundStandardNodes_AudioFormats"
 
@@ -87,8 +88,8 @@ namespace Metasound
 	// Special vertex keys for stereo input/output nodes.
 	namespace StereoAudioFormatVertexKeys
 	{
-		static const FVertexKey LeftChannelVertexKey = "Left";
-		static const FVertexKey RightChannelVertexKey = "Right";
+		static const FVertexName LeftChannelVertexKey = "Left";
+		static const FVertexName RightChannelVertexKey = "Right";
 	}
 
 	// Specialization of TOutputNode<FStereoAudio> to support direct connection
@@ -100,7 +101,7 @@ namespace Metasound
 		class FOutputOperator : public IOperator
 		{
 		public:
-			FOutputOperator(const FVertexKey& InOutputName, TDataReadReference<FAudioBuffer> InLeft, TDataReadReference<FAudioBuffer> InRight, TDataReadReference<FStereoAudioFormat> InStereo)
+			FOutputOperator(const FVertexName& InOutputName, TDataReadReference<FAudioBuffer> InLeft, TDataReadReference<FAudioBuffer> InRight, TDataReadReference<FStereoAudioFormat> InStereo)
 			: OutputName(InOutputName)
 			, Left(InLeft)
 			, Right(InRight)
@@ -136,7 +137,7 @@ namespace Metasound
 
 		private:
 
-			FVertexKey OutputName;
+			FVertexName OutputName;
 			TDataReadReference<FAudioBuffer> Left;
 			TDataReadReference<FAudioBuffer> Right;
 			TDataReadReference<FStereoAudioFormat> Stereo;
@@ -145,7 +146,7 @@ namespace Metasound
 		class FOutputOperatorFactory : public IOperatorFactory
 		{
 		public:
-			FOutputOperatorFactory(const FVertexKey& InOutputName)
+			FOutputOperatorFactory(const FVertexName& InOutputName)
 			: OutputName(InOutputName)
 			{
 			}
@@ -162,10 +163,10 @@ namespace Metasound
 			}
 
 		private:
-			FVertexKey OutputName;
+			FVertexName OutputName;
 		};
 
-		static FVertexInterface GetVertexInterface(const FVertexKey& InVertexName)
+		static FVertexInterface GetVertexInterface(const FVertexName& InVertexName)
 		{
 			return FVertexInterface(
 				FInputVertexInterface(
@@ -178,7 +179,7 @@ namespace Metasound
 			);
 		}
 
-		static FNodeClassMetadata GetNodeClassMetadata(const FVertexKey& InOutputName)
+		static FNodeClassMetadata GetNodeClassMetadata(const FVertexName& InOutputName)
 		{
 			FNodeClassMetadata Info;
 
@@ -196,7 +197,7 @@ namespace Metasound
 
 
 	public:
-		TOutputNode(const FString& InInstanceName, const FGuid& InInstanceID, const FVertexKey& InVertexName)
+		TOutputNode(const FVertexName& InInstanceName, const FGuid& InInstanceID, const FVertexName& InVertexName)
 		:	FNode(InInstanceName, InInstanceID, GetNodeClassMetadata(InVertexName))
 		,	VertexInterface(GetVertexInterface(InVertexName))
 		,	Factory(MakeShared<FOutputOperatorFactory, ESPMode::ThreadSafe>(InVertexName))
@@ -239,7 +240,7 @@ namespace Metasound
 		{
 		public:
 
-			FInputOperator(const FVertexKey& InInputName, TDataReadReference<FAudioBuffer> InLeft, TDataReadReference<FAudioBuffer> InRight, TDataReadReference<FStereoAudioFormat> InStereo)
+			FInputOperator(const FVertexName& InInputName, TDataReadReference<FAudioBuffer> InLeft, TDataReadReference<FAudioBuffer> InRight, TDataReadReference<FStereoAudioFormat> InStereo)
 			: InputName(InInputName)
 			, Left(InLeft)
 			, Right(InRight)
@@ -276,7 +277,7 @@ namespace Metasound
 
 		private:
 
-			FVertexKey InputName;
+			FVertexName InputName;
 			TDataReadReference<FAudioBuffer> Left;
 			TDataReadReference<FAudioBuffer> Right;
 			TDataReadReference<FStereoAudioFormat> Stereo;
@@ -285,7 +286,7 @@ namespace Metasound
 		class FInputOperatorFactory : public IOperatorFactory
 		{
 		public:
-			FInputOperatorFactory(const FVertexKey& InInputName)
+			FInputOperatorFactory(const FVertexName& InInputName)
 			: InputName(InInputName)
 			{
 			}
@@ -301,10 +302,10 @@ namespace Metasound
 			}
 
 		private:
-			FVertexKey InputName;
+			FVertexName InputName;
 		};
 
-		static FVertexInterface GetVertexInterface(const FVertexKey& InVertexName)
+		static FVertexInterface GetVertexInterface(const FVertexName& InVertexName)
 		{
 			return FVertexInterface(
 				FInputVertexInterface(
@@ -317,7 +318,7 @@ namespace Metasound
 			);
 		}
 
-		static FNodeClassMetadata GetNodeClassMetadata(const FVertexKey& InInputName)
+		static FNodeClassMetadata GetNodeClassMetadata(const FVertexName& InInputName)
 		{
 			FNodeClassMetadata Info;
 
@@ -337,7 +338,7 @@ namespace Metasound
 	public:
 		static constexpr bool bCanRegister = true;
 
-		TInputNode(const FString& InInstanceName, const FGuid& InInstanceID, const FVertexKey& InVertexName, FLiteral&& InLiteral)
+		TInputNode(const FVertexName& InInstanceName, const FGuid& InInstanceID, const FVertexName& InVertexName, FLiteral&& InLiteral)
 		:	FNode(InInstanceName, InInstanceID, GetNodeClassMetadata(InVertexName))
 		,	VertexInterface(GetVertexInterface(InVertexName))
 		,	Factory(MakeShared<FInputOperatorFactory, ESPMode::ThreadSafe>(InVertexName))
@@ -379,7 +380,7 @@ namespace Metasound
 		class FOutputOperator : public IOperator
 		{
 		public:
-			FOutputOperator(const FVertexKey& InOutputName, TDataReadReference<FAudioBuffer> InCenter, TDataReadReference<FMonoAudioFormat> InMono)
+			FOutputOperator(const FVertexName& InOutputName, TDataReadReference<FAudioBuffer> InCenter, TDataReadReference<FMonoAudioFormat> InMono)
 			: OutputName(InOutputName)
 			, Center(InCenter)
 			, Mono(InMono)
@@ -413,7 +414,7 @@ namespace Metasound
 
 		private:
 
-			FVertexKey OutputName;
+			FVertexName OutputName;
 			TDataReadReference<FAudioBuffer> Center;
 			TDataReadReference<FMonoAudioFormat> Mono;
 		};
@@ -421,7 +422,7 @@ namespace Metasound
 		class FOutputOperatorFactory : public IOperatorFactory
 		{
 		public:
-			FOutputOperatorFactory(const FVertexKey& InOutputName)
+			FOutputOperatorFactory(const FVertexName& InOutputName)
 			: OutputName(InOutputName)
 			{
 			}
@@ -436,10 +437,10 @@ namespace Metasound
 			}
 
 		private:
-			FVertexKey OutputName;
+			FVertexName OutputName;
 		};
 
-		static FVertexInterface GetVertexInterface(const FVertexKey& InVertexName)
+		static FVertexInterface GetVertexInterface(const FVertexName& InVertexName)
 		{
 			return FVertexInterface(
 				FInputVertexInterface(
@@ -451,7 +452,7 @@ namespace Metasound
 			);
 		}
 
-		static FNodeClassMetadata GetNodeClassMetadata(const FVertexKey& InOutputName)
+		static FNodeClassMetadata GetNodeClassMetadata(const FVertexName& InOutputName)
 		{
 			FNodeClassMetadata Info;
 
@@ -469,7 +470,7 @@ namespace Metasound
 
 
 	public:
-		TOutputNode(const FString& InInstanceName, const FGuid& InInstanceID, const FVertexKey& InVertexName)
+		TOutputNode(const FVertexName& InInstanceName, const FGuid& InInstanceID, const FVertexName& InVertexName)
 		:	FNode(InInstanceName, InInstanceID, GetNodeClassMetadata(InVertexName))
 		,	VertexInterface(GetVertexInterface(InVertexName))
 		,	Factory(MakeShared<FOutputOperatorFactory, ESPMode::ThreadSafe>(InVertexName))
@@ -512,7 +513,7 @@ namespace Metasound
 		{
 		public:
 
-			FInputOperator(const FVertexKey& InInputName, TDataReadReference<FAudioBuffer> InCenter, TDataReadReference<FMonoAudioFormat> InMono)
+			FInputOperator(const FVertexName& InInputName, TDataReadReference<FAudioBuffer> InCenter, TDataReadReference<FMonoAudioFormat> InMono)
 			: InputName(InInputName)
 			, Center(InCenter)
 			, Mono(InMono)
@@ -547,7 +548,7 @@ namespace Metasound
 
 		private:
 
-			FVertexKey InputName;
+			FVertexName InputName;
 			TDataReadReference<FAudioBuffer> Center;
 			TDataReadReference<FMonoAudioFormat> Mono;
 		};
@@ -555,7 +556,7 @@ namespace Metasound
 		class FInputOperatorFactory : public IOperatorFactory
 		{
 		public:
-			FInputOperatorFactory(const FVertexKey& InInputName)
+			FInputOperatorFactory(const FVertexName& InInputName)
 			: InputName(InInputName)
 			{
 			}
@@ -569,10 +570,10 @@ namespace Metasound
 			}
 
 		private:
-			FVertexKey InputName;
+			FVertexName InputName;
 		};
 
-		static FVertexInterface GetVertexInterface(const FVertexKey& InVertexName)
+		static FVertexInterface GetVertexInterface(const FVertexName& InVertexName)
 		{
 			return FVertexInterface(
 				FInputVertexInterface(
@@ -584,7 +585,7 @@ namespace Metasound
 			);
 		}
 
-		static FNodeClassMetadata GetNodeClassMetadata(const FVertexKey& InInputName)
+		static FNodeClassMetadata GetNodeClassMetadata(const FVertexName& InInputName)
 		{
 			FNodeClassMetadata Info;
 
@@ -604,7 +605,7 @@ namespace Metasound
 	public:
 		static constexpr bool bCanRegister = true;
 
-		TInputNode(const FString& InInstanceName, const FGuid& InInstanceID, const FVertexKey& InVertexName, FLiteral&& InLiteral)
+		TInputNode(const FVertexName& InInstanceName, const FGuid& InInstanceID, const FVertexName& InVertexName, FLiteral&& InLiteral)
 		:	FNode(InInstanceName, InInstanceID, GetNodeClassMetadata(InVertexName))
 		,	VertexInterface(GetVertexInterface(InVertexName))
 		,	Factory(MakeShared<FInputOperatorFactory, ESPMode::ThreadSafe>(InVertexName))
diff --git a/Engine/Plugins/Runtime/Metasound/Source/MetasoundStandardNodes/Private/MetasoundBPMToSecondsNode.cpp b/Engine/Plugins/Runtime/Metasound/Source/MetasoundStandardNodes/Private/MetasoundBPMToSecondsNode.cpp
--- a/Engine/Plugins/Runtime/Metasound/Source/MetasoundStandardNodes/Private/MetasoundBPMToSecondsNode.cpp
+++ b/Engine/Plugins/Runtime/Metasound/Source/MetasoundStandardNodes/Private/MetasoundBPMToSecondsNode.cpp
@@ -15,27 +15,27 @@ namespace Metasound
 {
 	namespace BPMToSecondsVertexNames
 	{
-		const FString& GetInputBPMTempoName()
+		const FVertexName& GetInputBPMTempoName()
 		{
-			static FString Name = TEXT("BPM");
+			static FVertexName Name = TEXT("BPM");
 			return Name;
 		}
 
-		const FString& GetInputBeatMultiplierName()
+		const FVertexName& GetInputBeatMultiplierName()
 		{
-			static FString Name = TEXT("Beat Multiplier");
+			static FVertexName Name = TEXT("Beat Multiplier");
 			return Name;
 		}
 
-		const FString& GetInputDivisionsOfWholeNoteName()
+		const FVertexName& GetInputDivisionsOfWholeNoteName()
 		{
-			static FString Name = TEXT("Divisions of Whole Note");
+			static FVertexName Name = TEXT("Divisions of Whole Note");
 			return Name;
 		}
 
-		const FString& GetOutputTimeSecondsName()
+		const FVertexName& GetOutputTimeSecondsName()
 		{
-			static FString Name = TEXT("Seconds");
+			static FVertexName Name = TEXT("Seconds");
 			return Name;
 		}
 	}
diff --git a/Engine/Plugins/Runtime/Metasound/Source/MetasoundStandardNodes/Private/MetasoundBasicFilters.cpp b/Engine/Plugins/Runtime/Metasound/Source/MetasoundStandardNodes/Private/MetasoundBasicFilters.cpp
--- a/Engine/Plugins/Runtime/Metasound/Source/MetasoundStandardNodes/Private/MetasoundBasicFilters.cpp
+++ b/Engine/Plugins/Runtime/Metasound/Source/MetasoundStandardNodes/Private/MetasoundBasicFilters.cpp
@@ -902,7 +902,7 @@ namespace Metasound
 	public:
 		// public node api needs to define two conversion constructors:
 		// 1.) from FString
-		FLadderFilterNode(const FString& InInstanceName, const FGuid& InInstanceID)
+		FLadderFilterNode(const FVertexName& InInstanceName, const FGuid& InInstanceID)
 			: FNodeFacade(InInstanceName, InInstanceID, TFacadeOperatorClass<FLadderFilterOperator>())
 		{ }
 
@@ -917,7 +917,7 @@ namespace Metasound
 	public:
 		// public node api needs to define two conversion constructors:
 		// 1.) from FString
-		FStateVariableFilterNode(const FString& InInstanceName, const FGuid& InInstanceID)
+		FStateVariableFilterNode(const FVertexName& InInstanceName, const FGuid& InInstanceID)
 			: FNodeFacade(InInstanceName, InInstanceID, TFacadeOperatorClass<FStateVariableFilterOperator>())
 		{ }
 
@@ -932,7 +932,7 @@ namespace Metasound
 	public:
 		// public node api needs to define two conversion constructors:
 		// 1.) from FString
-		FOnePoleLowPassFilterNode(const FString& InInstanceName, const FGuid& InInstanceID)
+		FOnePoleLowPassFilterNode(const FVertexName& InInstanceName, const FGuid& InInstanceID)
 			: FNodeFacade(InInstanceName, InInstanceID, TFacadeOperatorClass<FOnePoleLowPassFilterOperator>())
 		{ }
 
@@ -947,7 +947,7 @@ namespace Metasound
 	public:
 		// public node api needs to define two conversion constructors:
 		// 1.) from FString
-		FOnePoleHighPassFilterNode(const FString& InInstanceName, const FGuid& InInstanceID)
+		FOnePoleHighPassFilterNode(const FVertexName& InInstanceName, const FGuid& InInstanceID)
 			: FNodeFacade(InInstanceName, InInstanceID, TFacadeOperatorClass<FOnePoleHighPassFilterOperator>())
 		{ }
 
@@ -962,7 +962,7 @@ namespace Metasound
 	public:
 		// public node api needs to define two conversion constructors:
 		// 1.) from FString
-		FBiquadFilterNode(const FString& InInstanceName, const FGuid& InInstanceID)
+		FBiquadFilterNode(const FVertexName& InInstanceName, const FGuid& InInstanceID)
 			: FNodeFacade(InInstanceName, InInstanceID, TFacadeOperatorClass<FBiquadFilterOperator>())
 		{ }
 
diff --git a/Engine/Plugins/Runtime/Metasound/Source/MetasoundStandardNodes/Private/MetasoundCrossfadeNode.cpp b/Engine/Plugins/Runtime/Metasound/Source/MetasoundStandardNodes/Private/MetasoundCrossfadeNode.cpp
--- a/Engine/Plugins/Runtime/Metasound/Source/MetasoundStandardNodes/Private/MetasoundCrossfadeNode.cpp
+++ b/Engine/Plugins/Runtime/Metasound/Source/MetasoundStandardNodes/Private/MetasoundCrossfadeNode.cpp
@@ -3,15 +3,16 @@
 #include "MetasoundNodeRegistrationMacro.h"
 #include "MetasoundAudioBuffer.h"
 #include "CoreMinimal.h"
-#include "MetasoundFacade.h"
+#include "DSP/BufferVectorOperations.h"
+#include "Internationalization/Text.h"
 #include "MetasoundExecutableOperator.h"
+#include "MetasoundFacade.h"
+#include "MetasoundParamHelper.h"
 #include "MetasoundPrimitives.h"
 #include "MetasoundStandardNodesCategories.h"
 #include "MetasoundStandardNodesNames.h"
 #include "MetasoundTrigger.h"
-#include "Internationalization/Text.h"
-#include "MetasoundParamHelper.h"
-#include "DSP/BufferVectorOperations.h"
+#include "MetasoundVertex.h"
 
 #define LOCTEXT_NAMESPACE "MetasoundStandardNodes_CrossfadeNode"
 
@@ -24,9 +25,9 @@ namespace Metasound
 {
 	namespace CrossfadeVertexNames
 	{
-		const FString GetInputCrossfadeValueName()
+		const FVertexName GetInputCrossfadeValueName()
 		{
-			static const FString CrossfadeName = TEXT("Crossfade Value");
+			static const FVertexName CrossfadeName = TEXT("Crossfade Value");
 			return CrossfadeName;
 		}
 
@@ -36,9 +37,9 @@ namespace Metasound
 			return Desc;
 		}
 
-		const FString GetInputName(uint32 InIndex)
+		const FVertexName GetInputName(uint32 InIndex)
 		{
-			return FString::Format(TEXT("In {0}"), { InIndex });
+			return *FString::Format(TEXT("In {0}"), { InIndex });
 		}
 
 		const FText GetInputDescription(uint32 InIndex)
@@ -46,9 +47,9 @@ namespace Metasound
 			return FText::Format(LOCTEXT("CrossfadeInputDesc", "Cross fade {0} input."), InIndex);
 		}
 
-		const FString& GetOutputName()
+		const FVertexName& GetOutputName()
 		{
-			static const FString Name = TEXT("Out");
+			static const FVertexName Name = TEXT("Out");
 			return Name;
 		}
 
diff --git a/Engine/Plugins/Runtime/Metasound/Source/MetasoundStandardNodes/Private/MetasoundDebugLogNode.cpp b/Engine/Plugins/Runtime/Metasound/Source/MetasoundStandardNodes/Private/MetasoundDebugLogNode.cpp
--- a/Engine/Plugins/Runtime/Metasound/Source/MetasoundStandardNodes/Private/MetasoundDebugLogNode.cpp
+++ b/Engine/Plugins/Runtime/Metasound/Source/MetasoundStandardNodes/Private/MetasoundDebugLogNode.cpp
@@ -38,27 +38,27 @@ namespace Metasound
 	//Getters for the name of each parameter used in Print Log
 	namespace PrintLogVertexNames
 	{
-		const FString& GetInputTriggerName()
+		const FVertexName& GetInputTriggerName()
 		{
-			static const FString Name = TEXT("Trigger");
+			static const FVertexName Name = TEXT("Trigger");
 			return Name;
 		}
 
-		const FString& GetLabelPrintLogName()
+		const FVertexName& GetLabelPrintLogName()
 		{
-			static const FString Name = TEXT("Label");
+			static const FVertexName Name = TEXT("Label");
 			return Name;
 		}
 
-		const FString& GetToLogPrintLogName()
+		const FVertexName& GetToLogPrintLogName()
 		{
-			static const FString Name = TEXT("Value To Log");
+			static const FVertexName Name = TEXT("Value To Log");
 			return Name;
 		}
 
-		const FString& GetOutputPrintLogName()
+		const FVertexName& GetOutputPrintLogName()
 		{
-			static const FString Name = TEXT("Was Successful");
+			static const FVertexName Name = TEXT("Was Successful");
 			return Name;
 		}
 	}
diff --git a/Engine/Plugins/Runtime/Metasound/Source/MetasoundStandardNodes/Private/MetasoundITDPannerNode.cpp b/Engine/Plugins/Runtime/Metasound/Source/MetasoundStandardNodes/Private/MetasoundITDPannerNode.cpp
--- a/Engine/Plugins/Runtime/Metasound/Source/MetasoundStandardNodes/Private/MetasoundITDPannerNode.cpp
+++ b/Engine/Plugins/Runtime/Metasound/Source/MetasoundStandardNodes/Private/MetasoundITDPannerNode.cpp
@@ -14,6 +14,7 @@
 #include "DSP/Dsp.h"
 #include "MetasoundFacade.h"
 #include "MetasoundStandardNodesCategories.h"
+#include "MetasoundVertex.h"
 
 #define LOCTEXT_NAMESPACE "MetasoundStandardNodes_ITDPannerNode"
 
@@ -21,9 +22,9 @@ namespace Metasound
 {
 	namespace ITDPannerVertexNames
 	{
-		const FString& GetInputAudioName()
+		const FVertexName& GetInputAudioName()
 		{
-			static FString Name = TEXT("In");
+			static FVertexName Name = TEXT("In");
 			return Name;
 		}
 
@@ -33,9 +34,9 @@ namespace Metasound
 			return Desc;
 		}
 
-		const FString& GetInputPanAngleName()
+		const FVertexName& GetInputPanAngleName()
 		{
-			static FString Name = TEXT("Angle");
+			static FVertexName Name = TEXT("Angle");
 			return Name;
 		}
 
@@ -45,9 +46,9 @@ namespace Metasound
 			return Desc;
 		}
 
-		const FString& GetInputDistanceFactorName()
+		const FVertexName& GetInputDistanceFactorName()
 		{
-			static FString Name = TEXT("Distance Factor");
+			static FVertexName Name = TEXT("Distance Factor");
 			return Name;
 		}
 
@@ -57,9 +58,9 @@ namespace Metasound
 			return Desc;
 		}
 
-		const FString& GetInputHeadWidthName()
+		const FVertexName& GetInputHeadWidthName()
 		{
-			static FString Name = TEXT("Head Width");
+			static FVertexName Name = TEXT("Head Width");
 			return Name;
 		}
 
@@ -69,9 +70,9 @@ namespace Metasound
 			return Desc;
 		}
 
-		const FString& GetOutputAudioLeftName()
+		const FVertexName& GetOutputAudioLeftName()
 		{
-			static FString Name = TEXT("Out Left");
+			static FVertexName Name = TEXT("Out Left");
 			return Name;
 		}
 
@@ -81,9 +82,9 @@ namespace Metasound
 			return Desc;
 		}
 
-		const FString& GetOutputAudioRightName()
+		const FVertexName& GetOutputAudioRightName()
 		{
-			static FString Name = TEXT("Out Right");
+			static FVertexName Name = TEXT("Out Right");
 			return Name;
 		}
 
diff --git a/Engine/Plugins/Runtime/Metasound/Source/MetasoundStandardNodes/Private/MetasoundMathNodes.cpp b/Engine/Plugins/Runtime/Metasound/Source/MetasoundStandardNodes/Private/MetasoundMathNodes.cpp
--- a/Engine/Plugins/Runtime/Metasound/Source/MetasoundStandardNodes/Private/MetasoundMathNodes.cpp
+++ b/Engine/Plugins/Runtime/Metasound/Source/MetasoundStandardNodes/Private/MetasoundMathNodes.cpp
@@ -17,6 +17,8 @@
 #include "MetasoundStandardNodesNames.h"
 #include "MetasoundTime.h"
 #include "MetasoundTrigger.h"
+#include "MetasoundVertex.h"
+
 
 #define LOCTEXT_NAMESPACE "MetasoundMathOpNode"
 
@@ -50,9 +52,9 @@ namespace Metasound
 {
 	namespace MathOpNames
 	{
-		static const FString PrimaryOperandName = TEXT("PrimaryOperand");
+		static const FVertexName PrimaryOperandName = TEXT("PrimaryOperand");
 
-		static const FString AdditionalOperandsName = TEXT("AdditionalOperands");
+		static const FVertexName AdditionalOperandsName = TEXT("AdditionalOperands");
 
 		static const TArray<FText> AddKeywords = { LOCTEXT("AddMathKeyword", "+") };
 		static const TArray<FText> SubtractKeywords = { LOCTEXT("SubtractMathKeyword", "-") };
@@ -241,7 +243,7 @@ namespace Metasound
 				const FInputVertexInterface& InputInterface = MathOpNode.GetVertexInterface().GetInputInterface();
 
 				FLiteral DefaultValue = FLiteral::CreateInvalid();
-				const FString& PrimaryOperandName = MathOpNames::PrimaryOperandName;
+				const FVertexName& PrimaryOperandName = MathOpNames::PrimaryOperandName;
 				if (InputInterface.Contains(PrimaryOperandName))
 				{
 					DefaultValue = InputInterface[PrimaryOperandName].GetDefaultLiteral();
@@ -249,7 +251,7 @@ namespace Metasound
 				TDataClassReadRef PrimaryOperand = InParams.InputDataReferences.GetDataReadReferenceOrConstruct<TDataClass>(MathOpNames::PrimaryOperandName, TMathOpClass::GetDefault(InParams.OperatorSettings, DefaultValue));
 
 				// TODO: Support dynamic number of inputs
-				const FString& OpName = MathOpNames::AdditionalOperandsName;
+				const FVertexName& OpName = MathOpNames::AdditionalOperandsName;
 				if (InputInterface.Contains(OpName))
 				{
 					DefaultValue = InputInterface[OpName].GetDefaultLiteral();
diff --git a/Engine/Plugins/Runtime/Metasound/Source/MetasoundStandardNodes/Private/MetasoundMidiNoteQuantizerNode.cpp b/Engine/Plugins/Runtime/Metasound/Source/MetasoundStandardNodes/Private/MetasoundMidiNoteQuantizerNode.cpp
--- a/Engine/Plugins/Runtime/Metasound/Source/MetasoundStandardNodes/Private/MetasoundMidiNoteQuantizerNode.cpp
+++ b/Engine/Plugins/Runtime/Metasound/Source/MetasoundStandardNodes/Private/MetasoundMidiNoteQuantizerNode.cpp
@@ -219,7 +219,7 @@ namespace Metasound
 	public:
 		// public node api needs to define two conversion constructors:
 		// (1: from FString)
-		FMidiNoteQuantizerNode(const FString & InInstanceName, const FGuid & InInstanceID)
+		FMidiNoteQuantizerNode(const FVertexName& InInstanceName, const FGuid & InInstanceID)
 			: FNodeFacade(InInstanceName, InInstanceID, TFacadeOperatorClass <FMidiNoteQuantizerOperator>())
 		{ }
 
diff --git a/Engine/Plugins/Runtime/Metasound/Source/MetasoundStandardNodes/Private/MetasoundMidiScaleToArrayNode.cpp b/Engine/Plugins/Runtime/Metasound/Source/MetasoundStandardNodes/Private/MetasoundMidiScaleToArrayNode.cpp
--- a/Engine/Plugins/Runtime/Metasound/Source/MetasoundStandardNodes/Private/MetasoundMidiScaleToArrayNode.cpp
+++ b/Engine/Plugins/Runtime/Metasound/Source/MetasoundStandardNodes/Private/MetasoundMidiScaleToArrayNode.cpp
@@ -194,7 +194,7 @@ namespace Metasound
 	public:
 		// public node api needs to define two conversion constructors:
 		// (1: from FString)
-		FMusicalScaleToNoteArrayNode(const FString & InInstanceName, const FGuid & InInstanceID)
+		FMusicalScaleToNoteArrayNode(const FVertexName& InInstanceName, const FGuid & InInstanceID)
 			: FNodeFacade(InInstanceName, InInstanceID, TFacadeOperatorClass <FMusicalScaleToNoteArrayOperator>())
 		{ }
 
diff --git a/Engine/Plugins/Runtime/Metasound/Source/MetasoundStandardNodes/Private/MetasoundMixerNode.cpp b/Engine/Plugins/Runtime/Metasound/Source/MetasoundStandardNodes/Private/MetasoundMixerNode.cpp
--- a/Engine/Plugins/Runtime/Metasound/Source/MetasoundStandardNodes/Private/MetasoundMixerNode.cpp
+++ b/Engine/Plugins/Runtime/Metasound/Source/MetasoundStandardNodes/Private/MetasoundMixerNode.cpp
@@ -226,18 +226,18 @@ namespace Metasound
 		}
 
 #pragma region Name Gen
-		static const FString GetAudioInputName(uint32 InputIndex, uint32 ChannelIndex)
+		static const FVertexName GetAudioInputName(uint32 InputIndex, uint32 ChannelIndex)
 		{
 			if (NumChannels == 1)
 			{
-				return FString::Printf(TEXT("In %i"), InputIndex);
+				return *FString::Printf(TEXT("In %i"), InputIndex);
 			}
 			else if (NumChannels == 2)
 			{
-				return FString::Printf(TEXT("In %i %s"), InputIndex, (ChannelIndex == 0) ? TEXT("L") : TEXT("R"));
+				return *FString::Printf(TEXT("In %i %s"), InputIndex, (ChannelIndex == 0) ? TEXT("L") : TEXT("R"));
 			}
 
-			return FString::Printf(TEXT("In %i, %i"), InputIndex, ChannelIndex);
+			return *FString::Printf(TEXT("In %i, %i"), InputIndex, ChannelIndex);
 		}
 
 		static const FText GetAudioInputDescription(uint32 InputIndex, uint32 ChannelIndex)
@@ -245,9 +245,9 @@ namespace Metasound
 			return FText::Format(LOCTEXT("AudioMixerAudioInputDescription", "Audio Input #: {0}, Channel: {1}"), InputIndex, ChannelIndex);
 		}
 
-		static const FString GetGainInputName(uint32 InputIndex)
+		static const FVertexName GetGainInputName(uint32 InputIndex)
 		{
-			return FString::Printf(TEXT("Gain %i"), InputIndex);
+			return *FString::Printf(TEXT("Gain %i"), InputIndex);
 		}
 
 		static const FText GetGainInputDescription(uint32 InputIndex)
@@ -255,7 +255,7 @@ namespace Metasound
 			return FText::Format(LOCTEXT("AudioMixerGainInputDescription", "Gain Input #: {0}"), InputIndex);
 		}
 
-		static const FString GetAudioOutputName(uint32 ChannelIndex)
+		static const FVertexName GetAudioOutputName(uint32 ChannelIndex)
 		{
 			if (NumChannels == 1)
 			{
@@ -263,10 +263,10 @@ namespace Metasound
 			}
 			else if (NumChannels == 2)
 			{
-				return FString::Printf(TEXT("Out %s"), (ChannelIndex == 0) ? TEXT("L") : TEXT("R"));
+				return *FString::Printf(TEXT("Out %s"), (ChannelIndex == 0) ? TEXT("L") : TEXT("R"));
 			}
 
-			return FString::Printf(TEXT("Out %i"), ChannelIndex);
+			return *FString::Printf(TEXT("Out %i"), ChannelIndex);
 		}
 
 		static const FText GetAudioOutputDescription(uint32 ChannelIndex)
diff --git a/Engine/Plugins/Runtime/Metasound/Source/MetasoundStandardNodes/Private/MetasoundNoiseGenerator.cpp b/Engine/Plugins/Runtime/Metasound/Source/MetasoundStandardNodes/Private/MetasoundNoiseGenerator.cpp
--- a/Engine/Plugins/Runtime/Metasound/Source/MetasoundStandardNodes/Private/MetasoundNoiseGenerator.cpp
+++ b/Engine/Plugins/Runtime/Metasound/Source/MetasoundStandardNodes/Private/MetasoundNoiseGenerator.cpp
@@ -192,7 +192,7 @@ namespace Metasound
 		return Info;
 	}
 
-	FNoiseNode::FNoiseNode(const FString& InName, const FGuid& InInstanceID, int32 InDefaultSeed)
+	FNoiseNode::FNoiseNode(const FVertexName& InName, const FGuid& InInstanceID, int32 InDefaultSeed)
 		: FNodeFacade(InName, InInstanceID, TFacadeOperatorClass<FNoiseOperator>())
 		, DefaultSeed(InDefaultSeed)
 	{}
diff --git a/Engine/Plugins/Runtime/Metasound/Source/MetasoundStandardNodes/Private/MetasoundOscillatorNodes.cpp b/Engine/Plugins/Runtime/Metasound/Source/MetasoundStandardNodes/Private/MetasoundOscillatorNodes.cpp
--- a/Engine/Plugins/Runtime/Metasound/Source/MetasoundStandardNodes/Private/MetasoundOscillatorNodes.cpp
+++ b/Engine/Plugins/Runtime/Metasound/Source/MetasoundStandardNodes/Private/MetasoundOscillatorNodes.cpp
@@ -624,7 +624,7 @@ namespace Metasound
 		FAudioBufferReadRef Fm;
 	};
 
-	FOscilatorNodeBase::FOscilatorNodeBase(const FString& InInstanceName, const FGuid& InInstanceID, const FNodeClassMetadata& InInfo, const TSharedRef<IOperatorFactory, ESPMode::ThreadSafe>& InFactory, float InDefaultFrequency, float InDefaultGlideFactor, bool bInDefaultEnablement)
+	FOscilatorNodeBase::FOscilatorNodeBase(const FVertexName& InInstanceName, const FGuid& InInstanceID, const FNodeClassMetadata& InInfo, const TSharedRef<IOperatorFactory, ESPMode::ThreadSafe>& InFactory, float InDefaultFrequency, float InDefaultGlideFactor, bool bInDefaultEnablement)
 		: FNode(InInstanceName, InInstanceID, InInfo)
 		, Factory(InFactory)
 		, VertexInterface(GetMetadata().DefaultInterface)
@@ -748,8 +748,8 @@ namespace Metasound
 	private:
 	};
 
-	FSineOscilatorNode::FSineOscilatorNode(const FString& InInstanceName, const FGuid& InInstanceID, float InDefaultFrequency, float InDefautlGlideFactor, bool bInDefaultEnablement)
-		: FOscilatorNodeBase(InInstanceName,InInstanceID, FFactory::GetNodeInfo(), MakeShared<FFactory, ESPMode::ThreadSafe>(), InDefaultFrequency, InDefautlGlideFactor, bInDefaultEnablement )
+	FSineOscilatorNode::FSineOscilatorNode(const FVertexName& InInstanceName, const FGuid& InInstanceID, float InDefaultFrequency, float InDefautlGlideFactor, bool bInDefaultEnablement)
+		: FOscilatorNodeBase(InInstanceName, InInstanceID, FFactory::GetNodeInfo(), MakeShared<FFactory, ESPMode::ThreadSafe>(), InDefaultFrequency, InDefautlGlideFactor, bInDefaultEnablement )
 	{}
 
 	FSineOscilatorNode::FSineOscilatorNode(const FNodeInitData& InInitData)
@@ -867,7 +867,7 @@ namespace Metasound
 		return nullptr;
 	}
 
-	FSawOscilatorNode::FSawOscilatorNode(const FString& InInstanceName, const FGuid& InInstanceID, float InDefaultFrequency, float InDefaultGlideFactor, bool bInDefaultEnablement)
+	FSawOscilatorNode::FSawOscilatorNode(const FVertexName& InInstanceName, const FGuid& InInstanceID, float InDefaultFrequency, float InDefaultGlideFactor, bool bInDefaultEnablement)
 		: FOscilatorNodeBase(InInstanceName, InInstanceID, FFactory::GetNodeInfo(), MakeShared<FFactory, ESPMode::ThreadSafe>(), InDefaultFrequency, InDefaultGlideFactor, bInDefaultEnablement)
 	{}
 
@@ -1050,7 +1050,7 @@ namespace Metasound
 	private:
 	};
 	
-	FSquareOscilatorNode::FSquareOscilatorNode(const FString& InInstanceName, const FGuid& InInstanceID, float InDefaultFrequency, float InDefaultGlideFactor, bool bInDefaultEnablement)
+	FSquareOscilatorNode::FSquareOscilatorNode(const FVertexName& InInstanceName, const FGuid& InInstanceID, float InDefaultFrequency, float InDefaultGlideFactor, bool bInDefaultEnablement)
 		: FOscilatorNodeBase(InInstanceName, InInstanceID, FFactory::GetNodeInfo(), MakeShared<FFactory, ESPMode::ThreadSafe>(), InDefaultFrequency, InDefaultGlideFactor, bInDefaultEnablement)
 	{}
 
@@ -1165,7 +1165,7 @@ namespace Metasound
 	private:
 	};
 	
-	FTriangleOscilatorNode::FTriangleOscilatorNode(const FString& InInstanceName, const FGuid& InInstanceID, float InDefaultFrequency, float InDefaultGlideFactor, bool bInDefaultEnablement)
+	FTriangleOscilatorNode::FTriangleOscilatorNode(const FVertexName& InInstanceName, const FGuid& InInstanceID, float InDefaultFrequency, float InDefaultGlideFactor, bool bInDefaultEnablement)
 		: FOscilatorNodeBase(InInstanceName, InInstanceID, FFactory::GetNodeInfo(), MakeShared<FFactory, ESPMode::ThreadSafe>(), InDefaultFrequency, InDefaultGlideFactor, bInDefaultEnablement)
 	{
 
diff --git a/Engine/Plugins/Runtime/Metasound/Source/MetasoundStandardNodes/Private/MetasoundRandomNode.cpp b/Engine/Plugins/Runtime/Metasound/Source/MetasoundStandardNodes/Private/MetasoundRandomNode.cpp
--- a/Engine/Plugins/Runtime/Metasound/Source/MetasoundStandardNodes/Private/MetasoundRandomNode.cpp
+++ b/Engine/Plugins/Runtime/Metasound/Source/MetasoundStandardNodes/Private/MetasoundRandomNode.cpp
@@ -11,51 +11,51 @@ namespace Metasound
 {
 	namespace RandomNodeNames
 	{
-		const FString& GetInputNextTriggerName()
+		const FVertexName& GetInputNextTriggerName()
 		{
-			static const FString Name = TEXT("Next");
+			static const FVertexName Name = TEXT("Next");
 			return Name;
 		}
 
-		const FString& GetInputResetTriggerName()
+		const FVertexName& GetInputResetTriggerName()
 		{
-			static const FString Name = TEXT("Reset");
+			static const FVertexName Name = TEXT("Reset");
 			return Name;
 		}
 
-		const FString& GetInputSeedName()
+		const FVertexName& GetInputSeedName()
 		{
-			static const FString Name = TEXT("Seed");
+			static const FVertexName Name = TEXT("Seed");
 			return Name;
 		}
 
-		const FString& GetInputMinName()
+		const FVertexName& GetInputMinName()
 		{
-			static const FString Name = TEXT("Min");
+			static const FVertexName Name = TEXT("Min");
 			return Name;
 		}
 
-		const FString& GetInputMaxName()
+		const FVertexName& GetInputMaxName()
 		{
-			static const FString Name = TEXT("Max");
+			static const FVertexName Name = TEXT("Max");
 			return Name;
 		}
 
-		const FString& GetOutputOnNextTriggerName()
+		const FVertexName& GetOutputOnNextTriggerName()
 		{
-			static const FString Name = TEXT("On Next");
+			static const FVertexName Name = TEXT("On Next");
 			return Name;
 		}
 
-		const FString& GetOutputOnResetTriggerName()
+		const FVertexName& GetOutputOnResetTriggerName()
 		{
-			static const FString Name = TEXT("On Reset");
+			static const FVertexName Name = TEXT("On Reset");
 			return Name;
 		}
 
-		const FString& GetOutputValueName()
+		const FVertexName& GetOutputValueName()
 		{
-			static const FString Name = TEXT("Value");
+			static const FVertexName Name = TEXT("Value");
 			return Name;
 		}
 
diff --git a/Engine/Plugins/Runtime/Metasound/Source/MetasoundStandardNodes/Private/MetasoundStereoPannerNode.cpp b/Engine/Plugins/Runtime/Metasound/Source/MetasoundStandardNodes/Private/MetasoundStereoPannerNode.cpp
--- a/Engine/Plugins/Runtime/Metasound/Source/MetasoundStandardNodes/Private/MetasoundStereoPannerNode.cpp
+++ b/Engine/Plugins/Runtime/Metasound/Source/MetasoundStandardNodes/Private/MetasoundStereoPannerNode.cpp
@@ -13,6 +13,7 @@
 #include "DSP/BufferVectorOperations.h"
 #include "MetasoundStandardNodesCategories.h"
 #include "MetasoundFacade.h"
+#include "MetasoundVertex.h"
 
 #define LOCTEXT_NAMESPACE "MetasoundStandardNodes_StereoPanner"
 
@@ -20,9 +21,9 @@ namespace Metasound
 {
 	namespace StereoPannerVertexNames
 	{
-		const FString& GetInputAudioName()
+		const FVertexName& GetInputAudioName()
 		{
-			static FString Name = TEXT("In");
+			static FVertexName Name = TEXT("In");
 			return Name;
 		}
 
@@ -32,9 +33,9 @@ namespace Metasound
 			return Desc;
 		}
 
-		const FString& GetInputPanAmountName()
+		const FVertexName& GetInputPanAmountName()
 		{
-			static FString Name = TEXT("Pan Amount");
+			static FVertexName Name = TEXT("Pan Amount");
 			return Name;
 		}
 
@@ -44,9 +45,9 @@ namespace Metasound
 			return Desc;
 		}
 
-		const FString& GetInputPanningLawName()
+		const FVertexName& GetInputPanningLawName()
 		{
-			static FString Name = TEXT("Panning Law");
+			static FVertexName Name = TEXT("Panning Law");
 			return Name;
 		}
 
@@ -56,9 +57,9 @@ namespace Metasound
 			return Desc;
 		}
 
-		const FString& GetOutputAudioLeftName()
+		const FVertexName& GetOutputAudioLeftName()
 		{
-			static FString Name = TEXT("Out Left");
+			static FVertexName Name = TEXT("Out Left");
 			return Name;
 		}
 
@@ -68,9 +69,9 @@ namespace Metasound
 			return Desc;
 		}
 
-		const FString& GetOutputAudioRightName()
+		const FVertexName& GetOutputAudioRightName()
 		{
-			static FString Name = TEXT("Out Right");
+			static FVertexName Name = TEXT("Out Right");
 			return Name;
 		}
 
diff --git a/Engine/Plugins/Runtime/Metasound/Source/MetasoundStandardNodes/Private/MetasoundSwitchGateNode.cpp b/Engine/Plugins/Runtime/Metasound/Source/MetasoundStandardNodes/Private/MetasoundSwitchGateNode.cpp
--- a/Engine/Plugins/Runtime/Metasound/Source/MetasoundStandardNodes/Private/MetasoundSwitchGateNode.cpp
+++ b/Engine/Plugins/Runtime/Metasound/Source/MetasoundStandardNodes/Private/MetasoundSwitchGateNode.cpp
@@ -233,7 +233,7 @@ namespace Metasound
 		return MakeUnique<FSwitchOperator>(InParams.OperatorSettings, Selector, Crossfade, PreviewBuffers);
 	}
 
-	FSwitchNode::FSwitchNode(const FString& InName, const FGuid& InInstanceID)
+	FSwitchNode::FSwitchNode(const FVertexName& InName, const FGuid& InInstanceID)
 		: FNodeFacade(InName, InInstanceID, TFacadeOperatorClass<FSwitchOperator>())
 	{
 	}
@@ -446,7 +446,7 @@ namespace Metasound
 		return MakeUnique<FGateOperator>(InParams.OperatorSettings, Selector, Crossfade, Input);
 	}
 
-	FGateNode::FGateNode(const FString& InName, const FGuid& InInstanceID)
+	FGateNode::FGateNode(const FVertexName& InName, const FGuid& InInstanceID)
 		: FNodeFacade(InName, InInstanceID, TFacadeOperatorClass<FGateOperator>())
 	{
 	}
diff --git a/Engine/Plugins/Runtime/Metasound/Source/MetasoundStandardNodes/Private/MetasoundTriggerAnyNode.cpp b/Engine/Plugins/Runtime/Metasound/Source/MetasoundStandardNodes/Private/MetasoundTriggerAnyNode.cpp
--- a/Engine/Plugins/Runtime/Metasound/Source/MetasoundStandardNodes/Private/MetasoundTriggerAnyNode.cpp
+++ b/Engine/Plugins/Runtime/Metasound/Source/MetasoundStandardNodes/Private/MetasoundTriggerAnyNode.cpp
@@ -1,12 +1,13 @@
 // Copyright Epic Games, Inc. All Rights Reserved.
-#include "MetasoundTrigger.h"
-#include "MetasoundFacade.h"
-#include "MetasoundExecutableOperator.h"
+
 #include "Internationalization/Text.h"
+#include "MetasoundExecutableOperator.h"
+#include "MetasoundFacade.h"
 #include "MetasoundNodeRegistrationMacro.h"
+#include "MetasoundStandardNodesCategories.h"
 #include "MetasoundStandardNodesNames.h"
+#include "MetasoundTrigger.h"
 #include "MetasoundVertex.h"
-#include "MetasoundStandardNodesCategories.h"
 
 #define LOCTEXT_NAMESPACE "MetasoundStandardNodes_TriggerAny"
 
@@ -14,6 +15,7 @@
 	using FTriggerAnyNode_##Number = TTriggerAnyNode<Number>; \
 	METASOUND_REGISTER_NODE(FTriggerAnyNode_##Number) \
 
+
 namespace Metasound
 {
 	namespace MetasoundTriggerAnyNodePrivate
@@ -41,9 +43,9 @@ namespace Metasound
 
 	namespace TriggerAnyVertexNames
 	{
-		METASOUNDSTANDARDNODES_API const FString GetInputTriggerName(uint32 InIndex)
+		METASOUNDSTANDARDNODES_API const FVertexName GetInputTriggerName(uint32 InIndex)
 		{
-			return FString::Format(TEXT("In {0}"), { InIndex });
+			return *FString::Format(TEXT("In {0}"), { InIndex });
 		}
 
 		METASOUNDSTANDARDNODES_API const FText GetInputTriggerDescription(uint32 InIndex)
@@ -51,9 +53,9 @@ namespace Metasound
 			return FText::Format(LOCTEXT("TriggerAnyInputTriggerDesc", "Trigger {0} input. The output trigger is hit when any of the input triggers are hit."), InIndex);
 		}
 
-		METASOUNDSTANDARDNODES_API const FString& GetOutputTriggerName()
+		METASOUNDSTANDARDNODES_API const FVertexName& GetOutputTriggerName()
 		{
-			static const FString Name = TEXT("Out");
+			static const FVertexName Name = TEXT("Out");
 			return Name;
 		}
 
diff --git a/Engine/Plugins/Runtime/Metasound/Source/MetasoundStandardNodes/Private/MetasoundTriggerOnThresholdNode.cpp b/Engine/Plugins/Runtime/Metasound/Source/MetasoundStandardNodes/Private/MetasoundTriggerOnThresholdNode.cpp
--- a/Engine/Plugins/Runtime/Metasound/Source/MetasoundStandardNodes/Private/MetasoundTriggerOnThresholdNode.cpp
+++ b/Engine/Plugins/Runtime/Metasound/Source/MetasoundStandardNodes/Private/MetasoundTriggerOnThresholdNode.cpp
@@ -220,7 +220,7 @@ namespace Metasound
 	{
 	}
 
-	FTriggerOnThresholdNode::FTriggerOnThresholdNode(const FString& InName, const FGuid& InInstanceID, float InDefaultThreshold)
+	FTriggerOnThresholdNode::FTriggerOnThresholdNode(const FVertexName& InName, const FGuid& InInstanceID, float InDefaultThreshold)
 		: FNodeFacade(InName, InInstanceID, TFacadeOperatorClass<FTriggerOnThresholdOperator>())
 		, DefaultThreshold(InDefaultThreshold)
 	{
diff --git a/Engine/Plugins/Runtime/Metasound/Source/MetasoundStandardNodes/Private/MetasoundTriggerRouteNode.cpp b/Engine/Plugins/Runtime/Metasound/Source/MetasoundStandardNodes/Private/MetasoundTriggerRouteNode.cpp
--- a/Engine/Plugins/Runtime/Metasound/Source/MetasoundStandardNodes/Private/MetasoundTriggerRouteNode.cpp
+++ b/Engine/Plugins/Runtime/Metasound/Source/MetasoundStandardNodes/Private/MetasoundTriggerRouteNode.cpp
@@ -25,9 +25,9 @@ namespace Metasound
 {
 	namespace TriggerRouteVertexNames
 	{
-		const FString GetInputTriggerName(uint32 InIndex)
+		const FVertexName GetInputTriggerName(uint32 InIndex)
 		{
-			return FString::Format(TEXT("Set {0}"), { InIndex });
+			return *FString::Format(TEXT("Set {0}"), { InIndex });
 		}
 
 		const FText GetInputTriggerDescription(uint32 InIndex)
@@ -39,9 +39,9 @@ namespace Metasound
 			return FText::Format(LOCTEXT("TriggerRouteInputTriggerDesc", "The input trigger {0} to cause the corresponding input value to route to the output value."), InIndex);
 		}
 
-		const FString GetInputValueName(uint32 InIndex)
+		const FVertexName GetInputValueName(uint32 InIndex)
 		{
-			return FString::Format(TEXT("Value {0}"), { InIndex });
+			return *FString::Format(TEXT("Value {0}"), { InIndex });
 		}
 
 		const FText GetInputValueDescription(uint32 InIndex)
@@ -49,9 +49,9 @@ namespace Metasound
 			return FText::Format(LOCTEXT("TriggerRouteValueDesc", "The input value ({0}) to route to the output when triggered by Set {0}."), InIndex);
 		}
 
-		const FString& GetOutputTriggerName()
+		const FVertexName& GetOutputTriggerName()
 		{
-			static const FString Name = TEXT("On Set");
+			static const FVertexName Name = TEXT("On Set");
 			return Name;
 		}
 
@@ -61,9 +61,9 @@ namespace Metasound
 			return Desc;
 		}
 
-		const FString& GetOutputValueName()
+		const FVertexName& GetOutputValueName()
 		{
-			static const FString Name = TEXT("Value");
+			static const FVertexName Name = TEXT("Value");
 			return Name;
 		}
 
diff --git a/Engine/Plugins/Runtime/Metasound/Source/MetasoundStandardNodes/Private/MetasoundTriggerToggleNode.cpp b/Engine/Plugins/Runtime/Metasound/Source/MetasoundStandardNodes/Private/MetasoundTriggerToggleNode.cpp
--- a/Engine/Plugins/Runtime/Metasound/Source/MetasoundStandardNodes/Private/MetasoundTriggerToggleNode.cpp
+++ b/Engine/Plugins/Runtime/Metasound/Source/MetasoundStandardNodes/Private/MetasoundTriggerToggleNode.cpp
@@ -14,6 +14,7 @@
 #include "MetasoundPrimitives.h"
 #include "MetasoundStandardNodesNames.h"
 #include "MetasoundStandardNodesCategories.h"
+#include "MetasoundVertex.h"
 
 #define LOCTEXT_NAMESPACE "MetasoundStandardNodes_Toggle"
 
@@ -21,11 +22,11 @@ namespace Metasound
 {
 	namespace TriggerToggle
 	{
-		const FString GetInputOnTriggerName() { return TEXT("On"); }
-		const FString GetInputOffTriggerName() { return TEXT("Off"); }
-		const FString GetInputInitName() { return TEXT("Init"); }
-		const FString GetOutputTriggerName() { return TEXT("Out"); }
-		const FString GetOutputValueName() { return TEXT("Value"); }
+		const FVertexName GetInputOnTriggerName() { return TEXT("On"); }
+		const FVertexName GetInputOffTriggerName() { return TEXT("Off"); }
+		const FVertexName GetInputInitName() { return TEXT("Init"); }
+		const FVertexName GetOutputTriggerName() { return TEXT("Out"); }
+		const FVertexName GetOutputValueName() { return TEXT("Value"); }
 	}
 
 	class FTriggerToggleOperator : public TExecutableOperator<FTriggerToggleOperator>
diff --git a/Engine/Plugins/Runtime/Metasound/Source/MetasoundStandardNodes/Private/MetasoundWaveWriterNode.cpp b/Engine/Plugins/Runtime/Metasound/Source/MetasoundStandardNodes/Private/MetasoundWaveWriterNode.cpp
--- a/Engine/Plugins/Runtime/Metasound/Source/MetasoundStandardNodes/Private/MetasoundWaveWriterNode.cpp
+++ b/Engine/Plugins/Runtime/Metasound/Source/MetasoundStandardNodes/Private/MetasoundWaveWriterNode.cpp
@@ -10,16 +10,18 @@
 #include "MetasoundPrimitives.h"
 #include "MetasoundStandardNodesCategories.h"
 #include "MetasoundStandardNodesNames.h"
+#include "MetasoundVertex.h"
 #include "Misc/Paths.h"
 #include "Misc/ScopeLock.h"
 
 #define LOCTEXT_NAMESPACE "MetasoundStandardNodes_WaveWriterNode"
 
+
 namespace Metasound
-{	
+{
 	// Incremental wave writer class.
 	class FWaveWriter
-	{		
+	{
 		// Local definition so we don't depend on platform includes.
 		enum EFormatType { IEEE_FLOAT = 0x3 }; // WAVE_FORMAT_IEEE_FLOAT
 		struct FWaveFormatEx
@@ -305,13 +307,15 @@ namespace Metasound
 	constexpr const TCHAR* FWaveWriterOperator::EnabledPinName;
 	constexpr const TCHAR* FWaveWriterOperator::FilenamePrefixPinName;
 
-	FWaveWriterNode::FWaveWriterNode(const FString& InName, const FGuid& InInstanceID)
-		: FNodeFacade{ InName, InInstanceID, TFacadeOperatorClass<FWaveWriterOperator>() }
-	{}
+	FWaveWriterNode::FWaveWriterNode(const FVertexName& InName, const FGuid& InInstanceID)
+		: FNodeFacade(InName, InInstanceID, TFacadeOperatorClass<FWaveWriterOperator>())
+	{
+	}
 
 	FWaveWriterNode::FWaveWriterNode(const FNodeInitData& InInitData)
-		: FWaveWriterNode{ InInitData.InstanceName, InInitData.InstanceID }
-	{}
+		: FWaveWriterNode(InInitData.InstanceName, InInitData.InstanceID)
+	{
+	}
 
 	METASOUND_REGISTER_NODE(FWaveWriterNode)
 }
diff --git a/Engine/Plugins/Runtime/Metasound/Source/MetasoundStandardNodes/Public/MetasoundNoiseGenerator.h b/Engine/Plugins/Runtime/Metasound/Source/MetasoundStandardNodes/Public/MetasoundNoiseGenerator.h
--- a/Engine/Plugins/Runtime/Metasound/Source/MetasoundStandardNodes/Public/MetasoundNoiseGenerator.h
+++ b/Engine/Plugins/Runtime/Metasound/Source/MetasoundStandardNodes/Public/MetasoundNoiseGenerator.h
@@ -1,16 +1,17 @@
 // Copyright Epic Games, Inc. All Rights Reserved.
-
 #pragma once
 
 #include "CoreMinimal.h"
 #include "MetasoundFacade.h"
+#include "MetasoundVertex.h"
+
 
 namespace Metasound
 {
 	class METASOUNDSTANDARDNODES_API FNoiseNode : public FNodeFacade
 	{
 	public:
-		FNoiseNode(const FString& InName, const FGuid& InInstanceID, int32 InDefaultSeed);
+		FNoiseNode(const FVertexName& InName, const FGuid& InInstanceID, int32 InDefaultSeed);
 		FNoiseNode(const FNodeInitData& InInitData);
 
 		int32 GetDefaultSeed() const { return DefaultSeed; }
diff --git a/Engine/Plugins/Runtime/Metasound/Source/MetasoundStandardNodes/Public/MetasoundOscillatorNodes.h b/Engine/Plugins/Runtime/Metasound/Source/MetasoundStandardNodes/Public/MetasoundOscillatorNodes.h
--- a/Engine/Plugins/Runtime/Metasound/Source/MetasoundStandardNodes/Public/MetasoundOscillatorNodes.h
+++ b/Engine/Plugins/Runtime/Metasound/Source/MetasoundStandardNodes/Public/MetasoundOscillatorNodes.h
@@ -11,20 +11,22 @@
 #include "MetasoundFacade.h"
 #include "MetasoundNode.h"
 #include "MetasoundOperatorInterface.h"
+#include "MetasoundVertex.h"
+
 
 namespace Metasound
 {
 	class METASOUNDSTANDARDNODES_API FOscilatorNodeBase : public FNode
 	{
 	protected:
-		FOscilatorNodeBase(const FString& InInstanceName, const FGuid& InInstanceID, const FNodeClassMetadata& InInfo, const TSharedRef<IOperatorFactory, ESPMode::ThreadSafe>& InFactory, float InDefaultFrequency, float InDefaultGlideTime, bool bInDefaultEnablement);
+		FOscilatorNodeBase(const FVertexName& InInstanceName, const FGuid& InInstanceID, const FNodeClassMetadata& InInfo, const TSharedRef<IOperatorFactory, ESPMode::ThreadSafe>& InFactory, float InDefaultFrequency, float InDefaultGlideTime, bool bInDefaultEnablement);
 
-		float GetDefaultPhaseOffset() const 
+		float GetDefaultPhaseOffset() const
 		{
 			return DefaultPhaseOffset; 
 		}
 
-		float GetDefaultFrequency() const 
+		float GetDefaultFrequency() const
 		{
 			return DefaultFrequency;
 		}
@@ -34,7 +36,7 @@ namespace Metasound
 			return DefaultGlideFactor;
 		}
 
-		bool GetDefaultEnablement() const 
+		bool GetDefaultEnablement() const
 		{
 			return bDefaultEnablement;
 		}
@@ -76,8 +78,8 @@ namespace Metasound
 	class METASOUNDSTANDARDNODES_API FSineOscilatorNode : public FOscilatorNodeBase
 	{
 	public:
-		class FFactory;		
-		FSineOscilatorNode(const FString& InInstanceName, const FGuid& InInstanceID, float InDefaultFrequency, float InDefautlGlideTime, bool bInDefaultEnablement);
+		class FFactory;
+		FSineOscilatorNode(const FVertexName& InInstanceName, const FGuid& InInstanceID, float InDefaultFrequency, float InDefautlGlideTime, bool bInDefaultEnablement);
 		FSineOscilatorNode(const FNodeInitData& InInitData);
 	};
 		
@@ -85,7 +87,7 @@ namespace Metasound
 	{
 	public:
 		class FFactory;
-		FSawOscilatorNode(const FString& InInstanceName, const FGuid& InInstanceID, float InDefaultFrequency, float InDefautlGlideTime, bool bInDefaultEnablement);
+		FSawOscilatorNode(const FVertexName& InInstanceName, const FGuid& InInstanceID, float InDefaultFrequency, float InDefautlGlideTime, bool bInDefaultEnablement);
 		FSawOscilatorNode(const FNodeInitData& InInitData);
 	};
 
@@ -93,7 +95,7 @@ namespace Metasound
 	{
 	public:
 		class FFactory;
-		FTriangleOscilatorNode(const FString& InInstanceName, const FGuid& InInstanceID, float InDefaultFrequency, float InDefautlGlideTime, bool bInDefaultEnablement);
+		FTriangleOscilatorNode(const FVertexName& InInstanceName, const FGuid& InInstanceID, float InDefaultFrequency, float InDefautlGlideTime, bool bInDefaultEnablement);
 		FTriangleOscilatorNode(const FNodeInitData& InInitData);
 	};
 
@@ -101,7 +103,7 @@ namespace Metasound
 	{
 	public:
 		class FFactory;
-		FSquareOscilatorNode(const FString& InInstanceName, const FGuid& InInstanceID, float InDefaultFrequency, float InDefautlGlideTime, bool bInDefaultEnablement);
+		FSquareOscilatorNode(const FVertexName& InInstanceName, const FGuid& InInstanceID, float InDefaultFrequency, float InDefautlGlideTime, bool bInDefaultEnablement);
 		FSquareOscilatorNode(const FNodeInitData& InInitData);
 	};
 
diff --git a/Engine/Plugins/Runtime/Metasound/Source/MetasoundStandardNodes/Public/MetasoundParamHelper.h b/Engine/Plugins/Runtime/Metasound/Source/MetasoundStandardNodes/Public/MetasoundParamHelper.h
--- a/Engine/Plugins/Runtime/Metasound/Source/MetasoundStandardNodes/Public/MetasoundParamHelper.h
+++ b/Engine/Plugins/Runtime/Metasound/Source/MetasoundStandardNodes/Public/MetasoundParamHelper.h
@@ -50,6 +50,6 @@ Where NUMBER is a number (or whatever) to slot into the format specifier you def
 #define METASOUND_GET_PARAM_TT(NAME) NAME##Tooltip
 #define METASOUND_GET_PARAM_NAME_AND_TT(NAME) NAME##Name, NAME##Tooltip
 
-#define METASOUND_GET_VARIABLE_PARAM_NAME(NAME, INDEX) FString::Format(NAME##Name, {INDEX})
+#define METASOUND_GET_VARIABLE_PARAM_NAME(NAME, INDEX) *FString::Format(NAME##Name, {INDEX})
 #define METASOUND_GET_VARIABLE_TT(NAME, INDEX)  FText::Format(NAME##Tooltip, INDEX)
-#define METASOUND_GET_VARIABLE_PARAM_NAME_AND_TT(NAME, INDEX)  FString::Format(NAME##Name, {INDEX}), FText::Format(NAME##Tooltip, INDEX)
+#define METASOUND_GET_VARIABLE_PARAM_NAME_AND_TT(NAME, INDEX)  *FString::Format(NAME##Name, {INDEX}), FText::Format(NAME##Tooltip, INDEX)
diff --git a/Engine/Plugins/Runtime/Metasound/Source/MetasoundStandardNodes/Public/MetasoundRandomNode.h b/Engine/Plugins/Runtime/Metasound/Source/MetasoundStandardNodes/Public/MetasoundRandomNode.h
--- a/Engine/Plugins/Runtime/Metasound/Source/MetasoundStandardNodes/Public/MetasoundRandomNode.h
+++ b/Engine/Plugins/Runtime/Metasound/Source/MetasoundStandardNodes/Public/MetasoundRandomNode.h
@@ -11,6 +11,7 @@
 #include "Internationalization/Text.h"
 #include "MetasoundStandardNodesCategories.h"
 #include "MetasoundTime.h"
+#include "MetasoundVertex.h"
 
 #define LOCTEXT_NAMESPACE "MetasoundStandardNodes_RandomNode"
 
@@ -19,16 +20,16 @@ namespace Metasound
 	namespace RandomNodeNames
 	{
 		/** Input vertex names. */
-		const FString& GetInputNextTriggerName();
-		const FString& GetInputResetTriggerName();
-		const FString& GetInputSeedName();
-		const FString& GetInputMinName();
-		const FString& GetInputMaxName();
+		const FVertexName& GetInputNextTriggerName();
+		const FVertexName& GetInputResetTriggerName();
+		const FVertexName& GetInputSeedName();
+		const FVertexName& GetInputMinName();
+		const FVertexName& GetInputMaxName();
 
 		/** Output vertex names. */
-		const FString& GetOutputOnNextTriggerName();
-		const FString& GetOutputOnResetTriggerName();
-		const FString& GetOutputValueName();
+		const FVertexName& GetOutputOnNextTriggerName();
+		const FVertexName& GetOutputOnResetTriggerName();
+		const FVertexName& GetOutputValueName();
 
 		/** Descriptions. */
 		static FText GetNextTriggerDescription();
diff --git a/Engine/Plugins/Runtime/Metasound/Source/MetasoundStandardNodes/Public/MetasoundSwitchGateNode.h b/Engine/Plugins/Runtime/Metasound/Source/MetasoundStandardNodes/Public/MetasoundSwitchGateNode.h
--- a/Engine/Plugins/Runtime/Metasound/Source/MetasoundStandardNodes/Public/MetasoundSwitchGateNode.h
+++ b/Engine/Plugins/Runtime/Metasound/Source/MetasoundStandardNodes/Public/MetasoundSwitchGateNode.h
@@ -1,5 +1,4 @@
 // Copyright Epic Games, Inc. All Rights Reserved.
-
 #pragma once
 
 #include "CoreMinimal.h"
@@ -11,6 +10,8 @@
 #include "MetasoundFacade.h"
 #include "MetasoundNode.h"
 #include "MetasoundOperatorInterface.h"
+#include "MetasoundVertex.h"
+
 
 namespace Metasound
 {
@@ -27,7 +28,7 @@ namespace Metasound
 	class METASOUNDSTANDARDNODES_API FSwitchNode : public FNodeFacade
 	{
 	public:
-		FSwitchNode(const FString& InName, const FGuid& InInstanceID);
+		FSwitchNode(const FVertexName& InName, const FGuid& InInstanceID);
 		FSwitchNode(const FNodeInitData& InInitData);
 	};
 
@@ -39,7 +40,7 @@ namespace Metasound
 	class METASOUNDSTANDARDNODES_API FGateNode : public FNodeFacade
 	{
 	public:
-		FGateNode(const FString& InName, const FGuid& InInstanceID);
+		FGateNode(const FVertexName& InName, const FGuid& InInstanceID);
 		FGateNode(const FNodeInitData& InInitData);
 	};
 }
diff --git a/Engine/Plugins/Runtime/Metasound/Source/MetasoundStandardNodes/Public/MetasoundTriggerOnThresholdNode.h b/Engine/Plugins/Runtime/Metasound/Source/MetasoundStandardNodes/Public/MetasoundTriggerOnThresholdNode.h
--- a/Engine/Plugins/Runtime/Metasound/Source/MetasoundStandardNodes/Public/MetasoundTriggerOnThresholdNode.h
+++ b/Engine/Plugins/Runtime/Metasound/Source/MetasoundStandardNodes/Public/MetasoundTriggerOnThresholdNode.h
@@ -1,5 +1,4 @@
 // Copyright Epic Games, Inc. All Rights Reserved.
-
 #pragma once
 
 #include "CoreMinimal.h"
@@ -11,13 +10,15 @@
 #include "MetasoundFacade.h"
 #include "MetasoundNode.h"
 #include "MetasoundOperatorInterface.h"
+#include "MetasoundVertex.h"
+
 
 namespace Metasound
 {
 	class METASOUNDSTANDARDNODES_API FTriggerOnThresholdNode : public FNodeFacade
 	{
 	public:
-		FTriggerOnThresholdNode(const FString& InName, const FGuid& InInstanceID, float InDefaultThreshold);
+		FTriggerOnThresholdNode(const FVertexName& InName, const FGuid& InInstanceID, float InDefaultThreshold);
 		FTriggerOnThresholdNode(const FNodeInitData& InInitData);
 
 		float GetDefaultThreshold() const { return DefaultThreshold; }
diff --git a/Engine/Plugins/Runtime/Metasound/Source/MetasoundStandardNodes/Public/MetasoundWaveWriterNode.h b/Engine/Plugins/Runtime/Metasound/Source/MetasoundStandardNodes/Public/MetasoundWaveWriterNode.h
--- a/Engine/Plugins/Runtime/Metasound/Source/MetasoundStandardNodes/Public/MetasoundWaveWriterNode.h
+++ b/Engine/Plugins/Runtime/Metasound/Source/MetasoundStandardNodes/Public/MetasoundWaveWriterNode.h
@@ -1,18 +1,18 @@
 // Copyright Epic Games, Inc. All Rights Reserved.
-
 #pragma once
 
 #include "CoreMinimal.h"
 #include "MetasoundFacade.h"
+#include "MetasoundVertex.h"
+
 
 namespace Metasound
 {
 	class METASOUNDSTANDARDNODES_API FWaveWriterNode : public FNodeFacade
 	{
-	public:
-		FWaveWriterNode(const FString& InName, const FGuid& InInstanceID);
-		FWaveWriterNode(const FNodeInitData& InInitData);
-	private:
+		public:
+			FWaveWriterNode(const FVertexName& InName, const FGuid& InInstanceID);
+			FWaveWriterNode(const FNodeInitData& InInitData);
 	};
 } // namespace Metasound
 
