Fix bug with renaming MetaSound inputs/outputs introduced by backing out change locally to standardize names (in favor of waiting until we can address changes required in MS core for such a change)

diff --git a/Engine/Plugins/Runtime/Metasound/Source/MetasoundFrontend/Private/MetasoundFrontendStandardController.cpp b/Engine/Plugins/Runtime/Metasound/Source/MetasoundFrontend/Private/MetasoundFrontendStandardController.cpp
--- a/Engine/Plugins/Runtime/Metasound/Source/MetasoundFrontend/Private/MetasoundFrontendStandardController.cpp
+++ b/Engine/Plugins/Runtime/Metasound/Source/MetasoundFrontend/Private/MetasoundFrontendStandardController.cpp
@@ -2138,6 +2138,16 @@ namespace Metasound
 			if (FMetasoundFrontendNode* Node = NodePtr.Get())
 			{
 				Node->Name = InName;
+
+				for (FMetasoundFrontendVertex& Vertex : Node->Interface.Inputs)
+				{
+					Vertex.Name = InName;
+				}
+
+				for (FMetasoundFrontendVertex& Vertex : Node->Interface.Outputs)
+				{
+					Vertex.Name = InName;
+				}
 			}
 
 			// TODO: can we remove the const cast by constructing output nodes with a non-const access to class outputs?
@@ -2429,6 +2439,16 @@ namespace Metasound
 			if (FMetasoundFrontendNode* Node = NodePtr.Get())
 			{
 				Node->Name = InName;
+
+				for (FMetasoundFrontendVertex& Vertex : Node->Interface.Inputs)
+				{
+					Vertex.Name = InName;
+				}
+
+				for (FMetasoundFrontendVertex& Vertex : Node->Interface.Outputs)
+				{
+					Vertex.Name = InName;
+				}
 			}
 
 			if (FMetasoundFrontendClassInput* ClassInput = ConstCastAccessPtr<FClassInputAccessPtr>(OwningGraphClassInputPtr).Get())
diff --git a/Engine/Plugins/Runtime/Metasound/Source/MetasoundFrontend/Private/MetasoundFrontendStandardController.h b/Engine/Plugins/Runtime/Metasound/Source/MetasoundFrontend/Private/MetasoundFrontendStandardController.h
--- a/Engine/Plugins/Runtime/Metasound/Source/MetasoundFrontend/Private/MetasoundFrontendStandardController.h
+++ b/Engine/Plugins/Runtime/Metasound/Source/MetasoundFrontend/Private/MetasoundFrontendStandardController.h
@@ -398,8 +398,6 @@ namespace Metasound
 
 			const FVertexName& GetNodeName() const override;
 
-			// This only exists to allow for transform fix-ups to easily cleanup input/output
-			// node names & should not be used for standard edit or runtime callsites.
 			void SetNodeName(const FVertexName& InName) override { checkNoEntry(); }
 
 			/** Returns the readable display name of the given node (Used only within MetaSound
