- Add editor categories to frontend (utilized in subsequent CL)
- Minor clean-up with copy/pasta'ed category name
- Remove Audio category as it is nondescript in favor of new Mix category

diff --git a/Engine/Plugins/Runtime/Metasound/Source/MetasoundFrontend/Private/MetasoundFrontendNodesCategories.cpp b/Engine/Plugins/Runtime/Metasound/Source/MetasoundFrontend/Private/MetasoundFrontendNodesCategories.cpp
--- a/Engine/Plugins/Runtime/Metasound/Source/MetasoundFrontend/Private/MetasoundFrontendNodesCategories.cpp
+++ b/Engine/Plugins/Runtime/Metasound/Source/MetasoundFrontend/Private/MetasoundFrontendNodesCategories.cpp
@@ -8,7 +8,20 @@ namespace Metasound
 {
 	namespace NodeCategories
 	{
-		const FText Conversions  = { LOCTEXT("Metasound_ConversionsCategory", "Conversions") };
+		// Base category for all conversion functions.
+		const FText Conversions  = LOCTEXT("Metasound_ConversionsCategory", "Conversions");
+
+		// Base category for all standard MetaSound functions.
+		const FText Functions = LOCTEXT("Metasound_FunctionsCategory", "Functions");
+
+		// Base category reserved for all asset graph references.
+		const FText Graphs = LOCTEXT("Metasound_GraphsCategory", "Graphs");
+
+		// Base category for all inputs
+		const FText Inputs = LOCTEXT("Metasound_InputsCategory", "Inputs");
+
+		// Base cateogry for all outputs
+		const FText Outputs = LOCTEXT("Metasound_OutputsCategory", "Outputs");
 	}
 }
 
diff --git a/Engine/Plugins/Runtime/Metasound/Source/MetasoundFrontend/Public/MetasoundFrontendNodesCategories.h b/Engine/Plugins/Runtime/Metasound/Source/MetasoundFrontend/Public/MetasoundFrontendNodesCategories.h
--- a/Engine/Plugins/Runtime/Metasound/Source/MetasoundFrontend/Public/MetasoundFrontendNodesCategories.h
+++ b/Engine/Plugins/Runtime/Metasound/Source/MetasoundFrontend/Public/MetasoundFrontendNodesCategories.h
@@ -9,5 +9,9 @@ namespace Metasound
 	namespace NodeCategories
 	{
 		extern const FText METASOUNDFRONTEND_API Conversions;
-	}
-}
+		extern const FText METASOUNDFRONTEND_API Functions;
+		extern const FText METASOUNDFRONTEND_API Graphs;
+		extern const FText METASOUNDFRONTEND_API Inputs;
+		extern const FText METASOUNDFRONTEND_API Outputs;
+	} // namespace NodeCategories
+} // namespace Metasound
diff --git a/Engine/Plugins/Runtime/Metasound/Source/MetasoundStandardNodes/Private/MetasoundLinearToLogFrequencyNode.cpp b/Engine/Plugins/Runtime/Metasound/Source/MetasoundStandardNodes/Private/MetasoundLinearToLogFrequencyNode.cpp
--- a/Engine/Plugins/Runtime/Metasound/Source/MetasoundStandardNodes/Private/MetasoundLinearToLogFrequencyNode.cpp
+++ b/Engine/Plugins/Runtime/Metasound/Source/MetasoundStandardNodes/Private/MetasoundLinearToLogFrequencyNode.cpp
@@ -19,8 +19,8 @@ namespace Metasound
 		METASOUND_PARAM(InputValue, "Value", "Linear input value to map to log frequency output. Input and output are clamped to specified domain and range.");
 		METASOUND_PARAM(InputDomainMin, "Min Domain", "Min domain for the input value.");
 		METASOUND_PARAM(InputDomainMax, "Max Domain", "Max domain for the input value.");
-		METASOUND_PARAM(InputRangeMin, "Min Range", "Min positive range for the output frequency (hz) value.");
-		METASOUND_PARAM(InputRangeMax, "Max Range", "Max poasitive range for the output frequency (hz) value.");
+		METASOUND_PARAM(InputRangeMin, "Min Range", "Min positive range for the output frequency (Hz) value.");
+		METASOUND_PARAM(InputRangeMax, "Max Range", "Max positive range for the output frequency (Hz) value.");
 		METASOUND_PARAM(OutputFreq, "Frequency", "Output frequency value in hertz that is the log frequency of the input value.");
 	}
 
@@ -45,7 +45,7 @@ namespace Metasound
 				Info.Author = PluginAuthor;
 				Info.PromptIfMissing = PluginNodeMissingPrompt;
 				Info.DefaultInterface = GetVertexInterface();
-				Info.CategoryHierarchy.Emplace(NodeCategories::Music);
+				Info.CategoryHierarchy.Emplace(NodeCategories::Math);
 
 				return Info;
 			};
@@ -92,7 +92,7 @@ namespace Metasound
 		}
 
 		FLinearToLogFrequencyOperator(
-			const TDataReadReference<float>& InValue, 
+			const TDataReadReference<float>& InValue,
 			const TDataReadReference<float>& InDomainMin,
 			const TDataReadReference<float>& InDomainMax,
 			const TDataReadReference<float>& InRangeMin,
diff --git a/Engine/Plugins/Runtime/Metasound/Source/MetasoundStandardNodes/Private/MetasoundMixerNode.cpp b/Engine/Plugins/Runtime/Metasound/Source/MetasoundStandardNodes/Private/MetasoundMixerNode.cpp
--- a/Engine/Plugins/Runtime/Metasound/Source/MetasoundStandardNodes/Private/MetasoundMixerNode.cpp
+++ b/Engine/Plugins/Runtime/Metasound/Source/MetasoundStandardNodes/Private/MetasoundMixerNode.cpp
@@ -217,8 +217,8 @@ namespace Metasound
 				PluginAuthor,
 				PluginNodeMissingPrompt,
 				InDefaultInterface,
-				{ NodeCategories::Audio },
-				{ LOCTEXT("Metasound_AudioMixerKeyword", "Mixer")},
+				{ NodeCategories::Mix },
+				{ LOCTEXT("Metasound_AudioMixerKeyword", "Mixer") },
 				FNodeDisplayStyle{}
 			};
 
diff --git a/Engine/Plugins/Runtime/Metasound/Source/MetasoundStandardNodes/Private/MetasoundStandardNodesCategories.cpp b/Engine/Plugins/Runtime/Metasound/Source/MetasoundStandardNodes/Private/MetasoundStandardNodesCategories.cpp
--- a/Engine/Plugins/Runtime/Metasound/Source/MetasoundStandardNodes/Private/MetasoundStandardNodesCategories.cpp
+++ b/Engine/Plugins/Runtime/Metasound/Source/MetasoundStandardNodes/Private/MetasoundStandardNodesCategories.cpp
@@ -8,7 +8,6 @@ namespace Metasound
 {
 	namespace NodeCategories
 	{
-		const FText Audio = { LOCTEXT("Metasound_AudioCategory", "Audio") };
 		const FText Debug = { LOCTEXT("Metasound_DebugCategory", "Debug") };
 		const FText Delays = { LOCTEXT("Metasound_EffectsCategory", "Delays") };
 		const FText Dynamics = { LOCTEXT("Metasound_DynamicsCategory", "Dynamics") };
@@ -17,6 +16,7 @@ namespace Metasound
 		const FText Generators = { LOCTEXT("Metasound_GeneratorsCategory", "Generators") };
 		const FText Io = { LOCTEXT("Metasound_IoCategory", "External IO") };
 		const FText Math = { LOCTEXT("Metasound_MathCategory", "Math") };
+		const FText Mix = { LOCTEXT("Metasound_MixCategory", "Mix") };
 		const FText Music = { LOCTEXT("Metasound_MusicCategory", "Music") };
 		const FText RandomUtils = { LOCTEXT("Metasound_RandomCategory", "Random") };
 		const FText Spatialization = { LOCTEXT("Metasound_SpatializationCategory", "Spatialization") };
diff --git a/Engine/Plugins/Runtime/Metasound/Source/MetasoundStandardNodes/Public/MetasoundStandardNodesCategories.h b/Engine/Plugins/Runtime/Metasound/Source/MetasoundStandardNodes/Public/MetasoundStandardNodesCategories.h
--- a/Engine/Plugins/Runtime/Metasound/Source/MetasoundStandardNodes/Public/MetasoundStandardNodesCategories.h
+++ b/Engine/Plugins/Runtime/Metasound/Source/MetasoundStandardNodes/Public/MetasoundStandardNodesCategories.h
@@ -10,7 +10,6 @@ namespace Metasound
 {
 	namespace NodeCategories
 	{
-		extern const FText METASOUNDSTANDARDNODES_API Audio;
 		extern const FText METASOUNDSTANDARDNODES_API Debug;
 		extern const FText METASOUNDSTANDARDNODES_API Delays;
 		extern const FText METASOUNDSTANDARDNODES_API Dynamics;
@@ -19,6 +18,7 @@ namespace Metasound
 		extern const FText METASOUNDSTANDARDNODES_API Generators;
 		extern const FText METASOUNDSTANDARDNODES_API Io;
 		extern const FText METASOUNDSTANDARDNODES_API Math;
+		extern const FText METASOUNDSTANDARDNODES_API Mix;
 		extern const FText METASOUNDSTANDARDNODES_API Music;
 		extern const FText METASOUNDSTANDARDNODES_API RandomUtils;
 		extern const FText METASOUNDSTANDARDNODES_API Spatialization;
