Registering all MetaSound asset classes on editor load asyncronously
#rnx

diff --git a/Engine/Plugins/Runtime/Metasound/Source/MetasoundEditor/Private/MetasoundEditorModule.cpp b/Engine/Plugins/Runtime/Metasound/Source/MetasoundEditor/Private/MetasoundEditorModule.cpp
--- a/Engine/Plugins/Runtime/Metasound/Source/MetasoundEditor/Private/MetasoundEditorModule.cpp
+++ b/Engine/Plugins/Runtime/Metasound/Source/MetasoundEditor/Private/MetasoundEditorModule.cpp
@@ -7,11 +7,11 @@
 #include "AssetRegistryModule.h"
 #include "AssetTypeActions_Base.h"
 #include "Brushes/SlateImageBrush.h"
-#include "CoreMinimal.h"
 #include "EdGraph/EdGraphNode.h"
 #include "EdGraph/EdGraphPin.h"
 #include "EdGraphUtilities.h"
 #include "EditorStyleSet.h"
+#include "HAL/IConsoleManager.h"
 #include "IDetailCustomization.h"
 #include "ISettingsModule.h"
 #include "Metasound.h"
@@ -45,11 +45,20 @@
 #include "Styling/SlateTypes.h"
 #include "Templates/SharedPointer.h"
 #include "UObject/UObjectGlobals.h"
-#include "Animation/Skeleton.h"
+
 
 DEFINE_LOG_CATEGORY(LogMetasoundEditor);
 
 
+static int32 MetaSoundEditorAsyncRegistrationEnabledCVar = 1;
+FAutoConsoleVariableRef CVarMetaSoundEditorAsyncRegistrationEnabled(
+	TEXT("au.MetaSounds.Editor.AsyncRegistrationEnabled"),
+	MetaSoundEditorAsyncRegistrationEnabledCVar,
+	TEXT("Enable registering all MetaSound asset classes asyncronously on editor load.\n")
+	TEXT("0: Disabled, !0: Enabled (default)"),
+	ECVF_Default);
+
+
 namespace Metasound
 {
 	namespace Editor
@@ -150,7 +159,6 @@ namespace Metasound
 		{
 			void AddOrUpdateClassRegistryAsset(const FAssetData& InAssetData)
 			{
-				using namespace Metasound;
 				using namespace Metasound::Frontend;
 
 				if (IsMetaSoundAssetClass(InAssetData.AssetClass))
@@ -163,7 +171,8 @@ namespace Metasound
 					// Use the editor version of `RegisterGraphWithFrontend` so it re-registers any open MetaSound editors
 					AssetSubsystem->AddOrUpdateAsset(InAssetData, false /* bRegisterWithFrontend */);
 					
-					// May no longer be necessary, but left for safety
+					// Loading all assets necessary only in editor to register &
+					// populate potential graphs to reference in MetaSound editor.
 					if (InAssetData.IsAssetLoaded())
 					{
 						if (UObject* AssetObject = InAssetData.GetAsset())
@@ -171,6 +180,19 @@ namespace Metasound
 							FGraphBuilder::RegisterGraphWithFrontend(*AssetObject);
 						}
 					}
+					else if (MetaSoundEditorAsyncRegistrationEnabledCVar)
+					{
+						FSoftObjectPath Path = InAssetData.ToSoftObjectPath();
+						LoadPackageAsync(Path.GetLongPackageName(), FLoadPackageAsyncDelegate::CreateLambda(
+							[this, ObjectPath = MoveTemp(Path)](const FName& PackageName, UPackage* Package, EAsyncLoadingResult::Type Result)
+							{
+								if (Result == EAsyncLoadingResult::Succeeded)
+								{
+									FGraphBuilder::RegisterGraphWithFrontend(*ObjectPath.ResolveObject());
+								}
+							}
+						));
+					}
 				}
 			}
 
