New MetaSound ActionGraph node picker
- Fix rename from the MetaSound Editor API ActionMenu
#rnx

diff --git a/Engine/Plugins/Runtime/Metasound/Source/MetasoundEditor/Private/MetasoundDetailCustomization.h b/Engine/Plugins/Runtime/Metasound/Source/MetasoundEditor/Private/MetasoundDetailCustomization.h
--- a/Engine/Plugins/Runtime/Metasound/Source/MetasoundEditor/Private/MetasoundDetailCustomization.h
+++ b/Engine/Plugins/Runtime/Metasound/Source/MetasoundEditor/Private/MetasoundDetailCustomization.h
@@ -29,8 +29,6 @@ namespace Metasound
 			FName GetMetadataPropertyPath() const;
 
 			FName DocumentPropertyName;
-
-			TSharedPtr<SGraphActionMenu> GraphActionMenu;
 		};
 	} // namespace Editor
 } // namespace Metasound
\ No newline at end of file
diff --git a/Engine/Plugins/Runtime/Metasound/Source/MetasoundEditor/Private/MetasoundEditor.cpp b/Engine/Plugins/Runtime/Metasound/Source/MetasoundEditor/Private/MetasoundEditor.cpp
--- a/Engine/Plugins/Runtime/Metasound/Source/MetasoundEditor/Private/MetasoundEditor.cpp
+++ b/Engine/Plugins/Runtime/Metasound/Source/MetasoundEditor/Private/MetasoundEditor.cpp
@@ -43,6 +43,7 @@
 #include "Modules/ModuleManager.h"
 #include "PropertyEditorModule.h"
 #include "ScopedTransaction.h"
+#include "SMetasoundActionMenu.h"
 #include "SMetasoundPalette.h"
 #include "SNodePanel.h"
 #include "Templates/SharedPointer.h"
@@ -175,7 +176,8 @@ namespace Metasound
 				{
 					if (UMetasoundEditorGraphVariable* GraphVariable = MetasoundAction->GetVariable())
 					{
-						GraphVariable->SetDisplayName(InNewText);
+						GraphVariable->SetDisplayName(FText());
+						GraphVariable->SetNodeName(*InNewText.ToString());
 					}
 				}
 			}
@@ -1191,10 +1193,11 @@ namespace Metasound
 					}));
 			}
 
-			SGraphEditor::FGraphEditorEvents InEvents;
-			InEvents.OnSelectionChanged = SGraphEditor::FOnSelectionChanged::CreateSP(this, &FEditor::OnSelectedNodesChanged);
-			InEvents.OnTextCommitted = FOnNodeTextCommitted::CreateSP(this, &FEditor::OnNodeTitleCommitted);
-			InEvents.OnNodeDoubleClicked = FSingleNodeEvent::CreateSP(this, &FEditor::ExecuteNode);
+			SGraphEditor::FGraphEditorEvents GraphEvents;
+			GraphEvents.OnCreateActionMenu = SGraphEditor::FOnCreateActionMenu::CreateSP(this, &FEditor::OnCreateGraphActionMenu);
+			GraphEvents.OnNodeDoubleClicked = FSingleNodeEvent::CreateSP(this, &FEditor::ExecuteNode);
+			GraphEvents.OnSelectionChanged = SGraphEditor::FOnSelectionChanged::CreateSP(this, &FEditor::OnSelectedNodesChanged);
+			GraphEvents.OnTextCommitted = FOnNodeTextCommitted::CreateSP(this, &FEditor::OnNodeTitleCommitted);
 
 			FMetasoundAssetBase* MetasoundAsset = IMetasoundUObjectRegistry::Get().GetObjectAsAssetBase(Metasound);
 			check(MetasoundAsset);
@@ -1203,7 +1206,7 @@ namespace Metasound
 				.AdditionalCommands(GraphEditorCommands)
 				.Appearance(this, &FEditor::GetGraphAppearance)
 				.AutoExpandActionMenu(true)
-				.GraphEvents(InEvents)
+				.GraphEvents(GraphEvents)
 				.GraphToEdit(MetasoundAsset->GetGraph())
 				.IsEditable(this, &FEditor::IsGraphEditable)
 				.ShowGraphStateOverlay(false);
@@ -1941,6 +1944,19 @@ namespace Metasound
 			return FReply::Unhandled();
 		}
 
+		FActionMenuContent FEditor::OnCreateGraphActionMenu(UEdGraph* InGraph, const FVector2D& InNodePosition, const TArray<UEdGraphPin*>& InDraggedPins, bool bAutoExpand, SGraphEditor::FActionMenuClosed InOnMenuClosed)
+		{
+			TSharedRef<SMetasoundActionMenu> ActionMenu = SNew(SMetasoundActionMenu)
+				.AutoExpandActionMenu(bAutoExpand)
+				.Graph(&GetMetaSoundGraphChecked())
+				.NewNodePosition(InNodePosition)
+				.DraggedFromPins(InDraggedPins)
+				.OnClosedCallback(InOnMenuClosed);
+// 				.OnCloseReason(this, &FEditor::OnGraphActionMenuClosed);
+
+			return FActionMenuContent(ActionMenu, ActionMenu->GetFilterTextBox());
+		}
+
 		void FEditor::OnActionSelected(const TArray<TSharedPtr<FEdGraphSchemaAction>>& InActions, ESelectInfo::Type InSelectionType)
 		{
 			if (InSelectionType == ESelectInfo::OnMouseClick || InSelectionType == ESelectInfo::OnKeyPress || InSelectionType == ESelectInfo::OnNavigation || InActions.IsEmpty())
diff --git a/Engine/Plugins/Runtime/Metasound/Source/MetasoundEditor/Private/MetasoundEditor.h b/Engine/Plugins/Runtime/Metasound/Source/MetasoundEditor/Private/MetasoundEditor.h
--- a/Engine/Plugins/Runtime/Metasound/Source/MetasoundEditor/Private/MetasoundEditor.h
+++ b/Engine/Plugins/Runtime/Metasound/Source/MetasoundEditor/Private/MetasoundEditor.h
@@ -138,6 +138,7 @@ namespace Metasound
 			FText GetFilterText() const;
 			bool HandleActionMatchesName(FEdGraphSchemaAction* InAction, const FName& InName) const;
 			FReply OnActionDragged(const TArray<TSharedPtr<FEdGraphSchemaAction>>& InActions, const FPointerEvent& MouseEvent);
+			FActionMenuContent OnCreateGraphActionMenu(UEdGraph* InGraph, const FVector2D& InNodePosition, const TArray<UEdGraphPin*>& InDraggedPins, bool bAutoExpand, SGraphEditor::FActionMenuClosed InOnMenuClosed);
 			void OnActionSelected(const TArray<TSharedPtr<FEdGraphSchemaAction>>& InActions, ESelectInfo::Type InSelectionType);
 			FReply OnAddButtonClickedOnSection(int32 InSectionID);
 			TSharedRef<SWidget> OnGetMenuSectionWidget(TSharedRef<SWidget> RowWidget, int32 InSectionID);
diff --git a/Engine/Plugins/Runtime/Metasound/Source/MetasoundEditor/Private/MetasoundEditorGraphSchema.cpp b/Engine/Plugins/Runtime/Metasound/Source/MetasoundEditor/Private/MetasoundEditorGraphSchema.cpp
--- a/Engine/Plugins/Runtime/Metasound/Source/MetasoundEditor/Private/MetasoundEditorGraphSchema.cpp
+++ b/Engine/Plugins/Runtime/Metasound/Source/MetasoundEditor/Private/MetasoundEditorGraphSchema.cpp
@@ -51,68 +51,44 @@ namespace Metasound
 {
 	namespace Editor
 	{
-		namespace SchemaPrivate
+		const FText& GetContextGroupDisplayName(EPrimaryContextGroup InContextGroup)
 		{
-			static const FText CategoryDelim = LOCTEXT("MetaSoundActionsCategoryDelim", "|");
-			static const FText KeywordDelim = LOCTEXT("MetaSoundKeywordDelim", " ");
+			switch (InContextGroup)
+			{
+				case EPrimaryContextGroup::Inputs:
+					return NodeCategories::Inputs;
 
-			static const FText InputDisplayNameFormat = LOCTEXT("DisplayNameAddInputFormat", "Get {0}");
-			static const FText InputTooltipFormat = LOCTEXT("TooltipAddInputFormat", "Adds a getter for the input '{0}' to the graph.");
+				case EPrimaryContextGroup::Outputs:
+					return NodeCategories::Outputs;
 
-			static const FText OutputDisplayNameFormat = LOCTEXT("DisplayNameAddOutputFormat", "Set {0}");
-			static const FText OutputTooltipFormat = LOCTEXT("TooltipAddOutputFormat", "Adds a setter for the output '{0}' to the graph.");
+				case EPrimaryContextGroup::Graphs:
+					return NodeCategories::Graphs;
 
-			enum class EPrimaryContextGroup
-			{
-				Default = 0,
-				Inputs,
-				Outputs,
-				Graphs,
-				Functions,
-				Conversions
-			};
+				case EPrimaryContextGroup::Functions:
+					return NodeCategories::Functions;
 
-			const FText& GetContextGroupDisplayName(EPrimaryContextGroup InContextGroup)
-			{
-				switch (InContextGroup)
+				case EPrimaryContextGroup::Conversions:
+					return NodeCategories::Conversions;
+
+				case EPrimaryContextGroup::Common:
+				default:
 				{
-					case EPrimaryContextGroup::Inputs:
-					{
-						static const FText InputGroupName = LOCTEXT("InputActions", "Inputs");
-						return InputGroupName;
-					}
-					case EPrimaryContextGroup::Outputs:
-					{
-						static const FText OutputGroupName = LOCTEXT("OutputActions", "Outputs");
-						return OutputGroupName;
-					}
+					return FText::GetEmpty();
+				}
+			}
+		}
 
-					case EPrimaryContextGroup::Graphs:
-					{
-						static const FText GraphGroupName = LOCTEXT("GraphActions", "Graphs");
-						return GraphGroupName;
-					}
+		namespace SchemaPrivate
+		{
 
-					case EPrimaryContextGroup::Functions:
-					{
-						static const FText FunctionGroupName = LOCTEXT("FunctionActions", "Functions");
-						return FunctionGroupName;
-					}
+			static const FText CategoryDelim = LOCTEXT("MetaSoundActionsCategoryDelim", "|");
+			static const FText KeywordDelim = LOCTEXT("MetaSoundKeywordDelim", " ");
 
-					case EPrimaryContextGroup::Conversions:
-					{
-						static const FText ConversionGroupName = LOCTEXT("ConversionActions", "Conversions");
-						return ConversionGroupName;
-					}
+			static const FText InputDisplayNameFormat = LOCTEXT("DisplayNameAddInputFormat", "Get {0}");
+			static const FText InputTooltipFormat = LOCTEXT("TooltipAddInputFormat", "Adds a getter for the input '{0}' to the graph.");
 
-					case EPrimaryContextGroup::Default:
-					default:
-					{
-						checkNoEntry();
-						return FText::GetEmpty();
-					}
-				}
-			}
+			static const FText OutputDisplayNameFormat = LOCTEXT("DisplayNameAddOutputFormat", "Set {0}");
+			static const FText OutputTooltipFormat = LOCTEXT("TooltipAddOutputFormat", "Adds a setter for the output '{0}' to the graph.");
 
 			bool TryConnectNewNodeToPin(UEdGraphNode& NewGraphNode, UEdGraphPin* FromPin)
 			{
@@ -189,7 +165,7 @@ namespace Metasound
 						const FText NodeDisplayName = NodeHandle->GetDisplayName();
 						const FText Tooltip = FText::Format(InQuery.TooltipFormat, NodeDisplayName);
 						const FText DisplayName = FText::Format(InQuery.DisplayNameFormat, NodeDisplayName);
-						TSharedPtr<TAction> NewNodeAction = MakeShared<TAction>(GroupName, DisplayName, NodeHandle->GetID(), Tooltip, static_cast<int32>(InQuery.ContextGroup));
+						TSharedPtr<TAction> NewNodeAction = MakeShared<TAction>(GroupName, DisplayName, NodeHandle->GetID(), Tooltip, InQuery.ContextGroup);
 						InQuery.ActionMenuBuilder.AddAction(NewNodeAction);
 					}
 				}
@@ -257,8 +233,8 @@ UEdGraphNode* FMetasoundGraphSchemaAction_NewNode::PerformAction(UEdGraph* Paren
 	return nullptr;
 }
 
-FMetasoundGraphSchemaAction_NewInput::FMetasoundGraphSchemaAction_NewInput(FText InNodeCategory, FText InDisplayName, FGuid InNodeID, FText InToolTip, const int32 InGrouping)
-	: FMetasoundGraphSchemaAction(MoveTemp(InNodeCategory), MoveTemp(InDisplayName), MoveTemp(InToolTip), InGrouping)
+FMetasoundGraphSchemaAction_NewInput::FMetasoundGraphSchemaAction_NewInput(FText InNodeCategory, FText InDisplayName, FGuid InNodeID, FText InToolTip, Metasound::Editor::EPrimaryContextGroup InGroup)
+	: FMetasoundGraphSchemaAction(MoveTemp(InNodeCategory), MoveTemp(InDisplayName), MoveTemp(InToolTip), InGroup)
 	, NodeID(InNodeID)
 {
 }
@@ -321,8 +297,13 @@ UEdGraphNode* FMetasoundGraphSchemaAction_NewInput::PerformAction(UEdGraph* Pare
 	return nullptr;
 }
 
-FMetasoundGraphSchemaAction_PromoteToInput::FMetasoundGraphSchemaAction_PromoteToInput(FText InNodeCategory, FText InDisplayName, FText InToolTip, const int32 InGrouping)
-	: FMetasoundGraphSchemaAction_NewInput(MoveTemp(InNodeCategory), MoveTemp(InDisplayName), FGuid(), MoveTemp(InToolTip), InGrouping)
+FMetasoundGraphSchemaAction_PromoteToInput::FMetasoundGraphSchemaAction_PromoteToInput()
+	: FMetasoundGraphSchemaAction_NewInput(
+		FText(),
+		LOCTEXT("PromoteToInputName", "Promote To Graph Input"),
+		FGuid(),
+		LOCTEXT("PromoteToInputTooltip", "Promotes node input to graph input"),
+		Metasound::Editor::EPrimaryContextGroup::Common)
 {
 }
 
@@ -374,8 +355,8 @@ UEdGraphNode* FMetasoundGraphSchemaAction_PromoteToInput::PerformAction(UEdGraph
 	return nullptr;
 }
 
-FMetasoundGraphSchemaAction_NewOutput::FMetasoundGraphSchemaAction_NewOutput(FText InNodeCategory, FText InDisplayName, FGuid InOutputNodeID, FText InToolTip, const int32 InGrouping)
-	: FMetasoundGraphSchemaAction(MoveTemp(InNodeCategory), MoveTemp(InDisplayName), MoveTemp(InToolTip), InGrouping)
+FMetasoundGraphSchemaAction_NewOutput::FMetasoundGraphSchemaAction_NewOutput(FText InNodeCategory, FText InDisplayName, FGuid InOutputNodeID, FText InToolTip, Metasound::Editor::EPrimaryContextGroup InGroup)
+	: FMetasoundGraphSchemaAction(MoveTemp(InNodeCategory), MoveTemp(InDisplayName), MoveTemp(InToolTip), InGroup)
 	, NodeID(InOutputNodeID)
 {
 }
@@ -434,8 +415,13 @@ UEdGraphNode* FMetasoundGraphSchemaAction_NewOutput::PerformAction(UEdGraph* Par
 	return nullptr;
 }
 
-FMetasoundGraphSchemaAction_PromoteToOutput::FMetasoundGraphSchemaAction_PromoteToOutput(FText InNodeCategory, FText InDisplayName, FText InToolTip, const int32 InGrouping)
-	: FMetasoundGraphSchemaAction_NewOutput(MoveTemp(InNodeCategory), MoveTemp(InDisplayName), FGuid(), MoveTemp(InToolTip), InGrouping)
+FMetasoundGraphSchemaAction_PromoteToOutput::FMetasoundGraphSchemaAction_PromoteToOutput()
+	: FMetasoundGraphSchemaAction_NewOutput(
+		FText(),
+		LOCTEXT("PromoteToOutputName", "Promote To Graph Output"),
+		FGuid(),
+		LOCTEXT("PromoteToOutputTooltip", "Promotes node output to graph output"),
+		Metasound::Editor::EPrimaryContextGroup::Common)
 {
 }
 
@@ -626,18 +612,12 @@ void UMetasoundEditorGraphSchema::GetGraphContextActions(FGraphContextMenuBuilde
 				return bHasOutputOfType;
 			});
 
-			TSharedPtr<FMetasoundGraphSchemaAction_PromoteToInput> NewNodeAction = MakeShared<FMetasoundGraphSchemaAction_PromoteToInput>(
-				SchemaPrivate::GetContextGroupDisplayName(SchemaPrivate::EPrimaryContextGroup::Inputs),
-				LOCTEXT("PromoteToInputName", "Promote To Graph Input"),
-				LOCTEXT("PromoteToInputTooltip", "Promotes node input to graph input"),
-				static_cast<int32>(SchemaPrivate::EPrimaryContextGroup::Inputs));
-
-			static_cast<FGraphActionMenuBuilder&>(ContextMenuBuilder).AddAction(NewNodeAction);
+			static_cast<FGraphActionMenuBuilder&>(ContextMenuBuilder).AddAction(MakeShared<FMetasoundGraphSchemaAction_PromoteToInput>());
 		}
 
 		if (FromPin->Direction == EGPD_Output)
 		{
-			Frontend::FConstOutputHandle OutputHandle = Editor::FGraphBuilder::GetConstOutputHandleFromPin(FromPin);
+			FConstOutputHandle OutputHandle = FGraphBuilder::GetConstOutputHandleFromPin(FromPin);
 			ClassFilters.InputFilterFunction = [OutputHandle](const FMetasoundFrontendClassInput& InInput)
 			{
 				return InInput.TypeName == OutputHandle->GetDataType();
@@ -655,13 +635,7 @@ void UMetasoundEditorGraphSchema::GetGraphContextActions(FGraphContextMenuBuilde
 				return bHasInputOfType;
 			});
 
-			TSharedPtr<FMetasoundGraphSchemaAction_PromoteToOutput> NewNodeAction = MakeShared<FMetasoundGraphSchemaAction_PromoteToOutput>(
-				SchemaPrivate::GetContextGroupDisplayName(SchemaPrivate::EPrimaryContextGroup::Outputs),
-				LOCTEXT("PromoteToOutputName", "Promote To Graph Output"),
-				LOCTEXT("PromoteToOutputTooltip", "Promotes node output to graph output"),
-				static_cast<int32>(SchemaPrivate::EPrimaryContextGroup::Outputs));
-
-			static_cast<FGraphActionMenuBuilder&>(ContextMenuBuilder).AddAction(NewNodeAction);
+			static_cast<FGraphActionMenuBuilder&>(ContextMenuBuilder).AddAction(MakeShared<FMetasoundGraphSchemaAction_PromoteToOutput>());
 		}
 	}
 	else
@@ -669,7 +643,7 @@ void UMetasoundEditorGraphSchema::GetGraphContextActions(FGraphContextMenuBuilde
 		TSharedPtr<FEditor> MetasoundEditor = FGraphBuilder::GetEditorForGraph(*ContextMenuBuilder.CurrentGraph);
 		if (MetasoundEditor.IsValid() && MetasoundEditor->CanPasteNodes())
 		{
-			TSharedPtr<FMetasoundGraphSchemaAction_Paste> NewAction = MakeShared<FMetasoundGraphSchemaAction_Paste>(FText::GetEmpty(), LOCTEXT("PasteHereAction", "Paste here"), FText::GetEmpty(), 0);
+			TSharedPtr<FMetasoundGraphSchemaAction_Paste> NewAction = MakeShared<FMetasoundGraphSchemaAction_Paste>(FText::GetEmpty(), LOCTEXT("PasteHereAction", "Paste here"), FText::GetEmpty(), EPrimaryContextGroup::Common);
 			ContextMenuBuilder.AddAction(NewAction);
 		}
 
@@ -811,7 +785,7 @@ const FPinConnectionResponse UMetasoundEditorGraphSchema::CanCreateConnection(co
 		}
 
 		// Break existing connections on inputs only - multiple output connections are acceptable
-		if (InputPin->LinkedTo.Num() > 0)
+		if (!InputPin->LinkedTo.IsEmpty())
 		{
 			ECanCreateConnectionResponse ReplyBreakOutputs;
 			if (InputPin == PinA)
@@ -1143,7 +1117,7 @@ void UMetasoundEditorGraphSchema::GetConversionActions(FGraphActionMenuBuilder&
 				NodeCategories::Conversions,
 				Metadata.GetDisplayName(),
 				Tooltip,
-				static_cast<int32>(SchemaPrivate::EPrimaryContextGroup::Conversions),
+				EPrimaryContextGroup::Conversions,
 				KeywordsText
 			);
 
@@ -1166,7 +1140,7 @@ void UMetasoundEditorGraphSchema::GetDataTypeInputNodeActions(FGraphContextMenuB
 		ActionMenuBuilder,
 		Inputs,
 		InFilter,
-		SchemaPrivate::EPrimaryContextGroup::Inputs,
+		EPrimaryContextGroup::Inputs,
 		SchemaPrivate::InputDisplayNameFormat,
 		SchemaPrivate::InputTooltipFormat,
 		bShowSelectedActions
@@ -1203,7 +1177,7 @@ void UMetasoundEditorGraphSchema::GetDataTypeOutputNodeActions(FGraphContextMenu
 		ActionMenuBuilder,
 		Outputs,
 		InFilter,
-		SchemaPrivate::EPrimaryContextGroup::Outputs,
+		EPrimaryContextGroup::Outputs,
 		SchemaPrivate::OutputDisplayNameFormat,
 		SchemaPrivate::OutputTooltipFormat,
 		bShowSelectedActions
@@ -1267,18 +1241,18 @@ void UMetasoundEditorGraphSchema::GetFunctionActions(FGraphActionMenuBuilder& Ac
 
 		if (Metadata.GetCategoryHierarchy().IsEmpty() || Metadata.GetCategoryHierarchy()[0].CompareTo(Metasound::NodeCategories::Conversions))
 		{
-			TArray<FText> CategoryHeirarchy { GetContextGroupDisplayName(Path ? EPrimaryContextGroup::Graphs : EPrimaryContextGroup::Functions) };
-			CategoryHeirarchy.Append(Metadata.GetCategoryHierarchy());
+			const EPrimaryContextGroup ContextGroup = Path ? EPrimaryContextGroup::Graphs : EPrimaryContextGroup::Functions;
+			TArray<FText> CategoryHierarchy { GetContextGroupDisplayName(ContextGroup) };
+			CategoryHierarchy.Append(Metadata.GetCategoryHierarchy());
 			const FText KeywordsText = FText::Join(KeywordDelim, Metadata.GetKeywords());
-			const FText CategoryText = FText::Join(CategoryDelim, CategoryHeirarchy);
-			const int32 GroupID = static_cast<int32>(Path ? EPrimaryContextGroup::Graphs : EPrimaryContextGroup::Functions);
+			const FText CategoryText = FText::Join(CategoryDelim, CategoryHierarchy);
 
 			TSharedPtr<FMetasoundGraphSchemaAction_NewNode> NewNodeAction = MakeShared<FMetasoundGraphSchemaAction_NewNode>
 			(
 				CategoryText,
 				Metadata.GetDisplayName(),
 				Tooltip,
-				GroupID,
+				ContextGroup,
 				KeywordsText
 			);
 
@@ -1302,7 +1276,7 @@ void UMetasoundEditorGraphSchema::GetCommentAction(FGraphActionMenuBuilder& Acti
 			const FText MenuDescription = NumSelected > 0 ? LOCTEXT("CreateCommentAction", "Create Comment from Selection") : LOCTEXT("AddCommentAction", "Add Comment...");
 			const FText ToolTip = LOCTEXT("CreateCommentToolTip", "Creates a comment.");
 
-			TSharedPtr<FMetasoundGraphSchemaAction_NewComment> NewAction(new FMetasoundGraphSchemaAction_NewComment(FText::GetEmpty(), MenuDescription, ToolTip, 0));
+			TSharedPtr<FMetasoundGraphSchemaAction_NewComment> NewAction = MakeShared<FMetasoundGraphSchemaAction_NewComment>(FText::GetEmpty(), MenuDescription, ToolTip, EPrimaryContextGroup::Common);
 			ActionMenuBuilder.AddAction(NewAction);
 		}
 	}
diff --git a/Engine/Plugins/Runtime/Metasound/Source/MetasoundEditor/Private/SMetasoundActionMenu.cpp b/Engine/Plugins/Runtime/Metasound/Source/MetasoundEditor/Private/SMetasoundActionMenu.cpp
new file mode 100644
--- /dev/null
+++ b/Engine/Plugins/Runtime/Metasound/Source/MetasoundEditor/Private/SMetasoundActionMenu.cpp
@@ -0,0 +1,324 @@
+// Copyright Epic Games, Inc. All Rights Reserved.
+
+#include "SMetasoundActionMenu.h"
+
+#include "Framework/Application/SlateApplication.h"
+#include "EdGraphSchema_K2.h"
+#include "Editor/EditorPerProjectUserSettings.h"
+#include "EditorStyleSet.h"
+#include "IDocumentation.h"
+#include "MetasoundEditorGraphBuilder.h"
+#include "MetasoundEditorGraphSchema.h"
+#include "MetasoundEditorModule.h"
+#include "MetasoundEditorSettings.h"
+#include "MetasoundFrontendDocument.h"
+#include "MetasoundFrontendRegistries.h"
+#include "SGraphActionMenu.h"
+#include "SMetasoundPalette.h"
+#include "SSubobjectEditor.h"
+#include "Styling/SlateStyleRegistry.h"
+#include "Widgets/Images/SImage.h"
+#include "Widgets/Input/SCheckBox.h"
+#include "Widgets/Input/SComboButton.h"
+#include "Widgets/Layout/SBox.h"
+#include "Widgets/SBoxPanel.h"
+#include "Widgets/SToolTip.h"
+#include "Widgets/Text/STextBlock.h"
+#include "Widgets/Views/SExpanderArrow.h"
+
+#define LOCTEXT_NAMESPACE "MetaSoundEditor"
+
+
+namespace Metasound
+{
+	namespace Editor
+	{
+		// Custom expander to specify our desired padding
+		class SMetasoundActionMenuExpanderArrow : public SExpanderArrow
+		{
+			SLATE_BEGIN_ARGS(SMetasoundActionMenuExpanderArrow)
+			{
+			}
+
+			SLATE_ATTRIBUTE(float, IndentAmount)
+				SLATE_END_ARGS()
+
+		public:
+			void Construct(const FArguments& InArgs, const FCustomExpanderData& ActionMenuData)
+			{
+				OwnerRowPtr = ActionMenuData.TableRow;
+				IndentAmount = InArgs._IndentAmount;
+				ActionPtr = ActionMenuData.RowAction;
+
+				if (ActionPtr.IsValid())
+				{
+					ChildSlot
+					.Padding(TAttribute<FMargin>(this, &SMetasoundActionMenuExpanderArrow::GetCustomIndentPadding))
+					[
+						SNullWidget::NullWidget
+						// TODO: Add favorites
+						// SNew(SMetasoundActionFavoriteToggle, ActionMenuData)
+					];
+				}
+				else
+				{
+					SExpanderArrow::FArguments SuperArgs;
+					SuperArgs._IndentAmount = InArgs._IndentAmount;
+
+					SExpanderArrow::Construct(SuperArgs, ActionMenuData.TableRow);
+				}
+			}
+
+		private:
+			FMargin GetCustomIndentPadding() const
+			{
+				return SExpanderArrow::GetExpanderPadding();
+			}
+
+			TWeakPtr<FEdGraphSchemaAction> ActionPtr;
+		};
+	}
+}
+
+SMetasoundActionMenu::~SMetasoundActionMenu()
+{
+	OnClosedCallback.ExecuteIfBound();
+	OnCloseReasonCallback.ExecuteIfBound(bActionExecuted, false, !DraggedFromPins.IsEmpty());
+}
+
+void SMetasoundActionMenu::Construct(const FArguments& InArgs)
+{
+	using namespace Metasound::Editor;
+	using namespace Metasound::Frontend;
+
+	Graph = InArgs._Graph;
+	DraggedFromPins = InArgs._DraggedFromPins;
+	NewNodePosition = InArgs._NewNodePosition;
+	OnClosedCallback = InArgs._OnClosedCallback;
+	bAutoExpandActionMenu = InArgs._AutoExpandActionMenu;
+	OnCloseReasonCallback = InArgs._OnCloseReason;
+
+	FSlateColor TypeColor;
+	const FSlateBrush* PinBrush = nullptr;
+	if (!DraggedFromPins.IsEmpty())
+	{
+		if (const UEdGraphPin* Pin = DraggedFromPins[0])
+		{
+			FDataTypeRegistryInfo RegistryInfo;
+
+			const IMetasoundEditorModule& EditorModule = FModuleManager::GetModuleChecked<IMetasoundEditorModule>("MetaSoundEditor");
+			if (Pin->Direction == EGPD_Input)
+			{
+				FInputHandle InputHandle = FGraphBuilder::GetInputHandleFromPin(Pin);
+				IDataTypeRegistry::Get().GetDataTypeInfo(InputHandle->GetDataType(), RegistryInfo);
+			}
+			else
+			{
+				FOutputHandle OutputHandle = FGraphBuilder::GetOutputHandleFromPin(Pin);
+				IDataTypeRegistry::Get().GetDataTypeInfo(OutputHandle->GetDataType(), RegistryInfo);
+			}
+
+			TypeColor = FGraphBuilder::GetPinCategoryColor(Pin->PinType);
+
+			if (RegistryInfo.IsArrayType())
+			{
+				PinBrush = FEditorStyle::GetBrush("Graph.ArrayPin.Connected");
+			}
+			else if (Pin->PinType.PinCategory == FGraphBuilder::PinCategoryTrigger)
+			{
+				if (const ISlateStyle* MetasoundStyle = FSlateStyleRegistry::FindSlateStyle("MetaSoundStyle"))
+				{
+					PinBrush = MetasoundStyle->GetBrush(TEXT("MetasoundEditor.Graph.TriggerPin.Connected"));
+				}
+			}
+
+			if (!PinBrush)
+			{
+				PinBrush = FEditorStyle::GetBrush("Graph.Pin.Connected");
+			}
+		}
+	}
+
+	SBorder::Construct(SBorder::FArguments()
+		.BorderImage(FEditorStyle::GetBrush("Menu.Background"))
+		.Padding(5)
+		[
+			SNew(SBox)
+			.WidthOverride(400)
+			.HeightOverride(400)
+			[
+				SNew(SVerticalBox)
+				+SVerticalBox::Slot()
+				.AutoHeight()
+				.Padding(2, 2, 2, 5)
+				[
+					SNew(SHorizontalBox)
+					+SHorizontalBox::Slot()
+					.AutoWidth()
+					.VAlign(VAlign_Center)
+					.Padding(0, 0, 5, 0)
+					[
+						SNew(SImage)
+						.ColorAndOpacity(TypeColor)
+						.Visibility_Lambda([this]()
+						{
+							return DraggedFromPins.IsEmpty() ? EVisibility::Hidden : EVisibility::Visible;
+						})
+						.Image(PinBrush)
+					]
+					+SHorizontalBox::Slot()
+					.AutoWidth()
+					.VAlign(VAlign_Center)
+					[
+						SNew(STextBlock)
+						.Text_Lambda([this]()
+						{
+							if (DraggedFromPins.IsEmpty())
+							{
+								return LOCTEXT("ContextText", "All MetaSound Node Classes");
+							}
+
+							UEdGraphPin* Pin = DraggedFromPins[0];
+
+							if (DraggedFromPins[0]->Direction == EGPD_Input)
+							{
+								FInputHandle InputHandle = FGraphBuilder::GetInputHandleFromPin(Pin);
+								return FText::Format(LOCTEXT("ContextTypeFilteredText_Output", "Classes with output of type '{0}'"), FText::FromName(InputHandle->GetDataType()));
+							}
+							else
+							{
+								FOutputHandle OutputHandle = FGraphBuilder::GetOutputHandleFromPin(Pin);
+								return FText::Format(LOCTEXT("ContextTypeFilteredText_Input", "Classes with input of type '{0}'"), FText::FromName(OutputHandle->GetDataType()));
+							}
+						})
+						// TODO: Move to Metasound Style
+						.Font(FEditorStyle::GetFontStyle("BlueprintEditor.ActionMenu.ContextDescriptionFont"))
+						.ToolTip(IDocumentation::Get()->CreateToolTip(
+							LOCTEXT("ActionMenuContextTextTooltip", "Describes the current context of the action list"),
+							nullptr,
+							TEXT("Shared/Editors/MetasoundEditor"),
+							TEXT("MetasoundActionMenuContextText")))
+						.WrapTextAt(280)
+					]
+				]
+				+SVerticalBox::Slot()
+				[
+					SAssignNew(GraphActionMenu, SGraphActionMenu)
+						.OnActionSelected(this, &SMetasoundActionMenu::OnActionSelected)
+						.OnCreateWidgetForAction(SGraphActionMenu::FOnCreateWidgetForAction::CreateSP(this, &SMetasoundActionMenu::OnCreateWidgetForAction))
+						.OnCollectAllActions(this, &SMetasoundActionMenu::CollectAllActions)
+						.OnCreateCustomRowExpander_Lambda([](const FCustomExpanderData& InCustomExpanderData)
+						{
+							return SNew(SMetasoundActionMenuExpanderArrow, InCustomExpanderData);
+						})
+						.DraggedFromPins(DraggedFromPins)
+						.GraphObj(Graph)
+						.AlphaSortItems(true)
+				]
+			]
+		]
+	);
+}
+
+void SMetasoundActionMenu::CollectAllActions(FGraphActionListBuilderBase& OutAllActions)
+{
+	if (!Graph)
+	{
+		return;
+	}
+
+	FGraphContextMenuBuilder MenuBuilder(Graph);
+	if (!DraggedFromPins.IsEmpty())
+	{
+		MenuBuilder.FromPin = DraggedFromPins[0];
+	}
+
+	// Cannot call GetGraphContextActions() during serialization and GC due to its use of FindObject()
+	if(!GIsSavingPackage && !IsGarbageCollecting())
+	{
+		if (const UMetasoundEditorGraphSchema* Schema = GetDefault<UMetasoundEditorGraphSchema>())
+		{
+			Schema->GetGraphContextActions(MenuBuilder);
+		}
+	}
+
+	OutAllActions.Append(MenuBuilder);
+}
+
+TSharedRef<SEditableTextBox> SMetasoundActionMenu::GetFilterTextBox()
+{
+	return GraphActionMenu->GetFilterTextBox();
+}
+
+TSharedRef<SWidget> SMetasoundActionMenu::OnCreateWidgetForAction(FCreateWidgetForActionData* const InCreateData)
+{
+	using namespace Metasound;
+	using namespace Metasound::Frontend;
+
+	check(InCreateData);
+	InCreateData->bHandleMouseButtonDown = false;
+
+	const FSlateBrush* IconBrush = nullptr;
+	FLinearColor IconColor;
+	TSharedPtr<FMetasoundGraphSchemaAction> Action = StaticCastSharedPtr<FMetasoundGraphSchemaAction>(InCreateData->Action);
+	if (Action.IsValid())
+	{
+		IconBrush = Action->GetIconBrush();
+		IconColor = Action->GetIconColor();
+	}
+
+	TSharedPtr<SHorizontalBox> WidgetBox = SNew(SHorizontalBox);
+	if (IconBrush)
+	{
+		WidgetBox->AddSlot()
+			.AutoWidth()
+			.VAlign(VAlign_Center)
+			.Padding(5, 0, 0, 0)
+			[
+				SNew(SImage)
+				.ColorAndOpacity(IconColor)
+				.Image(IconBrush)
+			];
+	}
+
+	WidgetBox->AddSlot()
+		.AutoWidth()
+		.VAlign(VAlign_Center)
+		.Padding(0, 0, 0, 0)
+		[
+			SNew(SGraphPaletteItem, InCreateData)
+		];
+
+	return WidgetBox->AsShared();
+}
+
+void SMetasoundActionMenu::OnActionSelected(const TArray<TSharedPtr<FEdGraphSchemaAction>>& SelectedAction, ESelectInfo::Type InSelectionType)
+{
+	if (!Graph)
+	{
+		return;
+	}
+
+	if (InSelectionType != ESelectInfo::OnMouseClick  && InSelectionType != ESelectInfo::OnKeyPress && !SelectedAction.IsEmpty())
+	{
+		return;
+	}
+
+	for (const TSharedPtr<FEdGraphSchemaAction>& Action : SelectedAction)
+	{
+		if (Action.IsValid() && Graph)
+		{
+			if (!bActionExecuted && (Action->GetTypeId() != FEdGraphSchemaAction_Dummy::StaticGetTypeId()))
+			{
+				FSlateApplication::Get().DismissAllMenus();
+				bActionExecuted = true;
+			}
+
+			if (UEdGraphNode* ResultNode = Action->PerformAction(Graph, DraggedFromPins, NewNodePosition))
+			{
+				NewNodePosition += Metasound::Frontend::DisplayStyle::NodeLayout::DefaultOffsetX;
+			}
+		}
+	}
+}
+#undef LOCTEXT_NAMESPACE
diff --git a/Engine/Plugins/Runtime/Metasound/Source/MetasoundEditor/Private/SMetasoundActionMenu.h b/Engine/Plugins/Runtime/Metasound/Source/MetasoundEditor/Private/SMetasoundActionMenu.h
new file mode 100644
--- /dev/null
+++ b/Engine/Plugins/Runtime/Metasound/Source/MetasoundEditor/Private/SMetasoundActionMenu.h
@@ -0,0 +1,69 @@
+// Copyright Epic Games, Inc. All Rights Reserved.
+#pragma once
+
+#include "CoreMinimal.h"
+#include "Layout/Margin.h"
+#include "Layout/Visibility.h"
+#include "Widgets/DeclarativeSyntaxSupport.h"
+#include "Widgets/SWidget.h"
+#include "Widgets/Views/SExpanderArrow.h"
+#include "Widgets/Layout/SBorder.h"
+#include "EdGraph/EdGraph.h"
+#include "EdGraph/EdGraphPin.h"
+#include "EdGraph/EdGraphSchema.h"
+#include "GraphEditor.h"
+
+class SEditableTextBox;
+class SGraphActionMenu;
+class UEdGraph;
+struct FCustomExpanderData;
+
+
+class SMetasoundActionMenu : public SBorder
+{
+public:
+	/** Delegate for the OnCloseReason event which is always raised when the SMetasoundActionMenu closes */
+	DECLARE_DELEGATE_ThreeParams(FClosedReason, bool /*bActionExecuted*/, bool /*bContextSensitiveChecked*/, bool /*bGraphPinContext*/);
+
+	SLATE_BEGIN_ARGS(SMetasoundActionMenu)
+		: _Graph(static_cast<UEdGraph*>(nullptr))
+		, _NewNodePosition(FVector2D::ZeroVector)
+		, _AutoExpandActionMenu(false)
+	{
+	}
+		SLATE_ARGUMENT(UEdGraph*, Graph)
+		SLATE_ARGUMENT(FVector2D, NewNodePosition)
+		SLATE_ARGUMENT(TArray<UEdGraphPin*>, DraggedFromPins)
+		SLATE_ARGUMENT(SGraphEditor::FActionMenuClosed, OnClosedCallback)
+		SLATE_ARGUMENT(bool, AutoExpandActionMenu)
+		SLATE_EVENT(FClosedReason, OnCloseReason)
+	SLATE_END_ARGS()
+
+	void Construct(const FArguments& InArgs);
+
+	~SMetasoundActionMenu();
+
+	TSharedRef<SEditableTextBox> GetFilterTextBox();
+
+protected:
+	void OnActionSelected(const TArray<TSharedPtr<FEdGraphSchemaAction>>& SelectedAction, ESelectInfo::Type InSelectionType);
+
+	TSharedRef<SWidget> OnCreateWidgetForAction(struct FCreateWidgetForActionData* const InCreateData);
+
+	/** Callback used to populate all actions list in SGraphActionMenu */
+	void CollectAllActions(FGraphActionListBuilderBase& OutAllActions);
+
+private:
+	UEdGraph* Graph = nullptr;
+	bool bAutoExpandActionMenu = false;
+	bool bActionExecuted = false;
+
+	TArray<UEdGraphPin*> DraggedFromPins;
+	FVector2D NewNodePosition = FVector2D::ZeroVector;
+
+	SGraphEditor::FActionMenuClosed OnClosedCallback;
+	FClosedReason OnCloseReasonCallback;
+
+	TSharedPtr<SGraphActionMenu> GraphActionMenu;
+
+};
diff --git a/Engine/Plugins/Runtime/Metasound/Source/MetasoundEditor/Public/MetasoundEditorGraphSchema.h b/Engine/Plugins/Runtime/Metasound/Source/MetasoundEditor/Public/MetasoundEditorGraphSchema.h
--- a/Engine/Plugins/Runtime/Metasound/Source/MetasoundEditor/Public/MetasoundEditorGraphSchema.h
+++ b/Engine/Plugins/Runtime/Metasound/Source/MetasoundEditor/Public/MetasoundEditorGraphSchema.h
@@ -9,6 +9,7 @@
 #include "EdGraphUtilities.h"
 #include "MetasoundFrontend.h"
 #include "MetasoundFrontendController.h"
+#include "MetasoundFrontendNodesCategories.h"
 #include "Templates/Function.h"
 #include "UObject/ObjectMacros.h"
 
@@ -26,6 +27,20 @@ namespace Metasound
 {
 	namespace Editor
 	{
+		// Ordered such that highest priority is highest value
+		enum class EPrimaryContextGroup : uint8
+		{
+			Conversions = 0,
+			Graphs,
+			Functions,
+			Outputs,
+			Inputs,
+
+			Common
+		};
+
+		const FText& GetContextGroupDisplayName(EPrimaryContextGroup InContextGroup);
+
 		using FInputFilterFunction = TFunction<bool(const FMetasoundFrontendClassInput&)>;
 		using FOutputFilterFunction = TFunction<bool(const FMetasoundFrontendClassOutput&)>;
 		using FInterfaceNodeFilterFunction = TFunction<bool(Frontend::FConstNodeHandle)>;
@@ -47,8 +62,13 @@ struct METASOUNDEDITOR_API FMetasoundGraphSchemaAction : public FEdGraphSchemaAc
 		: FEdGraphSchemaAction()
 	{}
 
-	FMetasoundGraphSchemaAction(FText InNodeCategory, FText InMenuDesc, FText InToolTip, const int32 InGrouping, FText InKeywords = FText::GetEmpty())
-		: FEdGraphSchemaAction(MoveTemp(InNodeCategory), MoveTemp(InMenuDesc), MoveTemp(InToolTip), InGrouping, MoveTemp(InKeywords))
+	FMetasoundGraphSchemaAction(FText InNodeCategory, FText InMenuDesc, FText InToolTip, Metasound::Editor::EPrimaryContextGroup InGroup, FText InKeywords = FText::GetEmpty())
+		: FEdGraphSchemaAction(
+			MoveTemp(InNodeCategory),
+			MoveTemp(InMenuDesc),
+			MoveTemp(InToolTip),
+			/*InNodeCategory.IsEmpty() ? */static_cast<int32>(InGroup)/* : 0*/,
+			MoveTemp(InKeywords))
 	{}
 
 	virtual const FSlateBrush* GetIconBrush() const
@@ -77,7 +97,7 @@ struct METASOUNDEDITOR_API FMetasoundGraphSchemaAction_NewInput : public FMetaso
 		: FMetasoundGraphSchemaAction()
 	{}
 
-	FMetasoundGraphSchemaAction_NewInput(FText InNodeCategory, FText InDisplayName, FGuid InInputNodeID, FText InToolTip, const int32 InGrouping);
+	FMetasoundGraphSchemaAction_NewInput(FText InNodeCategory, FText InDisplayName, FGuid InInputNodeID, FText InToolTip, Metasound::Editor::EPrimaryContextGroup InGroup);
 
 	//~ Begin FMetasoundGraphSchemaAction Interface
 	virtual const FSlateBrush* GetIconBrush() const override;
@@ -96,11 +116,7 @@ struct METASOUNDEDITOR_API FMetasoundGraphSchemaAction_PromoteToInput : public F
 {
 	GENERATED_USTRUCT_BODY();
 
-	FMetasoundGraphSchemaAction_PromoteToInput()
-		: FMetasoundGraphSchemaAction_NewInput()
-	{}
-
-	FMetasoundGraphSchemaAction_PromoteToInput(FText InNodeCategory, FText InDisplayName, FText InToolTip, const int32 InGrouping);
+	FMetasoundGraphSchemaAction_PromoteToInput();
 
 	//~ Begin FEdGraphSchemaAction Interface
 	virtual UEdGraphNode* PerformAction(UEdGraph* ParentGraph, UEdGraphPin* FromPin, const FVector2D Location, bool bSelectNewNode = true) override;
@@ -120,7 +136,7 @@ struct METASOUNDEDITOR_API FMetasoundGraphSchemaAction_NewOutput : public FMetas
 		: FMetasoundGraphSchemaAction()
 	{}
 
-	FMetasoundGraphSchemaAction_NewOutput(FText InNodeCategory, FText InDisplayName, FGuid InOutputNodeID, FText InToolTip, const int32 InGrouping);
+	FMetasoundGraphSchemaAction_NewOutput(FText InNodeCategory, FText InDisplayName, FGuid InOutputNodeID, FText InToolTip, Metasound::Editor::EPrimaryContextGroup InGroup);
 
 	//~ Begin FMetasoundGraphSchemaAction Interface
 	virtual const FSlateBrush* GetIconBrush() const override;
@@ -139,11 +155,7 @@ struct METASOUNDEDITOR_API FMetasoundGraphSchemaAction_PromoteToOutput : public
 {
 	GENERATED_USTRUCT_BODY();
 
-	FMetasoundGraphSchemaAction_PromoteToOutput()
-		: FMetasoundGraphSchemaAction_NewOutput()
-	{}
-
-	FMetasoundGraphSchemaAction_PromoteToOutput(FText InNodeCategory, FText InDisplayName, FText InToolTip, const int32 InGrouping);
+	FMetasoundGraphSchemaAction_PromoteToOutput();
 
 	//~ Begin FEdGraphSchemaAction Interface
 	virtual UEdGraphNode* PerformAction(UEdGraph* ParentGraph, UEdGraphPin* FromPin, const FVector2D Location, bool bSelectNewNode = true) override;
@@ -163,8 +175,8 @@ struct METASOUNDEDITOR_API FMetasoundGraphSchemaAction_NewNode : public FMetasou
 		: FMetasoundGraphSchemaAction()
 	{}
 
-	FMetasoundGraphSchemaAction_NewNode(const FText& InNodeCategory, const FText& InMenuDesc, const FText& InToolTip, const int32 InGrouping, FText InKeywords = FText::GetEmpty())
-		: FMetasoundGraphSchemaAction(InNodeCategory, InMenuDesc, InToolTip, InGrouping, InKeywords)
+	FMetasoundGraphSchemaAction_NewNode(const FText& InNodeCategory, const FText& InMenuDesc, const FText& InToolTip, Metasound::Editor::EPrimaryContextGroup InGroup, FText InKeywords = FText::GetEmpty())
+		: FMetasoundGraphSchemaAction(InNodeCategory, InMenuDesc, InToolTip, InGroup, InKeywords)
 	{}
 
 	//~ Begin FMetasoundGraphSchemaAction Interface
@@ -188,8 +200,8 @@ struct METASOUNDEDITOR_API FMetasoundGraphSchemaAction_NewFromSelected : public
 		: FMetasoundGraphSchemaAction_NewNode()
 	{}
 
-	FMetasoundGraphSchemaAction_NewFromSelected(FText InNodeCategory, FText InMenuDesc, FText InToolTip, const int32 InGrouping)
-		: FMetasoundGraphSchemaAction_NewNode(MoveTemp(InNodeCategory), MoveTemp(InMenuDesc), MoveTemp(InToolTip), InGrouping) 
+	FMetasoundGraphSchemaAction_NewFromSelected(FText InNodeCategory, FText InMenuDesc, FText InToolTip, Metasound::Editor::EPrimaryContextGroup InGroup)
+		: FMetasoundGraphSchemaAction_NewNode(MoveTemp(InNodeCategory), MoveTemp(InMenuDesc), MoveTemp(InToolTip), InGroup)
 	{}
 
 	//~ Begin FEdGraphSchemaAction Interface
@@ -207,8 +219,8 @@ struct METASOUNDEDITOR_API FMetasoundGraphSchemaAction_NewComment : public FMeta
 		: FMetasoundGraphSchemaAction()
 	{}
 
-	FMetasoundGraphSchemaAction_NewComment(const FText& InNodeCategory, const FText& InMenuDesc, const FText& InToolTip, const int32 InGrouping)
-		: FMetasoundGraphSchemaAction(InNodeCategory, InMenuDesc, InToolTip, InGrouping)
+	FMetasoundGraphSchemaAction_NewComment(const FText& InNodeCategory, const FText& InMenuDesc, const FText& InToolTip, Metasound::Editor::EPrimaryContextGroup InGroup)
+		: FMetasoundGraphSchemaAction(InNodeCategory, InMenuDesc, InToolTip, InGroup)
 	{}
 
 	//~ Begin FMetasoundGraphSchemaAction Interface
@@ -232,8 +244,8 @@ struct METASOUNDEDITOR_API FMetasoundGraphSchemaAction_Paste : public FMetasound
 		: FMetasoundGraphSchemaAction()
 	{}
 
-	FMetasoundGraphSchemaAction_Paste(const FText& InNodeCategory, const FText& InMenuDesc, const FText& InToolTip, const int32 InGrouping)
-		: FMetasoundGraphSchemaAction(InNodeCategory, InMenuDesc, InToolTip, InGrouping)
+	FMetasoundGraphSchemaAction_Paste(const FText& InNodeCategory, const FText& InMenuDesc, const FText& InToolTip, Metasound::Editor::EPrimaryContextGroup InGroup)
+		: FMetasoundGraphSchemaAction(InNodeCategory, InMenuDesc, InToolTip, InGroup)
 	{}
 
 	//~ Begin FEdGraphSchemaAction Interface
