Cherry pick from CL 17538469

Speedup initial play sound calls to MetaSound objects

diff --git a/Engine/Plugins/Runtime/Metasound/Source/MetasoundEngine/Private/MetasoundSource.cpp b/Engine/Plugins/Runtime/Metasound/Source/MetasoundEngine/Private/MetasoundSource.cpp
--- a/Engine/Plugins/Runtime/Metasound/Source/MetasoundEngine/Private/MetasoundSource.cpp
+++ b/Engine/Plugins/Runtime/Metasound/Source/MetasoundEngine/Private/MetasoundSource.cpp
@@ -428,8 +428,7 @@ ISoundGeneratorPtr UMetaSoundSource::CreateSoundGenerator(const FSoundGeneratorI
 	FOperatorSettings InSettings = GetOperatorSettings(static_cast<FSampleRate>(SampleRate));
 	FMetasoundEnvironment Environment = CreateEnvironment(InParams);
 
-	// TODO: cache graph to avoid having to create it every call to `CreateSoundGenerator(...)`
-	TUniquePtr<IGraph> MetasoundGraph = BuildMetasoundDocument();
+	TSharedPtr<const IGraph, ESPMode::ThreadSafe> MetasoundGraph = GetMetasoundCoreGraph();
 	if (!MetasoundGraph.IsValid())
 	{
 		UE_LOG(LogMetaSound, Error, TEXT("Cannot create UMetaSoundSource SoundGenerator [Name:%s]. Failed to create MetaSound Graph"), *GetName());
@@ -439,7 +438,7 @@ ISoundGeneratorPtr UMetaSoundSource::CreateSoundGenerator(const FSoundGeneratorI
 	FMetasoundGeneratorInitParams InitParams =
 	{
 		InSettings,
-		MoveTemp(MetasoundGraph),
+		MetasoundGraph,
 		Environment,
 		GetName(),
 		GetAudioOutputVertexKeys(),
diff --git a/Engine/Plugins/Runtime/Metasound/Source/MetasoundFrontend/Private/MetasoundAssetBase.cpp b/Engine/Plugins/Runtime/Metasound/Source/MetasoundFrontend/Private/MetasoundAssetBase.cpp
--- a/Engine/Plugins/Runtime/Metasound/Source/MetasoundFrontend/Private/MetasoundAssetBase.cpp
+++ b/Engine/Plugins/Runtime/Metasound/Source/MetasoundFrontend/Private/MetasoundAssetBase.cpp
@@ -4,6 +4,7 @@
 
 #include "Algo/AnyOf.h"
 #include "Algo/Transform.h"
+#include "Containers/Set.h"
 #include "HAL/FileManager.h"
 #include "Internationalization/Text.h"
 #include "IStructSerializerBackend.h"
@@ -85,6 +86,8 @@ void FMetasoundAssetBase::RegisterGraphWithFrontend(FMetaSoundAssetRegistrationO
 	}
 
 	GetReferencedAssetClassCache().Reset();
+	// Triggers the existing runtime data to be out-of-date.
+	CurrentCachedRuntimeDataChangeID = FGuid::NewGuid();
 
 	if (InRegistrationOptions.bRebuildReferencedAssetClassKeys)
 	{
@@ -429,7 +432,7 @@ bool FMetasoundAssetBase::VersionAsset()
 	return bDidEdit;
 }
 
-TUniquePtr<Metasound::IGraph> FMetasoundAssetBase::BuildMetasoundDocument() const
+TSharedPtr<Metasound::IGraph, ESPMode::ThreadSafe> FMetasoundAssetBase::BuildMetasoundDocument() const
 {
 	using namespace Metasound;
 	using namespace Metasound::Frontend;
@@ -440,7 +443,7 @@ TUniquePtr<Metasound::IGraph> FMetasoundAssetBase::BuildMetasoundDocument() cons
 	if (nullptr == Doc)
 	{
 		UE_LOG(LogMetaSound, Error, TEXT("Cannot create graph. Null MetaSound document in MetaSound asset [Name:%s]"), *GetOwningAssetName());
-		return TUniquePtr<IGraph>(nullptr);
+		return TSharedPtr<IGraph, ESPMode::ThreadSafe>(nullptr);
 	}
 
 	// Create graph which can spawn instances. TODO: cache graph.
@@ -455,7 +458,9 @@ TUniquePtr<Metasound::IGraph> FMetasoundAssetBase::BuildMetasoundDocument() cons
 		}
 	}
 
-	return MoveTemp(FrontendGraph);
+	TSharedPtr<Metasound::IGraph, ESPMode::ThreadSafe> SharedGraph(FrontendGraph.Release());
+
+	return SharedGraph;
 }
 
 void FMetasoundAssetBase::RebuildReferencedAssetClassKeys()
@@ -560,27 +565,23 @@ TArray<FMetasoundAssetBase::FSendInfoAndVertexName> FMetasoundAssetBase::GetSend
 	using namespace Metasound::Frontend;
 	using FSendInfo = FMetaSoundParameterTransmitter::FSendInfo;
 
-	TArray<FSendInfoAndVertexName> SendInfos;
+	check(IsInGameThread() || IsInAudioThread());
 
-	FConstGraphHandle RootGraph = GetRootGraphHandle();
+	const FRuntimeData& RuntimeData = GetRuntimeData();
 
-	const TArray<FVertexName> SendVertices = GetTransmittableInputVertexNames();
-	for (const FVertexName& VertexName : SendVertices)
+	TArray<FSendInfoAndVertexName> SendInfos;
+
+	for (const FMetasoundFrontendClassInput& Vertex : RuntimeData.TransmittableInputs)
 	{
-		FConstNodeHandle InputNode = RootGraph->GetInputNodeWithName(VertexName);
-		for (FConstInputHandle InputHandle : InputNode->GetConstInputs())
-		{
-			FSendInfoAndVertexName Info;
+		FSendInfoAndVertexName Info;
 
-			// TODO: incorporate VertexID into address. But need to ensure that VertexID
-			// will be maintained after injecting Receive nodes. 
-			Info.SendInfo.Address = FMetaSoundParameterTransmitter::CreateSendAddressFromInstanceID(InInstanceID, VertexName, InputHandle->GetDataType());
-			Info.SendInfo.ParameterName = VertexName;
-			Info.SendInfo.TypeName = InputHandle->GetDataType();
-			Info.VertexName = VertexName;
+		Info.SendInfo.Address = FMetaSoundParameterTransmitter::CreateSendAddressFromInstanceID(InInstanceID, Vertex.Name, Vertex.TypeName);
+		Info.SendInfo.ParameterName = Vertex.Name;
+		Info.SendInfo.TypeName = Vertex.TypeName;
+		Info.VertexName = Vertex.Name;
 
-			SendInfos.Add(Info);
-		}
+		SendInfos.Add(Info);
+		
 	}
 
 	return SendInfos;
@@ -803,4 +804,80 @@ FString FMetasoundAssetBase::GetOwningAssetName() const
 	return FString();
 }
 
+TSharedPtr<const Metasound::IGraph, ESPMode::ThreadSafe> FMetasoundAssetBase::GetMetasoundCoreGraph() const
+{
+	return FMetasoundAssetBase::GetRuntimeData().Graph;
+}
+
+TArray<FMetasoundFrontendClassInput> FMetasoundAssetBase::GetTransmittableClassInputs() const
+{
+	using namespace Metasound;
+	using namespace Metasound::Frontend;
+	check(IsInGameThread() || IsInAudioThread());
+
+	TArray<FMetasoundFrontendClassInput> Inputs;
+
+	if (const FMetasoundFrontendDocument* Doc = GetDocument().Get())
+	{
+		TSet<FVertexName> ArchetypeInputs;
+
+		auto GetVertexKey = [](const FMetasoundFrontendClassVertex& InVertex) { return InVertex.Name; };
+
+		// Do not transmit vertices defined in archetype. Those are already accounted
+		// for by owning object.
+		FMetasoundFrontendArchetype Archetype;
+		GetArchetype(Archetype);
+		Algo::Transform(Archetype.Interface.Inputs, ArchetypeInputs, GetVertexKey);
+
+		// Do not transmit vertices which are not transmittable. Async communication 
+		// is not supported without transmission.
+		IDataTypeRegistry& Registry = IDataTypeRegistry::Get();
+
+		auto IsTransmittable = [&Registry, &ArchetypeInputs](const FMetasoundFrontendClassVertex& InVertex) -> bool
+		{	
+			if (!ArchetypeInputs.Contains(InVertex.Name))
+			{
+				FDataTypeRegistryInfo Info;
+				if (Registry.GetDataTypeInfo(InVertex.TypeName, Info))
+				{
+					return Info.bIsTransmittable;
+				}
+			}
+			return false;
+		};
+
+
+		for (const FMetasoundFrontendClassInput& Vertex : Doc->RootGraph.Interface.Inputs)
+		{
+			if (IsTransmittable(Vertex))
+			{
+				Inputs.Add(Vertex);
+			}
+		}
+	}
+
+	return Inputs;
+}
+
+const FMetasoundAssetBase::FRuntimeData& FMetasoundAssetBase::GetRuntimeData() const
+{
+	
+	// Check if a ChangeID has been generated before.
+	if (!CurrentCachedRuntimeDataChangeID.IsValid())
+	{
+		CurrentCachedRuntimeDataChangeID = FGuid::NewGuid();
+	}
+
+	// Check if CachedRuntimeData is out-of-date.
+	if (CachedRuntimeData.ChangeID != CurrentCachedRuntimeDataChangeID)
+	{
+		// Update CachedRuntimeData.
+		CachedRuntimeData.TransmittableInputs = GetTransmittableClassInputs();
+		CachedRuntimeData.Graph = BuildMetasoundDocument();
+		CachedRuntimeData.ChangeID = CurrentCachedRuntimeDataChangeID;
+	}
+
+	return CachedRuntimeData;
+}
+
 #undef LOCTEXT_NAMESPACE // "MetaSound"
diff --git a/Engine/Plugins/Runtime/Metasound/Source/MetasoundFrontend/Public/MetasoundAssetBase.h b/Engine/Plugins/Runtime/Metasound/Source/MetasoundFrontend/Public/MetasoundAssetBase.h
--- a/Engine/Plugins/Runtime/Metasound/Source/MetasoundFrontend/Public/MetasoundAssetBase.h
+++ b/Engine/Plugins/Runtime/Metasound/Source/MetasoundFrontend/Public/MetasoundAssetBase.h
@@ -1,8 +1,6 @@
 // Copyright Epic Games, Inc. All Rights Reserved.
 #pragma once
 
-#include "CoreMinimal.h"
-
 #include "MetasoundAccessPtr.h"
 #include "MetasoundFrontend.h"
 #include "MetasoundFrontendController.h"
@@ -15,12 +13,10 @@
 #include "UObject/SoftObjectPath.h"
 #include "UObject/WeakObjectPtrTemplates.h"
 
-
 // Forward Declarations
 class FMetasoundAssetBase;
 class UEdGraph;
 
-
 namespace Metasound
 {
 	namespace AssetTags
@@ -201,12 +197,10 @@ public:
 		Metasound::FVertexName VertexName;
 	};
 
-	// Builds the Metasound Document returned by `GetDocument() const`.
-	virtual TUniquePtr<Metasound::IGraph> BuildMetasoundDocument() const;
-
 protected:
 	virtual void SetReferencedAssetClassKeys(TSet<Metasound::Frontend::FNodeRegistryKey>&& InKeys) = 0;
 
+	// Get information for communicating asynchronously with MetaSound running instance.
 	TArray<FSendInfoAndVertexName> GetSendInfos(uint64 InInstanceID) const;
 
 #if WITH_EDITORONLY_DATA
@@ -219,18 +213,42 @@ protected:
 	// Returns an access pointer to the document.
 	virtual Metasound::Frontend::FConstDocumentAccessPtr GetDocument() const = 0;
 
-	// Returns the owning asset responsible for transactions applied to metasound
+	// Returns the owning asset responsible for transactions applied to MetaSound
 	virtual UObject* GetOwningAsset() = 0;
 
-	// Returns the owning asset responsible for transactions applied to metasound
+	// Returns the owning asset responsible for transactions applied to MetaSound
 	virtual const UObject* GetOwningAsset() const = 0;
 
 	FString GetOwningAssetName() const;
 
+	// Returns a shared instance of the core metasound graph.
+	TSharedPtr<const Metasound::IGraph, ESPMode::ThreadSafe> GetMetasoundCoreGraph() const;
 
 private:
 	Metasound::Frontend::FNodeRegistryKey RegistryKey;
 
+	// Container for runtime data of MetaSound graph.
+	struct FRuntimeData
+	{
+		// Current ID of graph.
+		FGuid ChangeID;
+
+		// Array of inputs which can be transmitted to.
+		TArray<FMetasoundFrontendClassInput> TransmittableInputs;
+
+		// Core graph.
+		TSharedPtr<Metasound::IGraph, ESPMode::ThreadSafe> Graph;
+	};
+	// Cache ID is used to determine whether CachedRuntimeData is out-of-date.
+	mutable FGuid CurrentCachedRuntimeDataChangeID;
+	mutable FRuntimeData CachedRuntimeData;
+
+	// Returns the current runtime data. If the cached data is out-of-date, it will
+	// be updated in this call.
+	const FRuntimeData& GetRuntimeData() const;
+
+	TSharedPtr<Metasound::IGraph, ESPMode::ThreadSafe> BuildMetasoundDocument() const;
+	TArray<FMetasoundFrontendClassInput> GetTransmittableClassInputs() const;
 	TSet<Metasound::FVertexName> GetNonTransmittableInputVertices(const FMetasoundFrontendDocument& InDoc) const;
 	TArray<Metasound::FVertexName> GetTransmittableInputVertexNames() const;
 	Metasound::FSendAddress CreateSendAddress(uint64 InInstanceID, const Metasound::FVertexName& InVertexName, const FName& InDataTypeName) const;
diff --git a/Engine/Plugins/Runtime/Metasound/Source/MetasoundGenerator/Public/MetasoundGenerator.h b/Engine/Plugins/Runtime/Metasound/Source/MetasoundGenerator/Public/MetasoundGenerator.h
--- a/Engine/Plugins/Runtime/Metasound/Source/MetasoundGenerator/Public/MetasoundGenerator.h
+++ b/Engine/Plugins/Runtime/Metasound/Source/MetasoundGenerator/Public/MetasoundGenerator.h
@@ -21,7 +21,7 @@ namespace Metasound
 	struct METASOUNDGENERATOR_API FMetasoundGeneratorInitParams
 	{
 		FOperatorSettings OperatorSettings;
-		TUniquePtr<IGraph> Graph;
+		TSharedPtr<const IGraph, ESPMode::ThreadSafe> Graph;
 		FMetasoundEnvironment Environment;
 		FString MetaSoundName;
 		TArray<FVertexName> AudioOutputNames;
