Fix MetaSound DisplayName FText being inappropriately generated for
- Core input/output/variable class types
- Asset class reference nodes
#rnx

diff --git a/Engine/Plugins/Runtime/Metasound/Source/MetasoundEditor/Private/MetasoundEditorGraphBuilder.cpp b/Engine/Plugins/Runtime/Metasound/Source/MetasoundEditor/Private/MetasoundEditorGraphBuilder.cpp
--- a/Engine/Plugins/Runtime/Metasound/Source/MetasoundEditor/Private/MetasoundEditorGraphBuilder.cpp
+++ b/Engine/Plugins/Runtime/Metasound/Source/MetasoundEditor/Private/MetasoundEditorGraphBuilder.cpp
@@ -89,6 +89,8 @@ namespace Metasound
 
 		UMetasoundEditorGraphExternalNode* FGraphBuilder::AddExternalNode(UObject& InMetaSound, Frontend::FNodeHandle& InNodeHandle, FVector2D InLocation, bool bInSelectNewNode)
 		{
+			using namespace Frontend;
+
 			UMetasoundEditorGraphExternalNode* NewGraphNode = nullptr;
 			if (!ensure(InNodeHandle->GetClassMetadata().GetType() == EMetasoundFrontendClassType::External))
 			{
@@ -101,6 +103,9 @@ namespace Metasound
 			FGraphNodeCreator<UMetasoundEditorGraphExternalNode> NodeCreator(Graph);
 
 			NewGraphNode = NodeCreator.CreateNode(bInSelectNewNode);
+
+			const FNodeRegistryKey RegistryKey = NodeRegistryKey::CreateKey(InNodeHandle->GetClassMetadata());
+			NewGraphNode->bIsClassNative = FMetasoundFrontendRegistryContainer::Get()->IsNodeNative(RegistryKey);
 			NewGraphNode->ClassName = InNodeHandle->GetClassMetadata().GetClassName();
 
 			NodeCreator.Finalize();
@@ -1598,7 +1603,7 @@ namespace Metasound
 
 		bool FGraphBuilder::SynchronizePinLiteral(UEdGraphPin& InPin)
 		{
-			using namespace Metasound::Frontend;
+			using namespace Frontend;
 
 			if (!ensure(InPin.Direction == EGPD_Input))
 			{
diff --git a/Engine/Plugins/Runtime/Metasound/Source/MetasoundEditor/Private/MetasoundEditorGraphSchema.cpp b/Engine/Plugins/Runtime/Metasound/Source/MetasoundEditor/Private/MetasoundEditorGraphSchema.cpp
--- a/Engine/Plugins/Runtime/Metasound/Source/MetasoundEditor/Private/MetasoundEditorGraphSchema.cpp
+++ b/Engine/Plugins/Runtime/Metasound/Source/MetasoundEditor/Private/MetasoundEditorGraphSchema.cpp
@@ -920,13 +920,7 @@ FText UMetasoundEditorGraphSchema::GetPinDisplayName(const UEdGraphPin* Pin) con
 				TArray<FConstInputHandle> InputHandles = NodeHandle->GetConstInputsWithVertexName(Pin->GetFName());
 				if (ensure(!InputHandles.IsEmpty()))
 				{
-					const FText& DisplayName = InputHandles[0]->GetDisplayName();
-					if (!DisplayName.IsEmptyOrWhitespace())
-					{
-						return DisplayName;
-					}
-
-					return FText::FromName(InputHandles[0]->GetName());
+					return InputHandles[0]->GetDisplayName();;
 				}
 			}
 			else
@@ -934,13 +928,7 @@ FText UMetasoundEditorGraphSchema::GetPinDisplayName(const UEdGraphPin* Pin) con
 				TArray<FConstOutputHandle> OutputHandles = NodeHandle->GetConstOutputsWithVertexName(Pin->GetFName());
 				if (ensure(!OutputHandles.IsEmpty()))
 				{
-					const FText& DisplayName = OutputHandles[0]->GetDisplayName();
-					if (!DisplayName.IsEmptyOrWhitespace())
-					{
-						return DisplayName;
-					}
-
-					return FText::FromName(OutputHandles[0]->GetName());
+					return OutputHandles[0]->GetDisplayName();
 				}
 			}
 
diff --git a/Engine/Plugins/Runtime/Metasound/Source/MetasoundFrontend/Private/MetasoundFrontendStandardController.cpp b/Engine/Plugins/Runtime/Metasound/Source/MetasoundFrontend/Private/MetasoundFrontendStandardController.cpp
--- a/Engine/Plugins/Runtime/Metasound/Source/MetasoundFrontend/Private/MetasoundFrontendStandardController.cpp
+++ b/Engine/Plugins/Runtime/Metasound/Source/MetasoundFrontend/Private/MetasoundFrontendStandardController.cpp
@@ -3315,10 +3315,17 @@ namespace Metasound
 				{
 					FMetasoundFrontendNode& Node = GraphClass->Graph.Nodes.Emplace_GetRef(*NodeClass);
 
+					// Cache the asset name on the node if it node is reference to asset-defined graph.
+					const FNodeRegistryKey RegistryKey = NodeRegistryKey::CreateKey(NodeClass->Metadata);
+					if (const FSoftObjectPath* Path = IMetaSoundAssetManager::GetChecked().FindObjectPathFromKey(RegistryKey))
+					{
+						const FString& AssetName = Path->GetAssetName();
+						Node.Name = *AssetName;
+					}
 					Node.UpdateID();
 
 					FNodeAccessPtr NodePtr = GraphClassPtr.GetNodeWithNodeID(Node.GetID());
-					return GetNodeHandle(FGraphController::FNodeAndClass{NodePtr, InExistingDependency});
+					return GetNodeHandle(FGraphController::FNodeAndClass { NodePtr, InExistingDependency });
 				}
 			}
 
diff --git a/Engine/Plugins/Runtime/Metasound/Source/MetasoundFrontend/Public/MetasoundFrontendDocument.h b/Engine/Plugins/Runtime/Metasound/Source/MetasoundFrontend/Public/MetasoundFrontendDocument.h
--- a/Engine/Plugins/Runtime/Metasound/Source/MetasoundFrontend/Public/MetasoundFrontendDocument.h
+++ b/Engine/Plugins/Runtime/Metasound/Source/MetasoundFrontend/Public/MetasoundFrontendDocument.h
@@ -995,9 +995,6 @@ struct FMetasoundFrontendClassStyle
 
 	UPROPERTY()
 	FMetasoundFrontendClassStyleDisplay Display;
-
-	UPROPERTY()
-	bool bIsUserGenerated = false;
 };
 
 USTRUCT()
diff --git a/Engine/Plugins/Runtime/Metasound/Source/MetasoundFrontend/Public/MetasoundFrontendTransform.h b/Engine/Plugins/Runtime/Metasound/Source/MetasoundFrontend/Public/MetasoundFrontendTransform.h
--- a/Engine/Plugins/Runtime/Metasound/Source/MetasoundFrontend/Public/MetasoundFrontendTransform.h
+++ b/Engine/Plugins/Runtime/Metasound/Source/MetasoundFrontend/Public/MetasoundFrontendTransform.h
@@ -1,6 +1,7 @@
 // Copyright Epic Games, Inc. All Rights Reserved.
 #pragma once
 
+#include "Internationalization/Text.h"
 #include "MetasoundAssetBase.h"
 #include "MetasoundFrontendArchetypeRegistry.h"
 #include "MetasoundFrontendController.h"
@@ -109,6 +110,7 @@ namespace Metasound
 		  * a new class (useful for when an asset is duplicated). */
 		class METASOUNDFRONTEND_API FRegenerateAssetClassName : public IDocumentTransform
 		{
+
 		public:
 			FRegenerateAssetClassName() = default;
 
diff --git a/Engine/Plugins/Runtime/Metasound/Source/MetasoundFrontend/Public/MetasoundInputNode.h b/Engine/Plugins/Runtime/Metasound/Source/MetasoundFrontend/Public/MetasoundInputNode.h
--- a/Engine/Plugins/Runtime/Metasound/Source/MetasoundFrontend/Public/MetasoundInputNode.h
+++ b/Engine/Plugins/Runtime/Metasound/Source/MetasoundFrontend/Public/MetasoundInputNode.h
@@ -277,10 +277,9 @@ namespace Metasound
 			{
 				FNodeClassMetadata Info;
 
-				Info.ClassName = {TEXT("Input"), GetMetasoundDataTypeName<DataType>(), TEXT("")};
+				Info.ClassName = { "Input", GetMetasoundDataTypeName<DataType>(), FName() };
 				Info.MajorVersion = 1;
 				Info.MinorVersion = 0;
-				Info.DisplayName = FText::Format(LOCTEXT("Metasound_InputNodeDisplayNameFormat", "Input {0}"), FText::FromName(GetMetasoundDataTypeName<DataType>()));
 				Info.Description = LOCTEXT("Metasound_InputNodeDescription", "Input into the parent Metasound graph.");
 				Info.Author = PluginAuthor;
 				Info.PromptIfMissing = PluginNodeMissingPrompt;
diff --git a/Engine/Plugins/Runtime/Metasound/Source/MetasoundFrontend/Public/MetasoundOutputNode.h b/Engine/Plugins/Runtime/Metasound/Source/MetasoundFrontend/Public/MetasoundOutputNode.h
--- a/Engine/Plugins/Runtime/Metasound/Source/MetasoundFrontend/Public/MetasoundOutputNode.h
+++ b/Engine/Plugins/Runtime/Metasound/Source/MetasoundFrontend/Public/MetasoundOutputNode.h
@@ -117,10 +117,9 @@ namespace Metasound
 		{
 			FNodeClassMetadata Info;
 
-			Info.ClassName = {TEXT("Output"), GetMetasoundDataTypeName<DataType>(), TEXT("")};
+			Info.ClassName = { "Output", GetMetasoundDataTypeName<DataType>(), FName() };
 			Info.MajorVersion = 1;
 			Info.MinorVersion = 0;
-			Info.DisplayName = FText::Format(LOCTEXT("Metasound_OutputNodeDisplayNameFormat", "Output {0}"), FText::FromName(GetMetasoundDataTypeName<DataType>()));
 			Info.Description = LOCTEXT("Metasound_OutputNodeDescription", "Output from the parent Metasound graph.");
 			Info.Author = PluginAuthor;
 			Info.PromptIfMissing = PluginNodeMissingPrompt;
diff --git a/Engine/Plugins/Runtime/Metasound/Source/MetasoundGraphCore/Public/MetasoundVariableNodes.h b/Engine/Plugins/Runtime/Metasound/Source/MetasoundGraphCore/Public/MetasoundVariableNodes.h
--- a/Engine/Plugins/Runtime/Metasound/Source/MetasoundGraphCore/Public/MetasoundVariableNodes.h
+++ b/Engine/Plugins/Runtime/Metasound/Source/MetasoundGraphCore/Public/MetasoundVariableNodes.h
@@ -111,10 +111,9 @@ namespace Metasound
 		{
 			FNodeClassMetadata Info;
 
-			Info.ClassName = {"InitVariable", GetMetasoundDataTypeName<DataType>(), ""};
+			Info.ClassName = { "InitVariable", GetMetasoundDataTypeName<DataType>(), FName() };
 			Info.MajorVersion = 1;
 			Info.MinorVersion = 0;
-			Info.DisplayName = FText::Format(LOCTEXT("Metasound_InitVariableNodeDisplayNameFormat", "Init {0} Variable"), FText::FromName(GetMetasoundDataTypeName<DataType>()));
 			Info.Description = LOCTEXT("Metasound_InitVariableNodeDescription", "Initialize a variable of a MetaSound graph.");
 			Info.Author = PluginAuthor;
 			Info.PromptIfMissing = PluginNodeMissingPrompt;
