Audio Test fix-ups
#rnx

diff --git a/Engine/Plugins/Runtime/Metasound/Source/MetasoundFrontend/Public/MetasoundAssetBase.h b/Engine/Plugins/Runtime/Metasound/Source/MetasoundFrontend/Public/MetasoundAssetBase.h
--- a/Engine/Plugins/Runtime/Metasound/Source/MetasoundFrontend/Public/MetasoundAssetBase.h
+++ b/Engine/Plugins/Runtime/Metasound/Source/MetasoundFrontend/Public/MetasoundAssetBase.h
@@ -58,7 +58,10 @@ class METASOUNDFRONTEND_API IMetaSoundAssetManager
 public:
 	static void Set(IMetaSoundAssetManager& InInterface)
 	{
-		check(!Instance);
+		if (!IsTesting())
+		{
+			check(!Instance);
+		}
 		Instance = &InInterface;
 	}
 
@@ -68,8 +71,16 @@ public:
 		return *Instance;
 	}
 
+	// Whether or not manager is being used to run tests or not (enabling instances to be reset without asserting.)
+	virtual bool IsTesting() const { return false; }
+
+	// Whether or not the class is eligible for auto-update
 	virtual bool CanAutoUpdate(const FMetasoundFrontendClassName& InClassName) const = 0;
+
+	// Returns asset associated with the given key (null if key is not registered with the AssetManager or was not loaded from asset)
 	virtual FMetasoundAssetBase* FindAssetFromKey(const Metasound::Frontend::FNodeRegistryKey& InRegistryKey) const = 0;
+
+	// Returns path associated with the given key (null if key is not registered with the AssetManager or was not loaded from asset)
 	virtual const FSoftObjectPath* FindObjectPathFromKey(const Metasound::Frontend::FNodeRegistryKey& InRegistryKey) const = 0;
 
 	// Rescans settings for blacklisted assets not to run reference auto-update against.
