Backout dumb "fix"

diff --git a/Engine/Plugins/Runtime/Metasound/Source/MetasoundFrontend/Public/MetasoundAssetBase.h b/Engine/Plugins/Runtime/Metasound/Source/MetasoundFrontend/Public/MetasoundAssetBase.h
--- a/Engine/Plugins/Runtime/Metasound/Source/MetasoundFrontend/Public/MetasoundAssetBase.h
+++ b/Engine/Plugins/Runtime/Metasound/Source/MetasoundFrontend/Public/MetasoundAssetBase.h
@@ -58,10 +58,7 @@ class METASOUNDFRONTEND_API IMetaSoundAssetManager
 public:
 	static void Set(IMetaSoundAssetManager& InInterface)
 	{
-		if (!IsTesting())
-		{
-			check(!Instance);
-		}
+		check(!Instance);
 		Instance = &InInterface;
 	}
 
@@ -71,9 +68,6 @@ public:
 		return *Instance;
 	}
 
-	// Whether or not manager is being used to run tests or not (enabling instances to be reset without asserting.)
-	virtual bool IsTesting() const { return false; }
-
 	// Whether or not the class is eligible for auto-update
 	virtual bool CanAutoUpdate(const FMetasoundFrontendClassName& InClassName) const = 0;
 
