- Swap visible meter channel count based on MetaSound archetype
- Hide output meter for non-UMetasoundSource types in MetaSound Asset Editor
- Clean-up & encapsulate MetaSound Editor Meter as precursor work for adding additional meters/multi-channel output
- Center MetaSound Editor Output Meter in analyzer window

diff --git a/Engine/Plugins/Runtime/AudioWidgets/Source/AudioWidgets/Private/SAudioMeter.cpp b/Engine/Plugins/Runtime/AudioWidgets/Source/AudioWidgets/Private/SAudioMeter.cpp
--- a/Engine/Plugins/Runtime/AudioWidgets/Source/AudioWidgets/Private/SAudioMeter.cpp
+++ b/Engine/Plugins/Runtime/AudioWidgets/Source/AudioWidgets/Private/SAudioMeter.cpp
@@ -10,21 +10,23 @@ SAudioMeter::SAudioMeter()
 {
 }
 
-void SAudioMeter::Construct(const SAudioMeter::FArguments& InDeclaration)
+void SAudioMeter::Construct(const SAudioMeter::FArguments& InArgs)
 {
-	check(InDeclaration._Style);
+	check(InArgs._Style);
 
-	Orientation = InDeclaration._Orientation;
+	Orientation = InArgs._Orientation;
 
-	BackgroundColor = InDeclaration._BackgroundColor;
-	MeterBackgroundColor = InDeclaration._MeterBackgroundColor;
-	MeterValueColor = InDeclaration._MeterValueColor;
-	MeterPeakColor = InDeclaration._MeterPeakColor;
-	MeterClippingColor = InDeclaration._MeterClippingColor;
+	BackgroundColor = InArgs._BackgroundColor;
+	MeterBackgroundColor = InArgs._MeterBackgroundColor;
+	MeterValueColor = InArgs._MeterValueColor;
+	MeterPeakColor = InArgs._MeterPeakColor;
+	MeterScaleColor = InArgs._MeterScaleColor;
+	MeterScaleLabelColor = InArgs._MeterScaleLabelColor;
+	MeterClippingColor = InArgs._MeterClippingColor;
 
-	Style = InDeclaration._Style;
+	Style = InArgs._Style;
 
-	MeterChannelInfoAttribute = InDeclaration._MeterChannelInfo;
+	MeterChannelInfoAttribute = InArgs._MeterChannelInfo;
 }
 
 int32 SAudioMeter::OnPaint(const FPaintArgs& Args, const FGeometry& AllottedGeometry, const FSlateRect& MyCullingRect, FSlateWindowElementList& OutDrawElements, int32 LayerId, const FWidgetStyle& InWidgetStyle, bool bParentEnabled) const
diff --git a/Engine/Plugins/Runtime/Metasound/Source/MetasoundEditor/Private/MetasoundDetailCustomization.cpp b/Engine/Plugins/Runtime/Metasound/Source/MetasoundEditor/Private/MetasoundDetailCustomization.cpp
--- a/Engine/Plugins/Runtime/Metasound/Source/MetasoundEditor/Private/MetasoundDetailCustomization.cpp
+++ b/Engine/Plugins/Runtime/Metasound/Source/MetasoundEditor/Private/MetasoundDetailCustomization.cpp
@@ -11,6 +11,8 @@
 #include "IDetailGroup.h"
 #include "Input/Events.h"
 #include "MetasoundAssetBase.h"
+#include "MetasoundEditor.h"
+#include "MetasoundEditorGraphBuilder.h"
 #include "MetasoundEditorSettings.h"
 #include "MetasoundFrontend.h"
 #include "MetasoundFrontendController.h"
@@ -66,9 +68,7 @@ namespace Metasound
 
 		void FMetasoundDetailCustomization::CustomizeDetails(IDetailLayoutBuilder& DetailLayout)
 		{
-			using namespace Metasound;
-			using namespace Metasound::Editor;
-			using namespace Metasound::Frontend;
+			using namespace Frontend;
 
 			EMetasoundActiveDetailView DetailsView = EMetasoundActiveDetailView::Metasound;
 			if (const UMetasoundEditorSettings* EditorSettings = GetDefault<UMetasoundEditorSettings>())
@@ -76,7 +76,7 @@ namespace Metasound
 				DetailsView = EditorSettings->DetailView;
 			}
 
-			// Currently only support modifying a single MetaSound at a time (Multiple MetaSound editting will be covered by separate tool)
+			// Currently only support modifying a single MetaSound at a time (Multiple MetaSound editing will be covered by separate tool)
 			TArray<TWeakObjectPtr<UObject>> Objects;
 			DetailLayout.GetObjectsBeingCustomized(Objects);
 			if (Objects.Num() > 1)
@@ -124,14 +124,38 @@ namespace Metasound
 							return false;
 						});
 
+						if (MetaSoundSource.IsValid())
+						{
+							OutputFormat->SetOnPropertyValuePreChange(FSimpleDelegate::CreateLambda([Source = MetaSoundSource]()
+							{
+								if (Source.IsValid())
+								{
+									TSharedPtr<FEditor> ParentEditor = FGraphBuilder::GetEditorForMetasound(*Source.Get());
+									if (ParentEditor.IsValid())
+									{
+										ParentEditor->DestroyAnalyzers();
+									};
+								}
+							}));
+
+							OutputFormat->SetOnPropertyValueChanged(FSimpleDelegate::CreateLambda([Source = MetaSoundSource]()
+							{
+								if (Source.IsValid())
+								{
+									TSharedPtr<FEditor> ParentEditor = FGraphBuilder::GetEditorForMetasound(*Source.Get());
+									if (ParentEditor.IsValid())
+									{
+										ParentEditor->CreateAnalyzers();
+									};
+								}
+							}));
+						}
+
 						TSharedRef<SWidget> OutputFormatValueWidget = OutputFormat->CreatePropertyValueWidget();
 						OutputFormatValueWidget->SetEnabled(IsEnabledAttribute);
 
 						static const FText OutputFormatName = LOCTEXT("MetasoundOutputFormatPropertyName", "Output Format");
 						GeneralCategoryBuilder.AddCustomRow(OutputFormatName)
-// 						.EditCondition(
-// 						{
-// 						}), nullptr)
 						.NameContent()
 						[
 							OutputFormat->CreatePropertyNameWidget()
diff --git a/Engine/Plugins/Runtime/Metasound/Source/MetasoundEditor/Private/MetasoundEditor.cpp b/Engine/Plugins/Runtime/Metasound/Source/MetasoundEditor/Private/MetasoundEditor.cpp
--- a/Engine/Plugins/Runtime/Metasound/Source/MetasoundEditor/Private/MetasoundEditor.cpp
+++ b/Engine/Plugins/Runtime/Metasound/Source/MetasoundEditor/Private/MetasoundEditor.cpp
@@ -237,10 +237,9 @@ namespace Metasound
 			.SetGroup(WorkspaceMenuCategoryRef)
 			.SetIcon(FSlateIcon(FEditorStyle::GetStyleSetName(), "LevelEditor.Tabs.Settings"));
 
-
-			InTabManager->RegisterTabSpawner(TabFactory::Names::Analyzers, FOnSpawnTab::CreateLambda([InMetasoundMeter = MetasoundMeter](const FSpawnTabArgs& Args)
+			InTabManager->RegisterTabSpawner(TabFactory::Names::Analyzers, FOnSpawnTab::CreateLambda([InAnalyzerWidget = BuildAnalyzerWidget()](const FSpawnTabArgs& Args)
 			{
-				return TabFactory::CreateAnalyzersTab(InMetasoundMeter, Args);
+				return TabFactory::CreateAnalyzersTab(InAnalyzerWidget, Args);
 			}))
 			.SetDisplayName(LOCTEXT("AnalyzersTab", "Analyzers"))
 			.SetGroup(WorkspaceMenuCategoryRef)
@@ -259,6 +258,39 @@ namespace Metasound
 			InTabManager->UnregisterTabSpawner(TabFactory::Names::Metasound);
 		}
 
+		TSharedPtr<SWidget> FEditor::BuildAnalyzerWidget() const
+		{
+			if (!OutputMeter.IsValid() || !OutputMeter->GetWidget().IsValid())
+			{
+				return SNullWidget::NullWidget->AsShared();
+			}
+
+			const ISlateStyle* MetaSoundStyle = FSlateStyleRegistry::FindSlateStyle("MetaSoundStyle");
+			FLinearColor BackgroundColor = FLinearColor::Transparent;
+			if (ensure(MetaSoundStyle))
+			{
+				BackgroundColor = MetaSoundStyle->GetColor("MetasoundEditor.Analyzers.BackgroundColor");
+			}
+
+			return SNew(SOverlay)
+			+ SOverlay::Slot()
+			[
+				SNew(SColorBlock)
+				.Color(BackgroundColor)
+			]
+			+ SOverlay::Slot()
+			[
+				SNew(SVerticalBox)
+				+ SVerticalBox::Slot()
+				.FillHeight(1.0f)
+				.HAlign(HAlign_Center)
+				.VAlign(VAlign_Fill)
+				[
+					OutputMeter->GetWidget().ToSharedRef()
+				]
+			];
+		}
+
 		bool FEditor::IsPlaying() const
 		{
 			if (Metasound)
@@ -304,10 +336,7 @@ namespace Metasound
 				NameChangeDelegateHandles.Reset();
 			}
 
-			if (MetasoundMeterAnalyzer.IsValid() && ResultsDelegateHandle.IsValid())
-			{
-				MetasoundMeterAnalyzer->OnLatestPerChannelMeterResultsNative.Remove(ResultsDelegateHandle);
-			}
+			DestroyAnalyzers();
 
 			check(GEditor);
 			GEditor->UnregisterForUndo(this);
@@ -339,7 +368,7 @@ namespace Metasound
 
 			CreateAnalyzers();
 
-			const TSharedRef<FTabManager::FLayout> StandaloneDefaultLayout = FTabManager::NewLayout("Standalone_MetasoundEditor_Layout_v8")
+			const TSharedRef<FTabManager::FLayout> StandaloneDefaultLayout = FTabManager::NewLayout("Standalone_MetasoundEditor_Layout_v9")
 				->AddArea
 				(
 					FTabManager::NewPrimaryArea()
@@ -369,14 +398,14 @@ namespace Metasound
 						->Split
 						(
 							FTabManager::NewStack()
-							->SetSizeCoefficient(0.80f)
+							->SetSizeCoefficient(0.77f)
 							->SetHideTabWell(true)
 							->AddTab(TabFactory::Names::GraphCanvas, ETabState::OpenedTab)
 						)
 						->Split
 						(
 							FTabManager::NewStack()
-							->SetSizeCoefficient(0.05f)
+							->SetSizeCoefficient(0.08f)
 							->SetHideTabWell(true)
 							->AddTab(TabFactory::Names::Analyzers, ETabState::OpenedTab)
 						)
@@ -396,11 +425,6 @@ namespace Metasound
 			return Metasound;
 		}
 
-		UObject* FEditor::GetMetasoundAudioBusObject() const
-		{
-			return MetasoundAudioBus.Get();
-		}
-
 		void FEditor::SetSelection(const TArray<UObject*>& SelectedObjects)
 		{
 			if (MetasoundDetails.IsValid())
@@ -532,62 +556,28 @@ namespace Metasound
 			return false;
 		}
 
-		void FEditor::OnMeterOutput(UMeterAnalyzer* InMeterAnalyzer, int32 ChannelIndex, const FMeterResults& InMeterResults)
+		void FEditor::CreateAnalyzers()
 		{
-			if (InMeterAnalyzer && MetasoundMeter.IsValid())
+			if (UMetaSoundSource* MetaSoundSource = Cast<UMetaSoundSource>(Metasound))
 			{
-				FMeterChannelInfo ChannelInfo;
-				ChannelInfo.MeterValue = InMeterResults.MeterValue;
-				ChannelInfo.PeakValue = InMeterResults.PeakValue;
-
-				check(ChannelIndex < MetasoundChannelInfo.Num());
-				MetasoundChannelInfo[ChannelIndex] = ChannelInfo;
-
-				// If it's the last channel, update the widget
-				if (ChannelIndex == MetasoundChannelInfo.Num() - 1)
+				if (!OutputMeter.IsValid())
 				{
-					MetasoundMeter->SetMeterChannelInfo(MetasoundChannelInfo);
+					OutputMeter = MakeShared<FEditorMeter>();
 				}
+				OutputMeter->Init(EAudioBusChannels::Stereo, MetaSoundSource->NumChannels);
+			}
+			else
+			{
+				OutputMeter.Reset();
 			}
 		}
 
-		void FEditor::CreateAnalyzers()
+		void FEditor::DestroyAnalyzers()
 		{
-			MetasoundAudioBus = TStrongObjectPtr<UAudioBus>(NewObject<UAudioBus>());
-			MetasoundAudioBus->AudioBusChannels = EAudioBusChannels::Stereo;
-
-			MetasoundMeterAnalyzer = TStrongObjectPtr<UMeterAnalyzer>(NewObject<UMeterAnalyzer>());
-
-			ResultsDelegateHandle = MetasoundMeterAnalyzer->OnLatestPerChannelMeterResultsNative.AddRaw(this, &FEditor::OnMeterOutput);
-
-			MetasoundMeterAnalyzerSettings = TStrongObjectPtr<UMeterSettings>(NewObject<UMeterSettings>());
-			MetasoundMeterAnalyzerSettings->PeakHoldTime = 4000.0f;
-
-			MetasoundMeterAnalyzer->Settings = MetasoundMeterAnalyzerSettings.Get();
-
-			UWorld* EditorWorld = GEditor->GetEditorWorldContext().World();
-			MetasoundMeterAnalyzer->StartAnalyzing(EditorWorld, MetasoundAudioBus.Get());
-
-			// Create the audio meter
-			MetasoundMeter = SNew(SAudioMeter);
-
-			// TODO: THIS IS TEMP, WIP
-			MetasoundMeter->SetOrientation(EOrientation::Orient_Vertical);
-			MetasoundMeter->SetBackgroundColor(FLinearColor(0.0075f, 0.0075f, 0.0075, 1.0f));
-			MetasoundMeter->SetMeterBackgroundColor(FLinearColor(0.031f, 0.031f, 0.031f, 1.0f));
-			MetasoundMeter->SetMeterValueColor(FLinearColor(0.025719f, 0.208333f, 0.069907f, 1.0f));
-			MetasoundMeter->SetMeterPeakColor(FLinearColor(0.24349f, 0.708333f, 0.357002f, 1.0f));
-			MetasoundMeter->SetMeterClippingColor(FLinearColor(1.0f, 0.0f, 0.112334f, 1.0f));
-			MetasoundMeter->SetMeterScaleColor(FLinearColor(0.017642f, 0.017642f, 0.017642f, 1.0f));
-			MetasoundMeter->SetMeterScaleLabelColor(FLinearColor(0.442708f, 0.442708f, 0.442708f, 1.0f));
-
-			FMeterChannelInfo DefaultInfo;
-			DefaultInfo.MeterValue = -60.0f;
-			DefaultInfo.PeakValue = -60.0f;
-			MetasoundChannelInfo.Add(DefaultInfo);
-			MetasoundChannelInfo.Add(DefaultInfo);
-			MetasoundMeter->SetMeterChannelInfo(MetasoundChannelInfo);
-
+			if (OutputMeter.IsValid())
+			{
+				OutputMeter->Teardown();
+			}
 		}
 
 		void FEditor::ExtendToolbar()
@@ -790,19 +780,14 @@ namespace Metasound
 
 		void FEditor::Export()
 		{
-			FMetasoundAssetBase* InMetasoundAsset = IMetasoundUObjectRegistry::Get().GetObjectAsAssetBase(Metasound);
-			check(InMetasoundAsset);
-
-			if (!InMetasoundAsset)
-			{
-				return;
-			}
+			FMetasoundAssetBase* MetasoundAsset = IMetasoundUObjectRegistry::Get().GetObjectAsAssetBase(Metasound);
+			check(MetasoundAsset);
 
 			static const FString MetasoundExtension(TEXT(".metasound"));
 
 			// TODO: We could just make this an object.
 			const FString Path = FPaths::ProjectSavedDir() / TEXT("MetaSounds") + FPaths::ChangeExtension(Metasound->GetPathName(), MetasoundExtension);
-			InMetasoundAsset->GetDocumentHandle()->ExportToJSONAsset(Path);
+			MetasoundAsset->GetDocumentHandle()->ExportToJSONAsset(Path);
 		}
 
 		void FEditor::Play()
@@ -831,9 +816,9 @@ namespace Metasound
 						SetPreviewID(ParamInterfaceObject->GetUniqueID());
 					}
 
-					if (MetasoundAudioBus.IsValid())
+					if (UAudioBus* AudioBus = OutputMeter->GetAudioBus())
 					{
-						PreviewComp->SetAudioBusSendPostEffect(MetasoundAudioBus.Get(), 1.0f);
+						PreviewComp->SetAudioBusSendPostEffect(AudioBus, 1.0f);
 					}
 				}
 
diff --git a/Engine/Plugins/Runtime/Metasound/Source/MetasoundEditor/Private/MetasoundEditor.h b/Engine/Plugins/Runtime/Metasound/Source/MetasoundEditor/Private/MetasoundEditor.h
--- a/Engine/Plugins/Runtime/Metasound/Source/MetasoundEditor/Private/MetasoundEditor.h
+++ b/Engine/Plugins/Runtime/Metasound/Source/MetasoundEditor/Private/MetasoundEditor.h
@@ -9,6 +9,7 @@
 #include "GraphEditor.h"
 #include "IDetailsView.h"
 #include "IMetasoundEditor.h"
+#include "MetasoundEditorMeter.h"
 #include "MetasoundFrontend.h"
 #include "MetasoundFrontendController.h"
 #include "Misc/NotifyHook.h"
@@ -61,7 +62,6 @@ namespace Metasound
 			COUNT
 		};
 
-
 		class FEditor : public IMetasoundEditor, public FGCObject, public FNotifyHook, public FEditorUndoClient
 		{
 		public:
@@ -78,7 +78,6 @@ namespace Metasound
 
 			/** IMetasoundEditor interface */
 			virtual UObject* GetMetasoundObject() const override;
-			virtual UObject* GetMetasoundAudioBusObject() const override;
 			virtual void SetSelection(const TArray<UObject*>& SelectedObjects) override;
 			virtual bool GetBoundsForSelectedNodes(FSlateRect& Rect, float Padding) override;
 
@@ -125,8 +124,11 @@ namespace Metasound
 			void OnInputNameChanged(FGuid InNodeID);
 			void OnOutputNameChanged(FGuid InNodeID);
 
-			void OnMeterOutput(UMeterAnalyzer* InMeterAnalyzer, int32 ChannelIndex, const FMeterResults& MeterResults);
+			/** Creates analyzers */
+			void CreateAnalyzers();
 
+			/** Destroys analyzers */
+			void DestroyAnalyzers();
 
 		protected:
 			// Callbacks for action tree
@@ -203,9 +205,6 @@ namespace Metasound
 			/** Creates all internal widgets for the tabs to point at */
 			void CreateInternalWidgets();
 
-			/** Creates analyzers */
-			void CreateAnalyzers();
-
 			/** Builds the toolbar widget for the Metasound editor */
 			void ExtendToolbar();
 
@@ -265,6 +264,8 @@ namespace Metasound
 			void CreateGraphEditorWidget();
 
 		private:
+			TSharedPtr<SWidget> BuildAnalyzerWidget() const;
+
 			void EditObjectSettings();
 
 			void NotifyNodePasteFailure_ReferenceLoop();
@@ -284,11 +285,8 @@ namespace Metasound
 			/** Metasound Interface menu */
 			TSharedPtr<SGraphActionMenu> MetasoundInterfaceMenu;
 
-			/** Metasound Output Meter widget */
-			TSharedPtr<SAudioMeter> MetasoundMeter;
-
-			/** Metasound channel info for the meter. */
-			TArray<FMeterChannelInfo> MetasoundChannelInfo;
+			/** Meter used in the analyzer tab for auditioning preview output. */
+			TSharedPtr<FEditorMeter> OutputMeter;
 
 			/** Palette of Node types */
 			TSharedPtr<SMetasoundPalette> Palette;
@@ -303,17 +301,6 @@ namespace Metasound
 			/** The Metasound asset being edited */
 			UObject* Metasound = nullptr;
 
-			/** The preview audio bus. Used for analysis. */
-			TStrongObjectPtr<UAudioBus> MetasoundAudioBus;
-
-			/** Metasound analyzer object. */
-			TStrongObjectPtr<UMeterAnalyzer> MetasoundMeterAnalyzer;
-
-			TStrongObjectPtr<UMeterSettings> MetasoundMeterAnalyzerSettings;
-
-			/** Handle for results delegate for metasound meter analyzer. */
-			FDelegateHandle ResultsDelegateHandle;
-
 			TMap<FGuid, FDelegateHandle> NameChangeDelegateHandles;
 
 			/** Whether or not metasound being edited is valid */
diff --git a/Engine/Plugins/Runtime/Metasound/Source/MetasoundEditor/Private/MetasoundEditorMeter.cpp b/Engine/Plugins/Runtime/Metasound/Source/MetasoundEditor/Private/MetasoundEditorMeter.cpp
new file mode 100644
--- /dev/null
+++ b/Engine/Plugins/Runtime/Metasound/Source/MetasoundEditor/Private/MetasoundEditorMeter.cpp
@@ -0,0 +1,107 @@
+// Copyright Epic Games, Inc. All Rights Reserved.
+#include "MetasoundEditorMeter.h"
+
+
+namespace Metasound
+{
+	namespace Editor
+	{
+		FEditorMeter::FEditorMeter()
+			: Widget(SNew(SAudioMeter)
+				.Orientation(EOrientation::Orient_Vertical)
+				.BackgroundColor(FLinearColor::Transparent)
+
+				// TODO: Move to editor style
+				.MeterBackgroundColor(FLinearColor(0.031f, 0.031f, 0.031f, 1.0f))
+				.MeterValueColor(FLinearColor(0.025719f, 0.208333f, 0.069907f, 1.0f))
+				.MeterPeakColor(FLinearColor(0.24349f, 0.708333f, 0.357002f, 1.0f))
+				.MeterClippingColor(FLinearColor(1.0f, 0.0f, 0.112334f, 1.0f))
+				.MeterScaleColor(FLinearColor(0.017642f, 0.017642f, 0.017642f, 1.0f))
+				.MeterScaleLabelColor(FLinearColor(0.442708f, 0.442708f, 0.442708f, 1.0f))
+			)
+		{
+		}
+
+		UAudioBus* FEditorMeter::GetAudioBus() const
+		{
+			return AudioBus.Get();
+		}
+
+		TSharedPtr<SAudioMeter> FEditorMeter::GetWidget() const
+		{
+			return Widget;
+		}
+
+		void FEditorMeter::Init(EAudioBusChannels InMaxBusChannels, int32 InNumChannels)
+		{
+			check(InNumChannels <= static_cast<int32>(InMaxBusChannels) + 1);
+
+			Settings = TStrongObjectPtr(NewObject<UMeterSettings>());
+			Settings->PeakHoldTime = 4000.0f;
+
+			Analyzer = TStrongObjectPtr(NewObject<UMeterAnalyzer>());
+			Analyzer->Settings = Settings.Get();
+
+			AudioBus = TStrongObjectPtr(NewObject<UAudioBus>());
+			AudioBus->AudioBusChannels = InMaxBusChannels;
+
+			ResultsDelegateHandle = Analyzer->OnLatestPerChannelMeterResultsNative.AddRaw(this, &FEditorMeter::OnMeterOutput);
+
+			if (ensure(GEditor))
+			{
+				UWorld* EditorWorld = GEditor->GetEditorWorldContext().World();
+				if (ensure(EditorWorld))
+				{
+					Analyzer->StartAnalyzing(EditorWorld, AudioBus.Get());
+				}
+			}
+
+			constexpr float DefaultMeterValue = -60.0f;
+			constexpr float DefaultPeakValue = -60.0f;
+			ChannelInfo.Init(FMeterChannelInfo{ DefaultMeterValue, DefaultPeakValue }, InNumChannels);
+
+			if (Widget.IsValid())
+			{
+				Widget->SetMeterChannelInfo(ChannelInfo);
+			}
+		}
+
+		void FEditorMeter::OnMeterOutput(UMeterAnalyzer* InMeterAnalyzer, int32 ChannelIndex, const FMeterResults& InMeterResults)
+		{
+			if (InMeterAnalyzer == Analyzer.Get())
+			{
+				FMeterChannelInfo NewChannelInfo;
+				NewChannelInfo.MeterValue = InMeterResults.MeterValue;
+				NewChannelInfo.PeakValue = InMeterResults.PeakValue;
+
+				if (ChannelIndex < ChannelInfo.Num())
+				{
+					ChannelInfo[ChannelIndex] = MoveTemp(NewChannelInfo);
+				}
+
+				// Update the widget if this is the last channel
+				if (ChannelIndex == ChannelInfo.Num() - 1)
+				{
+					if (Widget.IsValid())
+					{
+						Widget->SetMeterChannelInfo(ChannelInfo);
+					}
+				}
+			}
+		}
+
+		void FEditorMeter::Teardown()
+		{
+			if (Analyzer.IsValid() && ResultsDelegateHandle.IsValid())
+			{
+				Analyzer->OnLatestPerChannelMeterResultsNative.Remove(ResultsDelegateHandle);
+				Analyzer.Reset();
+				ResultsDelegateHandle.Reset();
+			}
+
+			AudioBus.Reset();
+			ChannelInfo.Reset();
+			Settings.Reset();
+		}
+	} // namespace Editor
+} // namespace Metasound
diff --git a/Engine/Plugins/Runtime/Metasound/Source/MetasoundEditor/Private/MetasoundEditorModule.cpp b/Engine/Plugins/Runtime/Metasound/Source/MetasoundEditor/Private/MetasoundEditorModule.cpp
--- a/Engine/Plugins/Runtime/Metasound/Source/MetasoundEditor/Private/MetasoundEditorModule.cpp
+++ b/Engine/Plugins/Runtime/Metasound/Source/MetasoundEditor/Private/MetasoundEditorModule.cpp
@@ -129,6 +129,9 @@ namespace Metasound
 					Set("MetasoundEditor.Graph.Node.Math.Power", new FSlateImageBrush(RootToContentDir(TEXT("/Graph/node_math_power_40x.png")), Icon40x40));
 					Set("MetasoundEditor.Graph.Node.Math.Logarithm", new FSlateImageBrush(RootToContentDir(TEXT("/Graph/node_math_logarithm_40x.png")), Icon40x40));
 
+					// Analyzers
+					Set("MetasoundEditor.Analyzers.BackgroundColor", FLinearColor(0.0075f, 0.0075f, 0.0075, 1.0f));
+
 					// Misc
 					Set("MetasoundEditor.Speaker", new FSlateImageBrush(RootToContentDir(TEXT("/Icons/speaker_144x.png")), FVector2D(144.0f, 144.0f)));
 					Set("MetasoundEditor.Metasound.Icon", new IMAGE_BRUSH_SVG(TEXT("Icons/metasound_icon"), Icon16));
diff --git a/Engine/Plugins/Runtime/Metasound/Source/MetasoundEditor/Private/MetasoundEditorTabFactory.h b/Engine/Plugins/Runtime/Metasound/Source/MetasoundEditor/Private/MetasoundEditorTabFactory.h
--- a/Engine/Plugins/Runtime/Metasound/Source/MetasoundEditor/Private/MetasoundEditorTabFactory.h
+++ b/Engine/Plugins/Runtime/Metasound/Source/MetasoundEditor/Private/MetasoundEditorTabFactory.h
@@ -10,6 +10,7 @@
 #include "Styling/ISlateStyle.h"
 #include "Styling/SlateStyleRegistry.h"
 #include "Templates/SharedPointer.h"
+#include "Widgets/Colors/SColorBlock.h"
 
 #define LOCTEXT_NAMESPACE "MetaSoundEditor"
 
@@ -29,15 +30,15 @@ namespace Metasound
 				const FName Palette = "MetasoundEditor_Palette";
 			}
 
-			TSharedRef<SDockTab> CreateAnalyzersTab(TSharedPtr<SAudioMeter> InAudioMeter, const FSpawnTabArgs& Args)
+			TSharedRef<SDockTab> CreateAnalyzersTab(TSharedPtr<SWidget> InAnalyzerWidget, const FSpawnTabArgs& Args)
 			{
 				check(Args.GetTabId() == Names::Analyzers);
 
 				return SNew(SDockTab)
-					.Label(LOCTEXT("MetasoundAnalyzersTitle", "Analyzers"))
-					[
-						InAudioMeter.ToSharedRef()
-					];
+				.Label(LOCTEXT("MetasoundAnalyzersTitle", "Analyzers"))
+				[
+					InAnalyzerWidget.ToSharedRef()
+				];
 			}
 
 			TSharedRef<SDockTab> CreateGraphCanvasTab(TSharedPtr<SWidget> GraphEditor, const FSpawnTabArgs& Args)
diff --git a/Engine/Plugins/Runtime/Metasound/Source/MetasoundEditor/Public/IMetasoundEditor.h b/Engine/Plugins/Runtime/Metasound/Source/MetasoundEditor/Public/IMetasoundEditor.h
--- a/Engine/Plugins/Runtime/Metasound/Source/MetasoundEditor/Public/IMetasoundEditor.h
+++ b/Engine/Plugins/Runtime/Metasound/Source/MetasoundEditor/Public/IMetasoundEditor.h
@@ -12,7 +12,6 @@
 class IMetasoundEditor : public FAssetEditorToolkit
 {
 	virtual UObject* GetMetasoundObject() const = 0;
-	virtual UObject* GetMetasoundAudioBusObject() const = 0;
 	virtual void SetSelection(const TArray<UObject*>& SelectedObjects) = 0;
 	virtual bool GetBoundsForSelectedNodes(FSlateRect& Rect, float Padding) = 0;
 };
diff --git a/Engine/Plugins/Runtime/Metasound/Source/MetasoundEditor/Public/MetasoundEditorMeter.h b/Engine/Plugins/Runtime/Metasound/Source/MetasoundEditor/Public/MetasoundEditorMeter.h
new file mode 100644
--- /dev/null
+++ b/Engine/Plugins/Runtime/Metasound/Source/MetasoundEditor/Public/MetasoundEditorMeter.h
@@ -0,0 +1,52 @@
+// Copyright Epic Games, Inc. All Rights Reserved.
+#pragma once
+
+#include "Delegates/IDelegateInstance.h"
+#include "Meter.h"
+#include "SAudioMeter.h"
+#include "Sound/AudioBus.h"
+#include "Templates/SharedPointer.h"
+#include "UObject/StrongObjectPtr.h"
+
+
+namespace Metasound
+{
+	namespace Editor
+	{
+		class FEditorMeter
+		{
+		public:
+			FEditorMeter();
+
+			UAudioBus* GetAudioBus() const;
+
+			TSharedPtr<SAudioMeter> GetWidget() const;
+
+			void Teardown();
+
+			void Init(EAudioBusChannels InMaxBusChannels, int32 InNumChannels);
+
+		protected:
+			void OnMeterOutput(UMeterAnalyzer* InMeterAnalyzer, int32 ChannelIndex, const FMeterResults& InMeterResults);
+
+		private:
+			/** Metasound analyzer object. */
+			TStrongObjectPtr<UMeterAnalyzer> Analyzer;
+
+			/** The audio bus used for analysis. */
+			TStrongObjectPtr<UAudioBus> AudioBus;
+
+			/** Cached channel info for the meter. */
+			TArray<FMeterChannelInfo> ChannelInfo;
+
+			/** Handle for results delegate for MetaSound meter analyzer. */
+			FDelegateHandle ResultsDelegateHandle;
+
+			/** Meter settings. */
+			TStrongObjectPtr<UMeterSettings> Settings;
+
+			/** MetaSound Output Meter widget */
+			TSharedPtr<SAudioMeter> Widget;
+		};
+	} // namespace Editor
+} // namespace Metasound
\ No newline at end of file
