Fix for SFC, add utility function for subsequent slider change.
#rnx

diff --git a/Engine/Plugins/Runtime/Metasound/Source/MetasoundEditor/Private/MetasoundEditor.cpp b/Engine/Plugins/Runtime/Metasound/Source/MetasoundEditor/Private/MetasoundEditor.cpp
--- a/Engine/Plugins/Runtime/Metasound/Source/MetasoundEditor/Private/MetasoundEditor.cpp
+++ b/Engine/Plugins/Runtime/Metasound/Source/MetasoundEditor/Private/MetasoundEditor.cpp
@@ -1750,6 +1750,14 @@ namespace Metasound
 			MetasoundGraphEditor->NotifyGraphChanged();
 		}
 
+		void FEditor::RefreshDetails()
+		{
+			if (MetasoundDetails.IsValid())
+			{
+				MetasoundDetails->ForceRefresh();
+			}
+		}
+
 		void FEditor::RefreshInterface()
 		{
 			if (MetasoundInterfaceMenu.IsValid())
@@ -1939,7 +1947,8 @@ namespace Metasound
 				.OnClosedCallback(InOnMenuClosed);
 // 				.OnCloseReason(this, &FEditor::OnGraphActionMenuClosed);
 
-			return FActionMenuContent(ActionMenu, ActionMenu->GetFilterTextBox());
+			TSharedPtr<SWidget> FilterTextBox = StaticCastSharedRef<SWidget>(ActionMenu->GetFilterTextBox());
+			return FActionMenuContent(StaticCastSharedRef<SWidget>(ActionMenu), FilterTextBox);
 		}
 
 		void FEditor::OnActionSelected(const TArray<TSharedPtr<FEdGraphSchemaAction>>& InActions, ESelectInfo::Type InSelectionType)
diff --git a/Engine/Plugins/Runtime/Metasound/Source/MetasoundEditor/Private/MetasoundEditor.h b/Engine/Plugins/Runtime/Metasound/Source/MetasoundEditor/Private/MetasoundEditor.h
--- a/Engine/Plugins/Runtime/Metasound/Source/MetasoundEditor/Private/MetasoundEditor.h
+++ b/Engine/Plugins/Runtime/Metasound/Source/MetasoundEditor/Private/MetasoundEditor.h
@@ -110,6 +110,9 @@ namespace Metasound
 			void PasteNodes(const FVector2D* InLocation = nullptr);
 			void PasteNodes(const FVector2D* InLocation, const FText& InTransactionText);
 
+			/** Forces all UX pertaining to the root graph's details panel to be refreshed. */
+			void RefreshDetails();
+
 			/** Forces all UX pertaining to the root graph's interface to be refreshed. */
 			void RefreshInterface();
 
