Non-unity fix in metasound editor
#rnx

diff --git a/Engine/Plugins/Runtime/Metasound/Source/MetasoundEditor/Private/SMetasoundActionMenu.h b/Engine/Plugins/Runtime/Metasound/Source/MetasoundEditor/Private/SMetasoundActionMenu.h
--- a/Engine/Plugins/Runtime/Metasound/Source/MetasoundEditor/Private/SMetasoundActionMenu.h
+++ b/Engine/Plugins/Runtime/Metasound/Source/MetasoundEditor/Private/SMetasoundActionMenu.h
@@ -5,6 +5,7 @@
 #include "Layout/Margin.h"
 #include "Layout/Visibility.h"
 #include "Widgets/DeclarativeSyntaxSupport.h"
+#include "Widgets/Input/SEditableTextBox.h"
 #include "Widgets/SWidget.h"
 #include "Widgets/Views/SExpanderArrow.h"
 #include "Widgets/Layout/SBorder.h"
