Fix ensure on undoing/redoing adding MetaSound inputs and outputs
- set MetaSound asset synchronization required instead of calling RefreshInterface directly

diff --git a/Engine/Plugins/Runtime/Metasound/Source/MetasoundEditor/Private/MetasoundEditor.cpp b/Engine/Plugins/Runtime/Metasound/Source/MetasoundEditor/Private/MetasoundEditor.cpp
--- a/Engine/Plugins/Runtime/Metasound/Source/MetasoundEditor/Private/MetasoundEditor.cpp
+++ b/Engine/Plugins/Runtime/Metasound/Source/MetasoundEditor/Private/MetasoundEditor.cpp
@@ -493,7 +493,10 @@ namespace Metasound
 				MetasoundGraphEditor->NotifyGraphChanged();
 			}
 
-			RefreshInterface();
+			// Synchronize will update the interface
+			FMetasoundAssetBase* MetasoundAsset = IMetasoundUObjectRegistry::Get().GetObjectAsAssetBase(Metasound);
+			check(MetasoundAsset);
+			MetasoundAsset->SetSynchronizationRequired();
 
 			FSlateApplication::Get().DismissAllMenus();
 		}
@@ -1301,7 +1304,10 @@ namespace Metasound
 					Graph.RemoveMember(*GraphMember);
 				}
 
-				RefreshInterface();
+				//Synchronize will update the interface
+				FMetasoundAssetBase* MetasoundAsset = IMetasoundUObjectRegistry::Get().GetObjectAsAssetBase(Metasound);
+				check(MetasoundAsset);
+				MetasoundAsset->SetSynchronizationRequired();
 
 				if (NextToSelect)
 				{
diff --git a/Engine/Plugins/Runtime/Metasound/Source/MetasoundEditor/Private/MetasoundEditorGraph.cpp b/Engine/Plugins/Runtime/Metasound/Source/MetasoundEditor/Private/MetasoundEditorGraph.cpp
--- a/Engine/Plugins/Runtime/Metasound/Source/MetasoundEditor/Private/MetasoundEditorGraph.cpp
+++ b/Engine/Plugins/Runtime/Metasound/Source/MetasoundEditor/Private/MetasoundEditorGraph.cpp
@@ -340,29 +340,22 @@ void UMetasoundEditorGraphInputLiteral::PostEditUndo()
 {
 	Super::PostEditUndo();
 
-	if (UMetasoundEditorGraphInput* Input = Cast<UMetasoundEditorGraphInput>(GetOuter()))
-	{
-		Input->UpdateDocumentInput(false /* bPostTransaction */);
-	}
+	UpdateDocumentInputLiteral(false /* bPostTransaction */);
 }
 
-void UMetasoundEditorGraphInput::UpdateDocumentInput(bool bPostTransaction)
+void UMetasoundEditorGraphInputLiteral::UpdateDocumentInputLiteral(bool bPostTransaction)
 {
 	using namespace Metasound;
 	using namespace Metasound::Frontend;
 
-	UMetasoundEditorGraph* MetasoundGraph = CastChecked<UMetasoundEditorGraph>(GetOuter());
+	UMetasoundEditorGraphInput* Input = Cast<UMetasoundEditorGraphInput>(GetOuter());
+	UMetasoundEditorGraph* MetasoundGraph = CastChecked<UMetasoundEditorGraph>(Input->GetOuter());
 	UObject* Metasound = MetasoundGraph->GetMetasound();
 	if (!ensure(Metasound))
 	{
 		return;
 	}
 
-	if (!ensure(Literal))
-	{
-		return;
-	}
-
 	const FScopedTransaction Transaction(LOCTEXT("Set Input Default", "Set MetaSound Input Default"), bPostTransaction);
 	Metasound->Modify();
 
@@ -370,17 +363,17 @@ void UMetasoundEditorGraphInput::UpdateDocumentInput(bool bPostTransaction)
 	check(MetasoundAsset);
 
 	FGraphHandle GraphHandle = MetasoundAsset->GetRootGraphHandle();
-	FNodeHandle NodeHandle = GraphHandle->GetNodeWithID(NodeID);
+	FNodeHandle NodeHandle = GraphHandle->GetNodeWithID(Input->NodeID);
 
 	const Metasound::FVertexName& NodeName = NodeHandle->GetNodeName();
 	const FGuid VertexID = GraphHandle->GetVertexIDForInputVertex(NodeName);
-	GraphHandle->SetDefaultInput(VertexID, Literal->GetDefault());
+	GraphHandle->SetDefaultInput(VertexID, GetDefault());
 
 	// Disabled as internal call to validation to all other open graphs
 	// is expensive and can be spammed by dragging values 
 // 	Metasound::Editor::FGraphBuilder::RegisterGraphWithFrontend(*Metasound);
 
-	const bool bIsPreviewing = CastChecked<UMetasoundEditorGraph>(GetOuter())->IsPreviewing();
+	const bool bIsPreviewing = MetasoundGraph->IsPreviewing();
 	if (bIsPreviewing)
 	{
 		UAudioComponent* PreviewComponent = GEditor->GetPreviewAudioComponent();
@@ -388,7 +381,7 @@ void UMetasoundEditorGraphInput::UpdateDocumentInput(bool bPostTransaction)
 
 		if (TScriptInterface<IAudioParameterInterface> ParamInterface = PreviewComponent)
 		{
-			Metasound::Frontend::FConstNodeHandle ConstNodeHandle = GetConstNodeHandle();
+			Metasound::Frontend::FConstNodeHandle ConstNodeHandle = Input->GetConstNodeHandle();
 			Metasound::FVertexName VertexKey = NodeHandle->GetNodeName();
 			UpdatePreviewInstance(VertexKey, ParamInterface);
 		}
@@ -400,14 +393,6 @@ Metasound::Editor::ENodeSection UMetasoundEditorGraphInput::GetSectionID() const
 	return Metasound::Editor::ENodeSection::Inputs;
 }
 
-void UMetasoundEditorGraphInput::UpdatePreviewInstance(const Metasound::FVertexName& InParameterName, TScriptInterface<IAudioParameterInterface>& InParamInterface) const
-{
-	if (ensure(Literal))
-	{
-		Literal->UpdatePreviewInstance(InParameterName, InParamInterface);
-	}
-}
-
 Metasound::Frontend::FNodeHandle UMetasoundEditorGraphInput::AddNodeHandle(const FName& InName, FName InDataType)
 {
 	using namespace Metasound::Editor;
@@ -453,8 +438,7 @@ void UMetasoundEditorGraphInput::PostEditUndo()
 		}
 		return;
 	}
-
-	UpdateDocumentInput(false /* bPostTransaction */);
+	Literal->UpdateDocumentInputLiteral(false /* bPostTransaction */);
 	UpdateEditorLiteralType();
 }
 
diff --git a/Engine/Plugins/Runtime/Metasound/Source/MetasoundEditor/Private/MetasoundEditorGraphInputNodes.cpp b/Engine/Plugins/Runtime/Metasound/Source/MetasoundEditor/Private/MetasoundEditorGraphInputNodes.cpp
--- a/Engine/Plugins/Runtime/Metasound/Source/MetasoundEditor/Private/MetasoundEditorGraphInputNodes.cpp
+++ b/Engine/Plugins/Runtime/Metasound/Source/MetasoundEditor/Private/MetasoundEditorGraphInputNodes.cpp
@@ -60,7 +60,10 @@ void UMetasoundEditorGraphInputNode::UpdatePreviewInstance(const Metasound::FVer
 {
 	if (Input)
 	{
-		Input->UpdatePreviewInstance(InParameterName, InParameterInterface);
+		if (UMetasoundEditorGraphInputLiteral* InputLiteral = Input->Literal)
+		{
+			Input->Literal->UpdatePreviewInstance(InParameterName, InParameterInterface);
+		}
 	}
 }
 
@@ -97,7 +100,10 @@ void UMetasoundEditorGraphInputNode::PostEditUndo()
 
 	if (Input)
 	{
-		Input->UpdateDocumentInput(false /* bPostTransaction */);
+		if (UMetasoundEditorGraphInputLiteral* InputLiteral = Input->Literal)
+		{
+			Input->Literal->UpdateDocumentInputLiteral(false /* bPostTransaction */);
+		}
 	}
 }
 #endif // WITH_EDITORONLY_DATA
diff --git a/Engine/Plugins/Runtime/Metasound/Source/MetasoundEditor/Private/MetasoundNodeDetailCustomization.cpp b/Engine/Plugins/Runtime/Metasound/Source/MetasoundEditor/Private/MetasoundNodeDetailCustomization.cpp
--- a/Engine/Plugins/Runtime/Metasound/Source/MetasoundEditor/Private/MetasoundNodeDetailCustomization.cpp
+++ b/Engine/Plugins/Runtime/Metasound/Source/MetasoundEditor/Private/MetasoundNodeDetailCustomization.cpp
@@ -128,7 +128,10 @@ namespace Metasound
 					{
 						if (Input.IsValid())
 						{
-							Input->UpdateDocumentInput();
+							if (UMetasoundEditorGraphInputLiteral* InputLiteral = Input->Literal)
+							{
+								InputLiteral->UpdateDocumentInputLiteral();
+							}
 						}
 					});
 					RangeHandle->SetOnPropertyValueChanged(UpdateDocumentInput);
@@ -508,7 +511,10 @@ namespace Metasound
 				{
 					if (GraphInput.IsValid())
 					{
-						GraphInput->UpdateDocumentInput();
+						if (UMetasoundEditorGraphInputLiteral* InputLiteral = GraphInput->Literal)
+						{
+							InputLiteral->UpdateDocumentInputLiteral();
+						}
 					}
 				}
 			});
@@ -829,7 +835,10 @@ namespace Metasound
 								{
 									if (GraphMember.IsValid())
 									{
-										GraphMember->UpdateDocumentInput();
+										if (UMetasoundEditorGraphInputLiteral* InputLiteral = GraphMember->Literal)
+										{
+											InputLiteral->UpdateDocumentInputLiteral();
+										}
 									}
 								});
 
diff --git a/Engine/Plugins/Runtime/Metasound/Source/MetasoundEditor/Private/SMetasoundGraphNode.cpp b/Engine/Plugins/Runtime/Metasound/Source/MetasoundEditor/Private/SMetasoundGraphNode.cpp
--- a/Engine/Plugins/Runtime/Metasound/Source/MetasoundEditor/Private/SMetasoundGraphNode.cpp
+++ b/Engine/Plugins/Runtime/Metasound/Source/MetasoundEditor/Private/SMetasoundGraphNode.cpp
@@ -548,7 +548,10 @@ TSharedRef<SWidget> SMetasoundGraphNode::CreateNodeContentArea()
 						{
 							float Output = InputWidget->GetOutputValue(Value);
 							InputFloat->SetDefault(Output);
-							GraphInput->UpdateDocumentInput();
+							if (UMetasoundEditorGraphInputLiteral* InputLiteral = GraphInput->Literal)
+							{
+								InputLiteral->UpdateDocumentInputLiteral();
+							}
 						}
 					};
 					if (InputFloat->InputWidgetType == EMetasoundInputWidget::Slider)
diff --git a/Engine/Plugins/Runtime/Metasound/Source/MetasoundEditor/Public/MetasoundEditorGraph.h b/Engine/Plugins/Runtime/Metasound/Source/MetasoundEditor/Public/MetasoundEditorGraph.h
--- a/Engine/Plugins/Runtime/Metasound/Source/MetasoundEditor/Public/MetasoundEditorGraph.h
+++ b/Engine/Plugins/Runtime/Metasound/Source/MetasoundEditor/Public/MetasoundEditorGraph.h
@@ -60,6 +60,7 @@ public:
 	}
 
 	virtual void PostEditUndo() override;
+	void UpdateDocumentInputLiteral(bool bPostTransaction = true);
 
 	const UMetasoundEditorGraphMember* GetParentMember() const
 	{
@@ -203,12 +204,8 @@ public:
 	UPROPERTY(VisibleAnywhere, Category = DefaultValue)
 	UMetasoundEditorGraphInputLiteral* Literal;
 
-	void UpdatePreviewInstance(const Metasound::FVertexName& InParameterName, TScriptInterface<IAudioParameterInterface>& InParameterInterface) const;
-
 	void OnDataTypeChanged() override;
 
-	void UpdateDocumentInput(bool bPostTransaction = true);
-
 	virtual const FText& GetGraphMemberLabel() const override;
 	virtual Metasound::Editor::ENodeSection GetSectionID() const override;
 
