Fix not being able to mouse click searched for action from the MetaSound action menu

#rnx

diff --git a/Engine/Plugins/Runtime/Metasound/Source/MetasoundEditor/Private/SMetasoundActionMenu.cpp b/Engine/Plugins/Runtime/Metasound/Source/MetasoundEditor/Private/SMetasoundActionMenu.cpp
--- a/Engine/Plugins/Runtime/Metasound/Source/MetasoundEditor/Private/SMetasoundActionMenu.cpp
+++ b/Engine/Plugins/Runtime/Metasound/Source/MetasoundEditor/Private/SMetasoundActionMenu.cpp
@@ -215,6 +215,7 @@ void SMetasoundActionMenu::Construct(const FArguments& InArgs)
 						.DraggedFromPins(DraggedFromPins)
 						.GraphObj(Graph)
 						.AlphaSortItems(true)
+						.bAllowPreselectedItemActivation(true)
 				]
 			]
 		]
