MetaSounds Interface Checkpoint
- Archetype to Interface rename & support for multiple interface versions stored on MetaSoundBase
- Added ability to register interfaces
- Added spatialization/attenuation interfaces
- Added UX to add/remove Interfaces
- Fix ensure when deleting UMetaSound asset
#rnx

diff --git a/Engine/Plugins/Runtime/Metasound/Source/MetasoundEditor/Private/MetasoundAssetTypeActions.cpp b/Engine/Plugins/Runtime/Metasound/Source/MetasoundEditor/Private/MetasoundAssetTypeActions.cpp
--- a/Engine/Plugins/Runtime/Metasound/Source/MetasoundEditor/Private/MetasoundAssetTypeActions.cpp
+++ b/Engine/Plugins/Runtime/Metasound/Source/MetasoundEditor/Private/MetasoundAssetTypeActions.cpp
@@ -155,9 +155,9 @@ namespace Metasound
 			using namespace AssetTypeActionsPrivate;
 			RegisterPresetAction<UMetaSoundSource, UMetaSoundSourceFactory>(SourcePresetLabel, SourcePresetToolTip);
 
-			// This is currently disabled because of the requirement that archetypes of a preset and its referenced
-			// asset must match.  TODO: A feature of archetypes/preset transform is required for this action to work
-			// that introduces the concept of an archetype "interface" (i.e. if the archetype one graph subscribes to
+			// This is currently disabled because of the requirement that interfaces of a preset and its referenced
+			// asset must match.  TODO: A feature of interfaces/preset transform is required for this action to work
+			// that introduces the concept of an interface "interface" (i.e. if the interface one graph subscribes to
 			// or contains all of the necessary inputs/outputs of another), so that the preset can "implement" the
 			// referenced graphs "interface."
 // 			RegisterPresetAction<UMetaSound, UMetaSoundFactory, UMetaSoundSource>(PresetLabel, PresetToolTip);
diff --git a/Engine/Plugins/Runtime/Metasound/Source/MetasoundEditor/Private/MetasoundDetailCustomization.cpp b/Engine/Plugins/Runtime/Metasound/Source/MetasoundEditor/Private/MetasoundDetailCustomization.cpp
--- a/Engine/Plugins/Runtime/Metasound/Source/MetasoundEditor/Private/MetasoundDetailCustomization.cpp
+++ b/Engine/Plugins/Runtime/Metasound/Source/MetasoundEditor/Private/MetasoundDetailCustomization.cpp
@@ -8,6 +8,8 @@
 #include "DetailLayoutBuilder.h"
 #include "DetailWidgetRow.h"
 #include "Framework/Notifications/NotificationManager.h"
+#include "IAudioGeneratorInterfaceRegistry.h"
+#include "IAudioParameterTransmitter.h"
 #include "IDetailGroup.h"
 #include "Input/Events.h"
 #include "MetasoundAssetBase.h"
@@ -15,9 +17,12 @@
 #include "MetasoundEditorGraphBuilder.h"
 #include "MetasoundEditorSettings.h"
 #include "MetasoundFrontend.h"
+#include "MetasoundFrontendArchetypeRegistry.h"
 #include "MetasoundFrontendController.h"
+#include "MetasoundFrontendSearchEngine.h"
 #include "MetasoundSource.h"
 #include "MetasoundUObjectRegistry.h"
+#include "PropertyCustomizationHelpers.h"
 #include "PropertyEditorDelegates.h"
 #include "PropertyHandle.h"
 #include "PropertyRestriction.h"
@@ -30,6 +35,9 @@
 #include "UObject/WeakObjectPtrTemplates.h"
 #include "Widgets/Notifications/SNotificationList.h"
 #include "Widgets/Text/STextBlock.h"
+#include "Styling/SlateTypes.h"
+#include "Styling/CoreStyle.h"
+#include "Widgets/Input/SCheckBox.h"
 
 
 #define LOCTEXT_NAMESPACE "MetaSoundEditor"
@@ -53,6 +61,198 @@ namespace Metasound
 			: IDetailCustomization()
 			, DocumentPropertyName(InDocumentPropertyName)
 		{
+			IsGraphEditableAttribute = TAttribute<bool>::Create([this]()
+			{
+				using namespace Metasound;
+				using namespace Metasound::Frontend;
+				if (const FMetasoundAssetBase* MetaSoundAsset = IMetasoundUObjectRegistry::Get().GetObjectAsAssetBase(MetaSound.Get()))
+				{
+					FConstGraphHandle GraphHandle = MetaSoundAsset->GetRootGraphHandle();
+					return GraphHandle->GetGraphStyle().bIsGraphEditable;
+				}
+
+				return false;
+			});
+		}
+
+		void FMetasoundDetailCustomization::CustomizeInterfaces(IDetailLayoutBuilder& DetailLayout)
+		{
+			UpdateInterfaceNames();
+
+			SAssignNew(InterfaceComboBox, SSearchableComboBox)
+				.OptionsSource(&AddableInterfaceNames)
+				.OnGenerateWidget_Lambda([](TSharedPtr<FString> InItem)
+				{
+					return SNew(STextBlock)
+						.Text(FText::FromString(*InItem));
+				})
+				.OnSelectionChanged_Lambda([this](TSharedPtr<FString> NameToAdd, ESelectInfo::Type InSelectInfo)
+				{
+					using namespace Metasound;
+					using namespace Metasound::Frontend;
+
+					FMetasoundAssetBase* MetaSoundAsset = IMetasoundUObjectRegistry::Get().GetObjectAsAssetBase(MetaSound.Get());
+					if (!ensure(MetaSoundAsset))
+					{
+						return;
+					}
+
+					if (InSelectInfo != ESelectInfo::OnNavigation)
+					{
+						FMetasoundFrontendInterface InterfaceToAdd;
+						if (ensure(ISearchEngine::Get().FindInterfaceWithHighestVersion(FName(*NameToAdd.Get()), InterfaceToAdd)))
+						{
+							FScopedTransaction(LOCTEXT("AddInterfaceTransaction", "Add MetaSound Interface"));
+							MetaSound.Get()->Modify();
+							MetaSoundAsset->GetGraphChecked().Modify();
+
+							FDocumentHandle DocumentHandle = MetaSoundAsset->GetDocumentHandle();
+							FModifyRootGraphInterfaces ModifyTransform({ }, { InterfaceToAdd });
+							ModifyTransform.SetDefaultNodeLocations(false); // Don't automatically add nodes to ed graph
+							ModifyTransform.Transform(DocumentHandle);
+							MetaSoundAsset->SetInterfacesUpdatedOnSynchronization();
+						}
+
+						UpdateInterfaceNames();
+						InterfaceComboBox->RefreshOptions();
+						MetaSoundAsset->SetSynchronizationRequired();
+					}
+				})
+				.Content()
+				[
+					SNew(STextBlock)
+					.Text(LOCTEXT("UpdateInterfaceAction", "Add Interface..."))
+					.IsEnabled(IsGraphEditableAttribute)
+				];
+
+			TSharedRef<SWidget> InterfaceUtilities = SNew(SHorizontalBox)
+			+ SHorizontalBox::Slot()
+			.Padding(2.0f)
+			.HAlign(HAlign_Left)
+			.VAlign(VAlign_Center)
+			.AutoWidth()
+			[
+				InterfaceComboBox->AsShared()
+			]
+			+ SHorizontalBox::Slot()
+			.Padding(2.0f)
+			.HAlign(HAlign_Left)
+			.VAlign(VAlign_Center)
+			.AutoWidth()
+			[
+				PropertyCustomizationHelpers::MakeDeleteButton(FSimpleDelegate::CreateLambda([this]()
+				{
+					using namespace Frontend;
+
+					FMetasoundAssetBase* MetaSoundAsset = IMetasoundUObjectRegistry::Get().GetObjectAsAssetBase(MetaSound.Get());
+					if (!ensure(MetaSoundAsset))
+					{
+						return;
+					}
+
+					TArray<FMetasoundFrontendInterface> ImplementedInterfaces;
+					Algo::Transform(ImplementedInterfaceNames, ImplementedInterfaces, [](const FName& Name)
+					{
+						FMetasoundFrontendInterface Interface;
+						ISearchEngine::Get().FindInterfaceWithHighestVersion(Name, Interface);
+						return Interface;
+					});
+
+					{
+						FScopedTransaction(LOCTEXT("RemoveAllInterfacesTransaction", "Remove All MetaSound Interfaces"));
+						MetaSound.Get()->Modify();
+						MetaSoundAsset->GetGraphChecked().Modify();
+
+						FDocumentHandle DocumentHandle = MetaSoundAsset->GetDocumentHandle();
+						FModifyRootGraphInterfaces({ ImplementedInterfaces }, { }).Transform(DocumentHandle);
+					}
+
+					UpdateInterfaceNames();
+					InterfaceComboBox->RefreshOptions();
+					MetaSoundAsset->SetInterfacesUpdatedOnSynchronization();
+					MetaSoundAsset->SetSynchronizationRequired();
+				}), LOCTEXT("RemoveInterfaceTooltip", "Removes all interfaces from the given MetaSound."))
+			];
+
+			const FText HeaderName = LOCTEXT("InterfacesGroupDisplayName", "Interfaces");
+			IDetailCategoryBuilder& InterfaceCategory = DetailLayout.EditCategory("Interfaces", HeaderName);
+
+			InterfaceCategory.AddCustomRow(HeaderName)
+			.ValueContent()
+			[
+				InterfaceUtilities
+			];
+
+			auto CreateInterfaceEntryWidget = [&](FName InInterfaceName) -> TSharedPtr<SWidget>
+			{
+				using namespace Frontend;
+
+				FMetasoundFrontendInterface InterfaceEntry;
+				if (!ensure(ISearchEngine::Get().FindInterfaceWithHighestVersion(InInterfaceName, InterfaceEntry)))
+				{
+					return SNullWidget::NullWidget;
+				}
+
+				TSharedRef<SWidget> RemoveButtonWidget = PropertyCustomizationHelpers::MakeDeleteButton(FSimpleDelegate::CreateLambda([this, InterfaceEntry]()
+				{
+					using namespace Frontend;
+
+					FMetasoundAssetBase* MetaSoundAsset = IMetasoundUObjectRegistry::Get().GetObjectAsAssetBase(MetaSound.Get());
+					if (!ensure(MetaSoundAsset))
+					{
+						return;
+					}
+
+					{
+						FScopedTransaction(LOCTEXT("RemoveInterfaceTransaction", "Remove MetaSound Interface"));
+						MetaSound.Get()->Modify();
+						MetaSoundAsset->GetGraphChecked().Modify();
+
+						FDocumentHandle DocumentHandle = MetaSoundAsset->GetDocumentHandle();
+						FModifyRootGraphInterfaces({ InterfaceEntry }, { }).Transform(DocumentHandle);
+					}
+
+					UpdateInterfaceNames();
+					InterfaceComboBox->RefreshOptions();
+					MetaSoundAsset->SetInterfacesUpdatedOnSynchronization();
+					MetaSoundAsset->SetSynchronizationRequired();
+				}), LOCTEXT("RemoveInterfaceTooltip", "Removes the associated interface from the MetaSound."));
+
+				return SNew(SHorizontalBox)
+					+ SHorizontalBox::Slot()
+					.Padding(2.0f)
+					.HAlign(HAlign_Center)
+					.VAlign(VAlign_Center)
+					.AutoWidth()
+					[
+						SNew(STextBlock)
+						.Text(FText::FromName(InterfaceEntry.Version.Name))
+					]
+					+ SHorizontalBox::Slot()
+					.Padding(2.0f)
+					.HAlign(HAlign_Center)
+					.VAlign(VAlign_Center)
+					.AutoWidth()
+					[
+						RemoveButtonWidget
+					];
+			};
+
+			TArray<FName> InterfaceNames = ImplementedInterfaceNames.Array();
+			InterfaceNames.Sort([](const FName& A, const FName& B) { return A.LexicalLess(B); });
+			for (const FName& InterfaceName : InterfaceNames)
+			{
+				InterfaceCategory.AddCustomRow(FText::FromName(InterfaceName))
+				.ValueContent()
+				[
+					CreateInterfaceEntryWidget(InterfaceName)->AsShared()
+				];
+			}
+		}
+
+		FName FMetasoundDetailCustomization::GetInterfaceVersionsPath() const
+		{
+			return Metasound::Editor::BuildChildPath(DocumentPropertyName, GET_MEMBER_NAME_CHECKED(FMetasoundFrontendDocument, InterfaceVersions));
 		}
 
 		FName FMetasoundDetailCustomization::GetMetadataRootClassPath() const
@@ -66,6 +266,44 @@ namespace Metasound
 			return Metasound::Editor::BuildChildPath(RootClass, GET_MEMBER_NAME_CHECKED(FMetasoundFrontendClass, Metadata));
 		}
 
+		void FMetasoundDetailCustomization::UpdateInterfaceNames()
+		{
+			using namespace Frontend;
+
+			AddableInterfaceNames.Reset();
+			ImplementedInterfaceNames.Reset();
+
+			if (const FMetasoundAssetBase* MetaSoundAsset = IMetasoundUObjectRegistry::Get().GetObjectAsAssetBase(MetaSound.Get()))
+			{
+				const TArray<FMetasoundFrontendVersion>& ImplementedInterfaces = MetaSoundAsset->GetDocumentChecked().InterfaceVersions;
+
+				auto GetVersionName = [](const FMetasoundFrontendVersion& Version) { return Version.Name; };
+				auto IsTransmittableVersion = [](const FMetasoundFrontendVersion& Version)
+				{
+					using namespace Metasound::Frontend;
+
+					const FInterfaceRegistryKey Key = GetInterfaceRegistryKey(Version);
+					if (const IInterfaceRegistryEntry* Entry = IInterfaceRegistry::Get().FindInterfaceRegistryEntry(Key))
+					{
+						return Entry->GetRouterName() == Audio::IParameterTransmitter::RouterName;
+					}
+
+					return false;
+				};
+
+				Algo::TransformIf(ImplementedInterfaces, ImplementedInterfaceNames, IsTransmittableVersion, GetVersionName);
+
+				Audio::IGeneratorInterfaceRegistry::Get().IterateInterfaces([this](Audio::FGeneratorInterfacePtr Interface)
+				{
+					if (!ImplementedInterfaceNames.Contains(Interface->Name))
+					{
+						FString Name = Interface->Name.ToString();
+						AddableInterfaceNames.Add(MakeShared<FString>(MoveTemp(Name)));
+					}
+				});
+			}
+		}
+
 		void FMetasoundDetailCustomization::CustomizeDetails(IDetailLayoutBuilder& DetailLayout)
 		{
 			using namespace Frontend;
@@ -76,15 +314,22 @@ namespace Metasound
 				DetailsView = EditorSettings->DetailView;
 			}
 
-			// Currently only support modifying a single MetaSound at a time (Multiple MetaSound editing will be covered by separate tool)
 			TArray<TWeakObjectPtr<UObject>> Objects;
 			DetailLayout.GetObjectsBeingCustomized(Objects);
+
+			// Only support modifying a single MetaSound at a time (Multiple
+			// MetaSound editing will be covered most likely by separate tool).
 			if (Objects.Num() > 1)
 			{
 				return;
 			}
+			MetaSound = Objects.Last();
+			if (!MetaSound.IsValid())
+			{
+				return;
+			}
 
-			TWeakObjectPtr<UMetaSoundSource> MetaSoundSource = Cast<UMetaSoundSource>(Objects[0].Get());
+			TWeakObjectPtr<UMetaSoundSource> MetaSoundSource = Cast<UMetaSoundSource>(MetaSound.Get());
 
 			switch (DetailsView)
 			{
@@ -102,28 +347,20 @@ namespace Metasound
 					const FName MajorVersionPropertyPath = BuildChildPath(VersionPropertyPath, GET_MEMBER_NAME_CHECKED(FMetasoundFrontendVersionNumber, Major));
 					const FName MinorVersionPropertyPath = BuildChildPath(VersionPropertyPath, GET_MEMBER_NAME_CHECKED(FMetasoundFrontendVersionNumber, Minor));
 
+					const FName InterfaceVersionsPropertyPath = GetInterfaceVersionsPath();
+
 					TSharedPtr<IPropertyHandle> AuthorHandle = DetailLayout.GetProperty(AuthorPropertyPath);
 					TSharedPtr<IPropertyHandle> CategoryHierarchyHandle = DetailLayout.GetProperty(CategoryHierarchyPropertyPath);
 					TSharedPtr<IPropertyHandle> ClassNameHandle = DetailLayout.GetProperty(ClassNameNamePropertyPath);
 					TSharedPtr<IPropertyHandle> DescHandle = DetailLayout.GetProperty(DescPropertyPath);
 					TSharedPtr<IPropertyHandle> MajorVersionHandle = DetailLayout.GetProperty(MajorVersionPropertyPath);
 					TSharedPtr<IPropertyHandle> MinorVersionHandle = DetailLayout.GetProperty(MinorVersionPropertyPath);
+					TSharedPtr<IPropertyHandle> InterfaceVersionsHandle = DetailLayout.GetProperty(InterfaceVersionsPropertyPath);
 
 					// Invalid for UMetaSounds
 					TSharedPtr<IPropertyHandle> OutputFormat = DetailLayout.GetProperty(GET_MEMBER_NAME_CHECKED(UMetaSoundSource, OutputFormat));
 					if (OutputFormat.IsValid())
 					{
-						TAttribute<bool> IsEnabledAttribute = TAttribute<bool>::Create([InSource = MetaSoundSource]()
-						{
-							if (const FMetasoundAssetBase* MetaSoundAsset = IMetasoundUObjectRegistry::Get().GetObjectAsAssetBase(Cast<UObject>(InSource)))
-							{
-								FConstGraphHandle GraphHandle = MetaSoundAsset->GetRootGraphHandle();
-								return GraphHandle->GetGraphStyle().bIsGraphEditable;
-							}
-
-							return false;
-						});
-
 						if (MetaSoundSource.IsValid())
 						{
 							OutputFormat->SetOnPropertyValuePreChange(FSimpleDelegate::CreateLambda([Source = MetaSoundSource]()
@@ -152,7 +389,7 @@ namespace Metasound
 						}
 
 						TSharedRef<SWidget> OutputFormatValueWidget = OutputFormat->CreatePropertyValueWidget();
-						OutputFormatValueWidget->SetEnabled(IsEnabledAttribute);
+						OutputFormatValueWidget->SetEnabled(IsGraphEditableAttribute);
 
 						static const FText OutputFormatName = LOCTEXT("MetasoundOutputFormatPropertyName", "Output Format");
 						GeneralCategoryBuilder.AddCustomRow(OutputFormatName)
@@ -191,6 +428,8 @@ namespace Metasound
 					];
 					GeneralCategoryBuilder.AddProperty(CategoryHierarchyHandle);
 
+					CustomizeInterfaces(DetailLayout);
+
 					DetailLayout.HideCategory("Advanced");
 					DetailLayout.HideCategory("Analysis");
 					DetailLayout.HideCategory("Attenuation");
diff --git a/Engine/Plugins/Runtime/Metasound/Source/MetasoundEditor/Private/MetasoundDetailCustomization.h b/Engine/Plugins/Runtime/Metasound/Source/MetasoundEditor/Private/MetasoundDetailCustomization.h
--- a/Engine/Plugins/Runtime/Metasound/Source/MetasoundEditor/Private/MetasoundDetailCustomization.h
+++ b/Engine/Plugins/Runtime/Metasound/Source/MetasoundEditor/Private/MetasoundDetailCustomization.h
@@ -3,9 +3,12 @@
 
 #include "IDetailCustomization.h"
 #include "SGraphActionMenu.h"
+#include "SSearchableComboBox.h"
+#include "Templates/SharedPointer.h"
 #include "Types/SlateEnums.h"
 #include "UObject/NameTypes.h"
 
+
 // Forward Declarations
 class FPropertyRestriction;
 class IDetailLayoutBuilder;
@@ -25,9 +28,22 @@ namespace Metasound
 			// End of IDetailCustomization interface
 
 		private:
+			void CustomizeInterfaces(IDetailLayoutBuilder& DetailLayout);
+
+			FName GetInterfaceVersionsPath() const;
 			FName GetMetadataRootClassPath() const;
 			FName GetMetadataPropertyPath() const;
 
+			void UpdateInterfaceNames();
+
+			TArray<TSharedPtr<FString>> AddableInterfaceNames;
+
+			TSet<FName> ImplementedInterfaceNames;
+			TSharedPtr<SSearchableComboBox> InterfaceComboBox;
+			TAttribute<bool> IsGraphEditableAttribute;
+
+			TWeakObjectPtr<UObject> MetaSound;
+
 			FName DocumentPropertyName;
 		};
 	} // namespace Editor
diff --git a/Engine/Plugins/Runtime/Metasound/Source/MetasoundEditor/Private/MetasoundEditor.cpp b/Engine/Plugins/Runtime/Metasound/Source/MetasoundEditor/Private/MetasoundEditor.cpp
--- a/Engine/Plugins/Runtime/Metasound/Source/MetasoundEditor/Private/MetasoundEditor.cpp
+++ b/Engine/Plugins/Runtime/Metasound/Source/MetasoundEditor/Private/MetasoundEditor.cpp
@@ -2,6 +2,7 @@
 
 #include "MetasoundEditor.h"
 
+#include "Algo/AnyOf.h"
 #include "AudioDevice.h"
 #include "AudioMeterStyle.h"
 #include "Components/AudioComponent.h"
@@ -115,15 +116,6 @@ namespace Metasound
 				return GetMetasoundAssetChecked().GetRootGraphHandle();
 			}
 
-			bool IsRequired() const
-			{
-				if (UMetasoundEditorGraphMember* Member = GetGraphMember())
-				{
-					return Member->IsRequired();
-				}
-				return false;
-			}
-
 			// FEdGraphSchemaAction interface
 			virtual bool IsParentable() const override
 			{
@@ -157,7 +149,6 @@ namespace Metasound
 		{
 		private:
 			TSharedPtr<FMetasoundGraphMemberSchemaAction> MetasoundAction;
-			bool bIsNameInvalid = false;
 
 		public:
 			SLATE_BEGIN_ARGS(SMetaSoundGraphPaletteItem)
@@ -364,7 +355,7 @@ namespace Metasound
 			using namespace Metasound::Frontend;
 
 			check(ObjectToEdit);
-			checkf(IMetasoundUObjectRegistry::Get().IsRegisteredClass(ObjectToEdit), TEXT("Object passed in was not registered as a valid metasound archetype!"));
+			checkf(IMetasoundUObjectRegistry::Get().IsRegisteredClass(ObjectToEdit), TEXT("Object passed in was not registered as a valid metasound interface!"));
 			
 			// Support undo/redo
 			ObjectToEdit->SetFlags(RF_Transactional);
@@ -372,6 +363,7 @@ namespace Metasound
 			Metasound = ObjectToEdit;
 
 			FGraphBuilder::RegisterGraphWithFrontend(*Metasound);
+			FGraphBuilder::SynchronizeGraph(*Metasound);
 
 			GEditor->RegisterForUndo(this);
 
@@ -760,24 +752,36 @@ namespace Metasound
 
 			if (Frontend::ImportJSONAssetToMetasound(InputPath, MetasoundDoc))
 			{
-				TArray<UClass*> ImportClasses = IMetasoundUObjectRegistry::Get().GetUClassesForArchetype(MetasoundDoc.ArchetypeVersion);
+				TSet<UClass*> ImportClasses;
+
+				for (const FMetasoundFrontendVersion& InterfaceVersion : MetasoundDoc.InterfaceVersions)
+				{
+					TArray<UClass*> InterfaceClasses = IMetasoundUObjectRegistry::Get().FindSupportedInterfaceClasses(InterfaceVersion);
+					ImportClasses.Append(MoveTemp(InterfaceClasses));
+				}
 
 				if (ImportClasses.Num() < 1)
 				{
-					UE_LOG(LogMetaSound, Warning, TEXT("Cannot create UObject from MetaSound document. No UClass supports archetype \"%s\""), *MetasoundDoc.ArchetypeVersion.ToString());
+					TArray<FString> InterfaceNames;
+					Algo::Transform(MetasoundDoc.InterfaceVersions, InterfaceNames, [] (const FMetasoundFrontendVersion& InterfaceVersion) { return InterfaceVersion.ToString(); });
+					UE_LOG(LogMetaSound, Warning, TEXT("Cannot create UObject from MetaSound document. No UClass supports interface(s) \"%s\""), *FString::Join(InterfaceNames, TEXT(",")));
 				}
 				else
 				{
-					if (ImportClasses.Num() > 1)
+					UClass* AnyClass = nullptr;
+					for (UClass* ImportClass : ImportClasses)
 					{
-						for (UClass* Cls : ImportClasses)
+						AnyClass = ImportClass;
+						if (ImportClasses.Num() > 1)
 						{
-							// TODO: could do a modal dialog to give user choice of import type.
-							UE_LOG(LogMetaSound, Warning, TEXT("Duplicate UClass support archetype \"%s\" with UClass \"%s\""), *MetasoundDoc.ArchetypeVersion.ToString(), *Cls->GetName());
+							// TODO: Modal dialog to give user choice of import type.
+							TArray<FString> InterfaceNames;
+							Algo::Transform(MetasoundDoc.InterfaceVersions, InterfaceNames, [](const FMetasoundFrontendVersion& InterfaceVersion) { return InterfaceVersion.ToString(); });
+							UE_LOG(LogMetaSound, Warning, TEXT("Duplicate UClass support interface(s) \"%s\" with UClass \"%s\""), *FString::Join(InterfaceNames, TEXT(",")), *ImportClass->GetName());
 						}
 					}
 
-					IMetasoundUObjectRegistry::Get().NewObject(ImportClasses[0], MetasoundDoc, OutputPath);
+					IMetasoundUObjectRegistry::Get().NewObject(AnyClass, MetasoundDoc, OutputPath);
 				}
 			}
 			else
@@ -1345,14 +1349,20 @@ namespace Metasound
 						if (MetasoundAction.IsValid())
 						{
 							const UMetasoundEditorGraphMember* GraphMember = MetasoundAction->GetGraphMember();
-							
 							if (ensure(nullptr != GraphMember))
 							{
-								if (GraphMember->IsRequired())
+								const FMetasoundFrontendVersion* InterfaceVersion = nullptr;
+								if (const UMetasoundEditorGraphVertex* Vertex = Cast<UMetasoundEditorGraphVertex>(GraphMember))
+								{
+									InterfaceVersion = &Vertex->GetInterfaceVersion();
+								}
+
+								if (InterfaceVersion && InterfaceVersion->IsValid())
 								{
 									if (MetasoundGraphEditor.IsValid())
 									{
-										FNotificationInfo Info(LOCTEXT("CannotDelete_RequiredGraphMember", "Delete failed: Input/Output is required."));
+										const FText Notification = FText::Format(LOCTEXT("CannotDeleteInterfaceMemberNotificationFormat", "Cannot delete individual member of interface '{0}'."), FText::FromName(InterfaceVersion->Name));
+										FNotificationInfo Info(Notification);
 										Info.bFireAndForget = true;
 										Info.bUseSuccessFailIcons = false;
 										Info.ExpireDuration = 5.0f;
@@ -1888,13 +1898,30 @@ namespace Metasound
 			UMetasoundEditorGraph& EdGraph = GetMetaSoundGraphChecked();
 			Frontend::FConstGraphHandle FrontendGraph = MetasoundAsset->GetRootGraphHandle();
 
-			FrontendGraph->IterateConstNodes([this, &EdGraph, ActionList = &OutAllActions](const Frontend::FConstNodeHandle& Input)
+			auto GetVertexCategory = [](const Frontend::FConstNodeHandle& NodeHandle)
+			{
+				FName InterfaceName;
+				FName MemberName;
+				Audio::IGeneratorInterfaceRegistry::SplitMemberFullName(NodeHandle->GetNodeName(), InterfaceName, MemberName);
+
+				if (InterfaceName.IsNone())
+				{
+					return FText::GetEmpty();
+				}
+
+				FString CategoryString = InterfaceName.ToString();
+				CategoryString.ReplaceInline(*Audio::IGeneratorInterfaceRegistry::NamespaceDelimiter, TEXT("|"));
+				return FText::FromString(CategoryString);
+			};
+
+			FrontendGraph->IterateConstNodes([this, &GetVertexCategory, &EdGraph, ActionList = &OutAllActions](const Frontend::FConstNodeHandle& Input)
 			{
 				const FText Tooltip = Input->GetDescription();
 				const FText MenuDesc = FGraphBuilder::GetDisplayName(*Input);
 				const FGuid NodeID = Input->GetID();
+				const FText Category = GetVertexCategory(Input);
 
-				TSharedPtr<FMetasoundGraphMemberSchemaAction> NewFuncAction = MakeShared<FMetasoundGraphMemberSchemaAction>(FText::GetEmpty(), MenuDesc, Tooltip, 1, ENodeSection::Inputs);
+				TSharedPtr<FMetasoundGraphMemberSchemaAction> NewFuncAction = MakeShared<FMetasoundGraphMemberSchemaAction>(Category, MenuDesc, Tooltip, 1, ENodeSection::Inputs);
 				NewFuncAction->Graph = &EdGraph;
 				NewFuncAction->MemberID = NodeID;
 
@@ -1911,13 +1938,14 @@ namespace Metasound
 				}
 			}, EMetasoundFrontendClassType::Input);
 
-			FrontendGraph->IterateConstNodes([this, &EdGraph, ActionList = &OutAllActions](const Frontend::FConstNodeHandle& Output)
+			FrontendGraph->IterateConstNodes([this, &GetVertexCategory, &EdGraph, ActionList = &OutAllActions](const Frontend::FConstNodeHandle& Output)
 			{
 				const FText Tooltip = Output->GetDescription();
 				const FText MenuDesc = FGraphBuilder::GetDisplayName(*Output);
 				const FGuid NodeID = Output->GetID();
+				const FText Category = GetVertexCategory(Output);
 
-				TSharedPtr<FMetasoundGraphMemberSchemaAction> NewFuncAction = MakeShared<FMetasoundGraphMemberSchemaAction>(FText::GetEmpty(), MenuDesc, Tooltip, 1, ENodeSection::Outputs);
+				TSharedPtr<FMetasoundGraphMemberSchemaAction> NewFuncAction = MakeShared<FMetasoundGraphMemberSchemaAction>(Category, MenuDesc, Tooltip, 1, ENodeSection::Outputs);
 				NewFuncAction->Graph = &EdGraph;
 				NewFuncAction->MemberID = Output->GetID();
 				ActionList->AddAction(NewFuncAction);
@@ -2065,7 +2093,30 @@ namespace Metasound
 
 			if (MetasoundAsset->GetSynchronizationPending())
 			{
+				// Capture before synchronizing as the flag is cleared therein.
+				const bool bInterfacesUpdated = MetasoundAsset->GetSynchronizationInterfacesUpdated();
+
 				FGraphBuilder::SynchronizeGraph(*Metasound);
+
+				if (MetasoundDetails.IsValid() && bInterfacesUpdated)
+				{
+					// Only refresh the details panel if viewing the graph details
+					auto ContainsMetaSound = [](TWeakObjectPtr<UObject> Obj)
+					{
+						if (!Obj.IsValid())
+						{
+							return false;
+						}
+						
+						return IMetasoundUObjectRegistry::Get().IsRegisteredClass(Obj.Get());
+					};
+
+					const bool bDetailsEditingGraph = Algo::AnyOf(MetasoundDetails->GetSelectedObjects(), ContainsMetaSound);
+					if (bDetailsEditingGraph)
+					{
+						RefreshDetails();
+					}
+				}
 			}
 		}
 
diff --git a/Engine/Plugins/Runtime/Metasound/Source/MetasoundEditor/Private/MetasoundEditor.h b/Engine/Plugins/Runtime/Metasound/Source/MetasoundEditor/Private/MetasoundEditor.h
--- a/Engine/Plugins/Runtime/Metasound/Source/MetasoundEditor/Private/MetasoundEditor.h
+++ b/Engine/Plugins/Runtime/Metasound/Source/MetasoundEditor/Private/MetasoundEditor.h
@@ -235,10 +235,12 @@ namespace Metasound
 
 			FSlateIcon GetSettingsImage() const;
 
-			/** Toolbar command methods */
+			// TODO: Move import/export out of editor and into import/export asset actions
 			void Import();
-			void ExecuteNode();
 			void Export();
+
+			/** Toolbar command methods */
+			void ExecuteNode();
 			void Play();
 			void Stop();
 
diff --git a/Engine/Plugins/Runtime/Metasound/Source/MetasoundEditor/Private/MetasoundEditorGraph.cpp b/Engine/Plugins/Runtime/Metasound/Source/MetasoundEditor/Private/MetasoundEditorGraph.cpp
--- a/Engine/Plugins/Runtime/Metasound/Source/MetasoundEditor/Private/MetasoundEditorGraph.cpp
+++ b/Engine/Plugins/Runtime/Metasound/Source/MetasoundEditor/Private/MetasoundEditorGraph.cpp
@@ -294,11 +294,16 @@ Metasound::Frontend::FConstNodeHandle UMetasoundEditorGraphVertex::GetConstNodeH
 	return MetasoundAsset->GetRootGraphHandle()->GetNodeWithID(NodeID);
 }
 
-bool UMetasoundEditorGraphVertex::IsRequired() const
+const FMetasoundFrontendVersion& UMetasoundEditorGraphVertex::GetInterfaceVersion() const
+{
+	return GetConstNodeHandle()->GetInterfaceVersion();
+}
+
+bool UMetasoundEditorGraphVertex::IsInterfaceMember() const
 {
 	using namespace Metasound;
 
-	return GetConstNodeHandle()->IsRequired();
+	return GetConstNodeHandle()->IsInterfaceMember();
 }
 
 bool UMetasoundEditorGraphVertex::CanRename(const FText& InNewName, FText& OutError) const
@@ -311,9 +316,9 @@ bool UMetasoundEditorGraphVertex::CanRename(const FText& InNewName, FText& OutEr
 		return false;
 	}
 
-	if (IsRequired())
+	if (IsInterfaceMember())
 	{
-		OutError = FText::Format(LOCTEXT("GraphVertexRenameInvalid_GraphVertexRequired", "{0} is required and cannot be renamed."), InNewName);
+		OutError = FText::Format(LOCTEXT("GraphVertexRenameInvalid_GraphVertexRequired", "{0} is interface member and cannot be renamed."), InNewName);
 		return false;
 	}
 
@@ -429,11 +434,10 @@ void UMetasoundEditorGraphInput::PostEditUndo()
 			if (UObject* Object = MetasoundGraph->GetMetasound())
 			{
 				Metasound::Editor::FGraphBuilder::RegisterGraphWithFrontend(*Object);
-				if (TSharedPtr<FEditor> MetaSoundEditor = FGraphBuilder::GetEditorForMetasound(*Object))
-				{
-					// Referesh details panel in case current member was selected.
-					MetaSoundEditor->RefreshDetails();
-				}
+
+				FMetasoundAssetBase* MetasoundAsset = Metasound::IMetasoundUObjectRegistry::Get().GetObjectAsAssetBase(Object);
+				check(MetasoundAsset);
+				MetasoundAsset->SetSynchronizationRequired();
 			}
 		}
 		return;
@@ -596,12 +600,6 @@ bool UMetasoundEditorGraphVariable::CanRename(const FText& InNewText, FText& Out
 		OutError = FText::Format(LOCTEXT("GraphVariableRenameInvalid_NameEmpty", "{0} cannot be empty string."), InNewText);
 		return false;
 	}
-
-	if (IsRequired())
-	{
-		OutError = FText::Format(LOCTEXT("GraphVariableRenameInvalid_GraphVariableRequired", "{0} is required and cannot be renamed."), InNewText);
-		return false;
-	}
 	
 	const FName InNewName = FName(*InNewText.ToString());
 	if (!InNewName.IsValid())
@@ -627,11 +625,6 @@ bool UMetasoundEditorGraphVariable::CanRename(const FText& InNewText, FText& Out
 	return true;
 }
 
-bool UMetasoundEditorGraphVariable::IsRequired() const
-{
-	return false;
-}
-
 TArray<UMetasoundEditorGraphNode*> UMetasoundEditorGraphVariable::GetNodes() const
 {
 	TArray<UMetasoundEditorGraphNode*> Nodes;
diff --git a/Engine/Plugins/Runtime/Metasound/Source/MetasoundEditor/Private/MetasoundEditorGraphBuilder.cpp b/Engine/Plugins/Runtime/Metasound/Source/MetasoundEditor/Private/MetasoundEditorGraphBuilder.cpp
--- a/Engine/Plugins/Runtime/Metasound/Source/MetasoundEditor/Private/MetasoundEditorGraphBuilder.cpp
+++ b/Engine/Plugins/Runtime/Metasound/Source/MetasoundEditor/Private/MetasoundEditorGraphBuilder.cpp
@@ -91,7 +91,10 @@ namespace Metasound
 			FText DisplayName = InFrontendNode.GetDisplayName();
 			if (DisplayName.IsEmptyOrWhitespace())
 			{
-				DisplayName = FText::FromName(InFrontendNode.GetNodeName());
+				FName InterfaceName;
+				FName ParameterName;
+				Audio::IGeneratorInterfaceRegistry::Get().SplitMemberFullName(InFrontendNode.GetNodeName(), InterfaceName, ParameterName);
+				return FText::FromName(ParameterName);
 			}
 			return DisplayName;
 		}
@@ -104,7 +107,6 @@ namespace Metasound
 				DisplayName = FText::FromName(InFrontendInput.GetName());
 			}
 			return DisplayName;
-				
 		}
 
 		FText FGraphBuilder::GetDisplayName(const Frontend::IOutputController& InFrontendOutput)
@@ -998,8 +1000,6 @@ namespace Metasound
 			TArray<FInputHandle> InputHandles;
 			TArray<UEdGraphPin*> InputPins;
 
-			UObject& Metasound = CastChecked<UMetasoundEditorGraphNode>(InPin.GetOwningNode())->GetMetasoundChecked();
-
 			if (InPin.Direction == EGPD_Input)
 			{
 				FNodeHandle NodeHandle = CastChecked<UMetasoundEditorGraphNode>(InPin.GetOwningNode())->GetNodeHandle();
@@ -1070,7 +1070,7 @@ namespace Metasound
 			DocMetadata.Version.Number = FVersionDocument::GetMaxVersion();
 			DocumentHandle->SetMetadata(DocMetadata);
 
-			MetaSoundAsset->ConformDocumentToArchetype();
+			MetaSoundAsset->AddDefaultInterfaces();
 
 			FGraphHandle GraphHandle = MetaSoundAsset->GetRootGraphHandle();
 			FVector2D InputNodeLocation = FVector2D::ZeroVector;
@@ -1129,7 +1129,7 @@ namespace Metasound
 			check(ReferencedAsset);
 
 			FRebuildPresetRootGraph(ReferencedAsset->GetDocumentHandle()).Transform(PresetAsset->GetDocumentHandle());
-			PresetAsset->ConformObjectDataToArchetype();
+			PresetAsset->ConformObjectDataToInterfaces();
 		}
 
 		bool FGraphBuilder::DeleteNode(UEdGraphNode& InNode)
diff --git a/Engine/Plugins/Runtime/Metasound/Source/MetasoundEditor/Private/MetasoundEditorGraphNode.cpp b/Engine/Plugins/Runtime/Metasound/Source/MetasoundEditor/Private/MetasoundEditorGraphNode.cpp
--- a/Engine/Plugins/Runtime/Metasound/Source/MetasoundEditor/Private/MetasoundEditorGraphNode.cpp
+++ b/Engine/Plugins/Runtime/Metasound/Source/MetasoundEditor/Private/MetasoundEditorGraphNode.cpp
@@ -491,7 +491,7 @@ FGuid UMetasoundEditorGraphOutputNode::GetNodeID() const
 
 bool UMetasoundEditorGraphOutputNode::CanUserDeleteNode() const
 {
-	return !GetNodeHandle()->IsRequired();
+	return !GetNodeHandle()->IsInterfaceMember();
 }
 
 void UMetasoundEditorGraphOutputNode::SetNodeID(FGuid InNodeID)
diff --git a/Engine/Plugins/Runtime/Metasound/Source/MetasoundEditor/Private/MetasoundNodeDetailCustomization.cpp b/Engine/Plugins/Runtime/Metasound/Source/MetasoundEditor/Private/MetasoundNodeDetailCustomization.cpp
--- a/Engine/Plugins/Runtime/Metasound/Source/MetasoundEditor/Private/MetasoundNodeDetailCustomization.cpp
+++ b/Engine/Plugins/Runtime/Metasound/Source/MetasoundEditor/Private/MetasoundNodeDetailCustomization.cpp
@@ -110,10 +110,9 @@ namespace Metasound
 					if (FloatLiteral.IsValid())
 					{
 						FloatLiteral->OnClampInputChanged.Remove(OnClampInputChangedDelegateHandle);
-						TSharedPtr<FEditor> ParentEditor = FGraphBuilder::GetEditorForMetasound(*FloatLiteral->GetOutermostObject());
-						if (ParentEditor.IsValid())
+						if (FMetasoundAssetBase* MetasoundAsset = Metasound::IMetasoundUObjectRegistry::Get().GetObjectAsAssetBase(FloatLiteral->GetOutermostObject()))
 						{
-							ParentEditor->RefreshDetails();
+							MetasoundAsset->SetSynchronizationRequired();
 						}
 					}
 				});
@@ -145,10 +144,9 @@ namespace Metasound
 							if (FloatLiteral.IsValid())
 							{
 								FloatLiteral->OnRangeChanged.Remove(InputWidgetOnRangeChangedDelegateHandle);
-								TSharedPtr<FEditor> ParentEditor = FGraphBuilder::GetEditorForMetasound(*FloatLiteral->GetOutermostObject());
-								if (ParentEditor.IsValid())
+								if (FMetasoundAssetBase* MetasoundAsset = Metasound::IMetasoundUObjectRegistry::Get().GetObjectAsAssetBase(FloatLiteral->GetOutermostObject()))
 								{
-									ParentEditor->RefreshDetails();
+									MetasoundAsset->SetSynchronizationRequired();
 								}
 							}
 						});
@@ -745,23 +743,23 @@ namespace Metasound
 
 			IDetailCategoryBuilder& CategoryBuilder = DetailLayout.EditCategory("General");
 
-			const bool bIsRequired = IsRequired();
+			const bool bIsInterfaceMember = CastChecked<UMetasoundEditorGraphVertex>(GraphMember)->IsInterfaceMember();
 			const bool bIsGraphEditable = IsGraphEditable();
 			NameEditableTextBox = SNew(SEditableTextBox)
 				.Text(this, &FMetasoundInputDetailCustomization::GetName)
 				.OnTextChanged(this, &FMetasoundInputDetailCustomization::OnNameChanged)
 				.OnTextCommitted(this, &FMetasoundInputDetailCustomization::OnNameCommitted)
-				.IsReadOnly(bIsRequired || !bIsGraphEditable)
+				.IsReadOnly(bIsInterfaceMember || !bIsGraphEditable)
 				.Font(IDetailLayoutBuilder::GetDetailFont());
 
 			DisplayNameEditableTextBox = SNew(SEditableTextBox)
 				.Text(this, &FMetasoundInputDetailCustomization::GetDisplayName)
 				.OnTextCommitted(this, &FMetasoundInputDetailCustomization::OnDisplayNameCommitted)
-				.IsReadOnly(bIsRequired || !bIsGraphEditable)
+				.IsReadOnly(bIsInterfaceMember || !bIsGraphEditable)
 				.Font(IDetailLayoutBuilder::GetDetailFont());
 
 			CategoryBuilder.AddCustomRow(MemberCustomizationPrivate::InputNameText)
-			.EditCondition(!bIsRequired && bIsGraphEditable, nullptr)
+			.EditCondition(!bIsInterfaceMember && bIsGraphEditable, nullptr)
 			.NameContent()
 			[
 				SNew(STextBlock)
@@ -776,7 +774,7 @@ namespace Metasound
 
 			// TODO: Enable and use proper FText property editor
 // 			CategoryBuilder.AddCustomRow(MemberCustomizationPrivate::InputDisplayNameText)
-// 			.EditCondition(!bIsRequired && bIsGraphEditable, nullptr)
+// 			.EditCondition(!bIsInterfaceMember && bIsGraphEditable, nullptr)
 // 			.NameContent()
 // 			[
 // 				SNew(STextBlock)
@@ -790,7 +788,7 @@ namespace Metasound
 // 			];
 
 			CategoryBuilder.AddCustomRow(MemberCustomizationPrivate::NodeTooltipText)
-			.EditCondition(!bIsRequired && bIsGraphEditable, nullptr)
+			.EditCondition(!bIsInterfaceMember && bIsGraphEditable, nullptr)
 			.NameContent()
 			[
 				SNew(STextBlock)
@@ -802,14 +800,14 @@ namespace Metasound
 				SNew(SMultiLineEditableTextBox)
 				.Text(this, &FMetasoundInputDetailCustomization::GetTooltip)
 				.OnTextCommitted(this, &FMetasoundInputDetailCustomization::OnTooltipCommitted)
-				.IsReadOnly(bIsRequired || !bIsGraphEditable)
+				.IsReadOnly(bIsInterfaceMember || !bIsGraphEditable)
 				.ModiferKeyForNewLine(EModifierKey::Shift)
 				.RevertTextOnEscape(true)
 				.WrapTextAt(MemberCustomizationPrivate::DetailsTitleMaxWidth - MemberCustomizationPrivate::DetailsTitleWrapPadding)
 				.Font(IDetailLayoutBuilder::GetDetailFont())
 			];
 
-			DataTypeSelector->AddDataTypeSelector(DetailLayout, MemberCustomizationPrivate::DataTypeNameText, GraphMember, !bIsRequired && bIsGraphEditable);
+			DataTypeSelector->AddDataTypeSelector(DetailLayout, MemberCustomizationPrivate::DataTypeNameText, GraphMember, !bIsInterfaceMember && bIsGraphEditable);
 
 			IDetailCategoryBuilder& DefaultCategoryBuilder = DetailLayout.EditCategory("DefaultValue");
 			TSharedPtr<IPropertyHandle> LiteralHandle = DetailLayout.GetProperty(GET_MEMBER_NAME_CHECKED(UMetasoundEditorGraphInput, Literal));
@@ -965,7 +963,6 @@ namespace Metasound
 			using namespace Frontend;
 
 			TMetasoundGraphMemberDetailCustomization<UMetasoundEditorGraphOutput>::CustomizeDetails(DetailLayout);
-
 			if (!GraphMember.IsValid())
 			{
 				return;
@@ -973,24 +970,24 @@ namespace Metasound
 
 			IDetailCategoryBuilder& CategoryBuilder = DetailLayout.EditCategory("General");
 
-			const bool bIsRequired = IsRequired();
+			const bool bIsInterfaceMember = CastChecked<UMetasoundEditorGraphVertex>(GraphMember)->IsInterfaceMember();
 			const bool bIsGraphEditable = IsGraphEditable();
 
 			NameEditableTextBox = SNew(SEditableTextBox)
 				.Text(this, &FMetasoundOutputDetailCustomization::GetName)
 				.OnTextChanged(this, &FMetasoundOutputDetailCustomization::OnNameChanged)
 				.OnTextCommitted(this, &FMetasoundOutputDetailCustomization::OnNameCommitted)
-				.IsReadOnly(bIsRequired || !bIsGraphEditable)
+				.IsReadOnly(bIsInterfaceMember || !bIsGraphEditable)
 				.Font(IDetailLayoutBuilder::GetDetailFont());
 
 			DisplayNameEditableTextBox = SNew(SEditableTextBox)
 				.Text(this, &FMetasoundOutputDetailCustomization::GetDisplayName)
 				.OnTextCommitted(this, &FMetasoundOutputDetailCustomization::OnDisplayNameCommitted)
-				.IsReadOnly(bIsRequired || !bIsGraphEditable)
+				.IsReadOnly(bIsInterfaceMember || !bIsGraphEditable)
 				.Font(IDetailLayoutBuilder::GetDetailFont());
 
 			CategoryBuilder.AddCustomRow(MemberCustomizationPrivate::OutputNameText)
-			.EditCondition(!bIsRequired && bIsGraphEditable, nullptr)
+			.EditCondition(!bIsInterfaceMember && bIsGraphEditable, nullptr)
 			.NameContent()
 			[
 				SNew(STextBlock)
@@ -1005,7 +1002,7 @@ namespace Metasound
 
 			// TODO: Enable and use proper FText property editor
 // 			CategoryBuilder.AddCustomRow(MemberCustomizationPrivate::OutputDisplayNameText)
-// 			.EditCondition(!bIsRequired && bIsGraphEditable, nullptr)
+// 			.EditCondition(!bIsInterfaceMember && bIsGraphEditable, nullptr)
 // 			.NameContent()
 // 			[
 // 				SNew(STextBlock)
@@ -1019,7 +1016,7 @@ namespace Metasound
 // 			];
 
 			CategoryBuilder.AddCustomRow(MemberCustomizationPrivate::NodeTooltipText)
-			.EditCondition(!bIsRequired && bIsGraphEditable, nullptr)
+			.EditCondition(!bIsInterfaceMember && bIsGraphEditable, nullptr)
 			.NameContent()
 			[
 				SNew(STextBlock)
@@ -1031,14 +1028,14 @@ namespace Metasound
 				SNew(SMultiLineEditableTextBox)
 				.Text(this, &FMetasoundOutputDetailCustomization::GetTooltip)
 				.OnTextCommitted(this, &FMetasoundOutputDetailCustomization::OnTooltipCommitted)
-				.IsReadOnly(bIsRequired || !bIsGraphEditable)
+				.IsReadOnly(bIsInterfaceMember || !bIsGraphEditable)
 				.ModiferKeyForNewLine(EModifierKey::Shift)
 				.RevertTextOnEscape(true)
 				.WrapTextAt(MemberCustomizationPrivate::DetailsTitleMaxWidth - MemberCustomizationPrivate::DetailsTitleWrapPadding)
 				.Font(IDetailLayoutBuilder::GetDetailFont())
 			];
 
-			DataTypeSelector->AddDataTypeSelector(DetailLayout, MemberCustomizationPrivate::DataTypeNameText, GraphMember, !bIsRequired && bIsGraphEditable);
+			DataTypeSelector->AddDataTypeSelector(DetailLayout, MemberCustomizationPrivate::DataTypeNameText, GraphMember, !bIsInterfaceMember && bIsGraphEditable);
 		}
 
 		void FMetasoundOutputDetailCustomization::SetDefaultPropertyMetaData(TSharedRef<IPropertyHandle> InDefaultPropertyHandle) const
diff --git a/Engine/Plugins/Runtime/Metasound/Source/MetasoundEditor/Private/MetasoundNodeDetailCustomization.h b/Engine/Plugins/Runtime/Metasound/Source/MetasoundEditor/Private/MetasoundNodeDetailCustomization.h
--- a/Engine/Plugins/Runtime/Metasound/Source/MetasoundEditor/Private/MetasoundNodeDetailCustomization.h
+++ b/Engine/Plugins/Runtime/Metasound/Source/MetasoundEditor/Private/MetasoundNodeDetailCustomization.h
@@ -138,7 +138,7 @@ namespace Metasound
 		class FMetasoundDataTypeSelector : public TSharedFromThis<FMetasoundDataTypeSelector>
 		{
 		public:
-			void AddDataTypeSelector(IDetailLayoutBuilder& InDetailLayoutBuilder, const FText& InRowName, TWeakObjectPtr<UMetasoundEditorGraphMember> InGraphMember, bool bIsRequired);
+			void AddDataTypeSelector(IDetailLayoutBuilder& InDetailLayoutBuilder, const FText& InRowName, TWeakObjectPtr<UMetasoundEditorGraphMember> InGraphMember, bool bInIsInterfaceMember);
 			void OnDataTypeArrayChanged(TWeakObjectPtr<UMetasoundEditorGraphMember> InGraphMember, ECheckBoxState InNewState);
 
 		protected:
@@ -247,16 +247,6 @@ namespace Metasound
 				return false;
 			}
 
-			bool IsRequired() const
-			{
-				if (GraphMember.IsValid())
-				{
-					return GraphMember->IsRequired();
-				}
-
-				return true;
-			}
-
 			void OnTooltipCommitted(const FText& InNewText, ETextCommit::Type InTextCommit)
 			{
 				using namespace Frontend;
diff --git a/Engine/Plugins/Runtime/Metasound/Source/MetasoundEditor/Private/MetasoundVariableDetailCustomization.cpp b/Engine/Plugins/Runtime/Metasound/Source/MetasoundEditor/Private/MetasoundVariableDetailCustomization.cpp
--- a/Engine/Plugins/Runtime/Metasound/Source/MetasoundEditor/Private/MetasoundVariableDetailCustomization.cpp
+++ b/Engine/Plugins/Runtime/Metasound/Source/MetasoundEditor/Private/MetasoundVariableDetailCustomization.cpp
@@ -41,23 +41,22 @@ namespace Metasound
 
 			IDetailCategoryBuilder& CategoryBuilder = DetailLayout.EditCategory("General");
 
-			const bool bIsRequired = IsRequired();
 			const bool bIsGraphEditable = IsGraphEditable();
 			NameEditableTextBox = SNew(SEditableTextBox)
 				.Text(this, &FMetasoundVariableDetailCustomization::GetName)
 				.OnTextChanged(this, &FMetasoundVariableDetailCustomization::OnNameChanged)
 				.OnTextCommitted(this, &FMetasoundVariableDetailCustomization::OnNameCommitted)
-				.IsReadOnly(bIsRequired || !bIsGraphEditable)
+				.IsReadOnly(!bIsGraphEditable)
 				.Font(IDetailLayoutBuilder::GetDetailFont());
 
 			DisplayNameEditableTextBox = SNew(SEditableTextBox)
 				.Text(this, &FMetasoundVariableDetailCustomization::GetDisplayName)
 				.OnTextCommitted(this, &FMetasoundVariableDetailCustomization::OnDisplayNameCommitted)
-				.IsReadOnly(bIsRequired || !bIsGraphEditable)
+				.IsReadOnly(!bIsGraphEditable)
 				.Font(IDetailLayoutBuilder::GetDetailFont());
 
 			CategoryBuilder.AddCustomRow(MemberCustomizationPrivate::VariableNameText)
-			.EditCondition(!bIsRequired && bIsGraphEditable, nullptr)
+			.EditCondition(bIsGraphEditable, nullptr)
 			.NameContent()
 			[
 				SNew(STextBlock)
@@ -72,7 +71,7 @@ namespace Metasound
 
 			// TODO: Enable and use proper FText property editor
 // 			CategoryBuilder.AddCustomRow(MemberCustomizationPrivate::VariableDisplayNameText)
-// 			.EditCondition(!bIsRequired && bIsGraphEditable, nullptr)
+// 			.EditCondition(!bIsInterfaceMember && bIsGraphEditable, nullptr)
 // 			.NameContent()
 // 			[
 // 				SNew(STextBlock)
@@ -86,7 +85,7 @@ namespace Metasound
 // 			];
 
 			CategoryBuilder.AddCustomRow(MemberCustomizationPrivate::NodeTooltipText)
-			.EditCondition(!bIsRequired && bIsGraphEditable, nullptr)
+			.EditCondition(bIsGraphEditable, nullptr)
 			.NameContent()
 			[
 				SNew(STextBlock)
@@ -98,14 +97,14 @@ namespace Metasound
 				SNew(SMultiLineEditableTextBox)
 				.Text(this, &FMetasoundVariableDetailCustomization::GetTooltip)
 				.OnTextCommitted(this, &FMetasoundVariableDetailCustomization::OnTooltipCommitted)
-				.IsReadOnly(bIsRequired || !bIsGraphEditable)
+				.IsReadOnly(!bIsGraphEditable)
 				.ModiferKeyForNewLine(EModifierKey::Shift)
 				.RevertTextOnEscape(true)
 				.WrapTextAt(MemberCustomizationPrivate::DetailsTitleMaxWidth - MemberCustomizationPrivate::DetailsTitleWrapPadding)
 				.Font(IDetailLayoutBuilder::GetDetailFont())
 			];
 
-			DataTypeSelector->AddDataTypeSelector(DetailLayout, MemberCustomizationPrivate::DataTypeNameText, GraphMember, !bIsRequired && bIsGraphEditable);
+			DataTypeSelector->AddDataTypeSelector(DetailLayout, MemberCustomizationPrivate::DataTypeNameText, GraphMember, bIsGraphEditable);
 
 			IDetailCategoryBuilder& DefaultCategoryBuilder = DetailLayout.EditCategory("DefaultValue");
 			TSharedPtr<IPropertyHandle> LiteralHandle = DetailLayout.GetProperty(GET_MEMBER_NAME_CHECKED(UMetasoundEditorGraphVariable, Literal));
diff --git a/Engine/Plugins/Runtime/Metasound/Source/MetasoundEditor/Public/MetasoundEditorGraph.h b/Engine/Plugins/Runtime/Metasound/Source/MetasoundEditor/Public/MetasoundEditorGraph.h
--- a/Engine/Plugins/Runtime/Metasound/Source/MetasoundEditor/Public/MetasoundEditorGraph.h
+++ b/Engine/Plugins/Runtime/Metasound/Source/MetasoundEditor/Public/MetasoundEditorGraph.h
@@ -91,9 +91,6 @@ public:
 	/** Return the section of where this member belongs. */
 	virtual Metasound::Editor::ENodeSection GetSectionID() const PURE_VIRTUAL(UMetasoundEditorGraphMember::GetSectionID, return Metasound::Editor::ENodeSection::None; );
 
-	/** If true, this member cannot be removed by the user. */
-	virtual bool IsRequired() const PURE_VIRTUAL(UMetasoundEditorGraphMember::IsRequired, return false; );
-
 	/** Return the nodes associated with this member */
 	virtual TArray<UMetasoundEditorGraphNode*> GetNodes() const PURE_VIRTUAL(UMetasoundEditorGraphMember::GetNodes, return TArray<UMetasoundEditorGraphNode*>(); );
 
@@ -136,6 +133,7 @@ public:
 	 * a UMetasoundEditorGraph, returns a nullptr. */
 	const UMetasoundEditorGraph* GetOwningGraph() const;
 
+	/** Marks all member's nodes as requiring a display refresh. */
 	void MarkNodesForRefresh();
 };
 
@@ -147,12 +145,7 @@ class METASOUNDEDITOR_API UMetasoundEditorGraphVertex : public UMetasoundEditorG
 
 protected:
 	/** Adds the node handle for a newly created vertex. */
-	virtual Metasound::Frontend::FNodeHandle AddNodeHandle(const FName& InNodeName, FName InDataType)
-	{
-		unimplemented();
-		return Metasound::Frontend::INodeController::GetInvalidHandle();
-	}
-
+	virtual Metasound::Frontend::FNodeHandle AddNodeHandle(const FName& InNodeName, FName InDataType) PURE_VIRTUAL(UMetasoundEditorGraphVertex::AddNodeHandle, return Metasound::Frontend::INodeController::GetInvalidHandle(); )
 
 public:
 	/** ID of Metasound Frontend node. */
@@ -165,12 +158,11 @@ public:
 
 	/* ~Begin UMetasoundEditorGraphMember interface */
 	virtual FGuid GetMemberID() const override;
-	virtual FName GetMemberName() const override; 
+	virtual FName GetMemberName() const override;
 	virtual FText GetDescription() const override;
 	virtual FText GetDisplayName() const override;
 
 	virtual bool CanRename(const FText& InNewName, FText& OutError) const override;
-	virtual bool IsRequired() const override;
 	virtual TArray<UMetasoundEditorGraphNode*> GetNodes() const override;
 	virtual void SetDescription(const FText& InDescription, bool bPostTransaction) override;
 	virtual void SetMemberName(const FName& InNewName, bool bPostTransaction) override;
@@ -178,6 +170,12 @@ public:
 	virtual void SetDataType(FName InNewType, bool bPostTransaction = true, bool bRegisterParentGraph = true) override;
 	/* ~End UMetasoundEditorGraphMember interface */
 
+	/** Version of interface membership, or invalid version if not an interface member. */
+	virtual const FMetasoundFrontendVersion& GetInterfaceVersion() const;
+
+	/** If true, this member cannot be modified by the user. */
+	virtual bool IsInterfaceMember() const;
+
 	/** Called when the data type changes. */
 	virtual void OnDataTypeChanged();
 
@@ -260,7 +258,6 @@ public:
 	virtual FText GetDisplayName() const override;
 	virtual void SetDisplayName(const FText& InNewName, bool bPostTransaction) override;
 
-	virtual bool IsRequired() const override;
 	virtual TArray<UMetasoundEditorGraphNode*> GetNodes() const override;
 
 	virtual const FText& GetGraphMemberLabel() const override;
diff --git a/Engine/Plugins/Runtime/Metasound/Source/MetasoundEngine/Private/Metasound.cpp b/Engine/Plugins/Runtime/Metasound/Source/MetasoundEngine/Private/Metasound.cpp
--- a/Engine/Plugins/Runtime/Metasound/Source/MetasoundEngine/Private/Metasound.cpp
+++ b/Engine/Plugins/Runtime/Metasound/Source/MetasoundEngine/Private/Metasound.cpp
@@ -41,7 +41,13 @@ UMetaSound::UMetaSound(const FObjectInitializer& ObjectInitializer)
 void UMetaSound::PostDuplicate(EDuplicateMode::Type InDuplicateMode)
 {
 	Super::PostDuplicate(InDuplicateMode);
-	Metasound::PostDuplicate(*this, InDuplicateMode);
+
+	// Guid is reset as asset may share implementation from
+	// asset duplicated from but should not be registered as such.
+	if (InDuplicateMode == EDuplicateMode::Normal)
+	{
+		Metasound::Frontend::FRegenerateAssetClassName().Transform(GetDocumentHandle());
+	}
 }
 
 void UMetaSound::PostEditUndo()
@@ -135,14 +141,4 @@ const FMetasoundFrontendVersion& UMetaSound::GetDefaultArchetypeVersion() const
 	static const FMetasoundFrontendVersion DefaultArchetypeVersion = Metasound::Engine::MetasoundV1_0::GetVersion();
 	return DefaultArchetypeVersion;
 }
-
-const TArray<FMetasoundFrontendVersion>& UMetaSound::GetSupportedArchetypeVersions() const
-{
-	static const TArray<FMetasoundFrontendVersion> Supported
-	{
-		Metasound::Engine::MetasoundV1_0::GetVersion()
-	};
-
-	return Supported;
-}
 #undef LOCTEXT_NAMESPACE // MetaSound
diff --git a/Engine/Plugins/Runtime/Metasound/Source/MetasoundEngine/Private/MetasoundAssetSubsystem.cpp b/Engine/Plugins/Runtime/Metasound/Source/MetasoundEngine/Private/MetasoundAssetSubsystem.cpp
--- a/Engine/Plugins/Runtime/Metasound/Source/MetasoundEngine/Private/MetasoundAssetSubsystem.cpp
+++ b/Engine/Plugins/Runtime/Metasound/Source/MetasoundEngine/Private/MetasoundAssetSubsystem.cpp
@@ -1,17 +1,20 @@
 // Copyright Epic Games, Inc. All Rights Reserved.
 #include "MetasoundAssetSubsystem.h"
 
+#include "Algo/Sort.h"
 #include "Algo/Transform.h"
 #include "AssetRegistry/AssetRegistryModule.h"
 #include "Async/Async.h"
 #include "Engine/AssetManager.h"
 #include "Metasound.h"
 #include "MetasoundAssetBase.h"
+#include "MetasoundFrontendDocument.h"
 #include "MetasoundFrontendRegistries.h"
 #include "MetasoundFrontendTransform.h"
 #include "MetasoundSettings.h"
 #include "MetasoundSource.h"
 #include "MetasoundTrace.h"
+#include "MetasoundUObjectRegistry.h"
 #include "UObject/NoExportTypes.h"
 
 
diff --git a/Engine/Plugins/Runtime/Metasound/Source/MetasoundEngine/Private/MetasoundEngineArchetypes.cpp b/Engine/Plugins/Runtime/Metasound/Source/MetasoundEngine/Private/MetasoundEngineArchetypes.cpp
--- a/Engine/Plugins/Runtime/Metasound/Source/MetasoundEngine/Private/MetasoundEngineArchetypes.cpp
+++ b/Engine/Plugins/Runtime/Metasound/Source/MetasoundEngine/Private/MetasoundEngineArchetypes.cpp
@@ -3,9 +3,15 @@
 #include "MetasoundEngineArchetypes.h"
 
 #include "MetasoundAudioFormats.h"
+#include "MetasoundDataReference.h"
 #include "MetasoundFrontendArchetypeRegistry.h"
 #include "MetasoundFrontendTransform.h"
+#include "MetasoundInterface.h"
+#include "MetasoundRouter.h"
 #include "MetasoundSource.h"
+#include "MetasoundUObjectRegistry.h"
+#include "UObject/Class.h"
+#include "UObject/NoExportTypes.h"
 
 #define LOCTEXT_NAMESPACE "MetasoundEngineArchetypes"
 
@@ -13,41 +19,6 @@ namespace Metasound
 {
 	namespace Engine
 	{
-		namespace MetasoundEngineArchetypesPrivate
-		{
-			// Entry for registered archetype.
-			class FArchetypeRegistryEntry : public Frontend::IArchetypeRegistryEntry
-			{
-			public:
-
-				// @param InArchetype - Archetype to register.
-				// @parma InTransform - Transform to convert from prior version of archetype to this version of archetype.
-				FArchetypeRegistryEntry(const FMetasoundFrontendArchetype& InArchetype, TUniquePtr<Frontend::IDocumentTransform>&& InTransform)
-				: Archetype(InArchetype)
-				, Transform(MoveTemp(InTransform))
-				{
-				}
-
-				virtual const FMetasoundFrontendArchetype& GetArchetype() const override
-				{
-					return Archetype;
-				}
-
-				virtual bool UpdateRootGraphArchetype(Frontend::FDocumentHandle InDocument) const override
-				{
-					if (Transform.IsValid())
-					{
-						return Transform->Transform(InDocument);
-					}
-					return false;
-				}
-
-			private:
-				FMetasoundFrontendArchetype Archetype;
-				TUniquePtr<Frontend::IDocumentTransform> Transform;
-			};
-		}
-
 		// MetasoundV1_0 is a metasound without any required inputs or outputs.
 		namespace MetasoundV1_0
 		{
@@ -57,12 +28,12 @@ namespace Metasound
 				return FMetasoundFrontendVersion{ VersionName, FMetasoundFrontendVersionNumber{1, 0} };
 			}
 
-			FMetasoundFrontendArchetype GetArchetype()
+			FMetasoundFrontendInterface GetInterface()
 			{
-				FMetasoundFrontendArchetype Archetype;
-				Archetype.Version = GetVersion();
+				FMetasoundFrontendInterface Interface;
+				Interface.Version = GetVersion();
 
-				return Archetype;
+				return Interface;
 			}
 		}
 
@@ -111,11 +82,11 @@ namespace Metasound
 				return SoundIsPreviewSoundName;
 			}
 
-			FMetasoundFrontendArchetype GetArchetype()
+			FMetasoundFrontendInterface GetInterface()
 			{
-				FMetasoundFrontendArchetype Archetype;
-				Archetype.Version.Name = "MonoSource";
-				Archetype.Version.Number = FMetasoundFrontendVersionNumber{1, 0};
+				FMetasoundFrontendInterface Interface;
+				Interface.Version.Name = "MonoSource";
+				Interface.Version.Number = FMetasoundFrontendVersionNumber{1, 0};
 				
 				// Inputs
 				FMetasoundFrontendClassVertex OnPlayTrigger;
@@ -126,7 +97,7 @@ namespace Metasound
 				OnPlayTrigger.Metadata.Description = LOCTEXT("OnPlayTriggerToolTip", "Trigger executed when this source is played.");
 				OnPlayTrigger.VertexID = FGuid::NewGuid();
 
-				Archetype.Interface.Inputs.Add(OnPlayTrigger);
+				Interface.Inputs.Add(OnPlayTrigger);
 
 				// Outputs 
 				FMetasoundFrontendClassVertex OnFinished;
@@ -137,7 +108,7 @@ namespace Metasound
 				OnFinished.Metadata.Description = LOCTEXT("OnFinishedToolTip", "Trigger executed to initiate stopping the source.");
 				OnFinished.VertexID = FGuid::NewGuid();
 
-				Archetype.Interface.Outputs.Add(OnFinished);
+				Interface.Outputs.Add(OnFinished);
 
 				FMetasoundFrontendClassVertex GeneratedAudio;
 				GeneratedAudio.Name = GetAudioOutputName();
@@ -146,16 +117,16 @@ namespace Metasound
 				GeneratedAudio.Metadata.Description = LOCTEXT("GeneratedAudioToolTip", "The resulting output audio from this source.");
 				GeneratedAudio.VertexID = FGuid::NewGuid();
 
-				Archetype.Interface.Outputs.Add(GeneratedAudio);
+				Interface.Outputs.Add(GeneratedAudio);
 
 				// Environment 
-				FMetasoundFrontendEnvironmentVariable AudioDeviceID;
+				FMetasoundFrontendClassEnvironmentVariable AudioDeviceID;
 				AudioDeviceID.Name = GetAudioDeviceIDVariableName();
 				AudioDeviceID.Metadata.Description = LOCTEXT("AudioDeviceIDToolTip", "Audio Device ID");
 
-				Archetype.Interface.Environment.Add(AudioDeviceID);
+				Interface.Environment.Add(AudioDeviceID);
 
-				return Archetype;
+				return Interface;
 			}
 		}
 
@@ -209,10 +180,10 @@ namespace Metasound
 				return SoundGraphName;
 			}
 
-			FMetasoundFrontendArchetype GetArchetype()
+			FMetasoundFrontendInterface GetInterface()
 			{
-				FMetasoundFrontendArchetype Archetype;
-				Archetype.Version = GetVersion();
+				FMetasoundFrontendInterface Interface;
+				Interface.Version = GetVersion();
 				
 				// Inputs
 				FMetasoundFrontendClassVertex OnPlayTrigger;
@@ -223,7 +194,7 @@ namespace Metasound
 				OnPlayTrigger.Metadata.Description = LOCTEXT("OnPlayTriggerToolTip", "Trigger executed when this source is played.");
 				OnPlayTrigger.VertexID = FGuid::NewGuid();
 
-				Archetype.Interface.Inputs.Add(OnPlayTrigger);
+				Interface.Inputs.Add(OnPlayTrigger);
 
 				// Outputs
 				FMetasoundFrontendClassVertex OnFinished;
@@ -234,7 +205,7 @@ namespace Metasound
 				OnFinished.Metadata.Description = LOCTEXT("OnFinishedToolTip", "Trigger executed to initiate stopping the source.");
 				OnFinished.VertexID = FGuid::NewGuid();
 
-				Archetype.Interface.Outputs.Add(OnFinished);
+				Interface.Outputs.Add(OnFinished);
 
 				FMetasoundFrontendClassVertex GeneratedAudio;
 				GeneratedAudio.Name = GetAudioOutputName();
@@ -243,16 +214,16 @@ namespace Metasound
 				GeneratedAudio.Metadata.Description = LOCTEXT("GeneratedAudioToolTip", "The resulting output audio from this source.");
 				GeneratedAudio.VertexID = FGuid::NewGuid();
 
-				Archetype.Interface.Outputs.Add(GeneratedAudio);
+				Interface.Outputs.Add(GeneratedAudio);
 
 				// Environment
-				FMetasoundFrontendEnvironmentVariable AudioDeviceID;
+				FMetasoundFrontendClassEnvironmentVariable AudioDeviceID;
 				AudioDeviceID.Name = GetAudioDeviceIDVariableName();
 				AudioDeviceID.Metadata.Description = LOCTEXT("AudioDeviceIDToolTip", "Audio Device ID");
 
-				Archetype.Interface.Environment.Add(AudioDeviceID);
+				Interface.Environment.Add(AudioDeviceID);
 
-				return Archetype;
+				return Interface;
 			}
 		}
 
@@ -324,10 +295,10 @@ namespace Metasound
 				return GeneratedAudio;
 			}
 
-			FMetasoundFrontendArchetype GetArchetype()
+			FMetasoundFrontendInterface GetInterface()
 			{
-				FMetasoundFrontendArchetype Archetype;
-				Archetype.Version = GetVersion();
+				FMetasoundFrontendInterface Interface;
+				Interface.Version = GetVersion();
 				
 				// Inputs
 				FMetasoundFrontendClassVertex OnPlayTrigger;
@@ -338,7 +309,7 @@ namespace Metasound
 				OnPlayTrigger.Metadata.Description = LOCTEXT("OnPlayTriggerToolTip", "Trigger executed when this source is played.");
 				OnPlayTrigger.VertexID = FGuid::NewGuid();
 
-				Archetype.Interface.Inputs.Add(OnPlayTrigger);
+				Interface.Inputs.Add(OnPlayTrigger);
 
 				// Outputs
 				FMetasoundFrontendClassVertex OnFinished;
@@ -349,29 +320,28 @@ namespace Metasound
 				OnFinished.Metadata.Description = LOCTEXT("OnFinishedToolTip", "Trigger executed to initiate stopping the source.");
 				OnFinished.VertexID = FGuid::NewGuid();
 
-				Archetype.Interface.Outputs.Add(OnFinished);
+				Interface.Outputs.Add(OnFinished);
 
 				FMetasoundFrontendClassVertex GeneratedAudio = GetClassAudioOutput();
-				Archetype.Interface.Outputs.Add(GeneratedAudio);
+				Interface.Outputs.Add(GeneratedAudio);
 
 				// Environment
-				FMetasoundFrontendEnvironmentVariable AudioDeviceID;
+				FMetasoundFrontendClassEnvironmentVariable AudioDeviceID;
 				AudioDeviceID.Name = GetAudioDeviceIDVariableName();
 				AudioDeviceID.Metadata.Description = LOCTEXT("AudioDeviceIDToolTip", "Audio Device ID");
 
-				Archetype.Interface.Environment.Add(AudioDeviceID);
+				Interface.Environment.Add(AudioDeviceID);
 
-				return Archetype;
+				return Interface;
 			}
 
-			// Update from a MetasoundSourceMonoV1_0 to MtasoundSourceMonoV1_1
-			class FUpdateArchetype : public Frontend::IDocumentTransform
+			// Update from a MetasoundSourceMonoV1_0 to MetasoundSourceMonoV1_1
+			class FUpdateInterface : public Frontend::IDocumentTransform
 			{
 			public:
 				virtual bool Transform(Frontend::FDocumentHandle InDocument) const override
 				{
 					// Swap FMonoAudioFormat output node to an FAudioBuffer output node.
-					
 					using namespace Frontend;
 
 					FGraphHandle Graph = InDocument->GetRootGraph();
@@ -380,8 +350,10 @@ namespace Metasound
 						return false;
 					}
 
-					FNodeHandle MonoFormatOutput = Graph->GetOutputNodeWithName(MetasoundSourceMonoV1_0::GetAudioOutputName());
+					InDocument->RemoveInterfaceVersion(MetasoundSourceMonoV1_0::GetInterface().Version);
+					InDocument->AddInterfaceVersion(MetasoundSourceMonoV1_1::GetInterface().Version);
 
+					FNodeHandle MonoFormatOutput = Graph->GetOutputNodeWithName(MetasoundSourceMonoV1_0::GetAudioOutputName());
 					FVector2D MonoFormatLocation;
 
 					FOutputHandle OutputToReconnect = IOutputController::GetInvalidHandle();
@@ -513,10 +485,10 @@ namespace Metasound
 				return GeneratedRightAudio;
 			}
 
-			FMetasoundFrontendArchetype GetArchetype()
+			FMetasoundFrontendInterface GetInterface()
 			{
-				FMetasoundFrontendArchetype Archetype;
-				Archetype.Version = GetVersion();
+				FMetasoundFrontendInterface Interface;
+				Interface.Version = GetVersion();
 				
 				// Inputs
 				FMetasoundFrontendClassVertex OnPlayTrigger;
@@ -527,7 +499,7 @@ namespace Metasound
 				OnPlayTrigger.Metadata.Description = LOCTEXT("OnPlayTriggerToolTip", "Trigger executed when this source is played.");
 				OnPlayTrigger.VertexID = FGuid::NewGuid();
 
-				Archetype.Interface.Inputs.Add(OnPlayTrigger);
+				Interface.Inputs.Add(OnPlayTrigger);
 
 				// Outputs
 				FMetasoundFrontendClassVertex OnFinished;
@@ -538,25 +510,25 @@ namespace Metasound
 				OnFinished.Metadata.Description = LOCTEXT("OnFinishedToolTip", "Trigger executed to initiate stopping the source.");
 				OnFinished.VertexID = FGuid::NewGuid();
 
-				Archetype.Interface.Outputs.Add(OnFinished);
+				Interface.Outputs.Add(OnFinished);
 
 				FMetasoundFrontendClassVertex GeneratedLeftAudio = GetClassLeftAudioOutput();
-				Archetype.Interface.Outputs.Add(GeneratedLeftAudio);
+				Interface.Outputs.Add(GeneratedLeftAudio);
 
 				FMetasoundFrontendClassVertex GeneratedRightAudio = GetClassRightAudioOutput();
-				Archetype.Interface.Outputs.Add(GeneratedRightAudio);
+				Interface.Outputs.Add(GeneratedRightAudio);
 
 				// Environment
-				FMetasoundFrontendEnvironmentVariable AudioDeviceID;
+				FMetasoundFrontendClassEnvironmentVariable AudioDeviceID;
 				AudioDeviceID.Name = GetAudioDeviceIDVariableName();
 				AudioDeviceID.Metadata.Description = LOCTEXT("AudioDeviceIDToolTip", "Audio Device ID");
-				Archetype.Interface.Environment.Add(AudioDeviceID);
+				Interface.Environment.Add(AudioDeviceID);
 
-				return Archetype;
+				return Interface;
 			};
 
 			// Update from MetasoundSourceStereoV1_0 to MetasoundSourceStereoV1_1
-			class FUpdateArchetype : public Frontend::IDocumentTransform
+			class FUpdateInterface : public Frontend::IDocumentTransform
 			{
 			public:
 				virtual bool Transform(Frontend::FDocumentHandle InDocument) const override
@@ -570,6 +542,9 @@ namespace Metasound
 						return false;
 					}
 
+					InDocument->RemoveInterfaceVersion(MetasoundSourceStereoV1_0::GetInterface().Version);
+					InDocument->AddInterfaceVersion(MetasoundSourceStereoV1_1::GetInterface().Version);
+
 					FNodeHandle StereoFormatOutput = Graph->GetOutputNodeWithName(MetasoundSourceStereoV1_0::GetAudioOutputName());
 					FOutputHandle LeftOutputToReconnect = IOutputController::GetInvalidHandle();
 					FOutputHandle RightOutputToReconnect = IOutputController::GetInvalidHandle();
@@ -715,20 +690,20 @@ namespace Metasound
 			}
 		}
 
-		void RegisterArchetypes()
+		void RegisterInternalInterfaces()
 		{
-			using namespace MetasoundEngineArchetypesPrivate;
-
-			Frontend::IArchetypeRegistry::Get().RegisterArchetype(MakeUnique<FArchetypeRegistryEntry>(MetasoundV1_0::GetArchetype(), nullptr));
-
-			Frontend::IArchetypeRegistry::Get().RegisterArchetype(MakeUnique<FArchetypeRegistryEntry>(MetasoundSourceMonoV1_0::GetArchetype(), nullptr));
-
-			Frontend::IArchetypeRegistry::Get().RegisterArchetype(MakeUnique<FArchetypeRegistryEntry>(MetasoundSourceStereoV1_0::GetArchetype(), nullptr));
-
-			Frontend::IArchetypeRegistry::Get().RegisterArchetype(MakeUnique<FArchetypeRegistryEntry>(MetasoundSourceMonoV1_1::GetArchetype(), MakeUnique<MetasoundSourceMonoV1_1::FUpdateArchetype>()));
-
-			Frontend::IArchetypeRegistry::Get().RegisterArchetype(MakeUnique<FArchetypeRegistryEntry>(MetasoundSourceStereoV1_1::GetArchetype(), MakeUnique<MetasoundSourceStereoV1_1::FUpdateArchetype>()));
+			{
+				constexpr bool bIsDefault = false;
+				RegisterInterface<UMetaSoundSource>(MetasoundSourceStereoV1_1::GetInterface(), MakeUnique<MetasoundSourceStereoV1_1::FUpdateInterface>(), bIsDefault, IDataReference::RouterName);
+				RegisterInterface<UMetaSoundSource>(MetasoundSourceStereoV1_0::GetInterface(), nullptr, bIsDefault, IDataReference::RouterName);
+			}
 
+			{
+				constexpr bool bIsDefault = true;
+				RegisterInterface<UMetaSound>(MetasoundV1_0::GetInterface(), nullptr, bIsDefault, IDataReference::RouterName);
+				RegisterInterface<UMetaSoundSource>(MetasoundSourceMonoV1_0::GetInterface(), nullptr, bIsDefault, IDataReference::RouterName);
+				RegisterInterface<UMetaSoundSource>(MetasoundSourceMonoV1_1::GetInterface(), MakeUnique<MetasoundSourceMonoV1_1::FUpdateInterface>(), bIsDefault, IDataReference::RouterName);
+			}
 		}
 	}
 }
diff --git a/Engine/Plugins/Runtime/Metasound/Source/MetasoundEngine/Private/MetasoundEngineArchetypes.h b/Engine/Plugins/Runtime/Metasound/Source/MetasoundEngine/Private/MetasoundEngineArchetypes.h
--- a/Engine/Plugins/Runtime/Metasound/Source/MetasoundEngine/Private/MetasoundEngineArchetypes.h
+++ b/Engine/Plugins/Runtime/Metasound/Source/MetasoundEngine/Private/MetasoundEngineArchetypes.h
@@ -97,8 +97,7 @@ namespace Metasound
 			const FVertexName& GetRightAudioOutputName();
 		}
 
-		// Register all MetasoundEngine archetypes.
-		void RegisterArchetypes();
+		// Register all interfaces defined in MetaSoundEngine.
+		void RegisterInternalInterfaces();
 	}
 }
-
diff --git a/Engine/Plugins/Runtime/Metasound/Source/MetasoundEngine/Private/MetasoundEngineModule.cpp b/Engine/Plugins/Runtime/Metasound/Source/MetasoundEngine/Private/MetasoundEngineModule.cpp
--- a/Engine/Plugins/Runtime/Metasound/Source/MetasoundEngine/Private/MetasoundEngineModule.cpp
+++ b/Engine/Plugins/Runtime/Metasound/Source/MetasoundEngine/Private/MetasoundEngineModule.cpp
@@ -5,6 +5,7 @@
 #include "MetasoundDataReference.h"
 #include "MetasoundDataTypeRegistrationMacro.h"
 #include "MetasoundEngineArchetypes.h"
+#include "MetasoundInterface.h"
 #include "MetasoundSource.h"
 #include "MetasoundUObjectRegistry.h"
 #include "MetasoundWave.h"
@@ -24,20 +25,14 @@ class FMetasoundEngineModule : public IMetasoundEngineModule
 		FModuleManager::Get().LoadModuleChecked("MetasoundGenerator");
 		FModuleManager::Get().LoadModuleChecked("AudioCodecEngine");
 
-		Metasound::Engine::RegisterArchetypes();
+		// Register interfaces
+		Metasound::Engine::RegisterInternalInterfaces();
+		Metasound::Engine::RegisterExternalInterfaces();
 
-		// If there is no archetype name, use UMetaSound
-		FMetasoundFrontendVersion DefaultVersion;
-		Metasound::IMetasoundUObjectRegistry::RegisterUClassArchetype<UMetaSound>(DefaultVersion);
-
-		// Register preferred archetypes
-		Metasound::IMetasoundUObjectRegistry::RegisterUClassPreferredArchetypes<UMetaSound>();
-		Metasound::IMetasoundUObjectRegistry::RegisterUClassPreferredArchetypes<UMetaSoundSource>();
-
-		// flush node registration queue
+		// Flush node registration queue
 		FMetasoundFrontendRegistryContainer::Get()->RegisterPendingNodes();
 
-		UE_LOG(LogMetasoundEngine, Log, TEXT("Metasound Engine Initialized"));
+		UE_LOG(LogMetasoundEngine, Log, TEXT("MetaSound Engine Initialized"));
 	}
 };
 
diff --git a/Engine/Plugins/Runtime/Metasound/Source/MetasoundEngine/Private/MetasoundInterface.cpp b/Engine/Plugins/Runtime/Metasound/Source/MetasoundEngine/Private/MetasoundInterface.cpp
new file mode 100644
--- /dev/null
+++ b/Engine/Plugins/Runtime/Metasound/Source/MetasoundEngine/Private/MetasoundInterface.cpp
@@ -0,0 +1,85 @@
+// Copyright Epic Games, Inc. All Rights Reserved.
+#include "MetasoundInterface.h"
+
+#include "IAudioParameterTransmitter.h"
+
+
+namespace Metasound
+{
+	namespace Engine
+	{
+		void RegisterExternalInterfaces()
+		{
+			using namespace Audio;
+
+			auto RegisterExternalInterface = [](FGeneratorInterfacePtr Interface)
+			{
+				auto ResolveMemberDataType = [](FName DataType, EAudioParameterType ParamType)
+				{
+					if (!DataType.IsNone())
+					{
+						const bool bIsRegisteredType = Frontend::IDataTypeRegistry::Get().IsRegistered(DataType);
+						if (ensureAlwaysMsgf(bIsRegisteredType, TEXT("Attempting to register Interface member with unregistered DataType '%s'."), *DataType.ToString()))
+						{
+							return DataType;
+						}
+					}
+
+					return Frontend::ConvertParameterToDataType(ParamType);
+				};
+
+				FMetasoundFrontendInterface FrontendInterface;
+				FrontendInterface.Version = { Interface->Name, FMetasoundFrontendVersionNumber { Interface->Version.Major, Interface->Version.Minor } };
+
+				Algo::Transform(Interface->Inputs, FrontendInterface.Inputs, [&](const FGeneratorInterface::FInput& Input)
+				{
+					FMetasoundFrontendClassInput ClassInput;
+					ClassInput.Name = IGeneratorInterfaceRegistry::GetMemberFullName(Interface->Name, Input.InitValue.ParamName);
+					ClassInput.DefaultLiteral = FMetasoundFrontendLiteral(Input.InitValue);
+					ClassInput.TypeName = ResolveMemberDataType(Input.DataType, Input.InitValue.ParamType);
+					ClassInput.Metadata.DisplayName = Input.DisplayName;
+					ClassInput.Metadata.Description = Input.Description;
+					ClassInput.VertexID = FGuid::NewGuid();
+
+					return ClassInput;
+				});
+
+				Algo::Transform(Interface->Outputs, FrontendInterface.Outputs, [&](const FGeneratorInterface::FOutput& Output)
+				{
+					FMetasoundFrontendClassOutput ClassOutput;
+					ClassOutput.Name = IGeneratorInterfaceRegistry::GetMemberFullName(Interface->Name, Output.ParamName);
+					ClassOutput.TypeName = ResolveMemberDataType(Output.DataType, Output.ParamType);
+					ClassOutput.Metadata.DisplayName = Output.DisplayName;
+					ClassOutput.Metadata.Description = Output.Description;
+					ClassOutput.VertexID = FGuid::NewGuid();
+
+					return ClassOutput;
+				});
+
+				Algo::Transform(Interface->Environment, FrontendInterface.Environment, [&](const FGeneratorInterface::FEnvironmentVariable& Environment)
+				{
+					FMetasoundFrontendClassEnvironmentVariable EnvironmentVariable;
+					EnvironmentVariable.Name = IGeneratorInterfaceRegistry::GetMemberFullName(Interface->Name, Environment.ParamName);
+					EnvironmentVariable.TypeName = ResolveMemberDataType(Environment.DataType, Environment.ParamType);
+					EnvironmentVariable.Metadata.DisplayName = Environment.DisplayName;
+					EnvironmentVariable.Metadata.Description = Environment.Description;
+
+					return EnvironmentVariable;
+				});
+
+				// For now, only support interfaces for UMetaSoundSource type
+				if (ensureAlways(UMetaSoundSource::StaticClass()->IsChildOf(Interface->Type)))
+				{
+					// For now, all external interfaces use the ParameterTransmitter and cannot be applied by default to new
+					// MetaSound objects. Could pipe either of these through in the future if required/desired.
+					constexpr bool bIsDefault = false;
+					RegisterInterface<UMetaSoundSource>(FrontendInterface, nullptr, bIsDefault, Audio::IParameterTransmitter::RouterName);
+					UE_LOG(LogMetaSound, Verbose, TEXT("Interface '%s' registered as MetaSoundSource assets."), *Interface->Name.ToString());
+				}
+			};
+
+			IGeneratorInterfaceRegistry::Get().IterateInterfaces(RegisterExternalInterface);
+			IGeneratorInterfaceRegistry::Get().OnRegistration(MoveTemp(RegisterExternalInterface));
+		}
+	} // namespace Engine
+} // namespace Metasound
diff --git a/Engine/Plugins/Runtime/Metasound/Source/MetasoundEngine/Private/MetasoundInterface.h b/Engine/Plugins/Runtime/Metasound/Source/MetasoundEngine/Private/MetasoundInterface.h
new file mode 100644
--- /dev/null
+++ b/Engine/Plugins/Runtime/Metasound/Source/MetasoundEngine/Private/MetasoundInterface.h
@@ -0,0 +1,102 @@
+// Copyright Epic Games, Inc. All Rights Reserved.
+#pragma once
+
+#include "Algo/Transform.h"
+#include "AudioParameterInterface.h"
+#include "IAudioGeneratorInterfaceRegistry.h"
+#include "MetasoundFrontendDataTypeRegistry.h"
+#include "MetasoundFrontendDocument.h"
+#include "MetasoundFrontendTransform.h"
+#include "MetasoundLog.h"
+#include "MetasoundParameterTransmitter.h"
+#include "MetasoundSource.h"
+#include "MetasoundUObjectRegistry.h"
+#include "Templates/UniquePtr.h"
+#include "UObject/Class.h"
+#include "UObject/NoExportTypes.h"
+
+
+namespace Metasound
+{
+	namespace Engine
+	{
+		void RegisterExternalInterfaces();
+
+		struct FInterfaceRegistryOptions
+		{
+			bool bIsDefault = false;
+			FName InputSystemName;
+			FName UClassName;
+		};
+
+		// Entry for registered interface.
+		class FInterfaceRegistryEntry : public Frontend::IInterfaceRegistryEntry
+		{
+		public:
+			FInterfaceRegistryEntry(
+				const FMetasoundFrontendInterface& InInterface,
+				TUniquePtr<Frontend::IDocumentTransform>&& InUpdateTransform,
+				FInterfaceRegistryOptions&& InOptions
+			)
+				: Interface(InInterface)
+				, UpdateTransform(MoveTemp(InUpdateTransform))
+				, Options(MoveTemp(InOptions))
+			{
+			}
+
+			virtual bool UClassIsSupported(FName InUClassName) const override
+			{
+				if (Options.UClassName.IsNone())
+				{
+					return true;
+				}
+
+				// TODO: Support child asset class types.
+				return Options.UClassName == InUClassName;
+			}
+
+			virtual bool IsDefault() const override
+			{
+				return Options.bIsDefault;
+			}
+
+			virtual FName GetRouterName() const override
+			{
+				return Options.InputSystemName;
+			}
+
+			virtual const FMetasoundFrontendInterface& GetInterface() const override
+			{
+				return Interface;
+			}
+
+			virtual bool UpdateRootGraphInterface(Frontend::FDocumentHandle InDocument) const override
+			{
+				if (UpdateTransform.IsValid())
+				{
+					return UpdateTransform->Transform(InDocument);
+				}
+				return false;
+			}
+
+			private:
+				FMetasoundFrontendInterface Interface;
+				TUniquePtr<Frontend::IDocumentTransform> UpdateTransform;
+				FInterfaceRegistryOptions Options;
+		};
+
+		template <typename UClassType>
+		void RegisterInterface(const FMetasoundFrontendInterface& InInterface, TUniquePtr<Frontend::IDocumentTransform>&& UpdateTransform, bool bInIsDefault, FName InRouterName)
+		{
+			FInterfaceRegistryOptions Options
+			{
+				bInIsDefault,
+				InRouterName,
+				UClassType::StaticClass()->GetFName()
+			};
+
+			IMetasoundUObjectRegistry::Get().RegisterUClassInterface(MakeUnique<TMetasoundUObjectRegistryEntry<UClassType>>(InInterface.Version));
+			Frontend::IInterfaceRegistry::Get().RegisterInterface(MakeUnique<FInterfaceRegistryEntry>(InInterface, MoveTemp(UpdateTransform), MoveTemp(Options)));
+		}
+	} // namespace Engine
+} // namespace Metasound
diff --git a/Engine/Plugins/Runtime/Metasound/Source/MetasoundEngine/Private/MetasoundSource.cpp b/Engine/Plugins/Runtime/Metasound/Source/MetasoundEngine/Private/MetasoundSource.cpp
--- a/Engine/Plugins/Runtime/Metasound/Source/MetasoundEngine/Private/MetasoundSource.cpp
+++ b/Engine/Plugins/Runtime/Metasound/Source/MetasoundEngine/Private/MetasoundSource.cpp
@@ -3,6 +3,7 @@
 
 #include "Algo/Transform.h"
 #include "AssetRegistryModule.h"
+#include "IAudioGeneratorInterfaceRegistry.h"
 #include "Internationalization/Text.h"
 #include "MetasoundAssetBase.h"
 #include "MetasoundAudioFormats.h"
@@ -93,12 +94,19 @@ void UMetaSoundSource::PostEditUndo()
 void UMetaSoundSource::PostDuplicate(EDuplicateMode::Type InDuplicateMode)
 {
 	Super::PostDuplicate(InDuplicateMode);
-	Metasound::PostDuplicate(*this, InDuplicateMode);
+
+	// Guid is reset as asset may share implementation from
+	// asset duplicated from but should not be registered as such.
+	if (InDuplicateMode == EDuplicateMode::Normal)
+	{
+		Metasound::Frontend::FRegenerateAssetClassName().Transform(GetDocumentHandle());
+	}
 }
 
 void UMetaSoundSource::PostEditChangeProperty(FPropertyChangedEvent& InEvent)
 {
 	using namespace Metasound;
+	using namespace Metasound::Engine;
 	using namespace Metasound::Frontend;
 
 	Super::PostEditChangeProperty(InEvent);
@@ -113,14 +121,19 @@ void UMetaSoundSource::PostEditChangeProperty(FPropertyChangedEvent& InEvent)
 		{
 			case EMetasoundSourceAudioFormat::Mono:
 			{
-				// TODO: utilize latest version `MetasoudnSourceMono`
-				bDidModifyDocument = Metasound::Frontend::FMatchRootGraphToArchetype(Metasound::Engine::MetasoundSourceMono::GetVersion()).Transform(GetDocumentHandle());
+				bDidModifyDocument = FModifyRootGraphInterfaces(
+					{ MetasoundSourceStereo::GetVersion() },
+					{ MetasoundSourceMono::GetVersion() }
+				).Transform(GetDocumentHandle());
 			}
 			break;
 
 			case EMetasoundSourceAudioFormat::Stereo:
 			{
-				bDidModifyDocument = Metasound::Frontend::FMatchRootGraphToArchetype(Metasound::Engine::MetasoundSourceStereo::GetVersion()).Transform(GetDocumentHandle());
+				bDidModifyDocument = FModifyRootGraphInterfaces(
+					{ MetasoundSourceMono::GetVersion() },
+					{ MetasoundSourceStereo::GetVersion() }
+				).Transform(GetDocumentHandle());
 			}
 			break;
 
@@ -133,7 +146,7 @@ void UMetaSoundSource::PostEditChangeProperty(FPropertyChangedEvent& InEvent)
 
 		if (bDidModifyDocument)
 		{
-			ConformObjectDataToArchetype();
+			ConformObjectDataToInterfaces();
 
 			// Use the editor form of register to ensure other editors'
 			// MetaSounds are auto-updated if they are referencing this graph.
@@ -147,10 +160,10 @@ void UMetaSoundSource::PostEditChangeProperty(FPropertyChangedEvent& InEvent)
 }
 #endif // WITH_EDITOR
 
-bool UMetaSoundSource::ConformObjectDataToArchetype()
+bool UMetaSoundSource::ConformObjectDataToInterfaces()
 {
-	const FMetasoundFrontendVersion& ArchetypeVersion = GetDocumentHandle()->GetArchetypeVersion();
-	if (ArchetypeVersion == Metasound::Engine::MetasoundSourceMono::GetVersion())
+	const TArray<FMetasoundFrontendVersion>& InterfaceVersions = GetDocumentHandle()->GetInterfaceVersions();
+	if (InterfaceVersions.Contains(Metasound::Engine::MetasoundSourceMono::GetVersion()))
 	{
 		if (OutputFormat != EMetasoundSourceAudioFormat::Mono || NumChannels != 1)
 		{
@@ -160,7 +173,7 @@ bool UMetaSoundSource::ConformObjectDataToArchetype()
 		}
 	}
 
-	if (ArchetypeVersion == Metasound::Engine::MetasoundSourceStereo::GetVersion())
+	if (InterfaceVersions.Contains(Metasound::Engine::MetasoundSourceStereo::GetVersion()))
 	{
 		if (OutputFormat != EMetasoundSourceAudioFormat::Stereo || NumChannels != 2)
 		{
@@ -416,7 +429,6 @@ bool UMetaSoundSource::SupportsSubtitles() const
 const FMetasoundFrontendVersion& UMetaSoundSource::GetDefaultArchetypeVersion() const
 {
 	static const FMetasoundFrontendVersion DefaultVersion = Metasound::Engine::MetasoundSourceMono::GetVersion();
-
 	return DefaultVersion;
 }
 
@@ -426,6 +438,17 @@ float UMetaSoundSource::GetDuration()
 	return Super::GetDuration();
 }
 
+bool UMetaSoundSource::ImplementsGeneratorInterface(Audio::FGeneratorInterfacePtr InInterface) const
+{
+	const FMetasoundFrontendVersionNumber Number = FMetasoundFrontendVersionNumber{ InInterface->Version.Major, InInterface->Version.Minor };
+	auto ContainsInterface = [Name = InInterface->Name, &Number](FMetasoundFrontendVersion& InVersion)
+	{
+		return InVersion.Name == Name && InVersion.Number == Number;
+	};
+
+	return GetDocumentChecked().InterfaceVersions.ContainsByPredicate(ContainsInterface);
+}
+
 ISoundGeneratorPtr UMetaSoundSource::CreateSoundGenerator(const FSoundGeneratorInitParams& InParams)
 {
 	using namespace Metasound;
@@ -698,18 +721,6 @@ Metasound::FOperatorSettings UMetaSoundSource::GetOperatorSettings(Metasound::FS
 	return Metasound::FOperatorSettings(InSampleRate, BlockRate);
 }
 
-const TArray<FMetasoundFrontendVersion>& UMetaSoundSource::GetSupportedArchetypeVersions() const 
-{
-	static const TArray<FMetasoundFrontendVersion> Supported(
-	{
-		Metasound::Engine::MetasoundSourceMono::GetVersion(),
-		Metasound::Engine::MetasoundSourceStereo::GetVersion()
-	});
-
-	return Supported;
-}
-
-
 Metasound::FMetasoundEnvironment UMetaSoundSource::CreateEnvironment() const
 {
 	using namespace Metasound;
diff --git a/Engine/Plugins/Runtime/Metasound/Source/MetasoundEngine/Private/MetasoundUObjectRegistry.cpp b/Engine/Plugins/Runtime/Metasound/Source/MetasoundEngine/Private/MetasoundUObjectRegistry.cpp
--- a/Engine/Plugins/Runtime/Metasound/Source/MetasoundEngine/Private/MetasoundUObjectRegistry.cpp
+++ b/Engine/Plugins/Runtime/Metasound/Source/MetasoundEngine/Private/MetasoundUObjectRegistry.cpp
@@ -21,30 +21,47 @@
 
 namespace Metasound
 {
-
 	class FMetasoundUObjectRegistry : public IMetasoundUObjectRegistry
 	{
 		public:
-			void RegisterUClassArchetype(TUniquePtr<IMetasoundUObjectRegistryEntry>&& InEntry) override
+			FMetasoundUObjectRegistry()
+			{
+				// Set default version to use base UMetaSound class implementation
+				FMetasoundFrontendVersion DefaultVersion;
+
+				using FRegistryEntryType = TMetasoundUObjectRegistryEntry<UMetaSound>;
+				RegisterUClassInterface(MakeUnique<FRegistryEntryType>(DefaultVersion));
+			}
+
+			void RegisterUClassInterface(TUniquePtr<IMetasoundUObjectRegistryEntry>&& InEntry) override
 			{
 				if (InEntry.IsValid())
 				{
-					Frontend::FArchetypeRegistryKey Key = Frontend::GetArchetypeRegistryKey(InEntry->GetArchetypeVersion());
+					Frontend::FInterfaceRegistryKey Key = Frontend::GetInterfaceRegistryKey(InEntry->GetInterfaceVersion());
 
-					EntriesByArchetype.Add(Key, InEntry.Get());
+					EntriesByInterface.Add(Key, InEntry.Get());
+					EntriesByName.Add(InEntry->GetInterfaceVersion().Name, InEntry.Get());
 					Entries.Add(InEntry.Get());
 					Storage.Add(MoveTemp(InEntry));
 				}
 			}
 
-			TArray<UClass*> GetUClassesForArchetype(const FMetasoundFrontendVersion& InArchetypeVersion) const override
+			TArray<const IMetasoundUObjectRegistryEntry*> FindInterfaceEntriesByName(FName InName) const override
+			{
+				TArray<const IMetasoundUObjectRegistryEntry*> EntriesWithName;
+				EntriesByName.MultiFind(InName, EntriesWithName);
+
+				return EntriesWithName;
+			}
+
+			TArray<UClass*> FindSupportedInterfaceClasses(const FMetasoundFrontendVersion& InInterfaceVersion) const override
 			{
 				TArray<UClass*> Classes;
 
-				TArray<const IMetasoundUObjectRegistryEntry*> EntriesForArchetype;
-				EntriesByArchetype.MultiFind(Frontend::GetArchetypeRegistryKey(InArchetypeVersion), EntriesForArchetype);
+				TArray<const IMetasoundUObjectRegistryEntry*> EntriesForInterface;
+				EntriesByInterface.MultiFind(Frontend::GetInterfaceRegistryKey(InInterfaceVersion), EntriesForInterface);
 
-				for (const IMetasoundUObjectRegistryEntry* Entry : EntriesForArchetype)
+				for (const IMetasoundUObjectRegistryEntry* Entry : EntriesForInterface)
 				{
 					if (nullptr != Entry)
 					{
@@ -60,18 +77,21 @@ namespace Metasound
 
 			UObject* NewObject(UClass* InClass, const FMetasoundFrontendDocument& InDocument, const FString& InPath) const override
 			{
-				TArray<const IMetasoundUObjectRegistryEntry*> EntriesForArchetype;
-				EntriesByArchetype.MultiFind(Frontend::GetArchetypeRegistryKey(InDocument.ArchetypeVersion), EntriesForArchetype);
+				TArray<const IMetasoundUObjectRegistryEntry*> AllInterfaceEntries;
+
+				for (const FMetasoundFrontendVersion& InterfaceVersion : InDocument.InterfaceVersions)
+				{
+					TArray<const IMetasoundUObjectRegistryEntry*> EntriesForInterface;
+					EntriesByInterface.MultiFind(Frontend::GetInterfaceRegistryKey(InterfaceVersion), EntriesForInterface);
+					AllInterfaceEntries.Append(MoveTemp(EntriesForInterface));
+				}
 
 				auto IsChildClassOfRegisteredClass = [&](const IMetasoundUObjectRegistryEntry* Entry)
 				{
 					return Entry->IsChildClass(InClass);
 				};
 
-			
-				const IMetasoundUObjectRegistryEntry* const* EntryForClass = EntriesForArchetype.FindByPredicate(IsChildClassOfRegisteredClass);
-
-				if (nullptr != EntryForClass)
+				if (const IMetasoundUObjectRegistryEntry* const* EntryForClass = AllInterfaceEntries.FindByPredicate(IsChildClassOfRegisteredClass))
 				{
 					return NewObject(**EntryForClass, InDocument, InPath);
 				}
@@ -131,7 +151,6 @@ namespace Metasound
 				if (ensure(nullptr != NewAssetBase))
 				{
 					NewAssetBase->SetDocument(InDocument);
-					NewAssetBase->ConformDocumentToArchetype();
 				}
 
 #if WITH_EDITOR
@@ -183,7 +202,8 @@ namespace Metasound
 			}
 
 			TArray<TUniquePtr<IMetasoundUObjectRegistryEntry>> Storage;
-			TMultiMap<Frontend::FArchetypeRegistryKey, const IMetasoundUObjectRegistryEntry*> EntriesByArchetype;
+			TMultiMap<Frontend::FInterfaceRegistryKey, const IMetasoundUObjectRegistryEntry*> EntriesByInterface;
+			TMultiMap<FName, const IMetasoundUObjectRegistryEntry*> EntriesByName;
 			TArray<const IMetasoundUObjectRegistryEntry*> Entries;
 	};
 
diff --git a/Engine/Plugins/Runtime/Metasound/Source/MetasoundEngine/Public/Metasound.h b/Engine/Plugins/Runtime/Metasound/Source/MetasoundEngine/Public/Metasound.h
--- a/Engine/Plugins/Runtime/Metasound/Source/MetasoundEngine/Public/Metasound.h
+++ b/Engine/Plugins/Runtime/Metasound/Source/MetasoundEngine/Public/Metasound.h
@@ -47,21 +47,6 @@ namespace Metasound
 	void PostEditChangeProperty(TMetaSoundObject& InMetaSound, FPropertyChangedEvent& InEvent)
 	{
 	}
-
-	template <typename TMetaSoundObject>
-	void PostDuplicate(TMetaSoundObject& InMetaSound, EDuplicateMode::Type InDuplicateMode)
-	{
-		if (InDuplicateMode == EDuplicateMode::Normal)
-		{
-			// Guid is reset as asset may share implementation from
-			// asset duplicated from but should not be registered as such.
-			FMetasoundAssetBase* MetasoundAsset = IMetasoundUObjectRegistry::Get().GetObjectAsAssetBase(&InMetaSound);
-			if (ensure(MetasoundAsset))
-			{
-				Metasound::Frontend::FRegenerateAssetClassName().Transform(MetasoundAsset->GetDocumentHandle());
-			}
-		}
-	}
 #endif // WITH_EDITOR
 
 	template <typename TMetaSoundObject>
@@ -247,7 +232,7 @@ public:
 
 	virtual const FMetasoundFrontendVersion& GetDefaultArchetypeVersion() const override;
 
-	virtual bool ConformObjectDataToArchetype() override { return false; }
+	virtual bool ConformObjectDataToInterfaces() override { return false; }
 
 	virtual const TSet<FString>& GetReferencedAssetClassKeys() const override
 	{
@@ -264,9 +249,6 @@ public:
 		return this;
 	}
 
-	// Get the most up to date archetype for metasound sources.
-	const TArray<FMetasoundFrontendVersion>& GetSupportedArchetypeVersions() const override;
-
 protected:
 	virtual void SetReferencedAssetClassKeys(TSet<Metasound::Frontend::FNodeRegistryKey>&& InKeys) override;
 
diff --git a/Engine/Plugins/Runtime/Metasound/Source/MetasoundEngine/Public/MetasoundSource.h b/Engine/Plugins/Runtime/Metasound/Source/MetasoundEngine/Public/MetasoundSource.h
--- a/Engine/Plugins/Runtime/Metasound/Source/MetasoundEngine/Public/MetasoundSource.h
+++ b/Engine/Plugins/Runtime/Metasound/Source/MetasoundEngine/Public/MetasoundSource.h
@@ -147,7 +147,7 @@ public:
 
 	virtual const FMetasoundFrontendVersion& GetDefaultArchetypeVersion() const override;
 
-	virtual bool ConformObjectDataToArchetype() override;
+	virtual bool ConformObjectDataToInterfaces() override;
 
 	UObject* GetOwningAsset() override
 	{
@@ -165,14 +165,11 @@ public:
 	virtual bool IsPlayable() const override;
 	virtual bool SupportsSubtitles() const override;
 	virtual float GetDuration() override;
+	virtual bool ImplementsGeneratorInterface(Audio::FGeneratorInterfacePtr InInterface) const override;
 	virtual ISoundGeneratorPtr CreateSoundGenerator(const FSoundGeneratorInitParams& InParams) override;
 	virtual TUniquePtr<Audio::IParameterTransmitter> CreateParameterTransmitter(Audio::FParameterTransmitterInitParams&& InParams) const override;
 	virtual bool IsParameterValid(const FAudioParameter& InParameter) const override;
 
-	// Get the most up to date archetype for metasound sources.
-	const TArray<FMetasoundFrontendVersion>& GetSupportedArchetypeVersions() const override;
-
-
 protected:
 	/** Gets all the default parameters for this Asset.  */
 	virtual bool GetAllDefaultParameters(TArray<FAudioParameter>& OutParameters) const override;
diff --git a/Engine/Plugins/Runtime/Metasound/Source/MetasoundEngine/Public/MetasoundUObjectRegistry.h b/Engine/Plugins/Runtime/Metasound/Source/MetasoundEngine/Public/MetasoundUObjectRegistry.h
--- a/Engine/Plugins/Runtime/Metasound/Source/MetasoundEngine/Public/MetasoundUObjectRegistry.h
+++ b/Engine/Plugins/Runtime/Metasound/Source/MetasoundEngine/Public/MetasoundUObjectRegistry.h
@@ -17,7 +17,7 @@ namespace Metasound
 {
 	/** Interface for an entry into the Metasound-UObject Registry. 
 	 *
-	 * An entry provides information linking a FMetasoundFrontendArchetype to a UClass.
+	 * An entry provides information linking a FMetasoundFrontendClassInterface to a UClass.
 	 * It also provides methods for accessing the FMetasoundAssetBase from a UObject.
 	 */
 	class IMetasoundUObjectRegistryEntry
@@ -25,8 +25,8 @@ namespace Metasound
 	public:
 		virtual ~IMetasoundUObjectRegistryEntry() = default;
 
-		/** Archetype name associated with this entry. */
-		virtual const FMetasoundFrontendVersion& GetArchetypeVersion() const = 0;
+		/** Interface version associated with this entry. */
+		virtual const FMetasoundFrontendVersion& GetInterfaceVersion() const = 0;
 
 		/** UClass associated with this entry. */
 		virtual UClass* GetUClass() const = 0;
@@ -54,7 +54,7 @@ namespace Metasound
 		friend class TMetasoundUObjectRegistryEntry;
 	};
 
-	/** An entry into the Metasound-UObject registry. 
+	/** An entry into the Metasound-UObject registry.
 	 *
 	 * @Tparam UClassType A class which derives from UObject and FMetasoundAssetBase.
 	 */
@@ -66,16 +66,16 @@ namespace Metasound
 		static_assert(std::is_base_of<UObject, UClassType>::value, "UClass must be derived from UObject");
 
 	public:
-		TMetasoundUObjectRegistryEntry(const FMetasoundFrontendVersion& InArchetypeVersion)
-		:	ArchetypeVersion(InArchetypeVersion)
+		TMetasoundUObjectRegistryEntry(const FMetasoundFrontendVersion& InInterfaceVersion)
+		:	InterfaceVersion(InInterfaceVersion)
 		{
 		}
 
 		virtual ~TMetasoundUObjectRegistryEntry() = default;
 
-		virtual const FMetasoundFrontendVersion& GetArchetypeVersion() const override
+		virtual const FMetasoundFrontendVersion& GetInterfaceVersion() const override
 		{
-			return ArchetypeVersion;
+			return InterfaceVersion;
 		}
 
 		UClass* GetUClass() const override
@@ -126,11 +126,11 @@ namespace Metasound
 
 	private:
 
-		FMetasoundFrontendVersion ArchetypeVersion;
+		FMetasoundFrontendVersion InterfaceVersion;
 	};
 
 
-	/** IMetaoundUObjectRegistry contains IMetasoundUObjectRegistryEntrys. 
+	/** IMetasoundUObjectRegistry contains IMetasoundUObjectRegistryEntries.
 	 *
 	 * Registered UObject classes can utilize the Metasound Editor. It also enables
 	 * the creation of a UObject directly from a FMetasoundFrontendDocument.
@@ -143,46 +143,20 @@ namespace Metasound
 			/** Return static singleton instance of the registry. */
 			static IMetasoundUObjectRegistry& Get();
 
-			/** Register all preferred archetypes of the UClass. 
-			 *
-			 * All Archtypes returned by the default objects GetSupportedArchetypeVersions() will be registered. 
-			 */
-			template<typename UClassType>
-			static void RegisterUClassPreferredArchetypes()
-			{
-				static_assert(std::is_base_of<FMetasoundAssetBase, UClassType>::value, "UClass must be derived from FMetasoundAssetBase");
-
-				const TArray<FMetasoundFrontendVersion>& SupportedArchetypeVersions = GetDefault<UClassType>()->GetSupportedArchetypeVersions();
-
-				for (const FMetasoundFrontendVersion& Version : SupportedArchetypeVersions)
-				{
-					IMetasoundUObjectRegistry::RegisterUClassArchetype<UClassType>(Version);
-				}
-			}
-
-			/** Register an archetype for the UClass. 
-			 *
-			 * @param InArchetypeVerison - The version of the FMetasoundFrontendArchetype to associate with the UClass.
-			 */
-			template<typename UClassType>
-			static void RegisterUClassArchetype(const FMetasoundFrontendVersion& InArchetypeVersion)
-			{
-				using FRegistryEntryType = TMetasoundUObjectRegistryEntry<UClassType>;
-
-				IMetasoundUObjectRegistry::Get().RegisterUClassArchetype(MakeUnique<FRegistryEntryType>(InArchetypeVersion));
-			}
-
 			/** Adds an entry to the registry. */
-			virtual void RegisterUClassArchetype(TUniquePtr<IMetasoundUObjectRegistryEntry>&& InEntry) = 0;
+			virtual void RegisterUClassInterface(TUniquePtr<IMetasoundUObjectRegistryEntry>&& InEntry) = 0;
+
+			/** Returns all RegistryEntries with the given name */
+			virtual TArray<const IMetasoundUObjectRegistryEntry*> FindInterfaceEntriesByName(FName InName) const = 0;
 
-			/** Returns all UClasses registered to the archetype name. */
-			virtual TArray<UClass*> GetUClassesForArchetype(const FMetasoundFrontendVersion& InArchetypeVersion) const = 0;
+			/** Returns all UClasses registered to the interface version. */
+			virtual TArray<UClass*> FindSupportedInterfaceClasses(const FMetasoundFrontendVersion& InInterfaceVersion) const = 0;
 
-			/** Creates a new object from a metasound document.
+			/** Creates a new object from a MetaSound document.
 			 *
 			 * @param InClass - A registered UClass to create.
 			 * @param InDocument - The FMetasoundFrontendDocument to use when creating the class.
-			 * @param InArchetypeVersion - The version of the FMetasoundFrontendArchetype to use when creating the class.
+			 * @param InInterfaceVersion - The version of the FMetasoundFrontendClassInterface to use when creating the class.
 			 * @param InPath - If in editor, the created asset will be stored at this content path.
 			 *
 			 * @return A new object. A nullptr on error.
diff --git a/Engine/Plugins/Runtime/Metasound/Source/MetasoundEngineTest/Private/EngineTestMetasoundSource.cpp b/Engine/Plugins/Runtime/Metasound/Source/MetasoundEngineTest/Private/EngineTestMetasoundSource.cpp
--- a/Engine/Plugins/Runtime/Metasound/Source/MetasoundEngineTest/Private/EngineTestMetasoundSource.cpp
+++ b/Engine/Plugins/Runtime/Metasound/Source/MetasoundEngineTest/Private/EngineTestMetasoundSource.cpp
@@ -64,7 +64,7 @@ namespace EngineTestMetasoundSourcePrivate
 		using namespace Metasound::Frontend;
 
 		FMetasoundFrontendDocument Document;
-		Document.ArchetypeVersion = FMetasoundFrontendVersion{"MonoSource", {1, 1}};
+		Document.InterfaceVersions.Add(FMetasoundFrontendVersion{"MonoSource", {1, 1}});
 
 		Document.RootGraph.Metadata.SetClassName(FMetasoundFrontendClassName { "Namespace", "Unit Test Node", *LexToString(FGuid::NewGuid()) });
 		Document.RootGraph.Metadata.SetType(EMetasoundFrontendClassType::Graph);
diff --git a/Engine/Plugins/Runtime/Metasound/Source/MetasoundFrontend/Private/MetasoundArchetype.cpp b/Engine/Plugins/Runtime/Metasound/Source/MetasoundFrontend/Private/MetasoundArchetype.cpp
--- a/Engine/Plugins/Runtime/Metasound/Source/MetasoundFrontend/Private/MetasoundArchetype.cpp
+++ b/Engine/Plugins/Runtime/Metasound/Source/MetasoundFrontend/Private/MetasoundArchetype.cpp
@@ -10,11 +10,11 @@ namespace Metasound
 {
 	namespace Frontend
 	{
-		namespace MetasoundArchetypeIntrinsics
+		namespace MetasoundInterfaceIntrinsics
 		{
 			bool IsLessThanVertex(const FMetasoundFrontendVertex& VertexA, const FMetasoundFrontendVertex& VertexB)
 			{
-				// For archetypes we can ignore metadata when comparing
+				// For interfaces we can ignore metadata when comparing
 				if (VertexA.Name == VertexB.Name)
 				{
 					return VertexA.TypeName.FastLess(VertexB.TypeName);
@@ -22,9 +22,9 @@ namespace Metasound
 				return VertexA.Name.FastLess(VertexB.Name);
 			}
 
-			bool IsLessThanEnvironmentVariable(const FMetasoundFrontendEnvironmentVariable& EnvironmentA, const FMetasoundFrontendEnvironmentVariable& EnvironmentB)
+			bool IsLessThanEnvironmentVariable(const FMetasoundFrontendClassEnvironmentVariable& EnvironmentA, const FMetasoundFrontendClassEnvironmentVariable& EnvironmentB)
 			{
-				// For archetypes we can ignore metadata when comparing
+				// For interfaces we can ignore metadata when comparing
 				if (EnvironmentA.Name == EnvironmentB.Name)
 				{
 					return EnvironmentA.TypeName.FastLess(EnvironmentB.TypeName);
@@ -32,13 +32,12 @@ namespace Metasound
 				return EnvironmentA.Name.FastLess(EnvironmentB.Name);
 			}
 
-			bool IsEquivalentEnvironmentVariable(const FMetasoundFrontendEnvironmentVariable& EnvironmentA, const FMetasoundFrontendEnvironmentVariable& EnvironmentB)
+			bool IsEquivalentEnvironmentVariable(const FMetasoundFrontendClassEnvironmentVariable& EnvironmentA, const FMetasoundFrontendClassEnvironmentVariable& EnvironmentB)
 			{
-				// For archetypes we can ignore display name and tooltip
+				// For interfaces we can ignore display name and tooltip
 				return (EnvironmentA.Name == EnvironmentB.Name) && (EnvironmentA.TypeName == EnvironmentB.TypeName);
 			}
 
-
 			template<typename Type>
 			TArray<const Type*> MakePointerToArrayElements(const TArray<Type>& InArray)
 			{
@@ -127,96 +126,96 @@ namespace Metasound
 			}
 		}
 
-		bool IsSubsetOfArchetype(const FMetasoundFrontendArchetype& InSubsetArchetype, const FMetasoundFrontendArchetype& InSupersetArchetype)
+		bool IsSubsetOfInterface(const FMetasoundFrontendInterface& InSubsetInterface, const FMetasoundFrontendInterface& InSupersetInterface)
 		{
-			using namespace MetasoundArchetypeIntrinsics;
+			using namespace MetasoundInterfaceIntrinsics;
 
-			const bool bIsInputSetIncluded = IsSetIncluded(InSubsetArchetype.Interface.Inputs, InSupersetArchetype.Interface.Inputs, &IsLessThanVertex);
+			const bool bIsInputSetIncluded = IsSetIncluded(InSubsetInterface.Inputs, InSupersetInterface.Inputs, &IsLessThanVertex);
 
-			const bool bIsOutputSetIncluded = IsSetIncluded(InSubsetArchetype.Interface.Outputs, InSupersetArchetype.Interface.Outputs, &IsLessThanVertex);
+			const bool bIsOutputSetIncluded = IsSetIncluded(InSubsetInterface.Outputs, InSupersetInterface.Outputs, &IsLessThanVertex);
 
-			const bool bIsEnvironmentVariableSetIncluded = IsSetIncluded(InSubsetArchetype.Interface.Environment, InSupersetArchetype.Interface.Environment, &IsLessThanEnvironmentVariable);
+			const bool bIsEnvironmentVariableSetIncluded = IsSetIncluded(InSubsetInterface.Environment, InSupersetInterface.Environment, &IsLessThanEnvironmentVariable);
 
 			return bIsInputSetIncluded && bIsOutputSetIncluded && bIsEnvironmentVariableSetIncluded;
 		}
 
-		bool IsSubsetOfClass(const FMetasoundFrontendArchetype& InSubsetArchetype, const FMetasoundFrontendClass& InSupersetClass)
+		bool IsSubsetOfClass(const FMetasoundFrontendInterface& InSubsetInterface, const FMetasoundFrontendClass& InSupersetClass)
 		{
 			// TODO: Environment variables are ignored as they are poorly supported and classes describe which environment variables are required,
 			// not which are supported 
-			using namespace MetasoundArchetypeIntrinsics;
-
-			const bool bIsInputSetIncluded = IsSetIncluded(InSubsetArchetype.Interface.Inputs, InSupersetClass.Interface.Inputs, &IsLessThanVertex);
+			using namespace MetasoundInterfaceIntrinsics;
 
-			const bool bIsOutputSetIncluded = IsSetIncluded(InSubsetArchetype.Interface.Outputs, InSupersetClass.Interface.Outputs, &IsLessThanVertex);
+			const bool bIsInputSetIncluded = IsSetIncluded(InSubsetInterface.Inputs, InSupersetClass.Interface.Inputs, &IsLessThanVertex);
+			const bool bIsOutputSetIncluded = IsSetIncluded(InSubsetInterface.Outputs, InSupersetClass.Interface.Outputs, &IsLessThanVertex);
 
 			return bIsInputSetIncluded && bIsOutputSetIncluded;
 		}
 
-		bool IsEquivalentArchetype(const FMetasoundFrontendArchetype& InInputArchetype, const FMetasoundFrontendArchetype& InTargetArchetype)
+		bool IsEquivalentInterface(const FMetasoundFrontendInterface& InInputInterface, const FMetasoundFrontendInterface& InTargetInterface)
 		{
-			using namespace MetasoundArchetypeIntrinsics;
-
-			const bool bIsInputSetEquivalent = IsSetEquivalent(InTargetArchetype.Interface.Inputs, InInputArchetype.Interface.Inputs, &IsLessThanVertex);
+			using namespace MetasoundInterfaceIntrinsics;
 
-			const bool bIsOutputSetEquivalent = IsSetEquivalent(InTargetArchetype.Interface.Outputs, InInputArchetype.Interface.Outputs, &IsLessThanVertex);
-
-			const bool bIsEnvironmentVariableSetEquivalent = IsSetEquivalent(InTargetArchetype.Interface.Environment, InInputArchetype.Interface.Environment, &IsLessThanEnvironmentVariable);
+			const bool bIsInputSetEquivalent = IsSetEquivalent(InTargetInterface.Inputs, InInputInterface.Inputs, &IsLessThanVertex);
+			const bool bIsOutputSetEquivalent = IsSetEquivalent(InTargetInterface.Outputs, InInputInterface.Outputs, &IsLessThanVertex);
+			const bool bIsEnvironmentVariableSetEquivalent = IsSetEquivalent(InTargetInterface.Environment, InInputInterface.Environment, &IsLessThanEnvironmentVariable);
 
 			return bIsInputSetEquivalent && bIsOutputSetEquivalent && bIsEnvironmentVariableSetEquivalent;
 		}
 
-		bool IsEqualArchetype(const FMetasoundFrontendArchetype& InInputArchetype, const FMetasoundFrontendArchetype& InTargetArchetype)
+		bool IsEqualInterface(const FMetasoundFrontendInterface& InInputInterface, const FMetasoundFrontendInterface& InTargetInterface)
 		{
-			bool bIsMetadataEqual = InInputArchetype.Version.Name == InTargetArchetype.Version.Name;
-			bIsMetadataEqual &= InInputArchetype.Version.Number.Major == InTargetArchetype.Version.Number.Major;
-			bIsMetadataEqual &= InInputArchetype.Version.Number.Minor == InTargetArchetype.Version.Number.Minor;
-
+			const bool bIsMetadataEqual = InInputInterface.Version == InTargetInterface.Version;
 			if (bIsMetadataEqual)
 			{
 				// If metadata is equal, then inputs/outputs/environment variables should also be equal. 
-				// Consider updating archetype version info if this ensure is hit.
-				return ensure(IsEquivalentArchetype(InInputArchetype, InTargetArchetype));
+				// Consider updating interface version info if this ensure is hit.
+				return ensure(IsEquivalentInterface(InInputInterface, InTargetInterface));
 			}
 
 			return false;
 		}
 
-		int32 InputOutputDifferenceCount(const FMetasoundFrontendClass& InClass, const FMetasoundFrontendArchetype& InArchetype)
+		bool DeclaredInterfaceVersionsMatch(const FMetasoundFrontendDocument& InDocumentA, const FMetasoundFrontendDocument& InDocumentB)
 		{
-			using namespace MetasoundArchetypeIntrinsics;
+			using namespace MetasoundInterfaceIntrinsics;
 
-			int32 DiffCount = SetDifferenceCount(InClass.Interface.Inputs, InArchetype.Interface.Inputs, &IsLessThanVertex);
+			auto IsLessThanVersion = [](const FMetasoundFrontendVersion& A, const FMetasoundFrontendVersion& B) { return A < B; };
+			return IsSetEquivalent(InDocumentA.InterfaceVersions, InDocumentB.InterfaceVersions, IsLessThanVersion);
+		}
 
-			DiffCount += SetDifferenceCount(InClass.Interface.Outputs, InArchetype.Interface.Outputs, &IsLessThanVertex);
+		int32 InputOutputDifferenceCount(const FMetasoundFrontendClass& InClass, const FMetasoundFrontendInterface& InInterface)
+		{
+			using namespace MetasoundInterfaceIntrinsics;
+
+			int32 DiffCount = SetDifferenceCount(InClass.Interface.Inputs, InInterface.Inputs, &IsLessThanVertex);
+			DiffCount += SetDifferenceCount(InClass.Interface.Outputs, InInterface.Outputs, &IsLessThanVertex);
 
 			return DiffCount;
 		}
 
-		int32 InputOutputDifferenceCount(const FMetasoundFrontendArchetype& InArchetypeA, const FMetasoundFrontendArchetype& InArchetypeB)
+		int32 InputOutputDifferenceCount(const FMetasoundFrontendInterface& InInterfaceA, const FMetasoundFrontendInterface& InInterfaceB)
 		{
-			using namespace MetasoundArchetypeIntrinsics;
-
-			int32 DiffCount = SetDifferenceCount(InArchetypeA.Interface.Inputs, InArchetypeB.Interface.Inputs, &IsLessThanVertex);
+			using namespace MetasoundInterfaceIntrinsics;
 
-			DiffCount += SetDifferenceCount(InArchetypeA.Interface.Outputs, InArchetypeB.Interface.Outputs, &IsLessThanVertex);
+			int32 DiffCount = SetDifferenceCount(InInterfaceA.Inputs, InInterfaceB.Inputs, &IsLessThanVertex);
+			DiffCount += SetDifferenceCount(InInterfaceA.Outputs, InInterfaceB.Outputs, &IsLessThanVertex);
 
 			return DiffCount;
 		}
 
-		void GatherRequiredEnvironmentVariables(const FMetasoundFrontendGraphClass& InRootGraph, const TArray<FMetasoundFrontendClass>& InDependencies, const TArray<FMetasoundFrontendGraphClass>& InSubgraphs, TArray<FMetasoundFrontendEnvironmentVariable>& OutEnvironmentVariables)
+		void GatherRequiredEnvironmentVariables(const FMetasoundFrontendGraphClass& InRootGraph, const TArray<FMetasoundFrontendClass>& InDependencies, const TArray<FMetasoundFrontendGraphClass>& InSubgraphs, TArray<FMetasoundFrontendClassEnvironmentVariable>& OutEnvironmentVariables)
 		{
 			auto GatherRequiredEnvironmentVariablesFromClass = [&](const FMetasoundFrontendClass& InClass)
 			{
-				using namespace MetasoundArchetypeIntrinsics;
+				using namespace MetasoundInterfaceIntrinsics;
 
 				for (const FMetasoundFrontendClassEnvironmentVariable& EnvVar : InClass.Interface.Environment)
 				{
 					if (EnvVar.bIsRequired)
 					{
-						auto IsEquivalentEnvVar = [&](const FMetasoundFrontendEnvironmentVariable& OtherEnvVar) 
-						{ 
-							return IsEquivalentEnvironmentVariable(EnvVar, OtherEnvVar); 
+						auto IsEquivalentEnvVar = [&](const FMetasoundFrontendClassEnvironmentVariable& OtherEnvVar)
+						{
+							return IsEquivalentEnvironmentVariable(EnvVar, OtherEnvVar);
 						};
 
 						// Basically same as TArray::AddUnique except uses the `IsEquivalentEnvironmentVariable` instead of `operator==` to test for uniqueness.
@@ -233,45 +232,44 @@ namespace Metasound
 			Algo::ForEach(InSubgraphs, GatherRequiredEnvironmentVariablesFromClass);
 		}
 
-		
-		const FMetasoundFrontendArchetype* FindMostSimilarArchetypeSupportingEnvironment(const FMetasoundFrontendGraphClass& InRootGraph, const TArray<FMetasoundFrontendClass>& InDependencies, const TArray<FMetasoundFrontendGraphClass>& InSubgraphs, const TArray<FMetasoundFrontendArchetype>& InCandidateArchetypes)
+		const FMetasoundFrontendInterface* FindMostSimilarInterfaceSupportingEnvironment(const FMetasoundFrontendGraphClass& InRootGraph, const TArray<FMetasoundFrontendClass>& InDependencies, const TArray<FMetasoundFrontendGraphClass>& InSubgraphs, const TArray<FMetasoundFrontendInterface>& InCandidateInterfaces)
 		{
-			using namespace MetasoundArchetypeIntrinsics;
+			using namespace MetasoundInterfaceIntrinsics;
 
-			const FMetasoundFrontendArchetype* Result = nullptr;
+			const FMetasoundFrontendInterface* Result = nullptr;
 
-			TArray<const FMetasoundFrontendArchetype*> CandidateArchs = MakePointerToArrayElements(InCandidateArchetypes);
+			TArray<const FMetasoundFrontendInterface*> CandidateInterfaces = MakePointerToArrayElements(InCandidateInterfaces);
 
-			TArray<FMetasoundFrontendEnvironmentVariable> AllEnvironmentVariables;
+			TArray<FMetasoundFrontendClassEnvironmentVariable> AllEnvironmentVariables;
 			GatherRequiredEnvironmentVariables(InRootGraph, InDependencies, InSubgraphs, AllEnvironmentVariables);
 
-			// Remove all archetypes which do not provide all environment variables. 
-			auto DoesNotProvideAllEnvironmentVariables = [&](const FMetasoundFrontendArchetype* CandidateArch)
+			// Remove all interfaces which do not provide all environment variables.
+			auto DoesNotProvideAllEnvironmentVariables = [&](const FMetasoundFrontendInterface* CandidateInterface)
 			{
-				return !IsSetIncluded(AllEnvironmentVariables, CandidateArch->Interface.Environment, &IsLessThanEnvironmentVariable);
+				return !IsSetIncluded(AllEnvironmentVariables, CandidateInterface->Environment, &IsLessThanEnvironmentVariable);
 			};
 
-			CandidateArchs.RemoveAll(DoesNotProvideAllEnvironmentVariables);
+			CandidateInterfaces.RemoveAll(DoesNotProvideAllEnvironmentVariables);
 
-			// Return the archetype with the least amount of differences.
-			auto DifferencesFromRootGraph = [&](const FMetasoundFrontendArchetype* CandidateArch) -> int32
+			// Return the interface with the least amount of differences.
+			auto DifferencesFromRootGraph = [&](const FMetasoundFrontendInterface* CandidateInterface) -> int32
 			{ 
-				return InputOutputDifferenceCount(InRootGraph, *CandidateArch); 
+				return InputOutputDifferenceCount(InRootGraph, *CandidateInterface);
 			};
 
-			Algo::SortBy(CandidateArchs, DifferencesFromRootGraph);
+			Algo::SortBy(CandidateInterfaces, DifferencesFromRootGraph);
 
-			if (0 == CandidateArchs.Num())
+			if (0 == CandidateInterfaces.Num())
 			{
 				return nullptr;
 			}
 
-			return CandidateArchs[0];
+			return CandidateInterfaces[0];
 		}
 		
-		const FMetasoundFrontendArchetype* FindMostSimilarArchetypeSupportingEnvironment(const FMetasoundFrontendDocument& InDocument, const TArray<FMetasoundFrontendArchetype>& InCandidateArchetypes)
+		const FMetasoundFrontendInterface* FindMostSimilarInterfaceSupportingEnvironment(const FMetasoundFrontendDocument& InDocument, const TArray<FMetasoundFrontendInterface>& InCandidateInterfaces)
 		{
-			return FindMostSimilarArchetypeSupportingEnvironment(InDocument.RootGraph, InDocument.Dependencies, InDocument.Subgraphs, InCandidateArchetypes);
+			return FindMostSimilarInterfaceSupportingEnvironment(InDocument.RootGraph, InDocument.Dependencies, InDocument.Subgraphs, InCandidateInterfaces);
 		}
 	}
 }
diff --git a/Engine/Plugins/Runtime/Metasound/Source/MetasoundFrontend/Private/MetasoundAssetBase.cpp b/Engine/Plugins/Runtime/Metasound/Source/MetasoundFrontend/Private/MetasoundAssetBase.cpp
--- a/Engine/Plugins/Runtime/Metasound/Source/MetasoundFrontend/Private/MetasoundAssetBase.cpp
+++ b/Engine/Plugins/Runtime/Metasound/Source/MetasoundFrontend/Private/MetasoundAssetBase.cpp
@@ -6,6 +6,7 @@
 #include "Algo/Transform.h"
 #include "Containers/Set.h"
 #include "HAL/FileManager.h"
+#include "IAudioParameterTransmitter.h"
 #include "Internationalization/Text.h"
 #include "IStructSerializerBackend.h"
 #include "Logging/LogMacros.h"
@@ -271,75 +272,40 @@ void FMetasoundAssetBase::SetMetadata(FMetasoundFrontendClassMetadata& InMetadat
 	MarkMetasoundDocumentDirty();
 }
 
-bool FMetasoundAssetBase::IsArchetypeSupported(const FMetasoundFrontendVersion& InArchetypeVersion) const
-{
-	return GetSupportedArchetypeVersions().Contains(InArchetypeVersion);
-}
-
-FMetasoundFrontendVersion FMetasoundAssetBase::GetPreferredArchetypeVersion(const FMetasoundFrontendDocument& InDocument) const
+bool FMetasoundAssetBase::GetDeclaredInterfaces(TArray<const Metasound::Frontend::IInterfaceRegistryEntry*>& OutInterfaces) const
 {
+	using namespace Metasound;
 	using namespace Metasound::Frontend;
 
-	// Default to archetype provided in case it is supported.
-	if (IsArchetypeSupported(InDocument.ArchetypeVersion))
+	if (const FMetasoundFrontendDocument* Document = GetDocument().Get())
 	{
-		return InDocument.ArchetypeVersion;
-	}
+		bool bInterfacesFound = true;
 
-	// If existing archetype is not supported, check if an updated version is preferred.
-	auto IsNewVersion = [&](const FMetasoundFrontendVersion& InVersion)
-	{
-		return (InDocument.ArchetypeVersion.Name == InVersion.Name) && (InVersion.Number > InDocument.ArchetypeVersion.Number);
-	};
-	TArray<FMetasoundFrontendVersion> UpdatedVersions = GetSupportedArchetypeVersions().FilterByPredicate(IsNewVersion);
-	if (UpdatedVersions.Num() > 0)
-	{
-		UpdatedVersions.Sort();
-		return UpdatedVersions.Last();
-	}
-	
-	// If existing archetype is not supported, get the most similar that still supports the documents environment.
-	TArray<FMetasoundFrontendArchetype> Archetypes;
-	for (const FMetasoundFrontendVersion& Version : GetSupportedArchetypeVersions())
-	{
-		const FArchetypeRegistryKey Key = GetArchetypeRegistryKey(Version);
-		FMetasoundFrontendArchetype Archetype;
-		if (IArchetypeRegistry::Get().FindArchetype(Key, Archetype))
+		Algo::Transform(Document->InterfaceVersions, OutInterfaces, [&](const FMetasoundFrontendVersion& Version)
 		{
-			Archetypes.Add(MoveTemp(Archetype));
-		}
-	}
-	const FMetasoundFrontendArchetype* SimilarArchetype = Metasound::Frontend::FindMostSimilarArchetypeSupportingEnvironment(InDocument, Archetypes);
+			const FInterfaceRegistryKey InterfaceKey = GetInterfaceRegistryKey(Version);
+			const IInterfaceRegistryEntry* RegistryEntry = IInterfaceRegistry::Get().FindInterfaceRegistryEntry(InterfaceKey);
+			if (!RegistryEntry)
+			{
+				bInterfacesFound = false;
+				UE_LOG(LogMetaSound, Warning, TEXT("No registered interface matching interface version on document [InterfaceVersion:%s]"), *Version.ToString());
+			}
 
-	if (nullptr != SimilarArchetype)
-	{
-		return SimilarArchetype->Version;
+			return RegistryEntry;
+		});
+	
+		return bInterfacesFound;
 	}
 
-	// Nothing found. Return the existing archetype for the FMetasoundAssetBase.
-	return GetDefaultArchetypeVersion();
+	return false;
 }
 
-bool FMetasoundAssetBase::GetArchetype(FMetasoundFrontendArchetype& OutArchetype) const
+bool FMetasoundAssetBase::IsInterfaceDeclared(FName InName) const
 {
-	using namespace Metasound;
-	using namespace Metasound::Frontend;
-
-	bool bFoundArchetype = false;
-
-	const FMetasoundFrontendDocument* Document = GetDocument().Get();
-	if (nullptr != Document)
+	return Algo::AnyOf(GetDocumentChecked().InterfaceVersions, [&InName](const FMetasoundFrontendVersion& FrontendVersion)
 	{
-		FArchetypeRegistryKey ArchetypeRegistryKey = GetArchetypeRegistryKey(Document->ArchetypeVersion);
-		bFoundArchetype = IArchetypeRegistry::Get().FindArchetype(ArchetypeRegistryKey, OutArchetype);
-
-		if (!bFoundArchetype)
-		{
-			UE_LOG(LogMetaSound, Warning, TEXT("No registered archetype matching archetype version on document [ArchetypeVersion:%s]"), *Document->ArchetypeVersion.ToString());
-		}
-	}
-
-	return bFoundArchetype;
+		return FrontendVersion.Name == InName;
+	});
 }
 
 void FMetasoundAssetBase::SetDocument(const FMetasoundFrontendDocument& InDocument)
@@ -349,22 +315,20 @@ void FMetasoundAssetBase::SetDocument(const FMetasoundFrontendDocument& InDocume
 	MarkMetasoundDocumentDirty();
 }
 
-void FMetasoundAssetBase::ConformDocumentToArchetype()
+void FMetasoundAssetBase::AddDefaultInterfaces()
 {
-	METASOUND_TRACE_CPUPROFILER_EVENT_SCOPE(MetaSoundAssetBase::ConformDocumentToArchetype);
-	FMetasoundFrontendDocument& Doc = GetDocumentChecked();
+	using namespace Metasound::Frontend;
 
-	FMetasoundFrontendVersion PreferredArchetypeVersion = GetPreferredArchetypeVersion(Doc);
+	UObject* OwningAsset = GetOwningAsset();
+	check(OwningAsset);
 
-	if (PreferredArchetypeVersion != Doc.ArchetypeVersion)
-	{
-		Metasound::Frontend::FMatchRootGraphToArchetype Transform(PreferredArchetypeVersion);
-		if (Transform.Transform(GetDocumentHandle()))
-		{
-			ConformObjectDataToArchetype();
-			MarkMetasoundDocumentDirty();
-		}
-	}
+	UClass* AssetClass = OwningAsset->GetClass();
+	check(AssetClass);
+
+	FDocumentHandle DocumentHandle = GetDocumentHandle();
+
+	TArray<FMetasoundFrontendInterface> InitInterfaces = ISearchEngine::Get().FindUClassDefaultInterfaces(AssetClass->GetFName());
+	FModifyRootGraphInterfaces({ }, InitInterfaces).Transform(DocumentHandle);
 }
 
 bool FMetasoundAssetBase::AutoUpdate(bool bInMarkDirty)
@@ -395,18 +359,40 @@ bool FMetasoundAssetBase::VersionAsset()
 		AssetPath = OwningAsset->GetPathName();
 	}
 
+	FMetasoundFrontendDocument* Doc = GetDocument().Get();
+	if (!ensure(Doc))
+	{
+		return false;
+	}
+
+	// Data migration for 5.0 Early Access data. ArchetypeVersion can be removed post 5.0 release.
 	bool bDidEdit = false;
-	FDocumentHandle DocumentHandle = GetDocumentHandle();
+	if (Doc->ArchetypeVersion.IsValid())
+	{
+		Doc->InterfaceVersions.Add(Doc->ArchetypeVersion);
+		Doc->ArchetypeVersion = FMetasoundFrontendVersion::GetInvalid();
+		bDidEdit = true;
+	}
+
+	FDocumentHandle DocHandle = GetDocumentHandle();
+
+	// Version Document Model
+	{
+		bDidEdit |= FVersionDocument(AssetName, AssetPath).Transform(DocHandle);
+	}
 
-	bDidEdit |= FVersionDocument(AssetName, AssetPath).Transform(DocumentHandle);
-	if (FMetasoundFrontendDocument* Doc = GetDocument().Get())
+	// Version Interfaces
 	{
-		FMetasoundFrontendVersion ArchetypeVersion = GetPreferredArchetypeVersion(*Doc);
-		bDidEdit |= FMatchRootGraphToArchetype(ArchetypeVersion).Transform(DocumentHandle);
-		if (bDidEdit)
+		bool bInterfaceUpdated = false;
+		for (const FMetasoundFrontendVersion& Version : Doc->InterfaceVersions)
 		{
-			ConformObjectDataToArchetype();
+			bInterfaceUpdated |= FUpdateRootGraphInterface(Version).Transform(DocHandle);
 		}
+		if (bInterfaceUpdated)
+		{
+			ConformObjectDataToInterfaces();
+		}
+		bDidEdit |= bInterfaceUpdated;
 	}
 
 	return bDidEdit;
@@ -423,6 +409,11 @@ bool FMetasoundAssetBase::GetSynchronizationClearUpdateNotes() const
 	return bSynchronizationClearUpdateNotes;
 }
 
+bool FMetasoundAssetBase::GetSynchronizationInterfacesUpdated() const
+{
+	return bSynchronizationInterfacesUpdated;
+}
+
 void FMetasoundAssetBase::SetSynchronizationRequired()
 {
 	bSynchronizationRequired = true;
@@ -433,9 +424,15 @@ void FMetasoundAssetBase::SetClearNodeNotesOnSynchronization()
 	bSynchronizationClearUpdateNotes = true;
 }
 
+void FMetasoundAssetBase::SetInterfacesUpdatedOnSynchronization()
+{
+	bSynchronizationInterfacesUpdated = true;
+}
+
 void FMetasoundAssetBase::ResetSynchronizationState()
 {
 	bSynchronizationClearUpdateNotes = false;
+	bSynchronizationInterfacesUpdated = false;
 	bSynchronizationRequired = false;
 }
 #endif // WITH_EDITORONLY_DATA
@@ -447,19 +444,16 @@ TSharedPtr<Metasound::IGraph, ESPMode::ThreadSafe> FMetasoundAssetBase::BuildMet
 
 	METASOUND_TRACE_CPUPROFILER_EVENT_SCOPE(MetaSoundAssetBase::BuildMetasoundDocument);
 
-	const FMetasoundFrontendDocument* Doc = GetDocument().Get();
-	if (nullptr == Doc)
-	{
-		UE_LOG(LogMetaSound, Error, TEXT("Cannot create graph. Null MetaSound document in MetaSound asset [Name:%s]"), *GetOwningAssetName());
-		return TSharedPtr<IGraph, ESPMode::ThreadSafe>(nullptr);
-	}
-
 	// Create graph which can spawn instances. TODO: cache graph.
-	TUniquePtr<FFrontendGraph> FrontendGraph = FFrontendGraphBuilder::CreateGraph(*Doc);
+	TUniquePtr<FFrontendGraph> FrontendGraph = FFrontendGraphBuilder::CreateGraph(GetDocumentChecked());
 	if (FrontendGraph.IsValid())
 	{
-		TSet<FVertexName> VerticesToSkip = GetNonTransmittableInputVertices(*Doc);
-		bool bSuccessfullyInjectedReceiveNodes = InjectReceiveNodes(*FrontendGraph, FMetaSoundParameterTransmitter::CreateSendAddressFromEnvironment, VerticesToSkip);
+		const TArray<const FMetasoundFrontendClassInput*> TransmittableInputs = GetTransmittableClassInputs();
+
+		TSet<FVertexName> TransmittableInputNames;
+		Algo::Transform(TransmittableInputs, TransmittableInputNames, [](const FMetasoundFrontendClassInput* Input) { return Input->Name; });
+
+		bool bSuccessfullyInjectedReceiveNodes = InjectReceiveNodes(*FrontendGraph, FMetaSoundParameterTransmitter::CreateSendAddressFromEnvironment, TransmittableInputNames);
 		if (!bSuccessfullyInjectedReceiveNodes)
 		{
 			UE_LOG(LogMetaSound, Error, TEXT("Error while injecting async communication hooks. Instance communication may not function properly [Name:%s]."), *GetOwningAssetName());
@@ -584,101 +578,6 @@ TArray<FMetasoundAssetBase::FSendInfoAndVertexName> FMetasoundAssetBase::GetSend
 	return SendInfos;
 }
 
-TSet<Metasound::FVertexName> FMetasoundAssetBase::GetNonTransmittableInputVertices(const FMetasoundFrontendDocument& InDoc) const
-{
-	using namespace Metasound;
-	using namespace Metasound::Frontend;
-
-	check(IsInGameThread() || IsInAudioThread());
-
-	TSet<FVertexName> NonTransmittableInputVertices;
-
-	auto GetVertexKey = [](const FMetasoundFrontendClassVertex& InVertex) { return InVertex.Name; };
-
-	// Do not transmit vertices defined in archetype. Those are already accounted
-	// for by owning object.
-	FMetasoundFrontendArchetype Archetype;
-	GetArchetype(Archetype);
-	Algo::Transform(Archetype.Interface.Inputs, NonTransmittableInputVertices, GetVertexKey);
-
-
-	// Do not transmit vertices which are not transmittable. Async communication 
-	// is not supported without transmission.
-	const IDataTypeRegistry& Registry = IDataTypeRegistry::Get();
-
-	auto IsNotTransmittable = [&Registry](const FMetasoundFrontendClassVertex& InVertex) -> bool
-	{	
-		FDataTypeRegistryInfo Info;
-		if (Registry.GetDataTypeInfo(InVertex.TypeName, Info))
-		{
-			return !Info.bIsTransmittable;
-		}
-		return true;
-	};
-
-	Algo::TransformIf(InDoc.RootGraph.Interface.Inputs, NonTransmittableInputVertices, IsNotTransmittable, GetVertexKey);
-
-	return NonTransmittableInputVertices;
-}
-
-TArray<Metasound::FVertexName> FMetasoundAssetBase::GetTransmittableInputVertexNames() const
-{
-	using namespace Metasound;
-	using namespace Metasound::Frontend;
-
-	FConstDocumentHandle Document = GetDocumentHandle();
-
-	// Find the archetype for the document.
-	FMetasoundFrontendArchetype Archetype;
-	FArchetypeRegistryKey ArchetypeRegistryKey = Frontend::GetArchetypeRegistryKey(Document->GetArchetypeVersion());
-	bool bFoundArchetype = IArchetypeRegistry::Get().FindArchetype(ArchetypeRegistryKey, Archetype);
-
-	if (!bFoundArchetype)
-	{
-		UE_LOG(LogMetaSound, Warning, TEXT("No registered archetype matching archetype version on document [ArchetypeVersion:%s]"), *Document->GetArchetypeVersion().ToString());
-	}
-
-	// Unused inputs are all input vertices that are not in the archetype.
-	TArray<FVertexName> ArchetypeInputVertexNames;
-	for (const FMetasoundFrontendClassVertex& Vertex : Archetype.Interface.Inputs)
-	{
-		ArchetypeInputVertexNames.Add(Vertex.Name);
-	}
-
-	Frontend::FConstGraphHandle RootGraph = Document->GetRootGraph();
-	TArray<FVertexName> GraphInputVertexNames = RootGraph->GetInputVertexNames();
-
-	// Filter graph inputs by archetype inputs.
-	GraphInputVertexNames = GraphInputVertexNames.FilterByPredicate([&](const FVertexName& InName) { return !ArchetypeInputVertexNames.Contains(InName); });
-
-	auto IsDataTypeTransmittable = [&](const FVertexName& InVertexName)
-	{
-		using namespace Metasound::Frontend;
-
-		Frontend::FConstClassInputAccessPtr ClassInputPtr = RootGraph->FindClassInputWithName(InVertexName);
-		if (const FMetasoundFrontendClassInput* ClassInput = ClassInputPtr.Get())
-		{
-			Frontend::FDataTypeRegistryInfo TypeInfo;
-			if (IDataTypeRegistry::Get().GetDataTypeInfo(ClassInput->TypeName, TypeInfo))
-			{
-				if (TypeInfo.bIsTransmittable)
-				{
-					Frontend::FConstNodeHandle InputNode = RootGraph->GetNodeWithID(ClassInput->NodeID);
-					if (InputNode->IsValid())
-					{
-						return true;
-					}
-				}
-			}
-		}
-		return false;
-	};
-
-	GraphInputVertexNames = GraphInputVertexNames.FilterByPredicate(IsDataTypeTransmittable);
-
-	return GraphInputVertexNames;
-}
-
 Metasound::Frontend::FNodeHandle FMetasoundAssetBase::AddInputPinForSendAddress(const Metasound::FMetaSoundParameterTransmitter::FSendInfo& InSendInfo, Metasound::Frontend::FGraphHandle InGraph) const
 {
 	FMetasoundFrontendClassInput Description;
@@ -806,52 +705,50 @@ TSharedPtr<const Metasound::IGraph, ESPMode::ThreadSafe> FMetasoundAssetBase::Ge
 	return FMetasoundAssetBase::GetRuntimeData().Graph;
 }
 
-TArray<FMetasoundFrontendClassInput> FMetasoundAssetBase::GetTransmittableClassInputs() const
+TArray<const FMetasoundFrontendClassInput*> FMetasoundAssetBase::GetTransmittableClassInputs() const
 {
 	using namespace Metasound;
 	using namespace Metasound::Frontend;
+
 	check(IsInGameThread() || IsInAudioThread());
+	TArray<const FMetasoundFrontendClassInput*> Inputs;
 
-	TArray<FMetasoundFrontendClassInput> Inputs;
+	const FMetasoundFrontendDocument& Doc = GetDocumentChecked();
+	auto GetInputName = [](const FMetasoundFrontendClassInput& InInput) { return InInput.Name; };
 
-	if (const FMetasoundFrontendDocument* Doc = GetDocument().Get())
+	// Do not transmit vertices defined in interface marked as non-transmittable
+	TArray<const IInterfaceRegistryEntry*> Interfaces;
+	TSet<FVertexName> NonTransmittableInputs;
+	GetDeclaredInterfaces(Interfaces);
+	for (const IInterfaceRegistryEntry* InterfaceEntry : Interfaces)
 	{
-		TSet<FVertexName> ArchetypeInputs;
-
-		auto GetVertexKey = [](const FMetasoundFrontendClassVertex& InVertex) { return InVertex.Name; };
-
-		// Do not transmit vertices defined in archetype. Those are already accounted
-		// for by owning object.
-		FMetasoundFrontendArchetype Archetype;
-		GetArchetype(Archetype);
-		Algo::Transform(Archetype.Interface.Inputs, ArchetypeInputs, GetVertexKey);
-
-		// Do not transmit vertices which are not transmittable. Async communication 
-		// is not supported without transmission.
-		IDataTypeRegistry& Registry = IDataTypeRegistry::Get();
-
-		auto IsTransmittable = [&Registry, &ArchetypeInputs](const FMetasoundFrontendClassVertex& InVertex) -> bool
-		{	
-			if (!ArchetypeInputs.Contains(InVertex.Name))
+		if (InterfaceEntry)
+		{
+			if (InterfaceEntry->GetRouterName() != Audio::IParameterTransmitter::RouterName)
 			{
-				FDataTypeRegistryInfo Info;
-				if (Registry.GetDataTypeInfo(InVertex.TypeName, Info))
-				{
-					return Info.bIsTransmittable;
-				}
+				const FMetasoundFrontendInterface& Interface = InterfaceEntry->GetInterface();
+				Algo::Transform(Interface.Inputs, NonTransmittableInputs, GetInputName);
 			}
-			return false;
-		};
-
+		}
+	}
 
-		for (const FMetasoundFrontendClassInput& Vertex : Doc->RootGraph.Interface.Inputs)
+	// Do not transmit vertices which are not transmittable. Async communication
+	// is not supported without transmission.
+	IDataTypeRegistry& Registry = IDataTypeRegistry::Get();
+	auto IsTransmittable = [&Registry, &NonTransmittableInputs](const FMetasoundFrontendClassVertex& InVertex)
+	{
+		if (!NonTransmittableInputs.Contains(InVertex.Name))
 		{
-			if (IsTransmittable(Vertex))
+			FDataTypeRegistryInfo Info;
+			if (Registry.GetDataTypeInfo(InVertex.TypeName, Info))
 			{
-				Inputs.Add(Vertex);
+				return Info.bIsTransmittable;
 			}
 		}
-	}
+
+		return false;
+	};
+	Algo::TransformIf(Doc.RootGraph.Interface.Inputs, Inputs, IsTransmittable, [] (const FMetasoundFrontendClassInput& Input) { return &Input; });
 
 	return Inputs;
 }
@@ -869,7 +766,10 @@ const FMetasoundAssetBase::FRuntimeData& FMetasoundAssetBase::GetRuntimeData() c
 	if (CachedRuntimeData.ChangeID != CurrentCachedRuntimeDataChangeID)
 	{
 		// Update CachedRuntimeData.
-		CachedRuntimeData.TransmittableInputs = GetTransmittableClassInputs();
+		CachedRuntimeData.TransmittableInputs.Reset();
+		TArray<const FMetasoundFrontendClassInput*> ClassInputs = GetTransmittableClassInputs();
+		Algo::Transform(ClassInputs, CachedRuntimeData.TransmittableInputs, [] (const FMetasoundFrontendClassInput* Input) { return *Input; });
+
 		CachedRuntimeData.Graph = BuildMetasoundDocument();
 		CachedRuntimeData.ChangeID = CurrentCachedRuntimeDataChangeID;
 	}
diff --git a/Engine/Plugins/Runtime/Metasound/Source/MetasoundFrontend/Private/MetasoundFrontendArchetypeRegistry.cpp b/Engine/Plugins/Runtime/Metasound/Source/MetasoundFrontend/Private/MetasoundFrontendArchetypeRegistry.cpp
--- a/Engine/Plugins/Runtime/Metasound/Source/MetasoundFrontend/Private/MetasoundFrontendArchetypeRegistry.cpp
+++ b/Engine/Plugins/Runtime/Metasound/Source/MetasoundFrontend/Private/MetasoundFrontendArchetypeRegistry.cpp
@@ -9,111 +9,107 @@ namespace Metasound
 	{
 		namespace MetasoundFrontendArchetypeRegistryPrivate
 		{
-			class FArchetypeRegistry : public IArchetypeRegistry
+			class FInterfaceRegistry : public IInterfaceRegistry
 			{
 			public:
-				virtual ~FArchetypeRegistry() = default;
+				virtual ~FInterfaceRegistry() = default;
 
-				virtual FArchetypeRegistryKey RegisterArchetype(TUniquePtr<IArchetypeRegistryEntry>&& InEntry) override
+				virtual void RegisterInterface(TUniquePtr<IInterfaceRegistryEntry>&& InEntry) override
 				{
 					if (InEntry.IsValid())
 					{
-						FArchetypeRegistryKey Key = GetArchetypeRegistryKey(InEntry->GetArchetype());
-						if (IsValidArchetypeRegistryKey(Key))
+						FInterfaceRegistryKey Key = GetInterfaceRegistryKey(InEntry->GetInterface());
+						if (IsValidInterfaceRegistryKey(Key))
 						{
-							if (const IArchetypeRegistryEntry* Entry = FindArchetypeRegistryEntry(Key))
+							if (const IInterfaceRegistryEntry* Entry = FindInterfaceRegistryEntry(Key))
 							{
-								UE_LOG(LogMetaSound, Warning, TEXT("Registration of archetype overwriting previously registered archetype [RegistryKey: %s]"), *Key);
+								UE_LOG(LogMetaSound, Warning, TEXT("Registration of interface overwriting previously registered interface [RegistryKey: %s]"), *Key);
 								
-								FArchetypeRegistryTransaction Transaction{FArchetypeRegistryTransaction::ETransactionType::ArchetypeUnregistration, Key, Entry->GetArchetype().Version};
+								FInterfaceRegistryTransaction Transaction{FInterfaceRegistryTransaction::ETransactionType::InterfaceUnregistration, Key, Entry->GetInterface().Version};
 								Transactions.Add(MoveTemp(Transaction));
 							}
 							
-							FArchetypeRegistryTransaction Transaction{FArchetypeRegistryTransaction::ETransactionType::ArchetypeRegistration, Key, InEntry->GetArchetype().Version};
+							FInterfaceRegistryTransaction Transaction{FInterfaceRegistryTransaction::ETransactionType::InterfaceRegistration, Key, InEntry->GetInterface().Version};
 							Transactions.Add(MoveTemp(Transaction));
 
-							Entries.Add(Key, MoveTemp(InEntry));
-
-							return Key;
+							Entries.Add(Key, MoveTemp(InEntry)).Get();
 						}
 					}
-
-					return FArchetypeRegistryKey{};
 				}
 
-				virtual const IArchetypeRegistryEntry* FindArchetypeRegistryEntry(const FArchetypeRegistryKey& InKey) const override
+				virtual const IInterfaceRegistryEntry* FindInterfaceRegistryEntry(const FInterfaceRegistryKey& InKey) const override
 				{
-					if (const TUniquePtr<IArchetypeRegistryEntry>* Entry = Entries.Find(InKey))
+					if (const TUniquePtr<IInterfaceRegistryEntry>* Entry = Entries.Find(InKey))
 					{
 						return Entry->Get();
 					}
 					return nullptr;
 				}
 
-				virtual bool FindArchetype(const FArchetypeRegistryKey& InKey, FMetasoundFrontendArchetype& OutArchetype) const override
+				virtual bool FindInterface(const FInterfaceRegistryKey& InKey, FMetasoundFrontendInterface& OutInterface) const override
 				{
-					if (const IArchetypeRegistryEntry* Entry = FindArchetypeRegistryEntry(InKey))
+					if (const IInterfaceRegistryEntry* Entry = FindInterfaceRegistryEntry(InKey))
 					{
-						OutArchetype = Entry->GetArchetype();
+						OutInterface = Entry->GetInterface();
 						return true;
 					}
 
 					return false;
 				}
 
-				virtual void ForEachRegistryTransactionSince(FRegistryTransactionID InSince, FRegistryTransactionID* OutCurrentRegistryTransactionID, TFunctionRef<void(const FArchetypeRegistryTransaction&)> InFunc) const override
+				virtual void ForEachRegistryTransactionSince(FRegistryTransactionID InSince, FRegistryTransactionID* OutCurrentRegistryTransactionID, TFunctionRef<void(const FInterfaceRegistryTransaction&)> InFunc) const override
 				{
 					Transactions.ForEachTransactionSince(InSince, OutCurrentRegistryTransactionID, InFunc);
 				}
 
 			private:
 
-				TMap<FArchetypeRegistryKey, TUniquePtr<IArchetypeRegistryEntry>> Entries;
-				TRegistryTransactionHistory<FArchetypeRegistryTransaction> Transactions;
+				TMap<FInterfaceRegistryKey, TUniquePtr<IInterfaceRegistryEntry>> Entries;
+				TRegistryTransactionHistory<FInterfaceRegistryTransaction> Transactions;
 			};
 		}
 
-		bool IsValidArchetypeRegistryKey(const FArchetypeRegistryKey& InKey)
+		bool IsValidInterfaceRegistryKey(const FInterfaceRegistryKey& InKey)
 		{
 			return !InKey.IsEmpty();
 		}
 
-		FArchetypeRegistryKey GetArchetypeRegistryKey(const FMetasoundFrontendVersion& InArchetypeVersion)
+		FInterfaceRegistryKey GetInterfaceRegistryKey(const FMetasoundFrontendVersion& InInterfaceVersion)
 		{
-			return FString::Format(TEXT("{0}_{1}.{2}"), { InArchetypeVersion.Name.ToString(), InArchetypeVersion.Number.Major, InArchetypeVersion.Number.Minor });
+			return FString::Format(TEXT("{0}_{1}.{2}"), { InInterfaceVersion.Name.ToString(), InInterfaceVersion.Number.Major, InInterfaceVersion.Number.Minor });
 
 		}
 
-		FArchetypeRegistryKey GetArchetypeRegistryKey(const FMetasoundFrontendArchetype& InArchetype)
+		FInterfaceRegistryKey GetInterfaceRegistryKey(const FMetasoundFrontendInterface& InInterface)
 		{
-			return GetArchetypeRegistryKey(InArchetype.Version);
+			return GetInterfaceRegistryKey(InInterface.Version);
 		}
 
-		FArchetypeRegistryTransaction::FArchetypeRegistryTransaction(ETransactionType InType, const FArchetypeRegistryKey& InKey, const FMetasoundFrontendVersion& InArchetypeVersion)
+		FInterfaceRegistryTransaction::FInterfaceRegistryTransaction(ETransactionType InType, const FInterfaceRegistryKey& InKey, const FMetasoundFrontendVersion& InInterfaceVersion)
 		: Type(InType)
 		, Key(InKey)
-		, ArchetypeVersion(InArchetypeVersion)
+		, InterfaceVersion(InInterfaceVersion)
 		{
 		}
 
-		FArchetypeRegistryTransaction::ETransactionType FArchetypeRegistryTransaction::GetTransactionType() const
+		FInterfaceRegistryTransaction::ETransactionType FInterfaceRegistryTransaction::GetTransactionType() const
 		{
 			return Type;
 		}
 
-		const FMetasoundFrontendVersion& FArchetypeRegistryTransaction::GetArchetypeVersion() const
+		const FMetasoundFrontendVersion& FInterfaceRegistryTransaction::GetInterfaceVersion() const
 		{
-			return ArchetypeVersion;
+			return InterfaceVersion;
 		}
 
-		const FArchetypeRegistryKey& FArchetypeRegistryTransaction::GetArchetypeRegistryKey() const
+		const FInterfaceRegistryKey& FInterfaceRegistryTransaction::GetInterfaceRegistryKey() const
 		{
 			return Key;
 		}
 
-		IArchetypeRegistry& IArchetypeRegistry::Get()
+		IInterfaceRegistry& IInterfaceRegistry::Get()
 		{
-			static MetasoundFrontendArchetypeRegistryPrivate::FArchetypeRegistry Registry;
+			static MetasoundFrontendArchetypeRegistryPrivate::FInterfaceRegistry Registry;
 			return Registry;
 		}
 	}
diff --git a/Engine/Plugins/Runtime/Metasound/Source/MetasoundFrontend/Private/MetasoundFrontendDataTypeRegistry.cpp b/Engine/Plugins/Runtime/Metasound/Source/MetasoundFrontend/Private/MetasoundFrontendDataTypeRegistry.cpp
--- a/Engine/Plugins/Runtime/Metasound/Source/MetasoundFrontend/Private/MetasoundFrontendDataTypeRegistry.cpp
+++ b/Engine/Plugins/Runtime/Metasound/Source/MetasoundFrontend/Private/MetasoundFrontendDataTypeRegistry.cpp
@@ -469,6 +469,8 @@ namespace Metasound
 
 				virtual ELiteralType GetDesiredLiteralType(const FName& InDataType) const override;
 
+				virtual bool IsRegistered(const FName& InDataType) const override;
+
 				virtual bool IsLiteralTypeSupported(const FName& InDataType, ELiteralType InLiteralType) const override;
 				virtual bool IsLiteralTypeSupported(const FName& InDataType, EMetasoundFrontendLiteralType InLiteralType) const override;
 
@@ -591,6 +593,11 @@ namespace Metasound
 				return false;
 			}
 
+			bool FDataTypeRegistry::IsRegistered(const FName& InDataType) const
+			{
+				return RegisteredDataTypes.Contains(InDataType);
+			}
+
 			// Return the enum interface for a data type. If the data type does not have 
 			// an enum interface, returns a nullptr.
 			TSharedPtr<const IEnumDataTypeInterface> FDataTypeRegistry::GetEnumInterfaceForDataType(const FName& InDataType) const
diff --git a/Engine/Plugins/Runtime/Metasound/Source/MetasoundFrontend/Private/MetasoundFrontendDocument.cpp b/Engine/Plugins/Runtime/Metasound/Source/MetasoundFrontend/Private/MetasoundFrontendDocument.cpp
--- a/Engine/Plugins/Runtime/Metasound/Source/MetasoundFrontend/Private/MetasoundFrontendDocument.cpp
+++ b/Engine/Plugins/Runtime/Metasound/Source/MetasoundFrontend/Private/MetasoundFrontendDocument.cpp
@@ -61,7 +61,7 @@ FString FMetasoundFrontendVersion::ToString() const
 
 bool FMetasoundFrontendVersion::IsValid() const
 {
-	return Number.IsValid() && Name.IsValid();
+	return Number != GetInvalid().Number && Name != GetInvalid().Name;
 }
 
 const FMetasoundFrontendVersion& FMetasoundFrontendVersion::GetInvalid()
diff --git a/Engine/Plugins/Runtime/Metasound/Source/MetasoundFrontend/Private/MetasoundFrontendDocumentController.cpp b/Engine/Plugins/Runtime/Metasound/Source/MetasoundFrontend/Private/MetasoundFrontendDocumentController.cpp
--- a/Engine/Plugins/Runtime/Metasound/Source/MetasoundFrontend/Private/MetasoundFrontendDocumentController.cpp
+++ b/Engine/Plugins/Runtime/Metasound/Source/MetasoundFrontend/Private/MetasoundFrontendDocumentController.cpp
@@ -117,20 +117,39 @@ namespace Metasound
 			return false;
 		}
 
-		const FMetasoundFrontendVersion& FDocumentController::GetArchetypeVersion() const
+		const TArray<FMetasoundFrontendVersion>& FDocumentController::GetInterfaceVersions() const
 		{
 			if (const FMetasoundFrontendDocument* Doc = DocumentPtr.Get())
 			{
-				return Doc->ArchetypeVersion;
+				return Doc->InterfaceVersions;
+			}
+
+			static const TArray<FMetasoundFrontendVersion> EmptyArray;
+			return EmptyArray;
+		}
+
+		void FDocumentController::AddInterfaceVersion(const FMetasoundFrontendVersion& InVersion)
+		{
+			if (FMetasoundFrontendDocument* Doc = DocumentPtr.Get())
+			{
+				Doc->InterfaceVersions.Add(InVersion);
+			}
+		}
+
+		void FDocumentController::RemoveInterfaceVersion(const FMetasoundFrontendVersion& InVersion)
+		{
+			if (FMetasoundFrontendDocument* Doc = DocumentPtr.Get())
+			{
+				constexpr bool bAllowShrinking = false;
+				Doc->InterfaceVersions.RemoveAllSwap([&](const FMetasoundFrontendVersion& InterfaceVersion) { return InterfaceVersion == InVersion; }, bAllowShrinking);
 			}
-			return FMetasoundFrontendVersion::GetInvalid();
 		}
 
-		void FDocumentController::SetArchetypeVersion(const FMetasoundFrontendVersion& InVersion)
+		void FDocumentController::ClearInterfaceVersions()
 		{
 			if (FMetasoundFrontendDocument* Doc = DocumentPtr.Get())
 			{
-				Doc->ArchetypeVersion = InVersion;
+				Doc->InterfaceVersions.Reset();
 			}
 		}
 
diff --git a/Engine/Plugins/Runtime/Metasound/Source/MetasoundFrontend/Private/MetasoundFrontendDocumentController.h b/Engine/Plugins/Runtime/Metasound/Source/MetasoundFrontend/Private/MetasoundFrontendDocumentController.h
--- a/Engine/Plugins/Runtime/Metasound/Source/MetasoundFrontend/Private/MetasoundFrontendDocumentController.h
+++ b/Engine/Plugins/Runtime/Metasound/Source/MetasoundFrontend/Private/MetasoundFrontendDocumentController.h
@@ -51,8 +51,10 @@ namespace Metasound
 
 			FGraphHandle AddDuplicateSubgraph(const IGraphController& InGraph) override;
 
-			const FMetasoundFrontendVersion& GetArchetypeVersion() const override;
-			void SetArchetypeVersion(const FMetasoundFrontendVersion& InVersion) override;
+			const TArray<FMetasoundFrontendVersion>& GetInterfaceVersions() const override;
+			void AddInterfaceVersion(const FMetasoundFrontendVersion& InVersion) override;
+			void RemoveInterfaceVersion(const FMetasoundFrontendVersion& InVersion) override;
+			void ClearInterfaceVersions() override;
 
 			void SetMetadata(const FMetasoundFrontendDocumentMetadata& InMetadata) override;
 			const FMetasoundFrontendDocumentMetadata& GetMetadata() const override;
diff --git a/Engine/Plugins/Runtime/Metasound/Source/MetasoundFrontend/Private/MetasoundFrontendGraphController.cpp b/Engine/Plugins/Runtime/Metasound/Source/MetasoundFrontend/Private/MetasoundFrontendGraphController.cpp
--- a/Engine/Plugins/Runtime/Metasound/Source/MetasoundFrontend/Private/MetasoundFrontendGraphController.cpp
+++ b/Engine/Plugins/Runtime/Metasound/Source/MetasoundFrontend/Private/MetasoundFrontendGraphController.cpp
@@ -656,6 +656,7 @@ namespace Metasound
 				GraphClass->Graph.Edges.Reset();
 				GraphClass->Interface.Inputs.Reset();
 				GraphClass->Interface.Outputs.Reset();
+				OwningDocument->ClearInterfaceVersions();
 				OwningDocument->SynchronizeDependencies();
 			}
 		}
diff --git a/Engine/Plugins/Runtime/Metasound/Source/MetasoundFrontend/Private/MetasoundFrontendInjectReceiveNodes.cpp b/Engine/Plugins/Runtime/Metasound/Source/MetasoundFrontend/Private/MetasoundFrontendInjectReceiveNodes.cpp
--- a/Engine/Plugins/Runtime/Metasound/Source/MetasoundFrontend/Private/MetasoundFrontendInjectReceiveNodes.cpp
+++ b/Engine/Plugins/Runtime/Metasound/Source/MetasoundFrontend/Private/MetasoundFrontendInjectReceiveNodes.cpp
@@ -184,7 +184,7 @@ namespace Metasound
 			return true;
 		}
 
-		bool InjectReceiveNodes(FFrontendGraph& InGraph, const FReceiveNodeAddressFunction& InAddressFunction, const TSet<FVertexName>& InInputVerticesToSkip)
+		bool InjectReceiveNodes(FFrontendGraph& InGraph, const FReceiveNodeAddressFunction& InAddressFunction, const TSet<FVertexName>& InInputVertexNames)
 		{
 			using namespace Metasound::Frontend;
 			bool bSuccess = true;
@@ -202,14 +202,11 @@ namespace Metasound
 				}
 
 				const FVertexName& VertexKey = InputDestination.Vertex.GetVertexName();
-				if (InInputVerticesToSkip.Contains(VertexKey))
+				if (InInputVertexNames.Contains(VertexKey))
 				{
-					// Skip this key as requested by caller.
-					continue;
+					const bool bInjectionSuccess = InjectReceiveNode(InGraph, InAddressFunction, InputDestination);
+					bSuccess = bSuccess || bInjectionSuccess;
 				}
-
-				const bool bInjectionSuccess = InjectReceiveNode(InGraph, InAddressFunction, InputDestination);
-				bSuccess = bSuccess || bInjectionSuccess;
 			}
 
 			return bSuccess;
diff --git a/Engine/Plugins/Runtime/Metasound/Source/MetasoundFrontend/Private/MetasoundFrontendInvalidController.h b/Engine/Plugins/Runtime/Metasound/Source/MetasoundFrontend/Private/MetasoundFrontendInvalidController.h
--- a/Engine/Plugins/Runtime/Metasound/Source/MetasoundFrontend/Private/MetasoundFrontendInvalidController.h
+++ b/Engine/Plugins/Runtime/Metasound/Source/MetasoundFrontend/Private/MetasoundFrontendInvalidController.h
@@ -219,7 +219,8 @@ namespace Metasound
 
 			virtual const FText& GetDescription() const override { return Invalid::GetInvalidText(); }
 
-			virtual bool IsRequired() const override { return false; }
+			virtual const FMetasoundFrontendVersion& GetInterfaceVersion() const { return FMetasoundFrontendVersion::GetInvalid(); }
+			virtual bool IsInterfaceMember() const override { return false; }
 
 			virtual bool DiffAgainstRegistryInterface(FClassInterfaceUpdates& OutInterfaceUpdates, bool bInUseHighestMinorVersion) const override { return false; }
 			virtual bool CanAutoUpdate(FClassInterfaceUpdates* OutInterfaceUpdates = nullptr) const override { return false; }
@@ -410,8 +411,10 @@ namespace Metasound
 				virtual FConstClassAccessPtr FindOrAddClass(const FMetasoundFrontendClassMetadata& InMetadata) override{ return FConstClassAccessPtr(); }
 				virtual FGraphHandle AddDuplicateSubgraph(const IGraphController& InGraph) override { return IGraphController::GetInvalidHandle(); }
 
-				virtual const FMetasoundFrontendVersion& GetArchetypeVersion() const override { return FMetasoundFrontendVersion::GetInvalid(); }
-				virtual void SetArchetypeVersion(const FMetasoundFrontendVersion& InVersion) override { }
+				virtual const TArray<FMetasoundFrontendVersion>& GetInterfaceVersions() const override { static const TArray<FMetasoundFrontendVersion> EmptyArray; return EmptyArray; }
+				virtual void AddInterfaceVersion(const FMetasoundFrontendVersion& InVersion) override { }
+				virtual void RemoveInterfaceVersion(const FMetasoundFrontendVersion& InVersion) override { }
+				virtual void ClearInterfaceVersions() override { }
 
 				virtual void SetMetadata(const FMetasoundFrontendDocumentMetadata& InMetadata) override { }
 				virtual const FMetasoundFrontendDocumentMetadata& GetMetadata() const override { return Invalid::GetInvalidDocumentMetadata(); }
diff --git a/Engine/Plugins/Runtime/Metasound/Source/MetasoundFrontend/Private/MetasoundFrontendLiteral.cpp b/Engine/Plugins/Runtime/Metasound/Source/MetasoundFrontend/Private/MetasoundFrontendLiteral.cpp
--- a/Engine/Plugins/Runtime/Metasound/Source/MetasoundFrontend/Private/MetasoundFrontendLiteral.cpp
+++ b/Engine/Plugins/Runtime/Metasound/Source/MetasoundFrontend/Private/MetasoundFrontendLiteral.cpp
@@ -77,6 +77,91 @@ namespace MetasoundFrontendLiteralPrivate
 	}
 }
 
+FMetasoundFrontendLiteral::FMetasoundFrontendLiteral(const FAudioParameter& InParameter)
+{
+	switch (InParameter.ParamType)
+	{
+		case EAudioParameterType::Boolean:
+		{
+			Set(InParameter.BoolParam);
+		}
+		break;
+
+		case EAudioParameterType::BooleanArray:
+		{
+			Set(InParameter.ArrayBoolParam);
+		}
+		break;
+
+		case EAudioParameterType::Float:
+		{
+			Set(InParameter.FloatParam);
+		}
+		break;
+
+		case EAudioParameterType::FloatArray:
+		{
+			Set(InParameter.ArrayFloatParam);
+		}
+		break;
+
+		case EAudioParameterType::Integer:
+		{
+			Set(InParameter.IntParam);
+		}
+		break;
+
+		case EAudioParameterType::IntegerArray:
+		{
+			Set(InParameter.ArrayIntParam);
+		}
+		break;
+
+		case EAudioParameterType::None:
+		{
+			Set(FMetasoundFrontendLiteral::FDefault());
+		}
+		break;
+
+		case EAudioParameterType::NoneArray:
+		{
+			Set(FMetasoundFrontendLiteral::FDefaultArray{ InParameter.IntParam });
+		}
+		break;
+
+		case EAudioParameterType::Object:
+		{
+			Set(InParameter.ObjectParam);
+		}
+		break;
+
+		case EAudioParameterType::ObjectArray:
+		{
+			Set(InParameter.ArrayObjectParam);
+		}
+		break;
+
+		case EAudioParameterType::String:
+		{
+			Set(InParameter.StringParam);
+		}
+		break;
+
+		case EAudioParameterType::StringArray:
+		{
+			Set(InParameter.ArrayStringParam);
+		}
+		break;
+
+		default:
+		{
+			static_assert(static_cast<int32>(EAudioParameterType::COUNT) == 12, "Possible missing switch case coverage");
+			checkNoEntry();
+		}
+	}
+}
+
+
 bool FMetasoundFrontendLiteral::IsArray() const
 {
 	switch (Type)
diff --git a/Engine/Plugins/Runtime/Metasound/Source/MetasoundFrontend/Private/MetasoundFrontendNodeController.cpp b/Engine/Plugins/Runtime/Metasound/Source/MetasoundFrontend/Private/MetasoundFrontendNodeController.cpp
--- a/Engine/Plugins/Runtime/Metasound/Source/MetasoundFrontend/Private/MetasoundFrontendNodeController.cpp
+++ b/Engine/Plugins/Runtime/Metasound/Source/MetasoundFrontend/Private/MetasoundFrontendNodeController.cpp
@@ -511,9 +511,14 @@ namespace Metasound
 			return IInputController::GetInvalidHandle();
 		}
 
-		bool FBaseNodeController::IsRequired() const
+		bool FBaseNodeController::IsInterfaceMember() const
 		{
-			return false;
+			return GetInterfaceVersion() != FMetasoundFrontendVersion::GetInvalid();
+		}
+
+		const FMetasoundFrontendVersion& FBaseNodeController::GetInterfaceVersion() const
+		{
+			return FMetasoundFrontendVersion::GetInvalid();
 		}
 
 		FConstInputHandle FBaseNodeController::GetInputWithID(FGuid InVertexID) const
@@ -1234,7 +1239,7 @@ namespace Metasound
 			return OutputDisplayTitle;
 		}
 
-		bool FOutputNodeController::IsRequired() const
+		const FMetasoundFrontendVersion& FOutputNodeController::GetInterfaceVersion() const
 		{
 			FConstDocumentHandle OwningDocument = OwningGraph->GetOwningDocument();
 			FConstGraphHandle RootGraph = OwningDocument->GetRootGraph();
@@ -1244,32 +1249,41 @@ namespace Metasound
 
 			if (bIsNodeOnRootGraph)
 			{
-				// If the node is on the root graph, test if it is in the archetypes
-				// required inputs or outputs.
-				FMetasoundFrontendArchetype Archetype;
-
-				FArchetypeRegistryKey ArchetypeKey = GetArchetypeRegistryKey(OwningDocument->GetArchetypeVersion());
+				// If the node is on the root graph, test if it is in the interfaces' required inputs or outputs.
+				FMetasoundFrontendInterface Interface;
 
-				bool bFoundArchetype = IArchetypeRegistry::Get().FindArchetype(ArchetypeKey, Archetype);
-				if (bFoundArchetype)
+				for (const FMetasoundFrontendVersion& InterfaceVersion : OwningDocument->GetInterfaceVersions())
 				{
-					if (const FMetasoundFrontendNode* Node = NodePtr.Get())
+					FInterfaceRegistryKey InterfaceKey = GetInterfaceRegistryKey(InterfaceVersion);
+					bool bFoundInterface = IInterfaceRegistry::Get().FindInterface(InterfaceKey, Interface);
+					if (bFoundInterface)
 					{
-						const FVertexName& Name = Node->Name;
-						auto IsVertexWithSameName = [&Name](const FMetasoundFrontendClassVertex& InVertex)
+						if (const FMetasoundFrontendNode* Node = NodePtr.Get())
 						{
-							return InVertex.Name == Name;
-						};
-						return Archetype.Interface.Outputs.ContainsByPredicate(IsVertexWithSameName);
+							if (ensure(Node->Interface.Outputs.Num() == 1))
+							{
+								const FVertexName& Name = Node->Name;
+								const FName& DataType = Node->Interface.Outputs[0].TypeName;
+								auto IsOutputWithSameNameAndType = [&Name, &DataType](const FMetasoundFrontendClassOutput& InOutput)
+								{
+									return InOutput.Name == Name && InOutput.TypeName == DataType;
+								};
+
+								if (Interface.Outputs.ContainsByPredicate(IsOutputWithSameNameAndType))
+								{
+									return InterfaceVersion;
+								}
+							}
+						}
+					}
+					else
+					{
+						UE_LOG(LogMetaSound, Warning, TEXT("Document using unregistered interface [InterfaceVersion:%s]"), *InterfaceVersion.ToString());
 					}
-				}
-				else
-				{
-					UE_LOG(LogMetaSound, Warning, TEXT("Document using unregistered archetype [ArchetypeVersion:%s]"), *OwningDocument->GetArchetypeVersion().ToString());
 				}
 			}
 
-			return false;
+			return FMetasoundFrontendVersion::GetInvalid();
 		}
 
 		bool FOutputNodeController::IsValid() const
@@ -1414,42 +1428,49 @@ namespace Metasound
 			return InputDisplayTitle;
 		}
 
-		bool FInputNodeController::IsRequired() const
+		const FMetasoundFrontendVersion& FInputNodeController::GetInterfaceVersion() const
 		{
 			FConstDocumentHandle OwningDocument = OwningGraph->GetOwningDocument();
 			FConstGraphHandle RootGraph = OwningDocument->GetRootGraph();
 
 			// Test if this node exists on the document's root graph.
 			const bool bIsNodeOnRootGraph = OwningGraph->IsValid() && (RootGraph->GetClassID() == OwningGraph->GetClassID());
-
 			if (bIsNodeOnRootGraph)
 			{
-				// If the node is on the root graph, test if it is in the archetypes
-				// required inputs or outputs. 
-				FMetasoundFrontendArchetype Archetype;
-
-				FArchetypeRegistryKey ArchetypeKey = GetArchetypeRegistryKey(OwningDocument->GetArchetypeVersion());
-
-				bool bFoundArchetype = IArchetypeRegistry::Get().FindArchetype(ArchetypeKey, Archetype);
-				if (bFoundArchetype)
+				for (const FMetasoundFrontendVersion& InterfaceVersion : OwningDocument->GetInterfaceVersions())
 				{
-					if (const FMetasoundFrontendNode* Node = NodePtr.Get())
+					// If the node is on the root graph, test if it is in the interfaces required inputs.
+					FMetasoundFrontendInterface Interface;
+					FInterfaceRegistryKey InterfaceKey = GetInterfaceRegistryKey(InterfaceVersion);
+					bool bFoundInterface = IInterfaceRegistry::Get().FindInterface(InterfaceKey, Interface);
+					if (bFoundInterface)
 					{
-						const FVertexName& Name = Node->Name;
-						auto IsVertexWithSameName = [&Name](const FMetasoundFrontendClassVertex& InVertex)
+						if (const FMetasoundFrontendNode* Node = NodePtr.Get())
 						{
-							return InVertex.Name == Name;
-						};
-						return Archetype.Interface.Inputs.ContainsByPredicate(IsVertexWithSameName);
+							const FVertexName& Name = Node->Name;
+							if (ensure(Node->Interface.Inputs.Num() == 1))
+							{
+								const FName& DataType = Node->Interface.Inputs[0].TypeName;
+								auto IsInputWithSameNameAndType = [&Name, &DataType](const FMetasoundFrontendClassInput& InInput)
+								{
+									return InInput.Name == Name && InInput.TypeName == DataType;
+								};
+
+								if (Interface.Inputs.ContainsByPredicate(IsInputWithSameNameAndType))
+								{
+									return InterfaceVersion;
+								}
+							}
+						}
+					}
+					else
+					{
+						UE_LOG(LogMetaSound, Warning, TEXT("Document using unregistered interface [InterfaceVersion:%s]"), *InterfaceVersion.ToString());
 					}
-				}
-				else
-				{
-					UE_LOG(LogMetaSound, Warning, TEXT("Document using unregistered archetype [ArchetypeVersion:%s]"), *OwningDocument->GetArchetypeVersion().ToString());
 				}
 			}
 
-			return false;
+			return FMetasoundFrontendVersion::GetInvalid();
 		}
 
 		void FInputNodeController::SetDescription(const FText& InDescription)
diff --git a/Engine/Plugins/Runtime/Metasound/Source/MetasoundFrontend/Private/MetasoundFrontendNodeController.h b/Engine/Plugins/Runtime/Metasound/Source/MetasoundFrontend/Private/MetasoundFrontendNodeController.h
--- a/Engine/Plugins/Runtime/Metasound/Source/MetasoundFrontend/Private/MetasoundFrontendNodeController.h
+++ b/Engine/Plugins/Runtime/Metasound/Source/MetasoundFrontend/Private/MetasoundFrontendNodeController.h
@@ -116,7 +116,8 @@ namespace Metasound
 			virtual FOutputHandle GetOutputWithVertexName(const FVertexName& InName) override;
 			virtual FConstOutputHandle GetConstOutputWithVertexName(const FVertexName& InName) const override;
 
-			virtual bool IsRequired() const override;
+			virtual bool IsInterfaceMember() const override;
+			virtual const FMetasoundFrontendVersion& GetInterfaceVersion() const override;
 
 			/** Returns an input with the given id. 
 			 *
@@ -279,7 +280,7 @@ namespace Metasound
 			virtual void SetDescription(const FText& InDescription) override;
 			virtual void SetDisplayName(const FText& InText) override;
 			virtual void SetNodeName(const FVertexName& InName) override;
-			virtual bool IsRequired() const override;
+			virtual const FMetasoundFrontendVersion& GetInterfaceVersion() const override;
 
 		protected:
 
@@ -337,7 +338,7 @@ namespace Metasound
 			virtual const FText& GetDescription() const override;
 			virtual FText GetDisplayName() const override;
 			virtual const FText& GetDisplayTitle() const override;
-			virtual bool IsRequired() const override;
+			virtual const FMetasoundFrontendVersion& GetInterfaceVersion() const override;
 			virtual bool IsValid() const override;
 			virtual void SetDescription(const FText& InDescription) override;
 			virtual void SetDisplayName(const FText& InText) override;
diff --git a/Engine/Plugins/Runtime/Metasound/Source/MetasoundFrontend/Private/MetasoundFrontendSearchEngine.cpp b/Engine/Plugins/Runtime/Metasound/Source/MetasoundFrontend/Private/MetasoundFrontendSearchEngine.cpp
--- a/Engine/Plugins/Runtime/Metasound/Source/MetasoundFrontend/Private/MetasoundFrontendSearchEngine.cpp
+++ b/Engine/Plugins/Runtime/Metasound/Source/MetasoundFrontend/Private/MetasoundFrontendSearchEngine.cpp
@@ -9,6 +9,7 @@
 #include "MetasoundFrontendQuerySteps.h"
 #include "MetasoundFrontendRegistryTransaction.h"
 #include "MetasoundTrace.h"
+#include "MetasoundAssetManager.h"
 
 static int32 ClearMetaSoundFrontendSearchEngineCacheCVar = 0;
 FAutoConsoleVariableRef CVarClearMetaSoundFrontendSearchEngineCache(
@@ -25,22 +26,22 @@ namespace Metasound
 	{
 		namespace SearchEngineQuerySteps
 		{
-			class FArchetypeRegistryTransactionSource : public IFrontendQuerySource
+			class FInterfaceRegistryTransactionSource : public IFrontendQuerySource
 			{
 			public:
-				FArchetypeRegistryTransactionSource()
+				FInterfaceRegistryTransactionSource()
 				: CurrentTransactionID(GetOriginRegistryTransactionID())
 				{
 				}
 
 				void Stream(TArray<FFrontendQueryEntry>& OutEntries) override
 				{
-					auto AddEntry = [&OutEntries](const FArchetypeRegistryTransaction& InTransaction)
+					auto AddEntry = [&OutEntries](const FInterfaceRegistryTransaction& InTransaction)
 					{
-						OutEntries.Emplace(FFrontendQueryEntry::FValue(TInPlaceType<FArchetypeRegistryTransaction>(), InTransaction));
+						OutEntries.Emplace(FFrontendQueryEntry::FValue(TInPlaceType<FInterfaceRegistryTransaction>(), InTransaction));
 					};
 					
-					IArchetypeRegistry::Get().ForEachRegistryTransactionSince(CurrentTransactionID, &CurrentTransactionID, AddEntry);
+					IInterfaceRegistry::Get().ForEachRegistryTransactionSince(CurrentTransactionID, &CurrentTransactionID, AddEntry);
 				}
 
 				void Reset() override
@@ -52,42 +53,42 @@ namespace Metasound
 				FRegistryTransactionID CurrentTransactionID;
 			};
 
-			class FMapArchetypeRegistryTransactionsToArchetypeRegistryKeys : public IFrontendQueryMapStep
+			class FMapInterfaceRegistryTransactionsToInterfaceRegistryKeys : public IFrontendQueryMapStep
 			{
 			public:
 				FFrontendQueryEntry::FKey Map(const FFrontendQueryEntry& InEntry) const override
 				{
-					FArchetypeRegistryKey RegistryKey;
-					if (ensure(InEntry.Value.IsType<FArchetypeRegistryTransaction>()))
+					FInterfaceRegistryKey RegistryKey;
+					if (ensure(InEntry.Value.IsType<FInterfaceRegistryTransaction>()))
 					{
-						RegistryKey = InEntry.Value.Get<FArchetypeRegistryTransaction>().GetArchetypeRegistryKey();
+						RegistryKey = InEntry.Value.Get<FInterfaceRegistryTransaction>().GetInterfaceRegistryKey();
 					}
 					return FFrontendQueryEntry::FKey{RegistryKey};
 				}
 			};
 
-			class FReduceArchetypeRegistryTransactionsToCurrentStatus : public IFrontendQueryReduceStep
+			class FReduceInterfaceRegistryTransactionsToCurrentStatus : public IFrontendQueryReduceStep
 			{
 			public:
-				void Reduce(FFrontendQueryEntry::FKey InKey, TArrayView<FFrontendQueryEntry * const>& InEntries, FReduceOutputView& OutResult) const override
+				void Reduce(FFrontendQueryEntry::FKey InKey, TArrayView<FFrontendQueryEntry* const>& InEntries, FReduceOutputView& OutResult) const override
 				{
 					int32 State = 0;
 					FFrontendQueryEntry* FinalEntry = nullptr;
 
 					for (FFrontendQueryEntry* Entry : InEntries)
 					{
-						if (ensure(Entry->Value.IsType<FArchetypeRegistryTransaction>()))
+						if (ensure(Entry->Value.IsType<FInterfaceRegistryTransaction>()))
 						{
-							const FArchetypeRegistryTransaction& Transaction = Entry->Value.Get<FArchetypeRegistryTransaction>();
+							const FInterfaceRegistryTransaction& Transaction = Entry->Value.Get<FInterfaceRegistryTransaction>();
 							
 							switch (Transaction.GetTransactionType())
 							{
-								case FArchetypeRegistryTransaction::ETransactionType::ArchetypeRegistration:
+								case FInterfaceRegistryTransaction::ETransactionType::InterfaceRegistration:
 									State++;
 									FinalEntry = Entry;
 									break;
 
-								case FArchetypeRegistryTransaction::ETransactionType::ArchetypeUnregistration:
+								case FInterfaceRegistryTransaction::ETransactionType::InterfaceUnregistration:
 									State--;
 									break;
 
@@ -104,49 +105,102 @@ namespace Metasound
 				}
 			};
 
-			class FTransformArchetypeRegistryTransactionToArchetype : public IFrontendQueryTransformStep
+			class FReduceInterfaceRegistryTransactionsBySupportedClassName : public IFrontendQueryReduceStep
+			{
+				FName UClassName;
+
+			public:
+				FReduceInterfaceRegistryTransactionsBySupportedClassName(FName InUClassName)
+					: UClassName(InUClassName)
+				{
+				}
+
+				void Reduce(FFrontendQueryEntry::FKey InKey, TArrayView<FFrontendQueryEntry* const>& InEntries, FReduceOutputView& OutResult) const override
+				{
+					for (FFrontendQueryEntry* Entry : InEntries)
+					{
+						if (ensure(Entry->Value.IsType<FInterfaceRegistryTransaction>()))
+						{
+							FInterfaceRegistryKey RegistryKey = Entry->Value.Get<FInterfaceRegistryTransaction>().GetInterfaceRegistryKey();
+							const IInterfaceRegistryEntry* RegEntry = IInterfaceRegistry::Get().FindInterfaceRegistryEntry(RegistryKey);
+							if (ensure(RegEntry))
+							{
+								if (RegEntry->UClassIsSupported(UClassName))
+								{
+									OutResult.Add(*Entry);
+								}
+							}
+						}
+					}
+				}
+			};
+
+			class FReduceInterfaceRegistryTransactionsByDefault : public IFrontendQueryReduceStep
+			{
+			public:
+				void Reduce(FFrontendQueryEntry::FKey InKey, TArrayView<FFrontendQueryEntry* const>& InEntries, FReduceOutputView& OutResult) const override
+				{
+					for (FFrontendQueryEntry* Entry : InEntries)
+					{
+						if (ensure(Entry->Value.IsType<FInterfaceRegistryTransaction>()))
+						{
+							FInterfaceRegistryKey RegistryKey = Entry->Value.Get<FInterfaceRegistryTransaction>().GetInterfaceRegistryKey();
+							const IInterfaceRegistryEntry* RegEntry = IInterfaceRegistry::Get().FindInterfaceRegistryEntry(RegistryKey);
+							if (ensure(RegEntry))
+							{
+								if (RegEntry->IsDefault())
+								{
+									OutResult.Add(*Entry);
+								}
+							}
+						}
+					}
+				}
+			};
+
+			class FTransformInterfaceRegistryTransactionToInterface : public IFrontendQueryTransformStep
 			{
 			public:
 				virtual void Transform(FFrontendQueryEntry::FValue& InValue) const override
 				{
-					FMetasoundFrontendArchetype Archetype;
-					if (ensure(InValue.IsType<FArchetypeRegistryTransaction>()))
+					FMetasoundFrontendInterface Interface;
+					if (ensure(InValue.IsType<FInterfaceRegistryTransaction>()))
 					{
-						FArchetypeRegistryKey RegistryKey = InValue.Get<FArchetypeRegistryTransaction>().GetArchetypeRegistryKey();
-						IArchetypeRegistry::Get().FindArchetype(RegistryKey, Archetype);
+						FInterfaceRegistryKey RegistryKey = InValue.Get<FInterfaceRegistryTransaction>().GetInterfaceRegistryKey();
+						IInterfaceRegistry::Get().FindInterface(RegistryKey, Interface);
 					}
-					InValue.Set<FMetasoundFrontendArchetype>(MoveTemp(Archetype));
+					InValue.Set<FMetasoundFrontendInterface>(MoveTemp(Interface));
 				}
 			};
 
-			class FMapArchetypeToArchetypeNameAndMajorVersion : public IFrontendQueryMapStep
+			class FMapInterfaceToInterfaceNameAndMajorVersion : public IFrontendQueryMapStep
 			{
 			public:
 				FFrontendQueryEntry::FKey Map(const FFrontendQueryEntry& InEntry) const override
 				{
 					FString Key;
-					if (ensure(InEntry.Value.IsType<FMetasoundFrontendArchetype>()))
+					if (ensure(InEntry.Value.IsType<FMetasoundFrontendInterface>()))
 					{
-						const FMetasoundFrontendArchetype& Arch = InEntry.Value.Get<FMetasoundFrontendArchetype>();
-						Key = FString::Format(TEXT("{0}_{1}"), { Arch.Version.Name.ToString(), Arch.Version.Number.Major });
+						const FMetasoundFrontendInterface& Interface = InEntry.Value.Get<FMetasoundFrontendInterface>();
+						Key = FString::Format(TEXT("{0}_{1}"), { Interface.Version.Name.ToString(), Interface.Version.Number.Major });
 					}
 					return FFrontendQueryEntry::FKey{Key};
 				}
 			};
 
-			class FReduceArchetypesToHighestVersion : public IFrontendQueryReduceStep
+			class FReduceInterfacesToHighestVersion : public IFrontendQueryReduceStep
 			{
 			public:
-				void Reduce(FFrontendQueryEntry::FKey InKey, TArrayView<FFrontendQueryEntry * const>& InEntries, FReduceOutputView& OutResult) const override
+				void Reduce(FFrontendQueryEntry::FKey InKey, TArrayView<FFrontendQueryEntry* const>& InEntries, FReduceOutputView& OutResult) const override
 				{
 					FFrontendQueryEntry* HighestVersionEntry = nullptr;
 					FMetasoundFrontendVersionNumber HighestVersion = FMetasoundFrontendVersionNumber::GetInvalid();
 
 					for (FFrontendQueryEntry* Entry : InEntries)
 					{
-						if (ensure(Entry->Value.IsType<FMetasoundFrontendArchetype>()))
+						if (ensure(Entry->Value.IsType<FMetasoundFrontendInterface>()))
 						{
-							const FMetasoundFrontendVersionNumber& VersionNumber = Entry->Value.Get<FMetasoundFrontendArchetype>().Version.Number;
+							const FMetasoundFrontendVersionNumber& VersionNumber = Entry->Value.Get<FMetasoundFrontendInterface>().Version.Number;
 
 							if (VersionNumber > HighestVersion)
 							{
@@ -163,19 +217,19 @@ namespace Metasound
 				}
 			};
 
-			class FFilterArchetypeRegistryTransactionsByName : public IFrontendQueryFilterStep
+			class FFilterInterfaceRegistryTransactionsByName : public IFrontendQueryFilterStep
 			{
 			public:
-				FFilterArchetypeRegistryTransactionsByName(const FName& InName)
+				FFilterInterfaceRegistryTransactionsByName(FName InName)
 				: Name(InName)
 				{
 				}
 
 				bool Filter(const FFrontendQueryEntry& InEntry) const override
 				{
-					if (ensure(InEntry.Value.IsType<FArchetypeRegistryTransaction>()))
+					if (ensure(InEntry.Value.IsType<FInterfaceRegistryTransaction>()))
 					{
-						return Name == InEntry.Value.Get<FArchetypeRegistryTransaction>().GetArchetypeVersion().Name;
+						return Name == InEntry.Value.Get<FInterfaceRegistryTransaction>().GetInterfaceVersion().Name;
 					}
 					return false;
 				}
@@ -205,10 +259,10 @@ namespace Metasound
 				FMetasoundFrontendClassName Name;
 			};
 
-			class FFilterArchetypeRegistryTransactionsByNameAndMajorVersion : public IFrontendQueryFilterStep
+			class FFilterInterfaceRegistryTransactionsByNameAndMajorVersion : public IFrontendQueryFilterStep
 			{
 			public:
-				FFilterArchetypeRegistryTransactionsByNameAndMajorVersion(const FName& InName, int32 InMajorVersion)
+				FFilterInterfaceRegistryTransactionsByNameAndMajorVersion(const FName& InName, int32 InMajorVersion)
 				: Name(InName)
 				, MajorVersion(InMajorVersion)
 				{
@@ -216,9 +270,9 @@ namespace Metasound
 
 				bool Filter(const FFrontendQueryEntry& InEntry) const override
 				{
-					if (ensure(InEntry.Value.IsType<FArchetypeRegistryTransaction>()))
+					if (ensure(InEntry.Value.IsType<FInterfaceRegistryTransaction>()))
 					{
-						const FMetasoundFrontendVersion& Version = InEntry.Value.Get<FArchetypeRegistryTransaction>().GetArchetypeVersion();
+						const FMetasoundFrontendVersion& Version = InEntry.Value.Get<FInterfaceRegistryTransaction>().GetInterfaceVersion();
 
 						return (Name == Version.Name) && (MajorVersion == Version.Number.Major);
 					}
@@ -230,29 +284,29 @@ namespace Metasound
 				int32 MajorVersion;
 			};
 
-			class FMapArchetypeToArchetypeName : public IFrontendQueryMapStep
+			class FMapInterfaceToInterfaceName : public IFrontendQueryMapStep
 			{
 			public:
 				FFrontendQueryEntry::FKey Map(const FFrontendQueryEntry& InEntry) const override
 				{
 					FString Key;
-					if (ensure(InEntry.Value.IsType<FMetasoundFrontendArchetype>()))
+					if (ensure(InEntry.Value.IsType<FMetasoundFrontendInterface>()))
 					{
-						Key = InEntry.Value.Get<FMetasoundFrontendArchetype>().Version.Name.ToString();
+						Key = InEntry.Value.Get<FMetasoundFrontendInterface>().Version.Name.ToString();
 					}
 					return FFrontendQueryEntry::FKey{Key};
 				}
 			};
 
-			class FTransformArchetypeRegistryTransactionToArchetypeVersion : public IFrontendQueryTransformStep
+			class FTransformInterfaceRegistryTransactionToInterfaceVersion : public IFrontendQueryTransformStep
 			{
 			public:
 				virtual void Transform(FFrontendQueryEntry::FValue& InValue) const override
 				{
 					FMetasoundFrontendVersion Version;
-					if (ensure(InValue.IsType<FArchetypeRegistryTransaction>()))
+					if (ensure(InValue.IsType<FInterfaceRegistryTransaction>()))
 					{
-						Version = InValue.Get<FArchetypeRegistryTransaction>().GetArchetypeVersion();
+						Version = InValue.Get<FInterfaceRegistryTransaction>().GetInterfaceVersion();
 					}
 					InValue.Set<FMetasoundFrontendVersion>(Version);
 				}
@@ -270,19 +324,16 @@ namespace Metasound
 				FSearchEngine() = default;
 				virtual ~FSearchEngine() = default;
 
-				TArray<FMetasoundFrontendClass> FindAllClasses(bool bInIncludeDeprecated) override;
-				TArray<FMetasoundFrontendClass> FindClassesWithName(const FNodeClassName& InName, bool bInSortByVersion) override;
-
-				bool FindClassWithHighestVersion(const FNodeClassName& InName, FMetasoundFrontendClass& OutClass) override;
-				bool FindClassWithMajorVersion(const FNodeClassName& InName, int32 InMajorVersion, FMetasoundFrontendClass& OutClass) override;
-
-				virtual TArray<FMetasoundFrontendArchetype> FindAllArchetypes(bool bInIncludeDeprecated) override;
+				virtual TArray<FMetasoundFrontendClass> FindAllClasses(bool bInIncludeDeprecated) override;
+				virtual TArray<FMetasoundFrontendClass> FindClassesWithName(const FNodeClassName& InName, bool bInSortByVersion) override;
+				virtual bool FindClassWithHighestVersion(const FNodeClassName& InName, FMetasoundFrontendClass& OutClass) override;
+				virtual bool FindClassWithMajorVersion(const FNodeClassName& InName, int32 InMajorVersion, FMetasoundFrontendClass& OutClass) override;
 
-				virtual TArray<FMetasoundFrontendVersion> FindAllRegisteredArchetypesWithName(const FName& InArchetypeName) override;
-
-				virtual bool FindArchetypeWithHighestVersion(const FName& InArchetypeName, FMetasoundFrontendArchetype& OutArchetype) override;
-
-				virtual bool FindArchetypeWithMajorVersion(const FName& InArchetypeName, int32 InMajorVersion, FMetasoundFrontendArchetype& OutArchetype) override;
+				virtual TArray<FMetasoundFrontendInterface> FindAllInterfaces(bool bInIncludeDeprecated) override;
+				virtual TArray<FMetasoundFrontendInterface> FindUClassDefaultInterfaces(FName InUClassName) override;
+				virtual TArray<FMetasoundFrontendVersion> FindAllRegisteredInterfacesWithName(FName InInterfaceName) override;
+				virtual bool FindInterfaceWithHighestVersion(FName InInterfaceName, FMetasoundFrontendInterface& OutInterface) override;
+				virtual bool FindInterfaceWithMajorVersion(FName InInterfaceName, int32 InMajorVersion, FMetasoundFrontendInterface& OutInterface) override;
 
 			private:
 				FFrontendQuery* FindQuery(const FString& InQueryName);
@@ -482,32 +533,74 @@ namespace Metasound
 			return false;
 		}
 
-		TArray<FMetasoundFrontendArchetype> FSearchEngine::FindAllArchetypes(bool bInIncludeDeprecated)
+		TArray<FMetasoundFrontendInterface> FSearchEngine::FindAllInterfaces(bool bInIncludeDeprecated)
 		{
 			using namespace SearchEngineQuerySteps;
 
-			const FString QueryName = FString::Format(TEXT("FindAllArchetypes_IncludeDeprecated:{0}"), { bInIncludeDeprecated });
+			const FString QueryName = FString::Format(TEXT("FindAllInterfaces_IncludeDeprecated:{0}"), { bInIncludeDeprecated });
 
 			FFrontendQuery* Query = FindQuery(QueryName);
 			if (!Query)
 			{
 				TUniquePtr<FFrontendQuery> NewQuery = MakeUnique<FFrontendQuery>();
 
-				NewQuery->AddStep<FArchetypeRegistryTransactionSource>()
-					.AddStep<FMapArchetypeRegistryTransactionsToArchetypeRegistryKeys>()
-					.AddStep<FReduceArchetypeRegistryTransactionsToCurrentStatus>()
-					.AddStep<FTransformArchetypeRegistryTransactionToArchetype>();
+				NewQuery->AddStep<FInterfaceRegistryTransactionSource>()
+					.AddStep<FMapInterfaceRegistryTransactionsToInterfaceRegistryKeys>()
+					.AddStep<FReduceInterfaceRegistryTransactionsToCurrentStatus>()
+					.AddStep<FTransformInterfaceRegistryTransactionToInterface>();
 
 				if (!bInIncludeDeprecated)
 				{
-					NewQuery->AddStep<FMapArchetypeToArchetypeNameAndMajorVersion>()
-					.AddStep<FReduceArchetypesToHighestVersion>();
+					NewQuery->AddStep<FMapInterfaceToInterfaceNameAndMajorVersion>()
+					.AddStep<FReduceInterfacesToHighestVersion>();
 				}
 
 				Query = AddQuery(QueryName, MoveTemp(NewQuery));
 			}
 
-			TArray<FMetasoundFrontendArchetype> Result;
+			TArray<FMetasoundFrontendInterface> Result;
+
+			if (ensure(Query))
+			{
+				FFrontendQuerySelectionView Selection = Query->Execute();
+				TArrayView<const FFrontendQueryEntry* const> Entries = Selection.GetSelection();
+
+				for (const FFrontendQueryEntry* Entry : Selection.GetSelection())
+				{
+					if (ensure(Entry->Value.IsType<FMetasoundFrontendInterface>()))
+					{
+						Result.Add(Entry->Value.Get<FMetasoundFrontendInterface>());
+					}
+				}
+			}
+
+			return Result;
+		}
+
+		TArray<FMetasoundFrontendInterface> FSearchEngine::FindUClassDefaultInterfaces(FName InUClassName)
+		{
+			using namespace SearchEngineQuerySteps;
+
+			const FString QueryName = FString::Format(TEXT("FindUClassDefaultInterfaces_Class:{0}"), { InUClassName.ToString() });
+
+			FFrontendQuery* Query = FindQuery(QueryName);
+			if (!Query)
+			{
+				TUniquePtr<FFrontendQuery> NewQuery = MakeUnique<FFrontendQuery>();
+
+				NewQuery->AddStep<FInterfaceRegistryTransactionSource>()
+					.AddStep<FMapInterfaceRegistryTransactionsToInterfaceRegistryKeys>()
+					.AddStep<FReduceInterfaceRegistryTransactionsToCurrentStatus>()
+					.AddStep<FReduceInterfaceRegistryTransactionsByDefault>()
+					.AddStep<FReduceInterfaceRegistryTransactionsBySupportedClassName>(InUClassName)
+					.AddStep<FTransformInterfaceRegistryTransactionToInterface>()
+					.AddStep<FMapInterfaceToInterfaceNameAndMajorVersion>()
+					.AddStep<FReduceInterfacesToHighestVersion>();
+
+				Query = AddQuery(QueryName, MoveTemp(NewQuery));
+			}
+
+			TArray<FMetasoundFrontendInterface> Result;
 
 			if (ensure(Query))
 			{
@@ -516,9 +609,9 @@ namespace Metasound
 
 				for (const FFrontendQueryEntry* Entry : Selection.GetSelection())
 				{
-					if (ensure(Entry->Value.IsType<FMetasoundFrontendArchetype>()))
+					if (ensure(Entry->Value.IsType<FMetasoundFrontendInterface>()))
 					{
-						Result.Add(Entry->Value.Get<FMetasoundFrontendArchetype>());
+						Result.Add(Entry->Value.Get<FMetasoundFrontendInterface>());
 					}
 				}
 			}
@@ -526,22 +619,22 @@ namespace Metasound
 			return Result;
 		}
 
-		TArray<FMetasoundFrontendVersion> FSearchEngine::FindAllRegisteredArchetypesWithName(const FName& InArchetypeName)
+		TArray<FMetasoundFrontendVersion> FSearchEngine::FindAllRegisteredInterfacesWithName(FName InInterfaceName)
 		{
 			using namespace SearchEngineQuerySteps;
 
-			const FString QueryName = FString::Format(TEXT("FindAllRegisteredArchetypesWithName_Name:{0}"), { InArchetypeName.ToString() });
+			const FString QueryName = FString::Format(TEXT("FindAllRegisteredInterfacesWithName_Name:{0}"), { InInterfaceName.ToString() });
 
 			FFrontendQuery* Query = FindQuery(QueryName);
 			if (!Query)
 			{
 				TUniquePtr<FFrontendQuery> NewQuery = MakeUnique<FFrontendQuery>();
 
-				NewQuery->AddStep<FArchetypeRegistryTransactionSource>()
-					.AddStep<FFilterArchetypeRegistryTransactionsByName>(InArchetypeName)
-					.AddStep<FMapArchetypeRegistryTransactionsToArchetypeRegistryKeys>()
-					.AddStep<FReduceArchetypeRegistryTransactionsToCurrentStatus>()
-					.AddStep<FTransformArchetypeRegistryTransactionToArchetypeVersion>();
+				NewQuery->AddStep<FInterfaceRegistryTransactionSource>()
+					.AddStep<FFilterInterfaceRegistryTransactionsByName>(InInterfaceName)
+					.AddStep<FMapInterfaceRegistryTransactionsToInterfaceRegistryKeys>()
+					.AddStep<FReduceInterfaceRegistryTransactionsToCurrentStatus>()
+					.AddStep<FTransformInterfaceRegistryTransactionToInterfaceVersion>();
 
 				Query = AddQuery(QueryName, MoveTemp(NewQuery));
 			}
@@ -565,23 +658,23 @@ namespace Metasound
 			return Result;
 		}
 
-		bool FSearchEngine::FindArchetypeWithHighestVersion(const FName& InArchetypeName, FMetasoundFrontendArchetype& OutArchetype)
+		bool FSearchEngine::FindInterfaceWithHighestVersion(FName InInterfaceName, FMetasoundFrontendInterface& OutInterface)
 		{
 			using namespace SearchEngineQuerySteps;
-			const FString QueryName = FString::Format(TEXT("FindArchetypeWithHighestVersion_Name:{0}"), { InArchetypeName.ToString() });
+			const FString QueryName = FString::Format(TEXT("FindInterfaceWithHighestVersion_Name:{0}"), { InInterfaceName.ToString() });
 
 			FFrontendQuery* Query = FindQuery(QueryName);
 			if (!Query)
 			{
 				TUniquePtr<FFrontendQuery> NewQuery = MakeUnique<FFrontendQuery>();
 
-				NewQuery->AddStep<FArchetypeRegistryTransactionSource>()
-					.AddStep<FFilterArchetypeRegistryTransactionsByName>(InArchetypeName)
-					.AddStep<FMapArchetypeRegistryTransactionsToArchetypeRegistryKeys>()
-					.AddStep<FReduceArchetypeRegistryTransactionsToCurrentStatus>()
-					.AddStep<FTransformArchetypeRegistryTransactionToArchetype>()
-					.AddStep<FMapArchetypeToArchetypeName>()
-					.AddStep<FReduceArchetypesToHighestVersion>();
+				NewQuery->AddStep<FInterfaceRegistryTransactionSource>()
+					.AddStep<FFilterInterfaceRegistryTransactionsByName>(InInterfaceName)
+					.AddStep<FMapInterfaceRegistryTransactionsToInterfaceRegistryKeys>()
+					.AddStep<FReduceInterfaceRegistryTransactionsToCurrentStatus>()
+					.AddStep<FTransformInterfaceRegistryTransactionToInterface>()
+					.AddStep<FMapInterfaceToInterfaceName>()
+					.AddStep<FReduceInterfacesToHighestVersion>();
 
 				Query = AddQuery(QueryName, MoveTemp(NewQuery));
 			}
@@ -597,7 +690,7 @@ namespace Metasound
 
 				if (ensure(Entries.Num() == 1))
 				{
-					OutArchetype = Entries[0]->Value.Get<FMetasoundFrontendArchetype>();
+					OutInterface = Entries[0]->Value.Get<FMetasoundFrontendInterface>();
 					return true;
 				}
 			}
@@ -605,23 +698,23 @@ namespace Metasound
 			return false;
 		}
 
-		bool FSearchEngine::FindArchetypeWithMajorVersion(const FName& InArchetypeName, int32 InMajorVersion, FMetasoundFrontendArchetype& OutArchetype)
+		bool FSearchEngine::FindInterfaceWithMajorVersion(FName InInterfaceName, int32 InMajorVersion, FMetasoundFrontendInterface& OutInterface)
 		{
 			using namespace SearchEngineQuerySteps;
-			const FString QueryName = FString::Format(TEXT("FindArchetypeWithMajorVersion_Name:{0}_MajorVersion:{1}"), { InArchetypeName.ToString(), InMajorVersion });
+			const FString QueryName = FString::Format(TEXT("FindInterfaceWithMajorVersion_Name:{0}_MajorVersion:{1}"), { InInterfaceName.ToString(), InMajorVersion });
 
 			FFrontendQuery* Query = FindQuery(QueryName);
 			if (!Query)
 			{
 				TUniquePtr<FFrontendQuery> NewQuery = MakeUnique<FFrontendQuery>();
 
-				NewQuery->AddStep<FArchetypeRegistryTransactionSource>()
-					.AddStep<FFilterArchetypeRegistryTransactionsByNameAndMajorVersion>(InArchetypeName, InMajorVersion)
-					.AddStep<FMapArchetypeRegistryTransactionsToArchetypeRegistryKeys>()
-					.AddStep<FReduceArchetypeRegistryTransactionsToCurrentStatus>()
-					.AddStep<FTransformArchetypeRegistryTransactionToArchetype>()
-					.AddStep<FMapArchetypeToArchetypeName>()
-					.AddStep<FReduceArchetypesToHighestVersion>();
+				NewQuery->AddStep<FInterfaceRegistryTransactionSource>()
+					.AddStep<FFilterInterfaceRegistryTransactionsByNameAndMajorVersion>(InInterfaceName, InMajorVersion)
+					.AddStep<FMapInterfaceRegistryTransactionsToInterfaceRegistryKeys>()
+					.AddStep<FReduceInterfaceRegistryTransactionsToCurrentStatus>()
+					.AddStep<FTransformInterfaceRegistryTransactionToInterface>()
+					.AddStep<FMapInterfaceToInterfaceName>()
+					.AddStep<FReduceInterfacesToHighestVersion>();
 
 				Query = AddQuery(QueryName, MoveTemp(NewQuery));
 			}
@@ -637,7 +730,7 @@ namespace Metasound
 
 				if (ensure(Entries.Num() == 1))
 				{
-					OutArchetype = Entries[0]->Value.Get<FMetasoundFrontendArchetype>();
+					OutInterface = Entries[0]->Value.Get<FMetasoundFrontendInterface>();
 					return true;
 				}
 			}
diff --git a/Engine/Plugins/Runtime/Metasound/Source/MetasoundFrontend/Private/MetasoundFrontendTransform.cpp b/Engine/Plugins/Runtime/Metasound/Source/MetasoundFrontend/Private/MetasoundFrontendTransform.cpp
--- a/Engine/Plugins/Runtime/Metasound/Source/MetasoundFrontend/Private/MetasoundFrontendTransform.cpp
+++ b/Engine/Plugins/Runtime/Metasound/Source/MetasoundFrontend/Private/MetasoundFrontendTransform.cpp
@@ -15,12 +15,57 @@ namespace Metasound
 {
 	namespace Frontend
 	{
-		FSwapGraphArchetype::FSwapGraphArchetype(const FMetasoundFrontendArchetype& InFromArchetype, const FMetasoundFrontendArchetype& InToArchetype) 
+		FModifyRootGraphInterfaces::FModifyRootGraphInterfaces(const TArray<FMetasoundFrontendInterface>& InInterfacesToRemove, const TArray<FMetasoundFrontendInterface>& InInterfacesToAdd)
+			: InterfacesToRemove(InInterfacesToRemove)
+			, InterfacesToAdd(InInterfacesToAdd)
 		{
-			InputsToAdd = InToArchetype.Interface.Inputs;
-			InputsToRemove = InFromArchetype.Interface.Inputs;
-			OutputsToAdd = InToArchetype.Interface.Outputs;
-			OutputsToRemove = InFromArchetype.Interface.Outputs;
+			Init();
+		}
+
+		FModifyRootGraphInterfaces::FModifyRootGraphInterfaces(const TArray<FMetasoundFrontendVersion>& InInterfaceVersionsToRemove, const TArray<FMetasoundFrontendVersion>& InInterfaceVersionsToAdd)
+		{
+			Algo::Transform(InInterfaceVersionsToRemove, InterfacesToRemove, [](const FMetasoundFrontendVersion& Version)
+			{
+				FMetasoundFrontendInterface Interface;
+				const bool bFromInterfaceFound = IInterfaceRegistry::Get().FindInterface(GetInterfaceRegistryKey(Version), Interface);
+				if (!ensureAlways(bFromInterfaceFound))
+				{
+					UE_LOG(LogMetaSound, Error, TEXT("Failed to find interface '%s' to remove"), *Version.ToString());
+				}
+				return Interface;
+			});
+
+			Algo::Transform(InInterfaceVersionsToAdd, InterfacesToAdd, [](const FMetasoundFrontendVersion& Version)
+			{
+				FMetasoundFrontendInterface Interface;
+				const bool bToInterfaceFound = IInterfaceRegistry::Get().FindInterface(GetInterfaceRegistryKey(Version), Interface);
+				if (!ensureAlways(bToInterfaceFound))
+				{
+					UE_LOG(LogMetaSound, Error, TEXT("Failed to find interface '%s' to add"), *Version.ToString());
+				}
+				return Interface;
+			});
+
+			Init();
+		}
+
+		void FModifyRootGraphInterfaces::SetDefaultNodeLocations(bool bInSetDefaultNodeLocations)
+		{
+			bSetDefaultNodeLocations = bInSetDefaultNodeLocations;
+		}
+
+		void FModifyRootGraphInterfaces::Init()
+		{
+			for (const FMetasoundFrontendInterface& FromInterface : InterfacesToRemove)
+			{
+				InputsToRemove.Append(FromInterface.Inputs);
+				OutputsToRemove.Append(FromInterface.Outputs);
+			}
+			for (const FMetasoundFrontendInterface& ToInterface : InterfacesToAdd)
+			{
+				InputsToAdd.Append(ToInterface.Inputs);
+				OutputsToAdd.Append(ToInterface.Outputs);
+			}
 
 			// Iterate in reverse to allow removal from `InputsToAdd`
 			for (int32 AddIndex = InputsToAdd.Num() - 1; AddIndex >= 0; AddIndex--)
@@ -28,9 +73,9 @@ namespace Metasound
 				const FMetasoundFrontendClassVertex& VertexToAdd = InputsToAdd[AddIndex];
 
 				int32 RemoveIndex = InputsToRemove.IndexOfByPredicate([&](const FMetasoundFrontendClassVertex& VertexToRemove)
-					{
-						return FMetasoundFrontendClassVertex::IsFunctionalEquivalent(VertexToAdd, VertexToRemove);
-					});
+				{
+					return FMetasoundFrontendClassVertex::IsFunctionalEquivalent(VertexToAdd, VertexToRemove);
+				});
 
 				if (INDEX_NONE != RemoveIndex)
 				{
@@ -59,19 +104,35 @@ namespace Metasound
 			}
 		}
 
-		bool FSwapGraphArchetype::Transform(FGraphHandle InGraph) const
+		bool FModifyRootGraphInterfaces::Transform(FDocumentHandle InDocument) const
 		{
 			bool bDidEdit = false;
 
+			FGraphHandle GraphHandle = InDocument->GetRootGraph();
+			if (!ensure(GraphHandle->IsValid()))
+			{
+				return false;
+			}
+
+			for (const FMetasoundFrontendInterface& Interface : InterfacesToRemove)
+			{
+				InDocument->RemoveInterfaceVersion(Interface.Version);
+			}
+
+			for (const FMetasoundFrontendInterface& Interface : InterfacesToAdd)
+			{
+				InDocument->AddInterfaceVersion(Interface.Version);
+			}
+
 			// Remove unsupported inputs
 			for (const FMetasoundFrontendClassVertex& InputToRemove : InputsToRemove)
 			{
-				if (const FMetasoundFrontendClassInput* ClassInput = InGraph->FindClassInputWithName(InputToRemove.Name).Get())
+				if (const FMetasoundFrontendClassInput* ClassInput = GraphHandle->FindClassInputWithName(InputToRemove.Name).Get())
 				{
 					if (FMetasoundFrontendClassInput::IsFunctionalEquivalent(*ClassInput, InputToRemove))
 					{
 						bDidEdit = true;
-						InGraph->RemoveInputVertex(InputToRemove.Name);
+						GraphHandle->RemoveInputVertex(InputToRemove.Name);
 					}
 				}
 			}
@@ -79,12 +140,12 @@ namespace Metasound
 			// Remove unrequired outputs
 			for (const FMetasoundFrontendClassVertex& OutputToRemove : OutputsToRemove)
 			{
-				if (const FMetasoundFrontendClassOutput* ClassOutput = InGraph->FindClassOutputWithName(OutputToRemove.Name).Get())
+				if (const FMetasoundFrontendClassOutput* ClassOutput = GraphHandle->FindClassOutputWithName(OutputToRemove.Name).Get())
 				{
 					if (FMetasoundFrontendClassOutput::IsFunctionalEquivalent(*ClassOutput, OutputToRemove))
 					{
 						bDidEdit = true;
-						InGraph->RemoveOutputVertex(OutputToRemove.Name);
+						GraphHandle->RemoveOutputVertex(OutputToRemove.Name);
 					}
 				}
 			}
@@ -138,27 +199,33 @@ namespace Metasound
 			};
 
 			// Add missing inputs
-			for (const FMetasoundFrontendClassVertex& InputToAdd : InputsToAdd)
+			for (const FMetasoundFrontendClassInput& InputToAdd : InputsToAdd)
 			{
 				bDidEdit = true;
-				FNodeHandle NewInputNode = InGraph->AddInputVertex(InputToAdd);
+				FNodeHandle NewInputNode = GraphHandle->AddInputVertex(InputToAdd);
 
-				FMetasoundFrontendNodeStyle Style = NewInputNode->GetNodeStyle();
-				const FVector2D LastOutputLocation = FindLowestNodeLocationOfClassType(EMetasoundFrontendClassType::Input, InGraph, InputToAdd.TypeName, InputDataTypeCompareFilter);
-				Style.Display.Locations.Add(FGuid(), LastOutputLocation + DisplayStyle::NodeLayout::DefaultOffsetY);
-				NewInputNode->SetNodeStyle(Style);
+				if (bSetDefaultNodeLocations)
+				{
+					FMetasoundFrontendNodeStyle Style = NewInputNode->GetNodeStyle();
+					const FVector2D LastOutputLocation = FindLowestNodeLocationOfClassType(EMetasoundFrontendClassType::Input, GraphHandle, InputToAdd.TypeName, InputDataTypeCompareFilter);
+					Style.Display.Locations.Add(FGuid(), LastOutputLocation + DisplayStyle::NodeLayout::DefaultOffsetY);
+					NewInputNode->SetNodeStyle(Style);
+				}
 			}
 
 			// Add missing outputs
-			for (const FMetasoundFrontendClassVertex& OutputToAdd : OutputsToAdd)
+			for (const FMetasoundFrontendClassOutput& OutputToAdd : OutputsToAdd)
 			{
 				bDidEdit = true;
-				FNodeHandle NewOutputNode = InGraph->AddOutputVertex(OutputToAdd);
+				FNodeHandle NewOutputNode = GraphHandle->AddOutputVertex(OutputToAdd);
 
-				FMetasoundFrontendNodeStyle Style = NewOutputNode->GetNodeStyle();
-				const FVector2D LastOutputLocation = FindLowestNodeLocationOfClassType(EMetasoundFrontendClassType::Output, InGraph, OutputToAdd.TypeName, OutputDataTypeCompareFilter);
-				Style.Display.Locations.Add(FGuid(), LastOutputLocation + DisplayStyle::NodeLayout::DefaultOffsetY);
-				NewOutputNode->SetNodeStyle(Style);
+				if (bSetDefaultNodeLocations)
+				{
+					FMetasoundFrontendNodeStyle Style = NewOutputNode->GetNodeStyle();
+					const FVector2D LastOutputLocation = FindLowestNodeLocationOfClassType(EMetasoundFrontendClassType::Output, GraphHandle, OutputToAdd.TypeName, OutputDataTypeCompareFilter);
+					Style.Display.Locations.Add(FGuid(), LastOutputLocation + DisplayStyle::NodeLayout::DefaultOffsetY);
+					NewOutputNode->SetNodeStyle(Style);
+				}
 			}
 
 			// Swap paired inputs.
@@ -172,21 +239,22 @@ namespace Metasound
 				// Cache off node locations and connections to push to new node
 				TMap<FGuid, FVector2D> Locations;
 				TArray<FInputHandle> ConnectedInputs;
-				if (const FMetasoundFrontendClassInput* ClassInput = InGraph->FindClassInputWithName(OriginalVertex.Name).Get())
+				if (const FMetasoundFrontendClassInput* ClassInput = GraphHandle->FindClassInputWithName(OriginalVertex.Name).Get())
 				{
 					if (FMetasoundFrontendVertex::IsFunctionalEquivalent(*ClassInput, OriginalVertex))
 					{
 						NewVertex.DefaultLiteral = ClassInput->DefaultLiteral;
-						FNodeHandle OriginalInputNode = InGraph->GetInputNodeWithName(OriginalVertex.Name);
+						NewVertex.NodeID = ClassInput->NodeID;
+						FNodeHandle OriginalInputNode = GraphHandle->GetInputNodeWithName(OriginalVertex.Name);
 						Locations = OriginalInputNode->GetNodeStyle().Display.Locations;
 
 						FOutputHandle OriginalInputNodeOutput = OriginalInputNode->GetOutputWithVertexName(OriginalVertex.Name);
 						ConnectedInputs = OriginalInputNodeOutput->GetConnectedInputs();
-						InGraph->RemoveInputVertex(OriginalVertex.Name);
+						GraphHandle->RemoveInputVertex(OriginalVertex.Name);
 					}
 				}
 
-				FNodeHandle NewInputNode = InGraph->AddInputVertex(NewVertex);
+				FNodeHandle NewInputNode = GraphHandle->AddInputVertex(NewVertex);
 				// Copy prior node locations
 				if (!Locations.IsEmpty())
 				{
@@ -210,7 +278,7 @@ namespace Metasound
 				bDidEdit = true;
 
 				const FMetasoundFrontendClassVertex& OriginalVertex = OutputPair.Get<0>();
-				const FMetasoundFrontendClassVertex& NewVertex = OutputPair.Get<1>();
+				FMetasoundFrontendClassVertex NewVertex = OutputPair.Get<1>();
 
 				// Cache off node locations to push to new node
 				TMap<FGuid, FVector2D> Locations;
@@ -218,19 +286,20 @@ namespace Metasound
 
 				// Default add output node to origin.
 				Locations.Add(FGuid(), FVector2D{0.f, 0.f});
-				if (const FMetasoundFrontendClassOutput* ClassOutput = InGraph->FindClassOutputWithName(OriginalVertex.Name).Get())
+				if (const FMetasoundFrontendClassOutput* ClassOutput = GraphHandle->FindClassOutputWithName(OriginalVertex.Name).Get())
 				{
 					if (FMetasoundFrontendVertex::IsFunctionalEquivalent(*ClassOutput, OriginalVertex))
 					{
-						FNodeHandle OriginalOutputNode = InGraph->GetOutputNodeWithName(OriginalVertex.Name);
+						NewVertex.NodeID = ClassOutput->NodeID;
+						FNodeHandle OriginalOutputNode = GraphHandle->GetOutputNodeWithName(OriginalVertex.Name);
 						Locations = OriginalOutputNode->GetNodeStyle().Display.Locations;
 						FInputHandle Input = OriginalOutputNode->GetInputWithVertexName(OriginalVertex.Name);
 						ConnectedOutput = Input->GetConnectedOutput();
-						InGraph->RemoveOutputVertex(OriginalVertex.Name);
+						GraphHandle->RemoveOutputVertex(OriginalVertex.Name);
 					}
 				}
 
-				FNodeHandle NewOutputNode = InGraph->AddOutputVertex(NewVertex);
+				FNodeHandle NewOutputNode = GraphHandle->AddOutputVertex(NewVertex);
 
 				if (Locations.Num() > 0)
 				{
@@ -247,52 +316,41 @@ namespace Metasound
 			return bDidEdit;
 		}
 
-		bool FMatchRootGraphToArchetype::Transform(FDocumentHandle InDocument) const
+		bool FUpdateRootGraphInterface::Transform(FDocumentHandle InDocument) const
 		{
 			bool bDidEdit = false;
 
-			if (!InDocument->IsValid())
+			if (!ensure(InDocument->IsValid()))
 			{
 				return bDidEdit;
 			}
 
-			// Find registered target archetype.
-			FMetasoundFrontendArchetype TargetArchetype;
-			bool bFoundTargetArchetype = IArchetypeRegistry::Get().FindArchetype(GetArchetypeRegistryKey(ArchetypeVersion), TargetArchetype);
-
-			if (!bFoundTargetArchetype)
+			// Find registered target interface.
+			FMetasoundFrontendInterface TargetInterface;
+			bool bFoundTargetInterface = ISearchEngine::Get().FindInterfaceWithHighestVersion(InterfaceVersion.Name, TargetInterface);
+			if (!bFoundTargetInterface)
 			{
-				UE_LOG(LogMetaSound, Error, TEXT("Target archetype is not registered [ArchetypeVersion:%s]"), *ArchetypeVersion.ToString());
+				UE_LOG(LogMetaSound, Error, TEXT("Target interface is not registered [InterfaceVersion:%s]"), *InterfaceVersion.ToString());
 				return false;
 			}
 
-			// Get current archetype version on document.
-			const FMetasoundFrontendVersion InitialArchetypeVersion = InDocument->GetArchetypeVersion();
-
-			// Attempt to upgrade
-			TArray<const IArchetypeRegistryEntry*> UpgradePath;
-			GetUpgradePathForDocument(InitialArchetypeVersion, ArchetypeVersion, UpgradePath);
-
-			if (UpgradeDocumentArchetype(UpgradePath, InDocument))
-			{
-				bDidEdit = true;
-			}
-
-			// Force archetype to conform
-			if (ConformDocumentToArchetype(TargetArchetype, InDocument))
+			if (TargetInterface.Version == InterfaceVersion)
 			{
-				bDidEdit = true;
+				return false;
 			}
 
-			return bDidEdit;
+			// Attempt to upgrade
+			TArray<const IInterfaceRegistryEntry*> UpgradePath;
+			GetUpdatePathForDocument(InterfaceVersion, TargetInterface.Version, UpgradePath);
+			return UpdateDocumentInterface(UpgradePath, InDocument);
 		}
 
-		void FMatchRootGraphToArchetype::GetUpgradePathForDocument(const FMetasoundFrontendVersion& InCurrentVersion, const FMetasoundFrontendVersion& InTargetVersion, TArray<const IArchetypeRegistryEntry*>& OutUpgradePath) const
+		void FUpdateRootGraphInterface::GetUpdatePathForDocument(const FMetasoundFrontendVersion& InCurrentVersion, const FMetasoundFrontendVersion& InTargetVersion, TArray<const IInterfaceRegistryEntry*>& OutUpgradePath) const
 		{
 			if (InCurrentVersion.Name == InTargetVersion.Name)
 			{
-				// Get all associated registered archetypes
-				TArray<FMetasoundFrontendVersion> RegisteredVersions = ISearchEngine::Get().FindAllRegisteredArchetypesWithName(InTargetVersion.Name);
+				// Get all associated registered interfaces
+				TArray<FMetasoundFrontendVersion> RegisteredVersions = ISearchEngine::Get().FindAllRegisteredInterfacesWithName(InTargetVersion.Name);
 
 				// Filter registry entries that exist between current version and target version
 				auto FilterRegistryEntries = [&InCurrentVersion, &InTargetVersion](const FMetasoundFrontendVersion& InVersion)
@@ -310,62 +368,37 @@ namespace Metasound
 				// Get registry entries from registry keys.
 				auto GetRegistryEntry = [](const FMetasoundFrontendVersion& InVersion)
 				{
-					FArchetypeRegistryKey Key = GetArchetypeRegistryKey(InVersion);
-					return IArchetypeRegistry::Get().FindArchetypeRegistryEntry(Key);
+					FInterfaceRegistryKey Key = GetInterfaceRegistryKey(InVersion);
+					return IInterfaceRegistry::Get().FindInterfaceRegistryEntry(Key);
 				};
 				Algo::Transform(RegisteredVersions, OutUpgradePath, GetRegistryEntry);
 			}
 		}
 
-		bool FMatchRootGraphToArchetype::UpgradeDocumentArchetype(const TArray<const IArchetypeRegistryEntry*>& InUpgradePath, FDocumentHandle InDocument) const
+		bool FUpdateRootGraphInterface::UpdateDocumentInterface(const TArray<const IInterfaceRegistryEntry*>& InUpgradePath, FDocumentHandle InDocument) const
 		{
-			bool bDidEdit = false;
-			for (const IArchetypeRegistryEntry* Entry : InUpgradePath)
+			const FMetasoundFrontendVersionNumber* LastVersionUpdated = nullptr;
+			for (const IInterfaceRegistryEntry* Entry : InUpgradePath)
 			{
 				if (ensure(nullptr != Entry))
 				{
-					if (Entry->UpdateRootGraphArchetype(InDocument))
+					if (Entry->UpdateRootGraphInterface(InDocument))
 					{
-						bDidEdit = true;
-						InDocument->SetArchetypeVersion(Entry->GetArchetype().Version);
+						LastVersionUpdated = &Entry->GetInterface().Version.Number;
 					}
 				}
 			}
-			return bDidEdit;
-		}
 
-		bool FMatchRootGraphToArchetype::ConformDocumentToArchetype(const FMetasoundFrontendArchetype& InTargetArchetype, FDocumentHandle InDocument) const
-		{
-			bool bDidEdit = false;
-
-			FMetasoundFrontendArchetype CurrentArchetype;
-			FArchetypeRegistryKey CurrentRegistryKey = GetArchetypeRegistryKey(InDocument->GetArchetypeVersion());
-			const bool bFoundCurrentArchetype = IArchetypeRegistry::Get().FindArchetype(CurrentRegistryKey, CurrentArchetype);
-
-			if (!bFoundCurrentArchetype)
-			{
-				UE_LOG(LogMetaSound, Warning, TEXT("Failed to find current archetype on document [ArchetypeVersion:%s]"), *InDocument->GetArchetypeVersion().ToString());
-			}
-
-			const bool bIsEqualArchetypeVersion = InDocument->GetArchetypeVersion() == InTargetArchetype.Version;
-			const bool bRequiredInterfaceExists = IsSubsetOfClass(InTargetArchetype, InDocument->GetRootGraphClass());
-
-			if (!bIsEqualArchetypeVersion)
-			{
-				InDocument->SetArchetypeVersion(InTargetArchetype.Version);
-				bDidEdit = true;
-			}
-
-			if (!(bRequiredInterfaceExists && bIsEqualArchetypeVersion))
+			if (LastVersionUpdated)
 			{
-				FGraphHandle Graph = InDocument->GetRootGraph();
-				if (FSwapGraphArchetype(CurrentArchetype, InTargetArchetype).Transform(Graph))
-				{
-					bDidEdit = true;
-				}
+				UE_LOG(LogMetaSound, Display, TEXT("Asset '%s' interface '%s' updated: '%s' --> '%s'"),
+					*InDocument->GetMetadata().Version.ToString(),
+					*InterfaceVersion.Number.ToString(),
+					*LastVersionUpdated->ToString());
+				return true;
 			}
 
-			return bDidEdit;
+			return false;
 		}
 
 		bool FAutoUpdateRootGraph::Transform(FDocumentHandle InDocument) const
@@ -421,7 +454,7 @@ namespace Metasound
 					FMetasoundAssetBase* ParentMetaSoundAsset = IMetaSoundAssetManager::GetChecked().TryLoadAssetFromKey(RegistryKey);
 					if (ensure(ParentMetaSoundAsset))
 					{
-						ParentMetaSoundAsset->ConformObjectDataToArchetype();
+						ParentMetaSoundAsset->ConformObjectDataToInterfaces();
 					}
 				}
 			}
@@ -467,10 +500,6 @@ namespace Metasound
 				return false;
 			}
 
-			// Run transform to ensure preset matches reference archetype
-			const FMetasoundFrontendVersion& RefArchetypeVersion = ReferencedDocument->GetArchetypeVersion();
-			FMatchRootGraphToArchetype(RefArchetypeVersion).Transform(InDocument);
-
 			// Determine the inputs and outputs needed in the wrapping graph. Also
 			// cache any exiting literals that have been set on the wrapping graph.
 			TArray<FMetasoundFrontendClassInput> ClassInputs = GenerateRequiredClassInputs(RootGraphHandle);
@@ -485,6 +514,15 @@ namespace Metasound
 			// Clear the root graph so it can be rebuilt.
 			RootGraphHandle->ClearGraph();
 
+			// Ensure preset interfaces match those found in referenced graph.  Referenced graph is assumed to be
+			// well-formed (i.e. all inputs/outputs/environment variables declared by interfaces are present, and
+			// of proper name & data type).
+			const TArray<FMetasoundFrontendVersion>& RefInterfaceVersions = ReferencedDocument->GetInterfaceVersions();
+			for (const FMetasoundFrontendVersion& Version : RefInterfaceVersions)
+			{
+				InDocument->AddInterfaceVersion(Version);
+			}
+
 			// Add referenced node
 			FMetasoundFrontendClassMetadata ReferencedClassMetadata = ReferencedGraphHandle->GetGraphMetadata();
 			// Swap type on look-up as it will be referenced as an externally defined class relative to the new Preset asset
@@ -818,34 +856,33 @@ namespace Metasound
 
 			void TransformInternal(FDocumentHandle InDocument) const override
 			{
-				
 				check(InDocument->GetMetadata().Version.Number.Major == 1);
 				check(InDocument->GetMetadata().Version.Number.Minor == 3);
 
-				FMetasoundFrontendVersion ArchetypeVersion = InDocument->GetArchetypeVersion();
+				const TArray<FMetasoundFrontendVersion>& InterfaceVersions = InDocument->GetInterfaceVersions();
 
-				// Version 1.3 did not have an "ArchetypeVersion" property on the 
+				// Version 1.3 did not have an "InterfaceVersion" property on the 
 				// document, so any document that is being updated should start off
-				// with an "Invalid" archetype version.
-				if (ensure(!ArchetypeVersion.IsValid()))
+				// with an "Invalid" interface version.
+				if (ensure(InterfaceVersions.IsEmpty()))
 				{
-					constexpr bool bIncludeDeprecatedArchetypes = true;
-					TArray<FMetasoundFrontendArchetype> AllArchetypes = ISearchEngine::Get().FindAllArchetypes(bIncludeDeprecatedArchetypes);
+					constexpr bool bIncludeDeprecatedInterfaces = true;
+					TArray<FMetasoundFrontendInterface> AllInterfaces = ISearchEngine::Get().FindAllInterfaces(bIncludeDeprecatedInterfaces);
 
 					const FMetasoundFrontendGraphClass& RootGraph = InDocument->GetRootGraphClass();
 					const TArray<FMetasoundFrontendClass>& Dependencies = InDocument->GetDependencies();
 					const TArray<FMetasoundFrontendGraphClass>& Subgraphs = InDocument->GetSubgraphs();
 
-					if (const FMetasoundFrontendArchetype* Arch = FindMostSimilarArchetypeSupportingEnvironment(RootGraph, Dependencies, Subgraphs, AllArchetypes))
+					if (const FMetasoundFrontendInterface* Interface = FindMostSimilarInterfaceSupportingEnvironment(RootGraph, Dependencies, Subgraphs, AllInterfaces))
 					{
-						UE_LOG(LogMetaSound, Display, TEXT("Assigned archetype [ArchetypeVersion:%s] to document [RootGraphClassName:%s]"),
-							*Arch->Version.ToString(), *RootGraph.Metadata.GetClassName().ToString());
+						UE_LOG(LogMetaSound, Display, TEXT("Assigned interface [InterfaceVersion:%s] to document [RootGraphClassName:%s]"),
+							*Interface->Version.ToString(), *RootGraph.Metadata.GetClassName().ToString());
 
-						InDocument->SetArchetypeVersion(Arch->Version);
+						InDocument->AddInterfaceVersion(Interface->Version);
 					}
 					else
 					{
-						UE_LOG(LogMetaSound, Warning, TEXT("Failed to find archetype for document [RootGraphClassName:%s]"),
+						UE_LOG(LogMetaSound, Warning, TEXT("Failed to find interface for document [RootGraphClassName:%s]"),
 							*RootGraph.Metadata.GetClassName().ToString());
 					}
 				}
@@ -913,7 +950,7 @@ namespace Metasound
 					// unique names and customized display names (ex. 'Audio' for both mono &
 					// L/R output, On Play, & 'On Finished'), so do not replace them by nulling
 					// out the guid as a name and using the converted FName of the FText DisplayName.
-					if (!NodeHandle->IsRequired())
+					if (!NodeHandle->IsInterfaceMember())
 					{
 						const FName NewNodeName = *NodeHandle->GetDisplayName().ToString();
 						NodeHandle->IterateInputs([&](FInputHandle InputHandle)
diff --git a/Engine/Plugins/Runtime/Metasound/Source/MetasoundFrontend/Private/MetasoundParameterTransmitter.cpp b/Engine/Plugins/Runtime/Metasound/Source/MetasoundFrontend/Private/MetasoundParameterTransmitter.cpp
--- a/Engine/Plugins/Runtime/Metasound/Source/MetasoundFrontend/Private/MetasoundParameterTransmitter.cpp
+++ b/Engine/Plugins/Runtime/Metasound/Source/MetasoundFrontend/Private/MetasoundParameterTransmitter.cpp
@@ -2,11 +2,130 @@
 
 #include "MetasoundParameterTransmitter.h"
 
+#include "IAudioGeneratorInterfaceRegistry.h"
 #include "MetasoundLog.h"
 
 
 namespace Metasound
 {
+	namespace Frontend
+	{
+		FLiteral ConvertParameterToLiteral(FAudioParameter&& InValue)
+		{
+			switch (InValue.ParamType)
+			{
+				case EAudioParameterType::Boolean:
+				{
+					return FLiteral(InValue.BoolParam);
+				}
+
+				case EAudioParameterType::BooleanArray:
+				{
+					return FLiteral(MoveTemp(InValue.ArrayBoolParam));
+				}
+
+				case EAudioParameterType::Float:
+				{
+					return FLiteral(InValue.FloatParam);
+				}
+
+				case EAudioParameterType::FloatArray:
+				{
+					return FLiteral(MoveTemp(InValue.ArrayFloatParam));
+				}
+
+				case EAudioParameterType::Integer:
+				{
+					return FLiteral(InValue.IntParam);
+				}
+
+				case EAudioParameterType::IntegerArray:
+				{
+					return FLiteral(MoveTemp(InValue.ArrayIntParam));
+				}
+
+				case EAudioParameterType::None:
+				{
+					return FLiteral();
+				}
+
+				case EAudioParameterType::NoneArray:
+				{
+					TArray<FLiteral::FNone> InitArray;
+					InitArray.Init(FLiteral::FNone(), InValue.IntParam);
+					return FLiteral(MoveTemp(InitArray));
+				}
+
+				case EAudioParameterType::Object:
+				{
+					if (InValue.ObjectProxies.IsEmpty())
+					{
+						return FLiteral();
+					}
+
+					return FLiteral(MoveTemp(InValue.ObjectProxies[0]));
+				}
+
+				case EAudioParameterType::ObjectArray:
+				{
+					return FLiteral(MoveTemp(InValue.ObjectProxies));
+				}
+
+				case EAudioParameterType::String:
+				{
+					return FLiteral(MoveTemp(InValue.StringParam));
+				}
+
+				case EAudioParameterType::StringArray:
+				{
+					return FLiteral(MoveTemp(InValue.ArrayStringParam));
+				}
+
+				default:
+				{
+					static_assert(static_cast<int32>(EAudioParameterType::COUNT) == 12, "Possible missing switch case coverage");
+					checkNoEntry();
+				}
+			}
+
+			return FLiteral();
+		}
+
+		FName ConvertParameterToDataType(EAudioParameterType InParameterType)
+		{
+			switch (InParameterType)
+			{
+				case EAudioParameterType::Boolean:
+					return GetMetasoundDataTypeName<bool>();
+				case EAudioParameterType::BooleanArray:
+					return GetMetasoundDataTypeName<TArray<bool>>();
+				case EAudioParameterType::Float:
+					return GetMetasoundDataTypeName<float>();
+				case EAudioParameterType::FloatArray:
+					return GetMetasoundDataTypeName<TArray<float>>();
+				case EAudioParameterType::Integer:
+					return GetMetasoundDataTypeName<int32>();
+				case EAudioParameterType::IntegerArray:
+					return GetMetasoundDataTypeName<TArray<int32>>();
+				case EAudioParameterType::String:
+					return GetMetasoundDataTypeName<FString>();
+				case EAudioParameterType::StringArray:
+					return GetMetasoundDataTypeName<TArray<FString>>();
+
+				case EAudioParameterType::Object:
+				case EAudioParameterType::ObjectArray:
+					// TODO: Add support for objects
+
+				case EAudioParameterType::None:
+				case EAudioParameterType::NoneArray:
+				default:
+					ensureAlwaysMsgf(false, TEXT("Failed to convert AudioParameterType to POD MetaSound DataType"));
+					static_assert(static_cast<int32>(EAudioParameterType::COUNT) == 12, "Possible missing case coverage");
+					return FName();
+			}
+		}
+	}
+
 	const FVertexName& FMetaSoundParameterTransmitter::GetInstanceIDEnvironmentVariableName()
 	{
 		static const FVertexName VarName = "TransmitterInstanceID";
@@ -55,84 +174,17 @@ namespace Metasound
 		return InstanceID;
 	}
 
-	bool FMetaSoundParameterTransmitter::SetParameter(FAudioParameter&& InValue)
+	bool FMetaSoundParameterTransmitter::SetParameter(FAudioParameter&& InParameter)
 	{
-		switch (InValue.ParamType)
-		{
-			case EAudioParameterType::Boolean:
-			{
-				return SetParameterWithLiteral(InValue.ParamName, FLiteral(InValue.BoolParam));
-			}
-
-			case EAudioParameterType::BooleanArray:
-			{
-				return SetParameterWithLiteral(InValue.ParamName, FLiteral(MoveTemp(InValue.ArrayBoolParam)));
-			}
-
-			case EAudioParameterType::Float:
-			{
-				return SetParameterWithLiteral(InValue.ParamName, FLiteral(InValue.FloatParam));
-			}
-
-			case EAudioParameterType::FloatArray:
-			{
-				return SetParameterWithLiteral(InValue.ParamName, FLiteral(MoveTemp(InValue.ArrayFloatParam)));
-			}
-
-			case EAudioParameterType::Integer:
-			{
-				return SetParameterWithLiteral(InValue.ParamName, FLiteral(InValue.IntParam));
-			}
-
-			case EAudioParameterType::IntegerArray:
-			{
-				return SetParameterWithLiteral(InValue.ParamName, FLiteral(MoveTemp(InValue.ArrayIntParam)));
-			}
-
-			case EAudioParameterType::None:
-			{
-				return SetParameterWithLiteral(InValue.ParamName, FLiteral());
-			}
-
-			case EAudioParameterType::NoneArray:
-			{
-				TArray<FLiteral::FNone> InitArray;
-				InitArray.Init(FLiteral::FNone(), InValue.IntParam);
-				return SetParameterWithLiteral(InValue.ParamName, FLiteral(MoveTemp(InitArray)));
-			}
-
-			case EAudioParameterType::Object:
-			{
-				if(InValue.ObjectProxies.IsEmpty())
-				{
-					return SetParameterWithLiteral(InValue.ParamName, FLiteral());
-				}
-
-				return SetParameterWithLiteral(InValue.ParamName, FLiteral(MoveTemp(InValue.ObjectProxies[0])));
-			}
-
-			case EAudioParameterType::ObjectArray:
-			{
-				return SetParameterWithLiteral(InValue.ParamName, FLiteral(MoveTemp(InValue.ObjectProxies)));
-			}
-
-			case EAudioParameterType::String:
-			{
-				return SetParameterWithLiteral(InValue.ParamName, FLiteral(MoveTemp(InValue.StringParam)));
-			}
-
-			case EAudioParameterType::StringArray:
-			{
-				return SetParameterWithLiteral(InValue.ParamName, FLiteral(MoveTemp(InValue.ArrayStringParam)));
-			}
-
-			default:
-			{
-				checkNoEntry();
-			}
-		}
+		const FName ParamName = InParameter.ParamName;
+		return SetParameterWithLiteral(ParamName, Frontend::ConvertParameterToLiteral(MoveTemp(InParameter)));
+	}
 
-		return false;
+	bool FMetaSoundParameterTransmitter::SetParameter(FName InInterfaceName, FAudioParameter&& InParameter)
+	{
+		InParameter.ParamName = Audio::IGeneratorInterfaceRegistry::GetMemberFullName(InInterfaceName, InParameter.ParamName);
+		const FName ParamName = InParameter.ParamName;
+		return SetParameterWithLiteral(ParamName, Frontend::ConvertParameterToLiteral(MoveTemp(InParameter)));
 	}
 
 	bool FMetaSoundParameterTransmitter::SetParameterWithLiteral(FName InParameterName, const FLiteral& InLiteral)
diff --git a/Engine/Plugins/Runtime/Metasound/Source/MetasoundFrontend/Public/MetasoundArchetype.h b/Engine/Plugins/Runtime/Metasound/Source/MetasoundFrontend/Public/MetasoundArchetype.h
--- a/Engine/Plugins/Runtime/Metasound/Source/MetasoundFrontend/Public/MetasoundArchetype.h
+++ b/Engine/Plugins/Runtime/Metasound/Source/MetasoundFrontend/Public/MetasoundArchetype.h
@@ -8,24 +8,25 @@ namespace Metasound
 {
 	namespace Frontend
 	{
+		METASOUNDFRONTEND_API bool IsSubsetOfInterface(const FMetasoundFrontendInterface& InSubsetInterface, const FMetasoundFrontendInterface& InSupersetInterface);
 
-		METASOUNDFRONTEND_API bool IsSubsetOfArchetype(const FMetasoundFrontendArchetype& InSubsetArchetype, const FMetasoundFrontendArchetype& InSupersetArchetype);
+		METASOUNDFRONTEND_API bool IsSubsetOfClass(const FMetasoundFrontendInterface& InSubsetInterface, const FMetasoundFrontendClass& InSupersetClass);
 
-		METASOUNDFRONTEND_API bool IsSubsetOfClass(const FMetasoundFrontendArchetype& InSubsetArchetype, const FMetasoundFrontendClass& InSupersetClass);
+		METASOUNDFRONTEND_API bool IsEquivalentInterface(const FMetasoundFrontendInterface& InInputInterface, const FMetasoundFrontendInterface& InTargetInterface);
 
-		METASOUNDFRONTEND_API bool IsEquivalentArchetype(const FMetasoundFrontendArchetype& InInputArchetype, const FMetasoundFrontendArchetype& InTargetArchetype);
+		METASOUNDFRONTEND_API bool IsEqualInterface(const FMetasoundFrontendInterface& InInputInterface, const FMetasoundFrontendInterface& InTargetInterface);
 
-		METASOUNDFRONTEND_API bool IsEqualArchetype(const FMetasoundFrontendArchetype& InInputArchetype, const FMetasoundFrontendArchetype& InTargetArchetype);
+		METASOUNDFRONTEND_API int32 InputOutputDifferenceCount(const FMetasoundFrontendClass& InClass, const FMetasoundFrontendInterface& InInterface);
 
-		METASOUNDFRONTEND_API int32 InputOutputDifferenceCount(const FMetasoundFrontendClass& InClass, const FMetasoundFrontendArchetype& InArchetype);
+		METASOUNDFRONTEND_API int32 InputOutputDifferenceCount(const FMetasoundFrontendInterface& InInterfaceA, const FMetasoundFrontendInterface& InInterfaceB);
 
-		METASOUNDFRONTEND_API int32 InputOutputDifferenceCount(const FMetasoundFrontendArchetype& InArchetypeA, const FMetasoundFrontendArchetype& InArchetypeB);
+		METASOUNDFRONTEND_API bool DeclaredInterfaceVersionsMatch(const FMetasoundFrontendDocument& InDocumentA, const FMetasoundFrontendDocument& InDocumentB);
 
-		METASOUNDFRONTEND_API void GatherRequiredEnvironmentVariables(const FMetasoundFrontendClass& InClass, TArray<FMetasoundFrontendEnvironmentVariable>& OutEnvironmentVariables);
+		METASOUNDFRONTEND_API void GatherRequiredEnvironmentVariables(const FMetasoundFrontendClass& InClass, TArray<FMetasoundFrontendClassEnvironmentVariable>& OutEnvironmentVariables);
 
-		METASOUNDFRONTEND_API const FMetasoundFrontendArchetype* FindMostSimilarArchetypeSupportingEnvironment(const FMetasoundFrontendGraphClass& InRootGraph, const TArray<FMetasoundFrontendClass>& InDependencies, const TArray<FMetasoundFrontendGraphClass>& InSubgraphs, const TArray<FMetasoundFrontendArchetype>& InCandidateArchetypes);
+		METASOUNDFRONTEND_API const FMetasoundFrontendInterface* FindMostSimilarInterfaceSupportingEnvironment(const FMetasoundFrontendGraphClass& InRootGraph, const TArray<FMetasoundFrontendClass>& InDependencies, const TArray<FMetasoundFrontendGraphClass>& InSubgraphs, const TArray<FMetasoundFrontendInterface>& InCandidateInterfaces);
 
-		METASOUNDFRONTEND_API const FMetasoundFrontendArchetype* FindMostSimilarArchetypeSupportingEnvironment(const FMetasoundFrontendDocument& InDocument, const TArray<FMetasoundFrontendArchetype>& InCandidateArchetypes);
+		METASOUNDFRONTEND_API const FMetasoundFrontendInterface* FindMostSimilarInterfaceSupportingEnvironment(const FMetasoundFrontendDocument& InDocument, const TArray<FMetasoundFrontendInterface>& InCandidateInterfaces);
 	}
 }
 
diff --git a/Engine/Plugins/Runtime/Metasound/Source/MetasoundFrontend/Public/MetasoundAssetBase.h b/Engine/Plugins/Runtime/Metasound/Source/MetasoundFrontend/Public/MetasoundAssetBase.h
--- a/Engine/Plugins/Runtime/Metasound/Source/MetasoundFrontend/Public/MetasoundAssetBase.h
+++ b/Engine/Plugins/Runtime/Metasound/Source/MetasoundFrontend/Public/MetasoundAssetBase.h
@@ -18,10 +18,17 @@
 class FMetasoundAssetBase;
 class UEdGraph;
 
+namespace Metasound
+{
+	namespace Frontend {
+		class IInterfaceRegistryEntry;
+	} // namespace Frontend
+} // namespace Metasound
+
 
 /** FMetasoundAssetBase is intended to be a mix-in subclass for UObjects which utilize
  * Metasound assets.  It provides consistent access to FMetasoundFrontendDocuments, control
- * over the FMetasoundFrontendArchetype of the FMetasoundFrontendDocument.  It also enables the UObject
+ * over the FMetasoundFrontendClassInterface of the FMetasoundFrontendDocument.  It also enables the UObject
  * to be utilized by a host of other engine tools built to support MetaSounds.
  */
 class METASOUNDFRONTEND_API FMetasoundAssetBase
@@ -53,10 +60,10 @@ public:
 	virtual void SetRegistryAssetClassInfo(const Metasound::Frontend::FNodeClassInfo& InClassInfo) = 0;
 #endif // WITH_EDITORONLY_DATA
 
-	// Called when the archetype is changed, presenting the opportunity for
-	// any reflected object data to be updated based on the new archetype.
+	// Called when the interface is changed, presenting the opportunity for
+	// any reflected object data to be updated based on the new interface.
 	// Returns whether or not any edits were made.
-	virtual bool ConformObjectDataToArchetype() = 0;
+	virtual bool ConformObjectDataToInterfaces() = 0;
 
 	// Registers the root graph of the given asset with the MetaSound Frontend.
 	void RegisterGraphWithFrontend(Metasound::Frontend::FMetaSoundAssetRegistrationOptions InRegistrationOptions = Metasound::Frontend::FMetaSoundAssetRegistrationOptions());
@@ -67,18 +74,14 @@ public:
 	// Sets/overwrites the root class metadata
 	virtual void SetMetadata(FMetasoundFrontendClassMetadata& InMetadata);
 
-	// Returns  a description of the required inputs and outputs for this metasound UClass.
+	// Returns a default interface required when a given asset is created.
 	virtual const FMetasoundFrontendVersion& GetDefaultArchetypeVersion() const = 0;
 
-	// Returns true if the archetype is supported by this object.
-	virtual bool IsArchetypeSupported(const FMetasoundFrontendVersion& InArchetypeVersion) const;
+	// Returns the interface entries declared by the given asset's document from the InterfaceRegistry.
+	bool GetDeclaredInterfaces(TArray<const Metasound::Frontend::IInterfaceRegistryEntry*>& OutInterfaces) const;
 
-	// Returns an array of archetypes preferred for this class.
-	virtual const TArray<FMetasoundFrontendVersion>& GetSupportedArchetypeVersions() const = 0;
-
-	// Returns the preferred archetype for the given document.
-	virtual FMetasoundFrontendVersion GetPreferredArchetypeVersion(const FMetasoundFrontendDocument& InDocument) const;
-	bool GetArchetype(FMetasoundFrontendArchetype& OutArchetype) const;
+	// Returns whether an interface with the given name is declared by the given asset's document.
+	bool IsInterfaceDeclared(FName InName) const;
 
 	// Gets the asset class info.
 	virtual Metasound::Frontend::FNodeClassInfo GetAssetClassInfo() const = 0;
@@ -113,24 +116,24 @@ public:
 	Metasound::Frontend::FGraphHandle GetRootGraphHandle();
 	Metasound::Frontend::FConstGraphHandle GetRootGraphHandle() const;
 
-	// Overwrites the existing document. If the document's archetype is not supported,
-	// the FMetasoundAssetBase be while queried for a new one using `GetPreferredArchetype`. If `bForceUpdateArchetype`
-	// is true, `GetPreferredArchetype` will be used whether or not the provided document's archetype
-	// is supported. 
+	// Overwrites the existing document. If the document's interface is not supported,
+	// the FMetasoundAssetBase be while queried for a new one using `GetPreferredInterface`.
 	void SetDocument(const FMetasoundFrontendDocument& InDocument);
 
 	FMetasoundFrontendDocument& GetDocumentChecked();
 	const FMetasoundFrontendDocument& GetDocumentChecked() const;
 
-	void ConformDocumentToArchetype();
+	void AddDefaultInterfaces();
 	bool AutoUpdate(bool bInMarkDirty = false);
 	bool VersionAsset();
 
 #if WITH_EDITORONLY_DATA
 	bool GetSynchronizationPending() const;
 	bool GetSynchronizationClearUpdateNotes() const;
+	bool GetSynchronizationInterfacesUpdated() const;
 	void ResetSynchronizationState();
 	void SetClearNodeNotesOnSynchronization();
+	void SetInterfacesUpdatedOnSynchronization();
 	void SetSynchronizationRequired();
 #endif // WITH_EDITORONLY_DATA
 
@@ -149,6 +152,8 @@ public:
 	// Returns the owning asset responsible for transactions applied to MetaSound
 	virtual const UObject* GetOwningAsset() const = 0;
 
+	FString GetOwningAssetName() const;
+
 protected:
 	virtual void SetReferencedAssetClassKeys(TSet<Metasound::Frontend::FNodeRegistryKey>&& InKeys) = 0;
 
@@ -165,14 +170,13 @@ protected:
 	// Returns an access pointer to the document.
 	virtual Metasound::Frontend::FConstDocumentAccessPtr GetDocument() const = 0;
 
-	FString GetOwningAssetName() const;
-
 	// Returns a shared instance of the core metasound graph.
 	TSharedPtr<const Metasound::IGraph, ESPMode::ThreadSafe> GetMetasoundCoreGraph() const;
 
 #if WITH_EDITORONLY_DATA
 	bool bSynchronizationRequired = true;
 	bool bSynchronizationClearUpdateNotes = false;
+	bool bSynchronizationInterfacesUpdated = false;
 #endif // WITH_EDITORONLY_DATA
 
 private:
@@ -199,9 +203,7 @@ private:
 	const FRuntimeData& GetRuntimeData() const;
 
 	TSharedPtr<Metasound::IGraph, ESPMode::ThreadSafe> BuildMetasoundDocument() const;
-	TArray<FMetasoundFrontendClassInput> GetTransmittableClassInputs() const;
-	TSet<Metasound::FVertexName> GetNonTransmittableInputVertices(const FMetasoundFrontendDocument& InDoc) const;
-	TArray<Metasound::FVertexName> GetTransmittableInputVertexNames() const;
+	TArray<const FMetasoundFrontendClassInput*> GetTransmittableClassInputs() const;
 	Metasound::FSendAddress CreateSendAddress(uint64 InInstanceID, const Metasound::FVertexName& InVertexName, const FName& InDataTypeName) const;
 	Metasound::Frontend::FNodeHandle AddInputPinForSendAddress(const Metasound::FMetaSoundParameterTransmitter::FSendInfo& InSendInfo, Metasound::Frontend::FGraphHandle InGraph) const;
 };
diff --git a/Engine/Plugins/Runtime/Metasound/Source/MetasoundFrontend/Public/MetasoundFrontendArchetypeRegistry.h b/Engine/Plugins/Runtime/Metasound/Source/MetasoundFrontend/Public/MetasoundFrontendArchetypeRegistry.h
--- a/Engine/Plugins/Runtime/Metasound/Source/MetasoundFrontend/Public/MetasoundFrontendArchetypeRegistry.h
+++ b/Engine/Plugins/Runtime/Metasound/Source/MetasoundFrontend/Public/MetasoundFrontendArchetypeRegistry.h
@@ -2,65 +2,78 @@
 
 #pragma once
 
-#include "CoreMinimal.h"
 #include "MetasoundFrontendController.h"
 #include "MetasoundFrontendDocument.h"
+#include "UObject/NoExportTypes.h"
 
 namespace Metasound
 {
 	namespace Frontend
 	{
-		using FArchetypeRegistryKey = FString;
+		using FInterfaceRegistryKey = FString;
 		using FRegistryTransactionID = int32;
 
-		METASOUNDFRONTEND_API bool IsValidArchetypeRegistryKey(const FArchetypeRegistryKey& InKey);
-		METASOUNDFRONTEND_API FArchetypeRegistryKey GetArchetypeRegistryKey(const FMetasoundFrontendVersion& InArchetypeVersion);
-		METASOUNDFRONTEND_API FArchetypeRegistryKey GetArchetypeRegistryKey(const FMetasoundFrontendArchetype& InArchetype);
+		METASOUNDFRONTEND_API bool IsValidInterfaceRegistryKey(const FInterfaceRegistryKey& InKey);
+		METASOUNDFRONTEND_API FInterfaceRegistryKey GetInterfaceRegistryKey(const FMetasoundFrontendVersion& InInterfaceVersion);
+		METASOUNDFRONTEND_API FInterfaceRegistryKey GetInterfaceRegistryKey(const FMetasoundFrontendInterface& InInterface);
 
-		class IArchetypeRegistryEntry
+		class METASOUNDFRONTEND_API IInterfaceRegistryEntry
 		{
 		public:
-			virtual ~IArchetypeRegistryEntry() = default;
+			virtual ~IInterfaceRegistryEntry() = default;
 
-			virtual const FMetasoundFrontendArchetype& GetArchetype() const = 0;
-			virtual bool UpdateRootGraphArchetype(FDocumentHandle InDocument) const = 0;
+			// MetaSound Interface definition
+			virtual const FMetasoundFrontendInterface& GetInterface() const = 0;
+
+			// Whether or not to immediately apply the given interface when a MetaSound
+			// asset of the supported UClass is created.
+			virtual bool IsDefault() const = 0;
+
+			// Name of routing system used to update interface inputs (ex. ParameterInterface or DataReference).
+			virtual FName GetRouterName() const = 0;
+
+			// Name of UClass supported by the given interface.
+			virtual bool UClassIsSupported(FName InUClassName) const = 0;
+
+			// How to update a given document if versioning is required to this interface from a deprecated version.
+			virtual bool UpdateRootGraphInterface(FDocumentHandle InDocument) const = 0;
 		};
 
-		class METASOUNDFRONTEND_API FArchetypeRegistryTransaction 
+		class METASOUNDFRONTEND_API FInterfaceRegistryTransaction
 		{
 		public:
 			/** Describes the type of transaction. */
 			enum class ETransactionType : uint8
 			{
-				ArchetypeRegistration,     //< Something was added to the registry.
-				ArchetypeUnregistration,  //< Something was removed from the registry.
+				InterfaceRegistration,     //< Something was added to the registry.
+				InterfaceUnregistration,  //< Something was removed from the registry.
 				Invalid
 			};
 
-			FArchetypeRegistryTransaction(ETransactionType InType, const FArchetypeRegistryKey& InKey, const FMetasoundFrontendVersion& InArchetypeVersion);
+			FInterfaceRegistryTransaction(ETransactionType InType, const FInterfaceRegistryKey& InKey, const FMetasoundFrontendVersion& InInterfaceVersion);
 
 			ETransactionType GetTransactionType() const;
-			const FMetasoundFrontendVersion& GetArchetypeVersion() const;
-			const FArchetypeRegistryKey& GetArchetypeRegistryKey() const;
+			const FMetasoundFrontendVersion& GetInterfaceVersion() const;
+			const FInterfaceRegistryKey& GetInterfaceRegistryKey() const;
 
 		private:
 
 			ETransactionType Type;
-			FArchetypeRegistryKey Key;
-			FMetasoundFrontendVersion ArchetypeVersion;
+			FInterfaceRegistryKey Key;
+			FMetasoundFrontendVersion InterfaceVersion;
 		};
 
-		class METASOUNDFRONTEND_API IArchetypeRegistry
+		class METASOUNDFRONTEND_API IInterfaceRegistry
 		{
 		public:
-			static IArchetypeRegistry& Get();
+			static IInterfaceRegistry& Get();
 
-			virtual ~IArchetypeRegistry() = default;
+			virtual ~IInterfaceRegistry() = default;
 
-			virtual FArchetypeRegistryKey RegisterArchetype(TUniquePtr<IArchetypeRegistryEntry>&& InEntry) = 0;
-			virtual const IArchetypeRegistryEntry* FindArchetypeRegistryEntry(const FArchetypeRegistryKey& InKey) const = 0;
-			virtual bool FindArchetype(const FArchetypeRegistryKey& InKey, FMetasoundFrontendArchetype& OutArchetype) const = 0;
-			virtual void ForEachRegistryTransactionSince(FRegistryTransactionID InSince, FRegistryTransactionID* OutCurrentRegistryTransactionID, TFunctionRef<void(const FArchetypeRegistryTransaction&)> InFunc) const = 0;
+			virtual void RegisterInterface(TUniquePtr<IInterfaceRegistryEntry>&& InEntry) = 0;
+			virtual const IInterfaceRegistryEntry* FindInterfaceRegistryEntry(const FInterfaceRegistryKey& InKey) const = 0;
+			virtual bool FindInterface(const FInterfaceRegistryKey& InKey, FMetasoundFrontendInterface& OutInterface) const = 0;
+			virtual void ForEachRegistryTransactionSince(FRegistryTransactionID InSince, FRegistryTransactionID* OutCurrentRegistryTransactionID, TFunctionRef<void(const FInterfaceRegistryTransaction&)> InFunc) const = 0;
 		};
 	}
 }
diff --git a/Engine/Plugins/Runtime/Metasound/Source/MetasoundFrontend/Public/MetasoundFrontendController.h b/Engine/Plugins/Runtime/Metasound/Source/MetasoundFrontend/Public/MetasoundFrontendController.h
--- a/Engine/Plugins/Runtime/Metasound/Source/MetasoundFrontend/Public/MetasoundFrontendController.h
+++ b/Engine/Plugins/Runtime/Metasound/Source/MetasoundFrontend/Public/MetasoundFrontendController.h
@@ -453,8 +453,11 @@ namespace Metasound
 			virtual FOutputHandle GetOutputWithVertexName(const FVertexName& InName) = 0;
 			virtual FConstOutputHandle GetConstOutputWithVertexName(const FVertexName& InName) const = 0;
 
-			/** Returns true if node is required to satisfy the document archetype. */
-			virtual bool IsRequired() const = 0;
+			/** Returns true if node is required to satisfy a document interface. */
+			virtual bool IsInterfaceMember() const = 0;
+
+			/** Returns interface version if node is a required member of a given interface, otherwise returns invalid version. */
+			virtual const FMetasoundFrontendVersion& GetInterfaceVersion() const = 0;
 
 			/** Returns the highest minor version number available from the class registry that shares this node's name & major version. */
 			virtual FMetasoundFrontendVersionNumber FindHighestVersionInRegistry() const = 0;
@@ -926,7 +929,7 @@ namespace Metasound
 			/** Returns true if the controller is in a valid state. */
 			virtual bool IsValid() const = 0;
 
-			// TODO: add info on environment variables. 
+			// TODO: add info on environment variables.
 			// TODO: consider find/add subgraph
 			// TODO: perhaps functions returning access pointers could be removed from main interface and only exist in FDocumentController.
 			
@@ -983,8 +986,10 @@ namespace Metasound
 			 */
 			virtual FGraphHandle AddDuplicateSubgraph(const IGraphController& InGraph) = 0;
 
-			virtual const FMetasoundFrontendVersion& GetArchetypeVersion() const = 0;
-			virtual void SetArchetypeVersion(const FMetasoundFrontendVersion& InVersion) = 0;
+			virtual const TArray<FMetasoundFrontendVersion>& GetInterfaceVersions() const = 0;
+			virtual void AddInterfaceVersion(const FMetasoundFrontendVersion& InVersion) = 0;
+			virtual void RemoveInterfaceVersion(const FMetasoundFrontendVersion& InVersion) = 0;
+			virtual void ClearInterfaceVersions() = 0;
 
 			/** Updates dependency with given registry key if found with that in the active registry
 			  *
diff --git a/Engine/Plugins/Runtime/Metasound/Source/MetasoundFrontend/Public/MetasoundFrontendDataTypeRegistry.h b/Engine/Plugins/Runtime/Metasound/Source/MetasoundFrontend/Public/MetasoundFrontendDataTypeRegistry.h
--- a/Engine/Plugins/Runtime/Metasound/Source/MetasoundFrontend/Public/MetasoundFrontendDataTypeRegistry.h
+++ b/Engine/Plugins/Runtime/Metasound/Source/MetasoundFrontend/Public/MetasoundFrontendDataTypeRegistry.h
@@ -202,6 +202,8 @@ namespace Metasound
 			virtual bool GetDataTypeInfo(const UObject* InObject, FDataTypeRegistryInfo& OutInfo) const = 0;
 			virtual bool GetDataTypeInfo(const FName& InDataType, FDataTypeRegistryInfo& OutInfo) const = 0;
 
+			virtual bool IsRegistered(const FName& InDataType) const = 0;
+
 			// Return the enum interface for a data type. If the data type does not have 
 			// an enum interface, returns a nullptr.
 			virtual TSharedPtr<const IEnumDataTypeInterface> GetEnumInterfaceForDataType(const FName& InDataType) const = 0;
diff --git a/Engine/Plugins/Runtime/Metasound/Source/MetasoundFrontend/Public/MetasoundFrontendDocument.h b/Engine/Plugins/Runtime/Metasound/Source/MetasoundFrontend/Public/MetasoundFrontendDocument.h
--- a/Engine/Plugins/Runtime/Metasound/Source/MetasoundFrontend/Public/MetasoundFrontendDocument.h
+++ b/Engine/Plugins/Runtime/Metasound/Source/MetasoundFrontend/Public/MetasoundFrontendDocument.h
@@ -77,11 +77,11 @@ struct METASOUNDFRONTEND_API FMetasoundFrontendVersionNumber
 	GENERATED_BODY()
 
 	// Major version number.
-	UPROPERTY(VisibleAnywhere, Category = General)
+	UPROPERTY(EditAnywhere, Category = General)
 	int32 Major = 1;
 
 	// Minor version number.
-	UPROPERTY(VisibleAnywhere, Category = General)
+	UPROPERTY(EditAnywhere, Category = General)
 	int32 Minor = 0;
 
 	static const FMetasoundFrontendVersionNumber& GetInvalid()
@@ -158,11 +158,11 @@ struct METASOUNDFRONTEND_API FMetasoundFrontendVersion
 	GENERATED_BODY()
 
 	// Name of version.
-	UPROPERTY()
+	UPROPERTY(EditAnywhere, Category = CustomView)
 	FName Name;
 
 	// Version number.
-	UPROPERTY()
+	UPROPERTY(EditAnywhere, Category = CustomView)
 	FMetasoundFrontendVersionNumber Number;
 
 	FString ToString() const;
@@ -545,7 +545,7 @@ struct FMetasoundFrontendVertexMetadata
 };
 
 USTRUCT()
-struct FMetasoundFrontendEnvironmentVariableMetadata
+struct FMetasoundFrontendClassEnvironmentVariableMetadata
 {
 	GENERATED_BODY()
 
@@ -656,12 +656,12 @@ struct METASOUNDFRONTEND_API FMetasoundFrontendClassOutput : public FMetasoundFr
 	virtual ~FMetasoundFrontendClassOutput() = default;
 };
 
-USTRUCT() 
-struct METASOUNDFRONTEND_API FMetasoundFrontendEnvironmentVariable
+USTRUCT()
+struct METASOUNDFRONTEND_API FMetasoundFrontendClassEnvironmentVariable
 {
 	GENERATED_BODY()
 
-	virtual ~FMetasoundFrontendEnvironmentVariable() = default;
+	virtual ~FMetasoundFrontendClassEnvironmentVariable() = default;
 
 	// Name of environment variable.
 	UPROPERTY()
@@ -673,17 +673,10 @@ struct METASOUNDFRONTEND_API FMetasoundFrontendEnvironmentVariable
 
 	// Metadata of environment variable.
 	UPROPERTY()
-	FMetasoundFrontendEnvironmentVariableMetadata Metadata;
-};
-
-USTRUCT() 
-struct METASOUNDFRONTEND_API FMetasoundFrontendClassEnvironmentVariable : public FMetasoundFrontendEnvironmentVariable
-{
-	GENERATED_BODY()
-
-	virtual ~FMetasoundFrontendClassEnvironmentVariable() = default;
+	FMetasoundFrontendClassEnvironmentVariableMetadata Metadata;
 
-	// True if the environment variable is needed in order to instantiate a node instance of the class. 
+	// True if the environment variable is needed in order to instantiate a node instance of the class.
+	// TODO: Should be deprecated?
 	UPROPERTY()
 	bool bIsRequired = true;
 };
@@ -815,6 +808,17 @@ public:
 };
 
 
+USTRUCT()
+struct METASOUNDFRONTEND_API FMetasoundFrontendInterface : public FMetasoundFrontendClassInterface
+{
+	GENERATED_BODY()
+
+	// Name and version number of the interface
+	UPROPERTY()
+	FMetasoundFrontendVersion Version;
+};
+
+
 // Name of a Metasound class
 USTRUCT()
 struct METASOUNDFRONTEND_API FMetasoundFrontendClassName
@@ -1119,41 +1123,6 @@ struct FMetasoundFrontendDocumentMetadata
 };
 
 
-USTRUCT()
-struct FMetasoundFrontendArchetypeInterface
-{
-	GENERATED_BODY()
-
-	UPROPERTY()
-	FGuid TransactionID;
-
-	UPROPERTY()
-	TArray<FMetasoundFrontendClassVertex> Inputs;
-
-	UPROPERTY()
-	TArray<FMetasoundFrontendClassVertex> Outputs;
-
-	UPROPERTY()
-	TArray<FMetasoundFrontendEnvironmentVariable> Environment;
-};
-
-/** Describes the expected minimum interface of a MetaSound root graph. This allows
- * systems to describe the minimum interface for a MetaSound to work within them. */
-USTRUCT()
-struct FMetasoundFrontendArchetype
-{
-	GENERATED_BODY()
-
-	/** Name and version number of the archetype. */
-	UPROPERTY()
-	FMetasoundFrontendVersion Version;
-
-	/** Expected minimum interface for MetaSounds satisfying archetype. */
-	UPROPERTY()
-	FMetasoundFrontendArchetypeInterface Interface;
-};
-
-
 USTRUCT()
 struct METASOUNDFRONTEND_API FMetasoundFrontendDocument
 {
@@ -1169,9 +1138,12 @@ struct METASOUNDFRONTEND_API FMetasoundFrontendDocument
 	UPROPERTY()
 	FMetasoundFrontendEditorData EditorData;
 
-	UPROPERTY()
+	UPROPERTY(meta = (DeprecatedProperty, DeprecationMessage = "5.0 - ArchetypeVersion has been migrated to InterfaceVersions array."))
 	FMetasoundFrontendVersion ArchetypeVersion;
 
+	UPROPERTY(EditAnywhere, Category = CustomView, meta = (DisplayName = "Interfaces"))
+	TArray<FMetasoundFrontendVersion> InterfaceVersions;
+
 	UPROPERTY(EditAnywhere, Category = CustomView)
 	FMetasoundFrontendGraphClass RootGraph;
 
diff --git a/Engine/Plugins/Runtime/Metasound/Source/MetasoundFrontend/Public/MetasoundFrontendInjectReceiveNodes.h b/Engine/Plugins/Runtime/Metasound/Source/MetasoundFrontend/Public/MetasoundFrontendInjectReceiveNodes.h
--- a/Engine/Plugins/Runtime/Metasound/Source/MetasoundFrontend/Public/MetasoundFrontendInjectReceiveNodes.h
+++ b/Engine/Plugins/Runtime/Metasound/Source/MetasoundFrontend/Public/MetasoundFrontendInjectReceiveNodes.h
@@ -35,12 +35,12 @@ namespace Metasound
 		/** Injects receive nodes between an graph inputs and the connected internal nodes.
 		 *
 		 * @param InGraph - The graph to manipulate.
-		 * @parma InAddressPolicy - A function which creates unique transmission addresses for each input.
-		 * @param InInputVerticesToSkip - A set of input vertices which should *not* have receive nodes injected.
+		 * @parma InAddressPolicy - Function which creates unique transmission addresses for each input.
+		 * @param InInputVertexNames - Set of input vertices which should have receive nodes injected.
 		 *
-		 * @return True on success, false on failure. 
+		 * @return True on success, false on failure.
 		 */
-		METASOUNDFRONTEND_API bool InjectReceiveNodes(FFrontendGraph& InGraph, const FReceiveNodeAddressFunction& InAddressPolicy, const TSet<FVertexName>& InInputVerticesToSkip);
+		METASOUNDFRONTEND_API bool InjectReceiveNodes(FFrontendGraph& InGraph, const FReceiveNodeAddressFunction& InAddressPolicy, const TSet<FVertexName>& InInputVertexNames);
 	}
 }
 
diff --git a/Engine/Plugins/Runtime/Metasound/Source/MetasoundFrontend/Public/MetasoundFrontendLiteral.h b/Engine/Plugins/Runtime/Metasound/Source/MetasoundFrontend/Public/MetasoundFrontendLiteral.h
--- a/Engine/Plugins/Runtime/Metasound/Source/MetasoundFrontend/Public/MetasoundFrontendLiteral.h
+++ b/Engine/Plugins/Runtime/Metasound/Source/MetasoundFrontend/Public/MetasoundFrontendLiteral.h
@@ -64,6 +64,9 @@ struct METASOUNDFRONTEND_API FMetasoundFrontendLiteral
 	struct FDefault {};
 	struct FDefaultArray { int32 Num = 0; };
 
+	FMetasoundFrontendLiteral() = default;
+	FMetasoundFrontendLiteral(const FAudioParameter& InParameter);
+
 private:
 	// The set type of this literal.
 	UPROPERTY()
@@ -143,7 +146,7 @@ public:
 	void Clear();
 
 private:
-	// Remove all values. 
+	// Remove all values.
 	void Empty();
 };
 
diff --git a/Engine/Plugins/Runtime/Metasound/Source/MetasoundFrontend/Public/MetasoundFrontendQuery.h b/Engine/Plugins/Runtime/Metasound/Source/MetasoundFrontend/Public/MetasoundFrontendQuery.h
--- a/Engine/Plugins/Runtime/Metasound/Source/MetasoundFrontend/Public/MetasoundFrontendQuery.h
+++ b/Engine/Plugins/Runtime/Metasound/Source/MetasoundFrontend/Public/MetasoundFrontendQuery.h
@@ -46,7 +46,7 @@ namespace Metasound
 	 */
 	struct METASOUNDFRONTEND_API FFrontendQueryEntry
 	{
-		using FValue = TVariant<Frontend::FNodeRegistryTransaction, FMetasoundFrontendVersion, FMetasoundFrontendClass, Frontend::FArchetypeRegistryTransaction, FMetasoundFrontendArchetype>;
+		using FValue = TVariant<Frontend::FNodeRegistryTransaction, FMetasoundFrontendVersion, FMetasoundFrontendClass, Frontend::FInterfaceRegistryTransaction, FMetasoundFrontendInterface>;
 		using FKey = FFrontendQueryKey;
 
 		FKey Key;
diff --git a/Engine/Plugins/Runtime/Metasound/Source/MetasoundFrontend/Public/MetasoundFrontendSearchEngine.h b/Engine/Plugins/Runtime/Metasound/Source/MetasoundFrontend/Public/MetasoundFrontendSearchEngine.h
--- a/Engine/Plugins/Runtime/Metasound/Source/MetasoundFrontend/Public/MetasoundFrontendSearchEngine.h
+++ b/Engine/Plugins/Runtime/Metasound/Source/MetasoundFrontend/Public/MetasoundFrontendSearchEngine.h
@@ -36,13 +36,15 @@ namespace Metasound
 				/** Find the class with the given ClassName & Major Version. Returns false if not found, true if found. */
 				virtual bool FindClassWithMajorVersion(const FNodeClassName& InName, int32 InMajorVersion, FMetasoundFrontendClass& OutClass) = 0;
 
-				virtual TArray<FMetasoundFrontendArchetype> FindAllArchetypes(bool bInIncludeDeprecated=false) = 0;
+				virtual TArray<FMetasoundFrontendInterface> FindAllInterfaces(bool bInIncludeDeprecated = false) = 0;
 
-				virtual TArray<FMetasoundFrontendVersion> FindAllRegisteredArchetypesWithName(const FName& InArchetypeName) = 0;
+				virtual TArray<FMetasoundFrontendInterface> FindUClassDefaultInterfaces(FName InUClassName) = 0;
 
-				virtual bool FindArchetypeWithHighestVersion(const FName& InArchetypeName, FMetasoundFrontendArchetype& OutArchetype) = 0;
+				virtual TArray<FMetasoundFrontendVersion> FindAllRegisteredInterfacesWithName(FName InInterfaceName) = 0;
 
-				virtual bool FindArchetypeWithMajorVersion(const FName& InArchetypeName, int32 InMajorVersion, FMetasoundFrontendArchetype& OutArchetype) = 0;
+				virtual bool FindInterfaceWithHighestVersion(FName InInterfaceName, FMetasoundFrontendInterface& OutInterface) = 0;
+
+				virtual bool FindInterfaceWithMajorVersion(FName InInterfaceName, int32 InMajorVersion, FMetasoundFrontendInterface& OutInterface) = 0;
 
 				virtual ~ISearchEngine() = default;
 
diff --git a/Engine/Plugins/Runtime/Metasound/Source/MetasoundFrontend/Public/MetasoundFrontendTransform.h b/Engine/Plugins/Runtime/Metasound/Source/MetasoundFrontend/Public/MetasoundFrontendTransform.h
--- a/Engine/Plugins/Runtime/Metasound/Source/MetasoundFrontend/Public/MetasoundFrontendTransform.h
+++ b/Engine/Plugins/Runtime/Metasound/Source/MetasoundFrontend/Public/MetasoundFrontendTransform.h
@@ -30,42 +30,51 @@ namespace Metasound
 			virtual bool Transform(FGraphHandle InGraph) const = 0;
 		};
 
-		/** Adds or swaps graph inputs and outputs to match the archetype. */
-		class METASOUNDFRONTEND_API FSwapGraphArchetype : public IGraphTransform
+		/** Adds or swaps document members (inputs, outputs) and removing any document members where necessary and adding those missing. */
+		class METASOUNDFRONTEND_API FModifyRootGraphInterfaces : public IDocumentTransform
 		{
 		public:
-			FSwapGraphArchetype(const FMetasoundFrontendArchetype& InFromArchetype, const FMetasoundFrontendArchetype& InToArchetype);
+			FModifyRootGraphInterfaces(const TArray<FMetasoundFrontendInterface>& InInterfacesToRemove, const TArray<FMetasoundFrontendInterface>& InInterfacesToAdd);
+			FModifyRootGraphInterfaces(const TArray<FMetasoundFrontendVersion>& InInterfaceVersionsToRemove, const TArray<FMetasoundFrontendVersion>& InInterfaceVersionsToAdd);
 
-			bool Transform(FGraphHandle InGraph) const override;
+			void SetDefaultNodeLocations(bool bInSetDefaultNodeLocations);
+			bool Transform(FDocumentHandle InDocument) const override;
 
 		private:
+			void Init();
+
+			bool bSetDefaultNodeLocations = true;
+
+			TArray<FMetasoundFrontendInterface> InterfacesToRemove;
+			TArray<FMetasoundFrontendInterface> InterfacesToAdd;
+
 			using FVertexPair = TTuple<FMetasoundFrontendClassVertex, FMetasoundFrontendClassVertex>;
 			TArray<FVertexPair> PairedInputs;
 			TArray<FVertexPair> PairedOutputs;
-			TArray<FMetasoundFrontendClassVertex> InputsToAdd;
-			TArray<FMetasoundFrontendClassVertex> InputsToRemove;
-			TArray<FMetasoundFrontendClassVertex> OutputsToAdd;
-			TArray<FMetasoundFrontendClassVertex> OutputsToRemove;
+
+			TArray<FMetasoundFrontendClassInput> InputsToAdd;
+			TArray<FMetasoundFrontendClassInput> InputsToRemove;
+			TArray<FMetasoundFrontendClassOutput> OutputsToAdd;
+			TArray<FMetasoundFrontendClassOutput> OutputsToRemove;
 
 		};
 
-		/** Adds or swaps root graph inputs and outputs to match the document's archetype. */
-		class METASOUNDFRONTEND_API FMatchRootGraphToArchetype : public IDocumentTransform
+		/** Updates document's given interface to the most recent version. */
+		class METASOUNDFRONTEND_API FUpdateRootGraphInterface : public IDocumentTransform
 		{
 		public:
-			FMatchRootGraphToArchetype(const FMetasoundFrontendVersion& InArchetypeVersion)
-				: ArchetypeVersion(InArchetypeVersion)
+			FUpdateRootGraphInterface(const FMetasoundFrontendVersion& InInterfaceVersion)
+				: InterfaceVersion(InInterfaceVersion)
 			{
 			}
 
 			bool Transform(FDocumentHandle InDocument) const override;
 
 		private:
-			void GetUpgradePathForDocument(const FMetasoundFrontendVersion& InCurrentVersion, const FMetasoundFrontendVersion& InTargetVersion, TArray<const IArchetypeRegistryEntry*>& OutUpgradePath) const;
-			bool UpgradeDocumentArchetype(const TArray<const IArchetypeRegistryEntry*>& InUpgradePath, FDocumentHandle InDocument) const;
-			bool ConformDocumentToArchetype(const FMetasoundFrontendArchetype& InTargetArchetype, FDocumentHandle InDocument) const;
+			void GetUpdatePathForDocument(const FMetasoundFrontendVersion& InCurrentVersion, const FMetasoundFrontendVersion& InTargetVersion, TArray<const IInterfaceRegistryEntry*>& OutUpgradePath) const;
+			bool UpdateDocumentInterface(const TArray<const IInterfaceRegistryEntry*>& InUpgradePath, FDocumentHandle InDocument) const;
 
-			FMetasoundFrontendVersion ArchetypeVersion;
+			FMetasoundFrontendVersion InterfaceVersion;
 		};
 
 		/** Completely rebuilds the graph connecting a preset's inputs to the reference
diff --git a/Engine/Plugins/Runtime/Metasound/Source/MetasoundFrontend/Public/MetasoundParameterTransmitter.h b/Engine/Plugins/Runtime/Metasound/Source/MetasoundFrontend/Public/MetasoundParameterTransmitter.h
--- a/Engine/Plugins/Runtime/Metasound/Source/MetasoundFrontend/Public/MetasoundParameterTransmitter.h
+++ b/Engine/Plugins/Runtime/Metasound/Source/MetasoundFrontend/Public/MetasoundParameterTransmitter.h
@@ -4,14 +4,23 @@
 #include "AudioParameterInterface.h"
 #include "IAudioParameterTransmitter.h"
 #include "MetasoundDataReference.h"
+#include "MetasoundFrontendLiteral.h"
 #include "MetasoundLog.h"
 #include "MetasoundOperatorSettings.h"
 #include "MetasoundPrimitives.h"
 #include "MetasoundRouter.h"
 
+struct FMetasoundFrontendLiteral;
+
 
 namespace Metasound
 {
+	namespace Frontend
+	{
+		METASOUNDFRONTEND_API FLiteral ConvertParameterToLiteral(FAudioParameter&& InValue);
+		METASOUNDFRONTEND_API FName ConvertParameterToDataType(EAudioParameterType InParameterType);
+	}
+
 	/** FMetaSoundParameterTransmitter provides a communication interface for 
 	 * sending values to a MetaSound instance. It relies on the send/receive transmission
 	 * system to ferry data from the transmitter to the MetaSound instance. Data will
@@ -22,8 +31,8 @@ namespace Metasound
 	{
 		FMetaSoundParameterTransmitter(const FMetaSoundParameterTransmitter&) = delete;
 		FMetaSoundParameterTransmitter& operator=(const FMetaSoundParameterTransmitter&) = delete;
+
 	public:
-		
 		/** FSendInfo describes the MetaSounds input parameters as well as the 
 		 * necessary information to route data to the instances inputs. 
 		 */
@@ -84,6 +93,7 @@ namespace Metasound
 		 * @param InParameter - Parameter to set.
 		 */
 		bool SetParameter(FAudioParameter&& InParameter) override;
+		bool SetParameter(FName InInterfaceName, FAudioParameter&& InParameter) override;
 
 		/** Set a parameter using a literal.
 		 *
diff --git a/Engine/Plugins/Runtime/Metasound/Source/MetasoundGraphCore/Private/MetasoundDataReference.cpp b/Engine/Plugins/Runtime/Metasound/Source/MetasoundGraphCore/Private/MetasoundDataReference.cpp
--- a/Engine/Plugins/Runtime/Metasound/Source/MetasoundGraphCore/Private/MetasoundDataReference.cpp
+++ b/Engine/Plugins/Runtime/Metasound/Source/MetasoundGraphCore/Private/MetasoundDataReference.cpp
@@ -11,4 +11,6 @@ namespace Metasound
 	}
 	const void* const TDataReferenceTypeInfo<void>::TypePtr = nullptr;
 	const void* const TDataReferenceTypeInfo<void>::TypeId = static_cast<const void* const>(TDataReferenceTypeInfo<void>::TypePtr);
+
+	const FName IDataReference::RouterName = "DataReference";
 }
diff --git a/Engine/Plugins/Runtime/Metasound/Source/MetasoundGraphCore/Private/MetasoundOperatorBuilder.cpp b/Engine/Plugins/Runtime/Metasound/Source/MetasoundGraphCore/Private/MetasoundOperatorBuilder.cpp
--- a/Engine/Plugins/Runtime/Metasound/Source/MetasoundGraphCore/Private/MetasoundOperatorBuilder.cpp
+++ b/Engine/Plugins/Runtime/Metasound/Source/MetasoundGraphCore/Private/MetasoundOperatorBuilder.cpp
@@ -302,7 +302,7 @@ namespace Metasound
 
 				// Denote a pruned node as a non-fatal error. In the future this
 				// may be simply a warning as some nodes are required to conform
-				// to metasound archetypes even if they are unused.
+				// to metasound interfaces even if they are unused.
 				BuildStatus |= FBuildStatus::NonFatalError;
 			}
 		}
diff --git a/Engine/Plugins/Runtime/Metasound/Source/MetasoundGraphCore/Public/MetasoundDataReference.h b/Engine/Plugins/Runtime/Metasound/Source/MetasoundGraphCore/Public/MetasoundDataReference.h
--- a/Engine/Plugins/Runtime/Metasound/Source/MetasoundGraphCore/Public/MetasoundDataReference.h
+++ b/Engine/Plugins/Runtime/Metasound/Source/MetasoundGraphCore/Public/MetasoundDataReference.h
@@ -91,11 +91,13 @@ namespace Metasound
 
 	/** A Data Reference Interface.
 	 *
-	 * A parameter references provides information and access to a shared object in the graph. 
+	 * A parameter references provides information and access to a shared object in the graph.
 	 */
-	class IDataReference
+	class METASOUNDGRAPHCORE_API IDataReference
 	{
 		public:
+		static const FName RouterName;
+
 		virtual ~IDataReference() = default;
 
 		/** Returns the name of the data type. */
diff --git a/Engine/Source/Runtime/AudioExtensions/Private/IAudioGeneratorInterfaceRegistry.cpp b/Engine/Source/Runtime/AudioExtensions/Private/IAudioGeneratorInterfaceRegistry.cpp
new file mode 100644
--- /dev/null
+++ b/Engine/Source/Runtime/AudioExtensions/Private/IAudioGeneratorInterfaceRegistry.cpp
@@ -0,0 +1,76 @@
+// Copyright Epic Games, Inc. All Rights Reserved.
+#include "IAudioGeneratorInterfaceRegistry.h"
+
+#include "Containers/Map.h"
+
+
+namespace Audio
+{
+	namespace ParameterInterfaceRegistryPrivate
+	{
+		class FGeneratorInterfaceRegistry : public IGeneratorInterfaceRegistry
+		{
+		public:
+			virtual void IterateInterfaces(TFunction<void(FGeneratorInterfacePtr)> InFunction) const override
+			{
+				FGeneratorInterfaceQueryResults GeneratorQuery;
+				QueryInterfacesDelegate.Broadcast(GeneratorQuery);
+				for (const FGeneratorInterfacePtr& InterfacePtr : GeneratorQuery.GetInterfaces())
+				{
+					InFunction(InterfacePtr);
+				}
+			}
+
+			virtual void RegisterInterface(FGeneratorInterfacePtr InInterface) override
+			{
+				RegistrationFunction(InInterface);
+				QueryInterfacesDelegate.AddLambda([Interface = InInterface](FGeneratorInterfaceQueryResults& Results)
+				{
+					Results.AddInterface(Interface);
+				});
+			}
+
+			virtual void OnRegistration(TUniqueFunction<void(FGeneratorInterfacePtr)>&& InFunction) override
+			{
+				RegistrationFunction = MoveTemp(InFunction);
+			}
+		};
+	} // namespace ParameterInterfaceRegistryPrivate
+
+	TUniquePtr<IGeneratorInterfaceRegistry> IGeneratorInterfaceRegistry::Instance;
+
+	const FString IGeneratorInterfaceRegistry::NamespaceDelimiter = TEXT(".");
+
+	FName IGeneratorInterfaceRegistry::GetMemberFullName(FName InInterfaceName, FName InParameterName)
+	{
+		const FString FullName = FString::Join(TArray<FString>({ *InInterfaceName.ToString(), *InParameterName.ToString() }), *NamespaceDelimiter);
+		return FName(FullName);
+	}
+
+	void IGeneratorInterfaceRegistry::SplitMemberFullName(FName InFullName, FName& OutInterfaceName, FName& OutParameterName)
+	{
+		FString FullName = InFullName.ToString();
+		const int32 IndexOfDelim = FullName.Find(NamespaceDelimiter, ESearchCase::IgnoreCase, ESearchDir::FromEnd);
+		if (IndexOfDelim != INDEX_NONE)
+		{
+			OutInterfaceName = FName(*FullName.Left(IndexOfDelim));
+			OutParameterName = FName(*FullName.RightChop(IndexOfDelim + 1));
+		}
+		else
+		{
+			OutInterfaceName = FName();
+			OutParameterName = InFullName;
+		}
+	}
+
+	IGeneratorInterfaceRegistry& IGeneratorInterfaceRegistry::Get()
+	{
+		using namespace ParameterInterfaceRegistryPrivate;
+
+		if (!Instance.IsValid())
+		{
+			Instance = MakeUnique<FGeneratorInterfaceRegistry>();
+		}
+		return *Instance;
+	}
+} // namespace Audio
diff --git a/Engine/Source/Runtime/AudioExtensions/Private/IAudioParameterTransmitter.cpp b/Engine/Source/Runtime/AudioExtensions/Private/IAudioParameterTransmitter.cpp
--- a/Engine/Source/Runtime/AudioExtensions/Private/IAudioParameterTransmitter.cpp
+++ b/Engine/Source/Runtime/AudioExtensions/Private/IAudioParameterTransmitter.cpp
@@ -5,6 +5,8 @@
 
 namespace Audio
 {
+	const FName IParameterTransmitter::RouterName = "ParameterTransmitter";
+
 	bool ILegacyParameterTransmitter::GetParameter(FName InName, FAudioParameter& OutValue) const
 	{
 		return false;
diff --git a/Engine/Source/Runtime/AudioExtensions/Public/IAudioGeneratorInterfaceRegistry.h b/Engine/Source/Runtime/AudioExtensions/Public/IAudioGeneratorInterfaceRegistry.h
new file mode 100644
--- /dev/null
+++ b/Engine/Source/Runtime/AudioExtensions/Public/IAudioGeneratorInterfaceRegistry.h
@@ -0,0 +1,105 @@
+// Copyright Epic Games, Inc. All Rights Reserved.
+#pragma once
+
+#include "AudioParameterInterface.h"
+#include "Delegates/Delegate.h"
+#include "Templates/Function.h"
+#include "UObject/Class.h"
+
+
+namespace Audio
+{
+	class IGeneratorInterfaceRegistry;
+
+	struct AUDIOEXTENSIONS_API FGeneratorInterface
+	{
+		FGeneratorInterface() = default;
+		FGeneratorInterface(FName InName, UClass* InType)
+			: Name(InName)
+			, Type(InType)
+		{
+		}
+
+		FName Name;
+		UClass* Type = nullptr;
+
+		struct FVersion
+		{
+			int32 Major = 1;
+			int32 Minor = 0;
+		} Version;
+
+		struct FInput
+		{
+			const FText DisplayName;
+			const FText Description;
+			const FName DataType;
+
+			const FAudioParameter InitValue;
+		};
+		TArray<FInput> Inputs;
+
+		struct FOutput
+		{
+			const FText DisplayName;
+			const FText Description;
+			const FName DataType;
+
+			const FName ParamName;
+			const EAudioParameterType ParamType;
+		};
+		TArray<FOutput> Outputs;
+
+		struct FEnvironmentVariable
+		{
+			const FText DisplayName;
+			const FText Description;
+			const FName DataType;
+
+			const FName ParamName;
+			const EAudioParameterType ParamType;
+		};
+		TArray<FEnvironmentVariable> Environment;
+	};
+	using FGeneratorInterfacePtr = TSharedPtr<FGeneratorInterface>;
+
+	struct FGeneratorInterfaceQueryResults
+	{
+	private:
+		TArray<FGeneratorInterfacePtr> Interfaces;
+
+	public:
+		void AddInterface(FGeneratorInterfacePtr InInterfaceToAdd)
+		{
+			Interfaces.Add(InInterfaceToAdd);
+		}
+
+		const TArray<FGeneratorInterfacePtr>& GetInterfaces() const { return Interfaces; }
+	};
+
+	DECLARE_MULTICAST_DELEGATE_OneParam(FOnQueryAudioGeneratorInterface, FGeneratorInterfaceQueryResults&);
+	DECLARE_MULTICAST_DELEGATE_OneParam(FOnRegisterAudioGeneratorInterface, FGeneratorInterfacePtr);
+
+	class AUDIOEXTENSIONS_API IGeneratorInterfaceRegistry
+	{
+		static TUniquePtr<IGeneratorInterfaceRegistry> Instance;
+
+	public:
+		static const FString NamespaceDelimiter;
+		static IGeneratorInterfaceRegistry& Get();
+
+		static FName GetMemberFullName(FName InInterfaceName, FName InParameterName);
+		static void SplitMemberFullName(FName InFullName, FName& OutInterfaceName, FName& OutParameterName);
+
+		virtual ~IGeneratorInterfaceRegistry() = default;
+
+		virtual void IterateInterfaces(TFunction<void(FGeneratorInterfacePtr)> InFunction) const = 0;
+		virtual void OnRegistration(TUniqueFunction<void(FGeneratorInterfacePtr)>&& InFunction) = 0;
+		virtual void RegisterInterface(FGeneratorInterfacePtr InInterface) = 0;
+
+	protected:
+		mutable FOnQueryAudioGeneratorInterface QueryInterfacesDelegate;
+
+		TUniqueFunction<void(FGeneratorInterfacePtr)> RegistrationFunction;
+	};
+} // namespace Audio
diff --git a/Engine/Source/Runtime/AudioExtensions/Public/IAudioParameterTransmitter.h b/Engine/Source/Runtime/AudioExtensions/Public/IAudioParameterTransmitter.h
--- a/Engine/Source/Runtime/AudioExtensions/Public/IAudioParameterTransmitter.h
+++ b/Engine/Source/Runtime/AudioExtensions/Public/IAudioParameterTransmitter.h
@@ -41,6 +41,8 @@ namespace Audio
 	class AUDIOEXTENSIONS_API IParameterTransmitter : public ILegacyParameterTransmitter
 	{
 		public:
+			static const FName RouterName;
+
 			virtual ~IParameterTransmitter() = default;
 
 			virtual bool Reset() = 0;
@@ -50,6 +52,7 @@ namespace Audio
 
 			// Parameter Setters
 			virtual bool SetParameter(FAudioParameter&& InValue) = 0;
+			virtual bool SetParameter(FName InInterfaceName, FAudioParameter&& InValue) = 0;
 
 			// Create a copy of the instance transmitter.
 			virtual TUniquePtr<Audio::IParameterTransmitter> Clone() const = 0;
diff --git a/Engine/Source/Runtime/Engine/Classes/Sound/SoundAttenuation.h b/Engine/Source/Runtime/Engine/Classes/Sound/SoundAttenuation.h
--- a/Engine/Source/Runtime/Engine/Classes/Sound/SoundAttenuation.h
+++ b/Engine/Source/Runtime/Engine/Classes/Sound/SoundAttenuation.h
@@ -4,6 +4,8 @@
 
 #include "Engine/Attenuation.h"
 #include "IAudioExtensionPlugin.h"
+#include "IAudioGeneratorInterfaceRegistry.h"
+
 #include "SoundAttenuation.generated.h"
 
 class USoundSubmixBase;
@@ -487,3 +489,34 @@ class USoundAttenuation : public UObject
 	UPROPERTY(EditAnywhere, BlueprintReadOnly, Category=Settings, meta = (CustomizeProperty))
 	FSoundAttenuationSettings Attenuation;
 };
+
+namespace Audio
+{
+	struct ENGINE_API FAttenuationInterface : public FGeneratorInterface
+	{
+		static const FName Name;
+
+		struct FInputs
+		{
+			static const FName Distance;
+		};
+
+		FAttenuationInterface();
+	};
+
+	struct ENGINE_API FSpatializationInterface : public FGeneratorInterface
+	{
+		static const FName Name;
+
+		struct FInputs
+		{
+			static const FName Azimuth;
+			static const FName Elevation;
+		};
+
+		FSpatializationInterface();
+	};
+
+	ENGINE_API FGeneratorInterfacePtr GetAttenuationInterface();
+	ENGINE_API FGeneratorInterfacePtr GetSpatializationInterface();
+} // namespace Audio
\ No newline at end of file
diff --git a/Engine/Source/Runtime/Engine/Classes/Sound/SoundBase.h b/Engine/Source/Runtime/Engine/Classes/Sound/SoundBase.h
--- a/Engine/Source/Runtime/Engine/Classes/Sound/SoundBase.h
+++ b/Engine/Source/Runtime/Engine/Classes/Sound/SoundBase.h
@@ -5,6 +5,7 @@
 #include "Audio.h"
 #include "CoreMinimal.h"
 #include "IAudioExtensionPlugin.h"
+#include "IAudioGeneratorInterfaceRegistry.h"
 #include "IAudioParameterTransmitter.h"
 #include "Sound/AudioSettings.h"
 #include "Sound/SoundClass.h"
@@ -295,6 +296,9 @@ public:
 	/** Called from the Game Thread prior to attempting to initialize a sound instance. */
 	virtual void InitResources() { }
 
+	/** Whether or not the given sound is a generator and implements an interface with the given name. */
+	virtual bool ImplementsGeneratorInterface(Audio::FGeneratorInterfacePtr InInterface) const { return false; }
+
 	/** Creates a sound generator instance from this sound base. Return true if this is being implemented by a subclass. Sound generators procedurally generate audio in the audio render thread. */
 	virtual ISoundGeneratorPtr CreateSoundGenerator(const FSoundGeneratorInitParams& InParams) { return nullptr; }
 
diff --git a/Engine/Source/Runtime/Engine/Private/ActiveSound.cpp b/Engine/Source/Runtime/Engine/Private/ActiveSound.cpp
--- a/Engine/Source/Runtime/Engine/Private/ActiveSound.cpp
+++ b/Engine/Source/Runtime/Engine/Private/ActiveSound.cpp
@@ -7,6 +7,7 @@
 #include "AudioThread.h"
 #include "AudioDevice.h"
 #include "IAudioExtensionPlugin.h"
+#include "ProfilingDebugging/CpuProfilerTrace.h"
 #include "Sound/AudioSettings.h"
 #include "Sound/SoundClass.h"
 #include "Sound/SoundCue.h"
@@ -561,6 +562,57 @@ bool FActiveSound::GetConcurrencyFadeDuration(float& OutFadeDuration) const
 	return true;
 }
 
+void FActiveSound::UpdateGeneratorParameters(const TArray<FListener>& InListeners)
+{
+	TRACE_CPUPROFILER_EVENT_SCOPE(FActiveSound::UpdateGeneratorParameters);
+
+	using namespace Audio;
+
+	if (!InstanceTransmitter.IsValid())
+	{
+		return;
+	}
+
+	if (!Sound || IsPreviewSound())
+	{
+		return;
+	}
+
+	if (!InListeners.IsValidIndex(ClosestListenerIndex))
+	{
+		return;
+	}
+
+	FGeneratorInterfacePtr AttenuationInterface = GetAttenuationInterface();
+	FGeneratorInterfacePtr SpatializationInterface = GetSpatializationInterface();
+
+	const bool bImplementsAttenuation = Sound->ImplementsGeneratorInterface(AttenuationInterface);
+	const bool bImplementsSpatialization = Sound->ImplementsGeneratorInterface(SpatializationInterface);
+
+	if (!bImplementsAttenuation && !bImplementsSpatialization)
+	{
+		return;
+	}
+
+	const FListener& Listener = InListeners[ClosestListenerIndex];
+	const FVector SourceDirection = Transform.GetLocation() - Listener.Transform.GetLocation();
+	if (bImplementsAttenuation)
+	{
+		const float Distance = SourceDirection.Size();
+		InstanceTransmitter->SetParameter(FAttenuationInterface::Name, { FAttenuationInterface::FInputs::Distance, Distance });
+	}
+
+	if (bImplementsSpatialization)
+	{
+		const FVector SourceDirectionNormal = Listener.Transform.InverseTransformVectorNoScale(SourceDirection).GetSafeNormal();
+		const FVector2D SourceAzimuthAndElevation = FMath::GetAzimuthAndElevation(SourceDirectionNormal, FVector::ForwardVector, FVector::RightVector, FVector::UpVector);
+		const float Azimuth = FMath::RadiansToDegrees(SourceAzimuthAndElevation.X);
+		const float Elevation = FMath::RadiansToDegrees(SourceAzimuthAndElevation.Y);
+		InstanceTransmitter->SetParameter(FSpatializationInterface::Name, { FSpatializationInterface::FInputs::Azimuth, Azimuth });
+		InstanceTransmitter->SetParameter(FSpatializationInterface::Name, { FSpatializationInterface::FInputs::Elevation, Elevation });
+	}
+}
+
 void FActiveSound::UpdateWaveInstances(TArray<FWaveInstance*> &InWaveInstances, const float DeltaTime)
 {
 	// Reset whether or not the active sound is playing audio.
diff --git a/Engine/Source/Runtime/Engine/Private/AudioDevice.cpp b/Engine/Source/Runtime/Engine/Private/AudioDevice.cpp
--- a/Engine/Source/Runtime/Engine/Private/AudioDevice.cpp
+++ b/Engine/Source/Runtime/Engine/Private/AudioDevice.cpp
@@ -3887,6 +3887,7 @@ int32 FAudioDevice::GetSortedActiveWaveInstances(TArray<FWaveInstance*>& WaveIns
 						UsedDeltaTime = 0.0f;
 					}
 
+					ActiveSound->UpdateGeneratorParameters(Listeners);
 					ActiveSound->UpdateWaveInstances(WaveInstances, UsedDeltaTime);
 				}
 			}
@@ -5376,12 +5377,12 @@ bool FAudioDevice::LocationIsAudible(const FVector& Location, const float MaxDis
 
 	const int32 ListenerCount = bInAudioThread ? Listeners.Num() : ListenerProxies.Num();
 	for (int32 i = 0; i < ListenerCount; ++i)
-		{
+	{
 		if (LocationIsAudible(Location, i, MaxDistance))
-			{
-				return true;
-			}
+		{
+			return true;
 		}
+	}
 
 	return false;
 }
@@ -5410,16 +5411,16 @@ bool FAudioDevice::LocationIsAudible(const FVector& Location, const FTransform&
 bool FAudioDevice::LocationIsAudible(const FVector& Location, int32 ListenerIndex, float MaxDistance) const
 {
 	if (MaxDistance >= WORLD_MAX)
-			{
-				return true;
-			}
+	{
+		return true;
+	}
 	
 	FVector ListenerTranslation;
 	const bool bAllowOverride = true;
 	if (ListenerIndex == INDEX_NONE || !GetListenerPosition(ListenerIndex, ListenerTranslation, bAllowOverride))
 	{
 		return false;
-		}
+	}
 	
 	const float MaxDistanceSquared = MaxDistance * MaxDistance;
 	return (ListenerTranslation - Location).SizeSquared() < MaxDistanceSquared;
@@ -5488,24 +5489,25 @@ bool FAudioDevice::GetListenerPosition(int32 ListenerIndex, FVector& OutPosition
 {
 	OutPosition = FVector::ZeroVector;
 	if (ListenerIndex == INDEX_NONE)
-		{
+	{
 		return false;
 	}
 
 	if (IsInAudioThread())
-			{
+	{
 		checkf(ListenerIndex < Listeners.Num(), TEXT("Listener Index %u out of range of available Listeners!"), ListenerIndex);
 		const FListener& Listener = Listeners[ListenerIndex];
 		OutPosition = Listener.GetPosition(bAllowOverride);
 		return true;
-			}
+	}
 	else // IsInGameThread()
 	{
 		checkf(ListenerIndex < ListenerProxies.Num(), TEXT("Listener Index %u out of range of available Listeners!"), ListenerIndex);
 		const FListenerProxy& Proxy = ListenerProxies[ListenerIndex];
 		OutPosition = Proxy.GetPosition(bAllowOverride);
 		return true;
-		}
+	}
+
 	return false;
 }
 
@@ -5523,16 +5525,17 @@ bool FAudioDevice::GetListenerTransform(int32 ListenerIndex, FTransform& OutTran
 		{
 			OutTransform = Listeners[ListenerIndex].Transform;
 			return true;
-			}
 		}
+	}
 	else // IsInGameThread()
 	{
 		if (ListenerIndex < ListenerProxies.Num())
-	{
+		{
 			OutTransform = ListenerProxies[ListenerIndex].Transform;
 			return true;
+		}
 	}
-	}
+
 	return false;
 }
 
@@ -5560,6 +5563,7 @@ bool FAudioDevice::GetListenerWorldID(int32 ListenerIndex, uint32& OutWorldID) c
 			return true;
 		}
 	}
+
 	return false;
 }
 
@@ -5730,13 +5734,13 @@ int32 FAudioDevice::FindClosestListenerIndex(const FVector& Position, float& Out
 
 		const float DistSq = FVector::DistSquared(Position, ListenerPosition);
 		if (DistSq < OutDistanceSq)
-				{
+		{
 			OutDistanceSq = DistSq;
-					ClosestListenerIndex = i;
-			}
+			ClosestListenerIndex = i;
 		}
+	}
 
-		return ClosestListenerIndex;
+	return ClosestListenerIndex;
 }
 
 void FAudioDevice::UnlinkActiveSoundFromComponent(const FActiveSound& InActiveSound)
diff --git a/Engine/Source/Runtime/Engine/Private/AudioSettings.cpp b/Engine/Source/Runtime/Engine/Private/AudioSettings.cpp
--- a/Engine/Source/Runtime/Engine/Private/AudioSettings.cpp
+++ b/Engine/Source/Runtime/Engine/Private/AudioSettings.cpp
@@ -4,8 +4,10 @@
 #include "AudioDevice.h"
 #include "AudioDeviceManager.h"
 #include "AudioMixerDevice.h"
+#include "IAudioGeneratorInterfaceRegistry.h"
 #include "Misc/ConfigCacheIni.h"
 #include "Misc/Paths.h"
+#include "Sound/SoundAttenuation.h"
 #include "Sound/SoundBase.h"
 #include "Sound/SoundClass.h"
 #include "Sound/SoundConcurrency.h"
@@ -257,6 +259,10 @@ void UAudioSettings::LoadDefaultObjects()
 	{
 		UE_LOG(LogAudio, Display, TEXT("No default SoundConcurrencyObject specified (or failed to load)."));
 	}
+
+	Audio::IGeneratorInterfaceRegistry& InterfaceRegistry = Audio::IGeneratorInterfaceRegistry::Get();
+	InterfaceRegistry.RegisterInterface(Audio::GetAttenuationInterface());
+	InterfaceRegistry.RegisterInterface(Audio::GetSpatializationInterface());
 }
 
 USoundClass* UAudioSettings::GetDefaultMediaSoundClass() const
diff --git a/Engine/Source/Runtime/Engine/Private/SoundAttenuation.cpp b/Engine/Source/Runtime/Engine/Private/SoundAttenuation.cpp
--- a/Engine/Source/Runtime/Engine/Private/SoundAttenuation.cpp
+++ b/Engine/Source/Runtime/Engine/Private/SoundAttenuation.cpp
@@ -1,9 +1,11 @@
 // Copyright Epic Games, Inc. All Rights Reserved.
 
-
 #include "Sound/SoundAttenuation.h"
-#include "EngineDefines.h"
+
 #include "AudioDevice.h"
+#include "EngineDefines.h"
+#include "Internationalization/Internationalization.h"
+#include "Sound/SoundBase.h"
 #include "UObject/AnimPhysObjectVersion.h"
 
 /*-----------------------------------------------------------------------------
@@ -141,3 +143,70 @@ USoundAttenuation::USoundAttenuation(const FObjectInitializer& ObjectInitializer
 	: Super(ObjectInitializer)
 {
 }
+
+#define LOCTEXT_NAMESPACE "AudioGeneratorInterface"
+namespace Audio
+{
+	const FName FAttenuationInterface::Name = "UE.Attenuation";
+	const FName FAttenuationInterface::FInputs::Distance = "Distance";
+
+	FAttenuationInterface::FAttenuationInterface()
+		: FGeneratorInterface(Name, USoundBase::StaticClass())
+	{
+		Inputs =
+		{
+			{
+				FText(),
+				NSLOCTEXT("AudioGeneratorInterface_Attenuation", "DistanceDescription", "Distance between listener and sound location in game units."),
+				FName(),
+				{ FInputs::Distance, 0.0f }
+			}
+		};
+	}
+
+	const FName FSpatializationInterface::Name = "UE.Spatialization";
+	const FName FSpatializationInterface::FInputs::Azimuth = "Azimuth";
+	const FName FSpatializationInterface::FInputs::Elevation = "Elevation";
+
+	FSpatializationInterface::FSpatializationInterface()
+		: FGeneratorInterface(Name, USoundBase::StaticClass())
+	{
+		Inputs =
+		{
+			{
+				FText(),
+				NSLOCTEXT("Spatialization", "AzimuthDescription", "Horizontal angle between listener forward and sound location in degrees."),
+				FName(),
+				{ FInputs::Azimuth, 0.0f }
+			},
+			{
+				FText(),
+				NSLOCTEXT("Spatialization", "ElevationDescription", "Vertical angle between listener forward and sound location in degrees."),
+				FName(),
+				{ FInputs::Elevation, 0.0f }
+			}
+		};
+	}
+
+	FGeneratorInterfacePtr GetAttenuationInterface()
+	{
+		static FGeneratorInterfacePtr InterfacePtr;
+		if (!InterfacePtr.IsValid())
+		{
+			InterfacePtr = MakeShared<FAttenuationInterface>();
+		}
+		return InterfacePtr;
+	}
+
+	FGeneratorInterfacePtr GetSpatializationInterface()
+	{
+		static FGeneratorInterfacePtr InterfacePtr;
+		if (!InterfacePtr.IsValid())
+		{
+			InterfacePtr = MakeShared<FSpatializationInterface>();
+		}
+
+		return InterfacePtr;
+	}
+} // namespace Audio
+#undef LOCTEXT_NAMESPACE
diff --git a/Engine/Source/Runtime/Engine/Private/SoundBase.cpp b/Engine/Source/Runtime/Engine/Private/SoundBase.cpp
--- a/Engine/Source/Runtime/Engine/Private/SoundBase.cpp
+++ b/Engine/Source/Runtime/Engine/Private/SoundBase.cpp
@@ -3,6 +3,7 @@
 
 #include "AudioDevice.h"
 #include "EngineDefines.h"
+#include "IAudioGeneratorInterfaceRegistry.h"
 #include "IAudioParameterTransmitter.h"
 #include "Sound/AudioSettings.h"
 #include "Sound/SoundClass.h"
@@ -393,6 +394,13 @@ TUniquePtr<Audio::IParameterTransmitter> USoundBase::CreateParameterTransmitter(
 			return false;
 		}
 
+		bool SetParameter(FName InInterfaceName, FAudioParameter&& InValue) override
+		{
+			InValue.ParamName = Audio::IGeneratorInterfaceRegistry::GetMemberFullName(InInterfaceName, InValue.ParamName);
+			const FName ParamName = InValue.ParamName;
+			return SetParameter(MoveTemp(InValue));
+		}
+
 		uint64 GetInstanceID() const override
 		{
 			return InstanceID;
diff --git a/Engine/Source/Runtime/Engine/Public/ActiveSound.h b/Engine/Source/Runtime/Engine/Public/ActiveSound.h
--- a/Engine/Source/Runtime/Engine/Public/ActiveSound.h
+++ b/Engine/Source/Runtime/Engine/Public/ActiveSound.h
@@ -663,6 +663,8 @@ public:
 	FColor DebugColor;
 #endif // ENABLE_AUDIO_DEBUG
 
+	void UpdateGeneratorParameters(const TArray<FListener>& InListeners);
+
 	// Updates the wave instances to be played.
 	void UpdateWaveInstances(TArray<FWaveInstance*> &OutWaveInstances, const float DeltaTime);
 
