Fix ensure when deleting UMetaSound asset
#rnx

diff --git a/Engine/Plugins/Runtime/Metasound/Source/MetasoundEngine/Private/MetasoundAssetSubsystem.cpp b/Engine/Plugins/Runtime/Metasound/Source/MetasoundEngine/Private/MetasoundAssetSubsystem.cpp
--- a/Engine/Plugins/Runtime/Metasound/Source/MetasoundEngine/Private/MetasoundAssetSubsystem.cpp
+++ b/Engine/Plugins/Runtime/Metasound/Source/MetasoundEngine/Private/MetasoundAssetSubsystem.cpp
@@ -60,11 +60,14 @@ namespace Metasound
 				return false;
 			};
 
+			// These values are optional and not necessary to return successfully as MetaSounds
+			// don't require inputs or outputs for asset tags to be valid (ex. a new MetaSound,
+			// non-source asset has no inputs or outputs)
 			OutInfo.InputTypes.Reset();
-			bSuccess &= ParseTypesString(AssetTags::RegistryInputTypes, OutInfo.InputTypes);
+			ParseTypesString(AssetTags::RegistryInputTypes, OutInfo.InputTypes);
 
 			OutInfo.OutputTypes.Reset();
-			bSuccess &= ParseTypesString(AssetTags::RegistryOutputTypes, OutInfo.OutputTypes);
+			ParseTypesString(AssetTags::RegistryOutputTypes, OutInfo.OutputTypes);
 #endif // WITH_EDITORONLY_DATA
 
 			return bSuccess;
