Remove analysis and advanced options for MetaSound source details panel

#rnx

diff --git a/Engine/Plugins/Runtime/Metasound/Source/MetasoundEditor/Private/MetasoundDetailCustomization.cpp b/Engine/Plugins/Runtime/Metasound/Source/MetasoundEditor/Private/MetasoundDetailCustomization.cpp
--- a/Engine/Plugins/Runtime/Metasound/Source/MetasoundEditor/Private/MetasoundDetailCustomization.cpp
+++ b/Engine/Plugins/Runtime/Metasound/Source/MetasoundEditor/Private/MetasoundDetailCustomization.cpp
@@ -430,8 +430,7 @@ namespace Metasound
 
 					CustomizeInterfaces(DetailLayout);
 
-					DetailLayout.HideCategory("Advanced");
-					DetailLayout.HideCategory("Analysis");
+
 					DetailLayout.HideCategory("Attenuation");
 					DetailLayout.HideCategory("Effects");
 					DetailLayout.HideCategory("Loading");
@@ -470,9 +469,6 @@ namespace Metasound
 					DetailLayout.EditCategory("Modulation").InitiallyCollapsed(bShouldBeInitiallyCollapsed);
 					DetailLayout.EditCategory("Voice Management").InitiallyCollapsed(bShouldBeInitiallyCollapsed);
 
-					DetailLayout.EditCategory("Analysis").InitiallyCollapsed(bShouldBeInitiallyCollapsed);
-					DetailLayout.EditCategory("Advanced").InitiallyCollapsed(bShouldBeInitiallyCollapsed);
-
 					break;
 			}
 
@@ -487,6 +483,8 @@ namespace Metasound
 			DetailLayout.HideCategory("Loading");
 			DetailLayout.HideCategory("Playback");
 			DetailLayout.HideCategory("Subtitles");
+			DetailLayout.HideCategory("Analysis");
+			DetailLayout.HideCategory("Advanced");
 		}
 	} // namespace Editor
 } // namespace Metasound
