Fix missing arg for formatted comment
#rnx

diff --git a/Engine/Plugins/Runtime/Metasound/Source/MetasoundFrontend/Private/MetasoundFrontendTransform.cpp b/Engine/Plugins/Runtime/Metasound/Source/MetasoundFrontend/Private/MetasoundFrontendTransform.cpp
--- a/Engine/Plugins/Runtime/Metasound/Source/MetasoundFrontend/Private/MetasoundFrontendTransform.cpp
+++ b/Engine/Plugins/Runtime/Metasound/Source/MetasoundFrontend/Private/MetasoundFrontendTransform.cpp
@@ -393,6 +393,7 @@ namespace Metasound
 			{
 				UE_LOG(LogMetaSound, Display, TEXT("Asset '%s' interface '%s' updated: '%s' --> '%s'"),
 					*InDocument->GetMetadata().Version.ToString(),
+					*InterfaceVersion.Name.ToString(),
 					*InterfaceVersion.Number.ToString(),
 					*LastVersionUpdated->ToString());
 				return true;
