Fix crash when unregistering metasound asset with another asset type being editted.

diff --git a/Engine/Plugins/Runtime/Metasound/Source/MetasoundEditor/Private/MetasoundEditorGraphBuilder.cpp b/Engine/Plugins/Runtime/Metasound/Source/MetasoundEditor/Private/MetasoundEditorGraphBuilder.cpp
--- a/Engine/Plugins/Runtime/Metasound/Source/MetasoundEditor/Private/MetasoundEditorGraphBuilder.cpp
+++ b/Engine/Plugins/Runtime/Metasound/Source/MetasoundEditor/Private/MetasoundEditorGraphBuilder.cpp
@@ -1331,12 +1331,12 @@ namespace Metasound
 				{
 					if (Asset != &InMetaSound)
 					{
-						FMetasoundAssetBase* EditedMetaSound = IMetasoundUObjectRegistry::Get().GetObjectAsAssetBase(Asset);
-						check(EditedMetaSound);
-
-						if (EditedMetaSound->IsReferencedAsset(*MetaSoundAsset))
+						if (FMetasoundAssetBase* EditedMetaSound = IMetasoundUObjectRegistry::Get().GetObjectAsAssetBase(Asset))
 						{
-							EditedMetaSound->SetSynchronizationRequired();
+							if (EditedMetaSound->IsReferencedAsset(*MetaSoundAsset))
+							{
+								EditedMetaSound->SetSynchronizationRequired();
+							}
 						}
 					}
 				}
