Properly early out if MetaSound Preset's reference could not be or is not loaded properly
#rnx

diff --git a/Engine/Plugins/Runtime/Metasound/Source/MetasoundFrontend/Private/MetasoundFrontendTransform.cpp b/Engine/Plugins/Runtime/Metasound/Source/MetasoundFrontend/Private/MetasoundFrontendTransform.cpp
--- a/Engine/Plugins/Runtime/Metasound/Source/MetasoundFrontend/Private/MetasoundFrontendTransform.cpp
+++ b/Engine/Plugins/Runtime/Metasound/Source/MetasoundFrontend/Private/MetasoundFrontendTransform.cpp
@@ -421,6 +421,10 @@ namespace Metasound
 				{
 					const FNodeRegistryKey RegistryKey = FMetasoundFrontendRegistryContainer::Get()->GetRegistryKey(ClassMetadata);
 					PresetReferencedMetaSoundAsset = IMetaSoundAssetManager::GetChecked().TryLoadAssetFromKey(RegistryKey);
+					if (!PresetReferencedMetaSoundAsset)
+					{
+						UE_LOG(LogMetaSound, Error, TEXT("Failed to update Preset. Referenced asset class '%s' not found."), *ClassMetadata.GetDisplayName().ToString());
+					}
 					return;
 				}
 
@@ -446,16 +450,21 @@ namespace Metasound
 
 			if (PresetReferencedMetaSoundAsset)
 			{
-				bDidEdit |= FRebuildPresetRootGraph(PresetReferencedMetaSoundAsset->GetDocumentHandle()).Transform(InDocument);
-				if (bDidEdit)
+				if (bInterfaceManaged)
 				{
-					FMetasoundFrontendClassMetadata ParentMetadata = InDocument->GetRootGraphClass().Metadata;
-					ParentMetadata.SetType(EMetasoundFrontendClassType::External);
-					const FNodeRegistryKey RegistryKey = FMetasoundFrontendRegistryContainer::Get()->GetRegistryKey(ParentMetadata);
-					FMetasoundAssetBase* ParentMetaSoundAsset = IMetaSoundAssetManager::GetChecked().TryLoadAssetFromKey(RegistryKey);
-					if (ensure(ParentMetaSoundAsset))
+					bDidEdit |= FRebuildPresetRootGraph(PresetReferencedMetaSoundAsset->GetDocumentHandle()).Transform(InDocument);
+					if (bDidEdit)
 					{
-						ParentMetaSoundAsset->ConformObjectDataToInterfaces();
+						FMetasoundFrontendClassMetadata ParentMetadata = InDocument->GetRootGraphClass().Metadata;
+						ParentMetadata.SetType(EMetasoundFrontendClassType::External);
+						const FNodeRegistryKey RegistryKey = FMetasoundFrontendRegistryContainer::Get()->GetRegistryKey(ParentMetadata);
+						FMetasoundAssetBase* ParentMetaSoundAsset = IMetaSoundAssetManager::GetChecked().TryLoadAssetFromKey(RegistryKey);
+						if (ensure(ParentMetaSoundAsset))
+						{
+							ParentMetaSoundAsset->ConformObjectDataToInterfaces();
+						}
+
+						InDocument->SynchronizeDependencies();
 					}
 				}
 			}
@@ -471,9 +480,10 @@ namespace Metasound
 					Style.bMessageNodeUpdated = NewNode->GetClassMetadata().GetVersion() > InitialVersion;
 					NewNode->SetNodeStyle(Style);
 				}
+
+				InDocument->SynchronizeDependencies();
 			}
 
-			InDocument->SynchronizeDependencies();
 			return bDidEdit;
 		}
 
