Fix for crash due to parent asset not being added to MetaSoundAsset registry early enough in graph registration when adding new asset/dynamically loading
#rnx

diff --git a/Engine/Plugins/Runtime/Metasound/Source/MetasoundEngine/Private/MetasoundAssetSubsystem.cpp b/Engine/Plugins/Runtime/Metasound/Source/MetasoundEngine/Private/MetasoundAssetSubsystem.cpp
--- a/Engine/Plugins/Runtime/Metasound/Source/MetasoundEngine/Private/MetasoundAssetSubsystem.cpp
+++ b/Engine/Plugins/Runtime/Metasound/Source/MetasoundEngine/Private/MetasoundAssetSubsystem.cpp
@@ -111,18 +111,29 @@ void UMetaSoundAssetSubsystem::PostInitAssetScan()
 	}
 }
 
-void UMetaSoundAssetSubsystem::AddAssetReferences(const FMetasoundAssetBase& InAssetBase)
+void UMetaSoundAssetSubsystem::AddAssetReferences(FMetasoundAssetBase& InAssetBase)
 {
 	using namespace Metasound;
 	using namespace Metasound::Frontend;
 
+	constexpr bool bRegisterWithFrontend = false;
+
+	const FNodeClassInfo AssetClassInfo = InAssetBase.GetAssetClassInfo();
+	const FNodeRegistryKey AssetClassKey = NodeRegistryKey::CreateKey(AssetClassInfo);
+
+	if (!ContainsKey(AssetClassKey))
+	{
+		AddOrUpdateAsset(*InAssetBase.GetOwningAsset(), bRegisterWithFrontend);
+		UE_LOG(LogMetaSound, Log, TEXT("Adding asset '%s' to MetaSoundAsset registry."), *InAssetBase.GetOwningAssetName());
+	}
+
 	bool bLoadFromPathCache = false;
-	const TSet<FString>& AssetClassKeys = InAssetBase.GetReferencedAssetClassKeys();
-	for (const FString& AssetClassKey : AssetClassKeys)
+	const TSet<FString>& ReferencedAssetClassKeys = InAssetBase.GetReferencedAssetClassKeys();
+	for (const FString& ReferencedAssetClassKey : ReferencedAssetClassKeys)
 	{
-		if (!ContainsKey(AssetClassKey))
+		if (!ContainsKey(ReferencedAssetClassKey))
 		{
-			UE_LOG(LogMetaSound, Log, TEXT("Missing referenced class '%s' asset entry."), *AssetClassKey);
+			UE_LOG(LogMetaSound, Log, TEXT("Missing referenced class '%s' asset entry."), *ReferencedAssetClassKey);
 			bLoadFromPathCache = true;
 		}
 	}
@@ -152,7 +163,6 @@ void UMetaSoundAssetSubsystem::AddAssetReferences(const FMetasoundAssetBase& InA
 						"registration request (asset scan likely not complete)."),
 					*ClassKey,
 					*AssetClassPath.ToString());
-				constexpr bool bRegisterWithFrontend = false;
 
 				UObject* MetaSoundObject = MetaSoundAsset->GetOwningAsset();
 				if (ensureAlways(MetaSoundObject))
diff --git a/Engine/Plugins/Runtime/Metasound/Source/MetasoundEngine/Public/MetasoundAssetSubsystem.h b/Engine/Plugins/Runtime/Metasound/Source/MetasoundEngine/Public/MetasoundAssetSubsystem.h
--- a/Engine/Plugins/Runtime/Metasound/Source/MetasoundEngine/Public/MetasoundAssetSubsystem.h
+++ b/Engine/Plugins/Runtime/Metasound/Source/MetasoundEngine/Public/MetasoundAssetSubsystem.h
@@ -39,7 +39,7 @@ public:
 	void RemoveAsset(const FAssetData& InAssetData, bool bInUnregisterWithFrontend = true);
 	void RenameAsset(const FAssetData& InAssetData, bool bInReregisterWithFrontend = true);
 
-	virtual void AddAssetReferences(const FMetasoundAssetBase& InAssetBase) override;
+	virtual void AddAssetReferences(FMetasoundAssetBase& InAssetBase) override;
 	virtual void AddOrUpdateAsset(UObject& InObject, bool bInRegisterWithFrontend = true) override;
 	virtual bool CanAutoUpdate(const FMetasoundFrontendClassName& InClassName) const override;
 	virtual bool ContainsKey(const Metasound::Frontend::FNodeRegistryKey& InRegistryKey) const override;
diff --git a/Engine/Plugins/Runtime/Metasound/Source/MetasoundFrontend/Public/MetasoundAssetManager.h b/Engine/Plugins/Runtime/Metasound/Source/MetasoundFrontend/Public/MetasoundAssetManager.h
--- a/Engine/Plugins/Runtime/Metasound/Source/MetasoundFrontend/Public/MetasoundAssetManager.h
+++ b/Engine/Plugins/Runtime/Metasound/Source/MetasoundFrontend/Public/MetasoundAssetManager.h
@@ -77,7 +77,7 @@ namespace Metasound
 
 			// Adds missing assets using the provided asset's local reference class cache. Used
 			// to prime system from asset attempting to register prior to asset scan being complete.
-			virtual void AddAssetReferences(const FMetasoundAssetBase& InAssetBase) = 0;
+			virtual void AddAssetReferences(FMetasoundAssetBase& InAssetBase) = 0;
 
 			// Add or Update a MetaSound Asset's entry data
 			virtual void AddOrUpdateAsset(UObject& InObject, bool bInRegisterWithFrontend = true) = 0;
