Fix metasound to use new sound proxy properties for error reporting

#buildfix

diff --git a/Engine/Plugins/Runtime/Metasound/Source/MetasoundEngine/Private/MetasoundWavePlayerNode.cpp b/Engine/Plugins/Runtime/Metasound/Source/MetasoundEngine/Private/MetasoundWavePlayerNode.cpp
--- a/Engine/Plugins/Runtime/Metasound/Source/MetasoundEngine/Private/MetasoundWavePlayerNode.cpp
+++ b/Engine/Plugins/Runtime/Metasound/Source/MetasoundEngine/Private/MetasoundWavePlayerNode.cpp
@@ -302,7 +302,7 @@ namespace Metasound
 
 			if (IsWaveValid())
 			{
-				UE_LOG(LogMetaSound, Verbose, TEXT("Starting Sound: '%s'"), *CurrentWaveAsset->GetFullName());
+				UE_LOG(LogMetaSound, Verbose, TEXT("Starting Sound: '%s'"), *CurrentWaveAsset->GetFName().ToString());
 				SoundAssetSampleRate = CurrentWaveAsset->GetSampleRate();
 				SoundAssetDurationSeconds = CurrentWaveAsset->GetDuration();
 				SoundAssetNumFrames = CurrentWaveAsset->GetNumFrames();
@@ -375,7 +375,7 @@ namespace Metasound
 			{
 				if (bReportToLog)
 				{
-					UE_LOG(LogMetaSound, Error, TEXT("Failed to initialize SoundWave decoder. WavePlayerNode only supports 2 channels max [%s: %d Channels])"), *WaveProxy->GetFullName(), NumChannels);
+					UE_LOG(LogMetaSound, Error, TEXT("Failed to initialize SoundWave decoder. WavePlayerNode only supports 2 channels max [Package %s Sound %s: %d Channels])"), *WaveProxy->GetPackageName().ToString(), *WaveProxy->GetFName().ToString(), NumChannels);
 				}
 				return false;
 			}
