Upgrades to Wavewriter to support multiple channels
#p4v-preflight-copy 18297712

diff --git a/Engine/Plugins/Runtime/Metasound/Source/MetasoundStandardNodes/Private/MetasoundWaveWriterNode.cpp b/Engine/Plugins/Runtime/Metasound/Source/MetasoundStandardNodes/Private/MetasoundWaveWriterNode.cpp
--- a/Engine/Plugins/Runtime/Metasound/Source/MetasoundStandardNodes/Private/MetasoundWaveWriterNode.cpp
+++ b/Engine/Plugins/Runtime/Metasound/Source/MetasoundStandardNodes/Private/MetasoundWaveWriterNode.cpp
@@ -1,7 +1,5 @@
 // Copyright Epic Games, Inc. All Rights Reserved.
 
-#include "MetasoundWaveWriterNode.h"
-
 #include "HAL/FileManager.h"
 #include "MetasoundBuildError.h"
 #include "MetasoundEnumRegistrationMacro.h"
@@ -191,20 +189,71 @@ namespace Metasound
 	};
 	const FString FNumberedFileCache::Seperator{ TEXT("_") };
 
-	class FWaveWriterOperator : public TExecutableOperator<FWaveWriterOperator>
+	namespace WaveWriterOperatorPrivate
+	{
+		// Need to keep this outside the template so there's only 1
+		TSharedPtr<FNumberedFileCache> GetNameCache()
+		{
+			static const TCHAR* WaveExt = TEXT(".wav");
+
+			// Build cache of numbered files (do this once only).
+			static TSharedPtr<FNumberedFileCache> NumberedFileCacheSP = MakeShared<FNumberedFileCache>(*FPaths::AudioCaptureDir(), WaveExt);
+			return NumberedFileCacheSP;
+		}
+
+		static const FString GetDefaultFileName()
+		{
+			static const FString DefaultFileName = TEXT("Output");
+			return DefaultFileName;
+		}
+		static const TCHAR* GetEnabledPinName()
+		{
+			static const TCHAR* EnabledPinName = TEXT("Enabled");
+			return EnabledPinName;
+		}
+		static const TCHAR* GetFilenamePrefixPinName()
+		{
+			static const TCHAR* FilenamePrefixPinName = TEXT("Filename Prefix");
+			return FilenamePrefixPinName;
+		}
+	}
+
+	template<int32 NumInputChannels>
+	class TWaveWriterOperator : public TExecutableOperator<TWaveWriterOperator<NumInputChannels>>
 	{
 	public:
-		FWaveWriterOperator(const FOperatorSettings& InSettings, FAudioBufferReadRef&& InAudioBuffer, FBoolReadRef&& InEnabled, TUniquePtr<FWaveWriter>&& InStream)
-			: AudioInput{ MoveTemp(InAudioBuffer) }
+		// Theoretical limit of .WAV files.
+		static_assert(NumInputChannels > 0 && NumInputChannels <= 65535, "Num Channels > 0 and <= 65535");
+
+		TWaveWriterOperator(const FOperatorSettings& InSettings, TArray<FAudioBufferReadRef>&& InAudioBuffers, FBoolReadRef&& InEnabled, const TSharedPtr<FNumberedFileCache, ESPMode::ThreadSafe>& InNumberedFileCache, const FString& InFilenamePrefix)
+			: AudioInputs{ MoveTemp(InAudioBuffers) }
 			, Enabled{ MoveTemp(InEnabled) }
-			, Writer{ MoveTemp(InStream) }
-		{}
+			, NumberedFileCacheSP{ InNumberedFileCache }
+			, FileNamePrefix{ InFilenamePrefix }
+			, SampleRate{ InSettings.GetSampleRate() }
+		{
+			check(AudioInputs.Num() == NumInputChannels);
+
+			// Make an interleave buffer if we need one.
+			if (NumInputChannels > 1)
+			{
+				InterleaveBuffer.SetNum(InSettings.GetNumFramesPerBlock() * NumInputChannels);
+				AudioInputBufferPtrs.SetNum(NumInputChannels);
+				for (int32 i = 0; i < NumInputChannels; ++i)
+				{
+					AudioInputBufferPtrs[i] = AudioInputs[i]->GetData();
+				}
+			}
+		}
 
 		FDataReferenceCollection GetInputs() const override
 		{
-			FDataReferenceCollection InputDataReferences;
-			InputDataReferences.AddDataReadReference(AudioInputPinName, AudioInput);
-			return InputDataReferences;
+			FDataReferenceCollection InputPins;
+			for (int32 i = 0; i < NumInputChannels; ++i)
+			{
+				InputPins.AddDataReadReference(GetAudioInputName(i), AudioInputs[i]);
+			}
+			return InputPins;
 		}
 		FDataReferenceCollection GetOutputs() const override
 		{
@@ -212,112 +261,243 @@ namespace Metasound
 			return OutputDataReferences;
 		}
 
-		static const FVertexInterface& DeclareVertexInterface() 
+		static const FVertexInterface& DeclareVertexInterface()
 		{
-			static const FVertexInterface Interface(
-				FInputVertexInterface(
-					TInputDataVertexModel<FString>(FilenamePrefixPinName, LOCTEXT("WaveWriterFilenamePrefixDescriptionX", "Filename Prefix of file you are writing."), FString(DefaultFileName)),
-					TInputDataVertexModel<bool>(EnabledPinName, LOCTEXT("WaveWriterEnabledDescription", "If this wave writer is enabled or not. File will remain open if disabled until graph termination."), true),
-					TInputDataVertexModel<FAudioBuffer>(AudioInputPinName, LOCTEXT("WaveWriterAudioInputDescription", "Audio input that you want serialized."))
-				),
-				FOutputVertexInterface()
-			);
-			return Interface;
+			auto CreateDefaultInterface = []()-> FVertexInterface
+			{
+				using namespace WaveWriterOperatorPrivate;
+
+				// inputs
+				FInputVertexInterface InputInterface(
+					TInputDataVertexModel<FString>(GetFilenamePrefixPinName(), LOCTEXT("WaveWriterFilenamePrefixDescription", "Filename Prefix of file you are writing."), GetDefaultFileName()),
+					TInputDataVertexModel<bool>(GetEnabledPinName(), LOCTEXT("WaveWriterEnabledDescription", "If this wave writer is enabled or not."), true)
+				);
+
+				// For backwards compatibility with previous (mono) node, in the case of 1 channels, just provide the old interface.
+				for (int32 InputIndex = 0; InputIndex < NumInputChannels; ++InputIndex)
+				{
+					InputInterface.Add(TInputDataVertexModel<FAudioBuffer>(GetAudioInputName(InputIndex), GetAudioInputDescription(InputIndex)));
+				}
+				FOutputVertexInterface OutputInterface;
+
+				return FVertexInterface(InputInterface, OutputInterface);
+			};
+
+			static const FVertexInterface DefaultInterface = CreateDefaultInterface();
+			return DefaultInterface;
 		}
 
 		static const FNodeClassMetadata& GetNodeInfo()
 		{
-			auto InitNodeInfo = []() -> FNodeClassMetadata
+			// used if NumChannels == 1
+			auto CreateNodeClassMetadataMono = []() -> FNodeClassMetadata
 			{
-				FNodeClassMetadata Info;
-				Info.ClassName = { StandardNodes::Namespace, TEXT("WaveWriter"), StandardNodes::AudioVariant };
-				Info.MajorVersion = 1;
-				Info.MinorVersion = 0;
-				Info.DisplayName = LOCTEXT("Metasound_WaveWriterNodeDisplayName", "Wave Writer");
-				Info.Description = LOCTEXT("Metasound_WaveWriterNodeDescription", "Write a the incoming audio signal to disk");
-				Info.Author = PluginAuthor;
-				Info.PromptIfMissing = PluginNodeMissingPrompt;
-				Info.DefaultInterface = DeclareVertexInterface();
-				Info.CategoryHierarchy.Emplace(NodeCategories::Io);
-				return Info;
+				// For backwards compatibility with previous (mono) WaveWriters keep the node name the same.
+				FName OperatorName = TEXT("WaveWriter");
+				FText NodeDisplayName = LOCTEXT("Metasound_WaveWriterNodeMonoDisplayName", "Wave Writer (Mono)");
+				FText NodeDescription = LOCTEXT("Metasound_WaveWriterNodeMonoDescription", "Write a mono audio signal to disk");
+				FVertexInterface NodeInterface = DeclareVertexInterface();
+
+				return CreateNodeClassMetadata(OperatorName, NodeDisplayName, NodeDescription, NodeInterface);
 			};
-			static const FNodeClassMetadata Info = InitNodeInfo();
-			return Info;
-		}
 
-		static TUniquePtr<IOperator> CreateOperator(const FCreateOperatorParams& InParams, FBuildErrorArray& OutErrors)
-		{
-			const FWaveWriterNode& Node = static_cast<const FWaveWriterNode&>(InParams.Node);
-			const FDataReferenceCollection& InputCol = InParams.InputDataReferences;
-			const FOperatorSettings& Settings = InParams.OperatorSettings;
-			const FInputVertexInterface& InputInterface = DeclareVertexInterface().GetInputInterface();
-			
-			static const TCHAR* WaveExt = TEXT(".wav");
+			// used if NumChannels == 2
+			auto CreateNodeClassMetadataStereo = []() -> FNodeClassMetadata
+			{
+				FName OperatorName = TEXT("Wave Writer (Stereo)");
+				FText NodeDisplayName = LOCTEXT("Metasound_WaveWriterNodeStereoDisplayName", "Wave Writer (Stereo)");
+				FText NodeDescription = LOCTEXT("Metasound_WaveWriterNodeStereoDescription", "Write a stereo audio signal to disk");
+				FVertexInterface NodeInterface = DeclareVertexInterface();
 
-			// Build cache of numbered files (do this once only).
-			static TUniquePtr<FNumberedFileCache> NumberedFileCache = MakeUnique<FNumberedFileCache>(*FPaths::AudioCaptureDir(), WaveExt);
-			
-			FStringReadRef FilenamePrefix = InputCol.GetDataReadReferenceOrConstructWithVertexDefault<FString>(InputInterface, FilenamePrefixPinName, Settings);
-			FString Filename = NumberedFileCache->GenerateNextNumberedFilename(*FilenamePrefix);
+				return  CreateNodeClassMetadata(OperatorName, NodeDisplayName, NodeDescription, NodeInterface);
+			};
 
-			if (TUniquePtr<FArchive> FileWriter = TUniquePtr<FArchive>(IFileManager::Get().CreateFileWriter(*Filename, IO_WRITE)))
+			// used if NumChannels > 2
+			auto CreateNodeClassMetadataMultiChan = []() -> FNodeClassMetadata
 			{
-				return MakeUnique<FWaveWriterOperator>(
-					Settings,
-					InputCol.GetDataReadReferenceOrConstruct<FAudioBuffer>(AudioInputPinName, Settings),
-					InputCol.GetDataReadReferenceOrConstructWithVertexDefault<bool>(InputInterface, EnabledPinName, Settings),
-					MakeUnique<FWaveWriter>(MoveTemp(FileWriter), Settings.GetSampleRate(), 1, true)
-				);
-			}
-			
-			// Failed to open a writer object. Log an error.
-			OutErrors.Emplace(MakeUnique<FFileWriteError>(Node, Filename));
-
-			// Create a default operator with no-writer which will do nothing.
-			return MakeUnique<FWaveWriterOperator>(
-				Settings,
-				InputCol.GetDataReadReferenceOrConstruct<FAudioBuffer>(AudioInputPinName, Settings),
-				InputCol.GetDataReadReferenceOrConstructWithVertexDefault<bool>(InputInterface, EnabledPinName, Settings),
-				nullptr
-			);
+				FName OperatorName = *FString::Printf(TEXT("Wave Writer (%d-Channel)"), NumInputChannels);
+				FText NodeDisplayName = FText::Format(LOCTEXT("Metasound_WaveWriterNodeMultiChannelDisplayName", "Wave Writer ({0}-channel)"), NumInputChannels);
+				FText NodeDescription = LOCTEXT("Metasound_WaveWriterNodeMultiDescription", "Write a multi-channel audio signal to disk");
+				FVertexInterface NodeInterface = DeclareVertexInterface();
+
+				return  CreateNodeClassMetadata(OperatorName, NodeDisplayName, NodeDescription, NodeInterface);
+			};
+
+			static const FNodeClassMetadata Metadata = (NumInputChannels == 1) ? CreateNodeClassMetadataMono()
+				: (NumInputChannels == 2) ? CreateNodeClassMetadataStereo() : CreateNodeClassMetadataMultiChan();
+			return Metadata;
 		}
 
+		static TUniquePtr<IOperator> CreateOperator(const FCreateOperatorParams& InParams, FBuildErrorArray& OutErrors);
+
 		void Execute()
 		{
+			// Enabled and wasn't before? Enable.
+			if (!bIsEnabled && *Enabled)
+			{
+				Enable();
+			}
+			// Disabled but currently Enabled? Disable.
+			else if (bIsEnabled && !*Enabled)
+			{
+				Disable();
+			}
+
+			// If we have a valid writer and enabled.
 			if (Writer && *Enabled)
 			{
-				Writer->Write(MakeArrayView(AudioInput->GetData(), AudioInput->Num()));
+				// Need to Interleave?
+				if (NumInputChannels > 1)
+				{
+					InterleaveChannels(AudioInputBufferPtrs.GetData(), NumInputChannels, AudioInputs[0]->Num(), InterleaveBuffer.GetData());
+					Writer->Write(MakeArrayView(InterleaveBuffer.GetData(), InterleaveBuffer.Num()));
+				}
+				else if (NumInputChannels == 1)
+				{
+					Writer->Write(MakeArrayView(AudioInputs[0]->GetData(), AudioInputs[0]->Num()));
+				}
 			}
 		}
 
 	protected:
-		FAudioBufferReadRef AudioInput;
+		static const FVertexName GetAudioInputName(int32 InInputIndex)
+		{
+			if (NumInputChannels == 1)
+			{
+				// To maintain backwards compatibility with previous implementation keep the pin name the same.
+				static const FName AudioInputPinName = TEXT("In");
+				return AudioInputPinName;
+			}
+			else if (NumInputChannels == 2)
+			{
+				return *FString::Printf(TEXT("In %d %s"), InInputIndex, (InInputIndex == 0) ? TEXT("L") : TEXT("R"));
+			}
+
+			return *FString::Printf(TEXT("In %d"), InInputIndex);
+		}
+
+		static const FText GetAudioInputDescription(int32 InputIndex)
+		{
+			return FText::Format(LOCTEXT("WaveWriterAudioInputDescription", "Audio Input #: {0}"), InputIndex);
+		}
+
+		static FNodeClassMetadata CreateNodeClassMetadata(const FName& InOperatorName, const FText& InDisplayName, const FText& InDescription, const FVertexInterface& InDefaultInterface)
+		{
+			FNodeClassMetadata Metadata
+			{
+				FNodeClassName { StandardNodes::Namespace, InOperatorName, StandardNodes::AudioVariant },
+				1, // Major Version
+				1, // Minor Version
+				InDisplayName,
+				InDescription,
+				PluginAuthor,
+				PluginNodeMissingPrompt,
+				InDefaultInterface,
+				{ NodeCategories::Io },
+				{ LOCTEXT("Metasound_AudioMixerKeyword", "Writer") },
+				FNodeDisplayStyle{}
+			};
+			return Metadata;
+		}
+
+		// TODO. Move to DSP lib.
+		static void InterleaveChannels(const float* RESTRICT InMonoChannelsToInterleave[], const int32 InNumChannelsToInterleave, const int32 NumSamplesPerChannel, float* RESTRICT OutInterleavedBuffer)
+		{
+			for (int32 Sample = 0; Sample < NumSamplesPerChannel; ++Sample)
+			{
+				for (int32 Channel = 0; Channel < InNumChannelsToInterleave; ++Channel)
+				{
+					*OutInterleavedBuffer++ = InMonoChannelsToInterleave[Channel][Sample];
+				}
+			}
+		}
+
+		void Enable()
+		{
+			if (ensure(!bIsEnabled))
+			{
+				bIsEnabled = true;
+				FString Filename = NumberedFileCacheSP->GenerateNextNumberedFilename(*FileNamePrefix);
+				TUniquePtr<FArchive> Stream{ IFileManager::Get().CreateFileWriter(*Filename, IO_WRITE) };
+				if (Stream.IsValid())
+				{
+					Writer = MakeUnique<FWaveWriter>(MoveTemp(Stream), SampleRate, NumInputChannels, true);
+				}
+			}
+		}
+		void Disable()
+		{
+			if (ensure(bIsEnabled))
+			{
+				bIsEnabled = false;
+				Writer.Reset();
+			}
+		}
+
+		TArray<FAudioBufferReadRef> AudioInputs;
+		TArray<const float*> AudioInputBufferPtrs;
+		TArray<float> InterleaveBuffer;
 		FBoolReadRef Enabled;
-		TUniquePtr<FArchive> Output;
 		TUniquePtr<FWaveWriter> Writer;
-		static constexpr const TCHAR* DefaultFileName = TEXT("Output");
-		static constexpr const TCHAR* AudioInputPinName = TEXT("In");
-		static constexpr const TCHAR* EnabledPinName = TEXT("Enabled");
-		static constexpr const TCHAR* FilenamePrefixPinName = TEXT("Filename Prefix");
+		TSharedPtr<FNumberedFileCache, ESPMode::ThreadSafe> NumberedFileCacheSP;
+		FString FileNamePrefix;
+		float SampleRate = 0.f;
+		bool bIsEnabled = false;
 	};
 
-	// Linkage for constexpr on older clang.
-	constexpr const TCHAR* FWaveWriterOperator::DefaultFileName;
-	constexpr const TCHAR* FWaveWriterOperator::AudioInputPinName;
-	constexpr const TCHAR* FWaveWriterOperator::EnabledPinName;
-	constexpr const TCHAR* FWaveWriterOperator::FilenamePrefixPinName;
-
-	FWaveWriterNode::FWaveWriterNode(const FVertexName& InName, const FGuid& InInstanceID)
-		: FNodeFacade(InName, InInstanceID, TFacadeOperatorClass<FWaveWriterOperator>())
+	template<int32 NumInputChannels>
+	TUniquePtr<Metasound::IOperator> TWaveWriterOperator<NumInputChannels>::CreateOperator(const FCreateOperatorParams& InParams, FBuildErrorArray& OutErrors)
 	{
+		using namespace WaveWriterOperatorPrivate;
+
+		const FDataReferenceCollection& InputCol = InParams.InputDataReferences;
+		const FOperatorSettings& Settings = InParams.OperatorSettings;
+		const FInputVertexInterface& InputInterface = DeclareVertexInterface().GetInputInterface();
+
+		FStringReadRef FilenamePrefix = InputCol.GetDataReadReferenceOrConstructWithVertexDefault<FString>(InputInterface, GetFilenamePrefixPinName(), Settings);
+
+		TArray<FAudioBufferReadRef> InputBuffers;
+		for (int32 i = 0; i < NumInputChannels; ++i)
+		{
+			InputBuffers.Add(InputCol.GetDataReadReferenceOrConstruct<FAudioBuffer>(GetAudioInputName(i), InParams.OperatorSettings));
+		}
+
+		return MakeUnique<TWaveWriterOperator>(
+			Settings,
+			MoveTemp(InputBuffers),
+			InputCol.GetDataReadReferenceOrConstructWithVertexDefault<bool>(InputInterface, GetEnabledPinName(), Settings),
+			GetNameCache(),
+			*FilenamePrefix
+			);
 	}
 
-	FWaveWriterNode::FWaveWriterNode(const FNodeInitData& InInitData)
-		: FWaveWriterNode(InInitData.InstanceName, InInitData.InstanceID)
+	template<int32 NumInputChannels>
+	class METASOUNDSTANDARDNODES_API TWaveWriterNode : public FNodeFacade
 	{
-	}
+	public:
+		TWaveWriterNode(const FVertexName& InName, const FGuid& InInstanceID)
+			: FNodeFacade(InName, InInstanceID, TFacadeOperatorClass<TWaveWriterOperator<NumInputChannels>>())
+		{
+		}
+
+		TWaveWriterNode(const FNodeInitData& InInitData)
+			: TWaveWriterNode(InInitData.InstanceName, InInitData.InstanceID)
+		{
+		}
+	};
 
-	METASOUND_REGISTER_NODE(FWaveWriterNode)
+	#define REGISTER_WAVEWRITER_NODE(A) \
+		using FWaveWriterNode_##A = TWaveWriterNode<A>; \
+		METASOUND_REGISTER_NODE(FWaveWriterNode_##A)
+	
+	REGISTER_WAVEWRITER_NODE(1);
+	REGISTER_WAVEWRITER_NODE(2);
+	REGISTER_WAVEWRITER_NODE(3);
+	REGISTER_WAVEWRITER_NODE(4);
+	REGISTER_WAVEWRITER_NODE(5);
+	REGISTER_WAVEWRITER_NODE(6);
+	REGISTER_WAVEWRITER_NODE(7);
+	REGISTER_WAVEWRITER_NODE(8);
 }
 
 #undef LOCTEXT_NAMESPACE
diff --git a/Engine/Plugins/Runtime/Metasound/Source/MetasoundStandardNodes/Public/MetasoundWaveWriterNode.h b/Engine/Plugins/Runtime/Metasound/Source/MetasoundStandardNodes/Public/MetasoundWaveWriterNode.h
deleted file mode 100644
