MetaSounds Interfaces Checkpoint 2:
- Version up source archetypes to become two interfaces: channel interfaces (mono/stereo) & base source namespace
- Clean-up Interface panel to support namespacing better
- Fix bugs with assuming interfaces are always and the only base namespace members
- Allow namespacing for any arbitrary interface member
- Add lock icon to clarify what interface members cannot be modified individually (i.e. cannot add, remove, or rename them as they are interface members)
- Organize members alphabetically
#rnx

diff --git a/Engine/Config/BaseEngine.ini b/Engine/Config/BaseEngine.ini
--- a/Engine/Config/BaseEngine.ini
+++ b/Engine/Config/BaseEngine.ini
@@ -963,11 +963,17 @@ NearClipPlane=10.0
 +ClassRedirects=(OldName="/Script/MovieSceneTracks.MovieSceneVectorTrackExtensions",NewName="/Script/SequencerScripting.MovieSceneFloatVectorTrackExtensions")
 +StructRedirects=(OldName="/Script/MovieSceneTracks.MovieSceneVectorKeyStructBase",NewName="/Script/MovieSceneTracks.MovieSceneFloatVectorKeyStructBase")
 +StructRedirects=(OldName="/Script/MovieSceneTracks.MovieSceneVectorKeyStruct",NewName="/Script/MovieSceneTracks.MovieSceneVector3fKeyStruct")
+
+; Audio Component Parameter Updates
 +StructRedirects=(OldName="AudioComponentParam", NewName="/Script/AudioExtensions.AudioParameter")
 +PropertyRedirects=(OldName="AudioParameter.SoundWaveParam", NewName="AudioParameter.ObjectParam")
-+FunctionRedirects=(OldName="AudioComponent.SetBoolParameter", NewName="SoundGeneratorParameterInterface.SetBoolParameter")
-+FunctionRedirects=(OldName="AudioComponent.SetFloatParameter", NewName="SoundGeneratorParameterInterface.SetFloatParameter")
-+FunctionRedirects=(OldName="AudioComponent.SetIntParameter", NewName="SoundGeneratorParameterInterface.SetIntParameter")
+
++FunctionRedirects=(OldName="AudioComponent.SetBoolParameter", NewName="SoundParameterControllerInterface.SetBoolParameter")
++FunctionRedirects=(OldName="AudioComponent.SetFloatParameter", NewName="SoundParameterControllerInterface.SetFloatParameter")
++FunctionRedirects=(OldName="AudioComponent.SetIntParameter", NewName="SoundParameterControllerInterface.SetIntParameter")
++FunctionRedirects=(OldName="SoundGeneratorParameterInterface.SetBoolParameter", NewName="SoundParameterControllerInterface.SetBoolParameter")
++FunctionRedirects=(OldName="SoundGeneratorParameterInterface.SetFloatParameter", NewName="SoundParameterControllerInterface.SetFloatParameter")
++FunctionRedirects=(OldName="SoundGeneratorParameterInterface.SetIntParameter", NewName="SoundParameterControllerInterface.SetIntParameter")
 +PropertyRedirects=(OldName="AudioComponent.InstanceParameters", NewName="AudioComponent.DefaultParameters")
 
 ; LiveLink data moving to runtime core module
diff --git a/Engine/Plugins/Runtime/GameplayAbilities/Source/GameplayAbilities/Public/GameplayCueNotifyTypes.h b/Engine/Plugins/Runtime/GameplayAbilities/Source/GameplayAbilities/Public/GameplayCueNotifyTypes.h
--- a/Engine/Plugins/Runtime/GameplayAbilities/Source/GameplayAbilities/Public/GameplayCueNotifyTypes.h
+++ b/Engine/Plugins/Runtime/GameplayAbilities/Source/GameplayAbilities/Public/GameplayCueNotifyTypes.h
@@ -355,7 +355,7 @@ public:
 /**
  * FGameplayCueNotify_SoundParameterInterfaceInfo
  *
- *	Properties that specify how to interface with the ISoundGeneratorParameterInterface
+ *	Properties that specify how to interface with the ISoundParameterControllerInterface
  */
 USTRUCT(BlueprintType)
 struct FGameplayCueNotify_SoundParameterInterfaceInfo
diff --git a/Engine/Plugins/Runtime/Metasound/Source/MetasoundEditor/Private/MetasoundDetailCustomization.cpp b/Engine/Plugins/Runtime/Metasound/Source/MetasoundEditor/Private/MetasoundDetailCustomization.cpp
--- a/Engine/Plugins/Runtime/Metasound/Source/MetasoundEditor/Private/MetasoundDetailCustomization.cpp
+++ b/Engine/Plugins/Runtime/Metasound/Source/MetasoundEditor/Private/MetasoundDetailCustomization.cpp
@@ -8,7 +8,7 @@
 #include "DetailLayoutBuilder.h"
 #include "DetailWidgetRow.h"
 #include "Framework/Notifications/NotificationManager.h"
-#include "IAudioGeneratorInterfaceRegistry.h"
+#include "IAudioParameterInterfaceRegistry.h"
 #include "IAudioParameterTransmitter.h"
 #include "IDetailGroup.h"
 #include "Input/Events.h"
@@ -253,7 +253,7 @@ namespace Metasound
 
 		FName FMetasoundDetailCustomization::GetInterfaceVersionsPath() const
 		{
-			return Metasound::Editor::BuildChildPath(DocumentPropertyName, GET_MEMBER_NAME_CHECKED(FMetasoundFrontendDocument, InterfaceVersions));
+			return Metasound::Editor::BuildChildPath(DocumentPropertyName, GET_MEMBER_NAME_CHECKED(FMetasoundFrontendDocument, Interfaces));
 		}
 
 		FName FMetasoundDetailCustomization::GetMetadataRootClassPath() const
@@ -276,8 +276,6 @@ namespace Metasound
 
 			if (const FMetasoundAssetBase* MetaSoundAsset = IMetasoundUObjectRegistry::Get().GetObjectAsAssetBase(MetaSound.Get()))
 			{
-				const TArray<FMetasoundFrontendVersion>& ImplementedInterfaces = MetaSoundAsset->GetDocumentChecked().InterfaceVersions;
-
 				auto GetVersionName = [](const FMetasoundFrontendVersion& Version) { return Version.Name; };
 				auto IsTransmittableVersion = [](const FMetasoundFrontendVersion& Version)
 				{
@@ -292,13 +290,14 @@ namespace Metasound
 					return false;
 				};
 
+				const TSet<FMetasoundFrontendVersion>& ImplementedInterfaces = MetaSoundAsset->GetDocumentChecked().Interfaces;
 				Algo::TransformIf(ImplementedInterfaces, ImplementedInterfaceNames, IsTransmittableVersion, GetVersionName);
 
-				Audio::IGeneratorInterfaceRegistry::Get().IterateInterfaces([this](Audio::FGeneratorInterfacePtr Interface)
+				Audio::IAudioParameterInterfaceRegistry::Get().IterateInterfaces([this](Audio::FParameterInterfacePtr Interface)
 				{
-					if (!ImplementedInterfaceNames.Contains(Interface->Name))
+					if (!ImplementedInterfaceNames.Contains(Interface->GetName()))
 					{
-						FString Name = Interface->Name.ToString();
+						FString Name = Interface->GetName().ToString();
 						AddableInterfaceNames.Add(MakeShared<FString>(MoveTemp(Name)));
 					}
 				});
diff --git a/Engine/Plugins/Runtime/Metasound/Source/MetasoundEditor/Private/MetasoundEditor.cpp b/Engine/Plugins/Runtime/Metasound/Source/MetasoundEditor/Private/MetasoundEditor.cpp
--- a/Engine/Plugins/Runtime/Metasound/Source/MetasoundEditor/Private/MetasoundEditor.cpp
+++ b/Engine/Plugins/Runtime/Metasound/Source/MetasoundEditor/Private/MetasoundEditor.cpp
@@ -42,6 +42,7 @@
 #include "MetasoundUObjectRegistry.h"
 #include "Misc/Attribute.h"
 #include "Modules/ModuleManager.h"
+#include "PropertyCustomizationHelpers.h"
 #include "PropertyEditorModule.h"
 #include "ScopedTransaction.h"
 #include "SMetasoundActionMenu.h"
@@ -149,6 +150,7 @@ namespace Metasound
 		{
 		private:
 			TSharedPtr<FMetasoundGraphMemberSchemaAction> MetasoundAction;
+			FMetasoundFrontendVersion InterfaceVersion;
 
 		public:
 			SLATE_BEGIN_ARGS(SMetaSoundGraphPaletteItem)
@@ -162,6 +164,11 @@ namespace Metasound
 				TSharedPtr<FEdGraphSchemaAction> Action = InCreateData->Action;
 				MetasoundAction = StaticCastSharedPtr<FMetasoundGraphMemberSchemaAction>(Action);
 
+				if (UMetasoundEditorGraphVertex* GraphVertex = Cast<UMetasoundEditorGraphVertex>(MetasoundAction->GetGraphMember()))
+				{
+					InterfaceVersion = GraphVertex->GetInterfaceVersion();
+				}
+
 				SGraphPaletteItem::Construct(SGraphPaletteItem::FArguments(), InCreateData);
 			}
 
@@ -170,6 +177,11 @@ namespace Metasound
 			{
 				using namespace Frontend;
 
+				if (InterfaceVersion.IsValid())
+				{
+					return;
+				}
+
 				if (MetasoundAction.IsValid())
 				{
 					if (UMetasoundEditorGraphMember* GraphMember = MetasoundAction->GetGraphMember())
@@ -184,6 +196,69 @@ namespace Metasound
 				}
 			}
 
+			virtual TSharedRef<SWidget> CreateTextSlotWidget(FCreateWidgetForActionData* const InCreateData, TAttribute<bool> bIsReadOnly) override
+			{
+				TSharedRef<SWidget> TextWidget = SGraphPaletteItem::CreateTextSlotWidget(InCreateData, bIsReadOnly);
+
+				TSharedRef<SHorizontalBox> LayoutWidget = SNew(SHorizontalBox)
+				+ SHorizontalBox::Slot()
+				.AutoWidth()
+				.HAlign(HAlign_Left)
+				.VAlign(VAlign_Center)
+				[
+					SNew(SImage)
+					.Image(FEditorStyle::GetBrush(InterfaceVersion.IsValid() ? "Icons.Lock" : "Icons.BulletPoint"))
+					.ToolTipText(InterfaceVersion.IsValid() ? FText::Format(LOCTEXT("InterfaceMemberToolTipFormat", "Cannot Add/Remove: Member of interface '{0}'"), FText::FromName(InterfaceVersion.Name)): FText())
+					.ColorAndOpacity(FSlateColor::UseForeground())
+					.DesiredSizeOverride(FVector2D(16, 16))
+				]
+				+ SHorizontalBox::Slot()
+				.AutoWidth()
+				.HAlign(HAlign_Left)
+				.VAlign(VAlign_Center)
+				.Padding(2, 0, 0, 0)
+				[
+					TextWidget
+				];
+
+				// TODO: Move to right-click menu. Currently, too easy to accidentally hit & clutters interface
+// 				if (!InterfaceVersion.IsValid())
+// 				{
+// 					LayoutWidget->AddSlot()
+// 					.AutoWidth()
+// 					.HAlign(HAlign_Left)
+// 					.VAlign(VAlign_Center)
+// 					.Padding(2, 0, 0, 0)
+// 					[
+// 						PropertyCustomizationHelpers::MakeDeleteButton(FSimpleDelegate::CreateLambda([this]()
+// 						{
+// 							UMetasoundEditorGraphMember* GraphMember = MetasoundAction->GetGraphMember();
+// 							if (ensure(GraphMember))
+// 							{
+// 								if (UMetasoundEditorGraph* Graph = Cast<UMetasoundEditorGraph>(MetasoundAction->Graph))
+// 								{
+// 									if (UObject* MetaSound = Graph->GetMetasound())
+// 									{
+// 										const FScopedTransaction Transaction(LOCTEXT("MetaSoundEditorDeleteOnClick", "Delete MetaSound Graph Member"));
+// 										MetaSound->Modify();
+// 										Graph->Modify();
+// 										Graph->RemoveMember(*GraphMember);
+// 
+// 										//Synchronize will update the interface
+// 										if (FMetasoundAssetBase* MetasoundAsset = IMetasoundUObjectRegistry::Get().GetObjectAsAssetBase(MetaSound))
+// 										{
+// 											MetasoundAsset->SetSynchronizationRequired();
+// 										}
+// 									}
+// 								}
+// 							}
+// 						}))
+// 					];
+// 				}
+
+				return LayoutWidget;
+			}
+
 			virtual bool OnNameTextVerifyChanged(const FText& InNewText, FText& OutErrorMessage) override
 			{
 				if (MetasoundAction.IsValid())
@@ -529,7 +604,7 @@ namespace Metasound
 			Args.NotifyHook = this;
 
 			SAssignNew(MetasoundInterfaceMenu, SGraphActionMenu, false)
-				.AlphaSortItems(false)
+				.AlphaSortItems(true)
 // 				.OnActionDoubleClicked(this, &FEditor::OnActionDoubleClicked)
 				.OnActionDragged(this, &FEditor::OnActionDragged)
 				.OnActionMatchesName(this, &FEditor::HandleActionMatchesName)
@@ -543,6 +618,10 @@ namespace Metasound
 				.OnGetFilterText(this, &FEditor::GetFilterText)
 				.OnGetSectionTitle(this, &FEditor::OnGetSectionTitle)
 				.OnGetSectionWidget(this, &FEditor::OnGetMenuSectionWidget)
+				.OnCreateCustomRowExpander_Lambda([](const FCustomExpanderData& InCustomExpanderData)
+				{
+					return SNew(SMetasoundActionMenuExpanderArrow, InCustomExpanderData);
+				})
 				.UseSectionStyling(true);
 
 			FPropertyEditorModule& PropertyModule = FModuleManager::LoadModuleChecked<FPropertyEditorModule>("PropertyEditor");
@@ -754,7 +833,7 @@ namespace Metasound
 			{
 				TSet<UClass*> ImportClasses;
 
-				for (const FMetasoundFrontendVersion& InterfaceVersion : MetasoundDoc.InterfaceVersions)
+				for (const FMetasoundFrontendVersion& InterfaceVersion : MetasoundDoc.Interfaces)
 				{
 					TArray<UClass*> InterfaceClasses = IMetasoundUObjectRegistry::Get().FindSupportedInterfaceClasses(InterfaceVersion);
 					ImportClasses.Append(MoveTemp(InterfaceClasses));
@@ -763,7 +842,7 @@ namespace Metasound
 				if (ImportClasses.Num() < 1)
 				{
 					TArray<FString> InterfaceNames;
-					Algo::Transform(MetasoundDoc.InterfaceVersions, InterfaceNames, [] (const FMetasoundFrontendVersion& InterfaceVersion) { return InterfaceVersion.ToString(); });
+					Algo::Transform(MetasoundDoc.Interfaces, InterfaceNames, [] (const FMetasoundFrontendVersion& InterfaceVersion) { return InterfaceVersion.ToString(); });
 					UE_LOG(LogMetaSound, Warning, TEXT("Cannot create UObject from MetaSound document. No UClass supports interface(s) \"%s\""), *FString::Join(InterfaceNames, TEXT(",")));
 				}
 				else
@@ -776,7 +855,7 @@ namespace Metasound
 						{
 							// TODO: Modal dialog to give user choice of import type.
 							TArray<FString> InterfaceNames;
-							Algo::Transform(MetasoundDoc.InterfaceVersions, InterfaceNames, [](const FMetasoundFrontendVersion& InterfaceVersion) { return InterfaceVersion.ToString(); });
+							Algo::Transform(MetasoundDoc.Interfaces, InterfaceNames, [](const FMetasoundFrontendVersion& InterfaceVersion) { return InterfaceVersion.ToString(); });
 							UE_LOG(LogMetaSound, Warning, TEXT("Duplicate UClass support interface(s) \"%s\" with UClass \"%s\""), *FString::Join(InterfaceNames, TEXT(",")), *ImportClass->GetName());
 						}
 					}
@@ -1302,7 +1381,7 @@ namespace Metasound
 				UMetasoundEditorGraphMember* NextToSelect = Graph.FindAdjacentMember(*GraphMember);
 
 				{
-					const FScopedTransaction Transaction(LOCTEXT("MetaSoundEditorDeleteSelectedNode", "Delete MetaSound GraphMember"));
+					const FScopedTransaction Transaction(LOCTEXT("MetaSoundEditorDeleteSelectedNode", "Delete MetaSound Graph Member"));
 					Metasound->Modify();
 					Graph.Modify();
 					Graph.RemoveMember(*GraphMember);
@@ -1902,7 +1981,7 @@ namespace Metasound
 			{
 				FName InterfaceName;
 				FName MemberName;
-				Audio::IGeneratorInterfaceRegistry::SplitMemberFullName(NodeHandle->GetNodeName(), InterfaceName, MemberName);
+				Audio::FParameterPath::SplitName(NodeHandle->GetNodeName(), InterfaceName, MemberName);
 
 				if (InterfaceName.IsNone())
 				{
@@ -1910,14 +1989,15 @@ namespace Metasound
 				}
 
 				FString CategoryString = InterfaceName.ToString();
-				CategoryString.ReplaceInline(*Audio::IGeneratorInterfaceRegistry::NamespaceDelimiter, TEXT("|"));
+				CategoryString.ReplaceInline(*Audio::FParameterPath::NamespaceDelimiter, TEXT("|"));
 				return FText::FromString(CategoryString);
 			};
 
 			FrontendGraph->IterateConstNodes([this, &GetVertexCategory, &EdGraph, ActionList = &OutAllActions](const Frontend::FConstNodeHandle& Input)
 			{
+				constexpr bool bIncludeNamespace = false;
 				const FText Tooltip = Input->GetDescription();
-				const FText MenuDesc = FGraphBuilder::GetDisplayName(*Input);
+				const FText MenuDesc = FGraphBuilder::GetDisplayName(*Input, bIncludeNamespace);
 				const FGuid NodeID = Input->GetID();
 				const FText Category = GetVertexCategory(Input);
 
@@ -1940,8 +2020,10 @@ namespace Metasound
 
 			FrontendGraph->IterateConstNodes([this, &GetVertexCategory, &EdGraph, ActionList = &OutAllActions](const Frontend::FConstNodeHandle& Output)
 			{
+				constexpr bool bIncludeNamespace = false;
+
 				const FText Tooltip = Output->GetDescription();
-				const FText MenuDesc = FGraphBuilder::GetDisplayName(*Output);
+				const FText MenuDesc = FGraphBuilder::GetDisplayName(*Output, bIncludeNamespace);
 				const FGuid NodeID = Output->GetID();
 				const FText Category = GetVertexCategory(Output);
 
@@ -2098,21 +2180,21 @@ namespace Metasound
 
 				FGraphBuilder::SynchronizeGraph(*Metasound);
 
-				if (MetasoundDetails.IsValid() && bInterfacesUpdated)
+				if (MetasoundDetails.IsValid())
 				{
-					// Only refresh the details panel if viewing the graph details
-					auto ContainsMetaSound = [](TWeakObjectPtr<UObject> Obj)
+					// Only refresh the details panel if viewing the graph details of the object is no longer valid
+					auto ShouldRefreshDetails = [bInterfacesUpdated](TWeakObjectPtr<UObject> Obj)
 					{
-						if (!Obj.IsValid())
+						if (!Obj.IsValid() || !IsValid(Obj.Get()))
 						{
-							return false;
+							return true;
 						}
-						
-						return IMetasoundUObjectRegistry::Get().IsRegisteredClass(Obj.Get());
+
+						return bInterfacesUpdated && IMetasoundUObjectRegistry::Get().IsRegisteredClass(Obj.Get());
 					};
 
-					const bool bDetailsEditingGraph = Algo::AnyOf(MetasoundDetails->GetSelectedObjects(), ContainsMetaSound);
-					if (bDetailsEditingGraph)
+					const bool bShouldRefreshDetails = Algo::AnyOf(MetasoundDetails->GetSelectedObjects(), ShouldRefreshDetails);
+					if (bShouldRefreshDetails)
 					{
 						RefreshDetails();
 					}
diff --git a/Engine/Plugins/Runtime/Metasound/Source/MetasoundEditor/Private/MetasoundEditorGraph.cpp b/Engine/Plugins/Runtime/Metasound/Source/MetasoundEditor/Private/MetasoundEditorGraph.cpp
--- a/Engine/Plugins/Runtime/Metasound/Source/MetasoundEditor/Private/MetasoundEditorGraph.cpp
+++ b/Engine/Plugins/Runtime/Metasound/Source/MetasoundEditor/Private/MetasoundEditorGraph.cpp
@@ -2,7 +2,7 @@
 #include "MetasoundEditorGraph.h"
 
 #include "Algo/Transform.h"
-#include "AudioParameterInterface.h"
+#include "AudioParameterControllerInterface.h"
 #include "Components/AudioComponent.h"
 #include "EdGraph/EdGraphNode.h"
 #include "Interfaces/ITargetPlatform.h"
@@ -151,8 +151,8 @@ void UMetasoundEditorGraphVertex::SetMemberName(const FName& InNewName, bool bPo
 
 FText UMetasoundEditorGraphVertex::GetDisplayName() const
 {
-	
-	return Metasound::Editor::FGraphBuilder::GetDisplayName(*GetConstNodeHandle());
+	constexpr bool bIncludeNamespace = true;
+	return Metasound::Editor::FGraphBuilder::GetDisplayName(*GetConstNodeHandle(), bIncludeNamespace);
 }
 
 void UMetasoundEditorGraphVertex::SetDisplayName(const FText& InNewName, bool bPostTransaction)
@@ -323,14 +323,14 @@ bool UMetasoundEditorGraphVertex::CanRename(const FText& InNewName, FText& OutEr
 	}
 
 	bool bIsNameValid = true;
-	const FName NewFName = *InNewName.ToString();
+	const FString NewName = InNewName.ToString();
 	FConstNodeHandle NodeHandle = GetConstNodeHandle();
 	FConstGraphHandle GraphHandle = NodeHandle->GetOwningGraph();
 	GraphHandle->IterateConstNodes([&](FConstNodeHandle NodeToCompare)
 	{
 		if (NodeID != NodeToCompare->GetID())
 		{
-			if (NewFName == NodeToCompare->GetNodeName())
+			if (NewName == NodeToCompare->GetNodeName().ToString())
 			{
 				bIsNameValid = false;
 				OutError = FText::Format(LOCTEXT("GraphVertexRenameInvalid_NameTaken", "{0} is already in use"), InNewName);
@@ -345,6 +345,11 @@ void UMetasoundEditorGraphInputLiteral::PostEditUndo()
 {
 	Super::PostEditUndo();
 
+	if (!IsValid(this))
+	{
+		return;
+	}
+
 	UpdateDocumentInputLiteral(false /* bPostTransaction */);
 }
 
@@ -384,7 +389,7 @@ void UMetasoundEditorGraphInputLiteral::UpdateDocumentInputLiteral(bool bPostTra
 		UAudioComponent* PreviewComponent = GEditor->GetPreviewAudioComponent();
 		check(PreviewComponent);
 
-		if (TScriptInterface<IAudioParameterInterface> ParamInterface = PreviewComponent)
+		if (TScriptInterface<IAudioParameterControllerInterface> ParamInterface = PreviewComponent)
 		{
 			Metasound::Frontend::FConstNodeHandle ConstNodeHandle = Input->GetConstNodeHandle();
 			Metasound::FVertexName VertexKey = NodeHandle->GetNodeName();
@@ -424,6 +429,11 @@ void UMetasoundEditorGraphInput::PostEditUndo()
 {
 	Super::PostEditUndo();
 
+	if (!IsValid(this))
+	{
+		return;
+	}
+
 	if (!Literal)
 	{
 		if (UMetasoundEditorGraph* MetasoundGraph = Cast<UMetasoundEditorGraph>(GetOuter()))
@@ -442,6 +452,7 @@ void UMetasoundEditorGraphInput::PostEditUndo()
 		}
 		return;
 	}
+
 	Literal->UpdateDocumentInputLiteral(false /* bPostTransaction */);
 	UpdateEditorLiteralType();
 }
@@ -928,6 +939,11 @@ void UMetasoundEditorGraphVariable::PostEditUndo()
 {
 	Super::PostEditUndo();
 
+	if (!IsValid(this))
+	{
+		return;
+	}
+
 	if (!Literal)
 	{
 		if (UMetasoundEditorGraph* MetasoundGraph = GetOwningGraph())
diff --git a/Engine/Plugins/Runtime/Metasound/Source/MetasoundEditor/Private/MetasoundEditorGraphBuilder.cpp b/Engine/Plugins/Runtime/Metasound/Source/MetasoundEditor/Private/MetasoundEditorGraphBuilder.cpp
--- a/Engine/Plugins/Runtime/Metasound/Source/MetasoundEditor/Private/MetasoundEditorGraphBuilder.cpp
+++ b/Engine/Plugins/Runtime/Metasound/Source/MetasoundEditor/Private/MetasoundEditorGraphBuilder.cpp
@@ -86,16 +86,26 @@ namespace Metasound
 			}
 		} // namespace GraphBuilderPrivate
 
-		FText FGraphBuilder::GetDisplayName(const Frontend::INodeController& InFrontendNode)
+		FText FGraphBuilder::GetDisplayName(const Frontend::INodeController& InFrontendNode, bool bInIncludeNamespace)
 		{
+			FName Namespace;
+			FName ParameterName;
+			Audio::FParameterPath::SplitName(InFrontendNode.GetNodeName(), Namespace, ParameterName);
+
 			FText DisplayName = InFrontendNode.GetDisplayName();
 			if (DisplayName.IsEmptyOrWhitespace())
 			{
-				FName InterfaceName;
-				FName ParameterName;
-				Audio::IGeneratorInterfaceRegistry::Get().SplitMemberFullName(InFrontendNode.GetNodeName(), InterfaceName, ParameterName);
-				return FText::FromName(ParameterName);
+				DisplayName = FText::FromName(ParameterName);
 			}
+
+			if (bInIncludeNamespace)
+			{
+				if (!Namespace.IsNone())
+				{
+					return FText::Format(LOCTEXT("MemberDisplayNameWithNamespaceFormat", "{0} ({1})"), DisplayName, FText::FromName(Namespace));
+				}
+			}
+
 			return DisplayName;
 		}
 
@@ -1879,7 +1889,8 @@ namespace Metasound
 				{
 					if (bLogChanges)
 					{
-						FText NodeDisplayName = FGraphBuilder::GetDisplayName(*InNode);
+						constexpr bool bIncludeNamespace = true;
+						const FText NodeDisplayName = FGraphBuilder::GetDisplayName(*InNode, bIncludeNamespace);
 						UE_LOG(LogMetasoundEditor, Verbose, TEXT("Synchronizing Node '%s' Pins: Removing Excess Editor Pin '%s'"), *NodeDisplayName.ToString(), *Pin->GetName());
 					}
 					InEditorNode.RemovePin(Pin);
@@ -1895,8 +1906,9 @@ namespace Metasound
 				{
 					if (bLogChanges)
 					{
-						FText NodeDisplayName = FGraphBuilder::GetDisplayName(*InNode);
-						FText InputDisplayName = FGraphBuilder::GetDisplayName(*InputHandle);
+						constexpr bool bIncludeNamespace = true;
+						const FText NodeDisplayName = FGraphBuilder::GetDisplayName(*InNode, bIncludeNamespace);
+						const FText InputDisplayName = FGraphBuilder::GetDisplayName(*InputHandle);
 						UE_LOG(LogMetasoundEditor, Verbose, TEXT("Synchronizing Node '%s' Pins: Adding missing Editor Input Pin '%s'"), *NodeDisplayName.ToString(), *InputDisplayName.ToString());
 					}
 					AddPinToNode(InEditorNode, InputHandle);
@@ -1911,8 +1923,9 @@ namespace Metasound
 				{
 					if (bLogChanges)
 					{
-						FText NodeDisplayName = FGraphBuilder::GetDisplayName(*InNode);
-						FText OutputDisplayName = FGraphBuilder::GetDisplayName(*OutputHandle);
+						constexpr bool bIncludeNamespace = true;
+						const FText NodeDisplayName = FGraphBuilder::GetDisplayName(*InNode, bIncludeNamespace);
+						const FText OutputDisplayName = FGraphBuilder::GetDisplayName(*OutputHandle);
 						UE_LOG(LogMetasoundEditor, Verbose, TEXT("Synchronizing Node '%s' Pins: Adding missing Editor Output Pin '%s'"), *NodeDisplayName.ToString(), *OutputDisplayName.ToString());
 					}
 					AddPinToNode(InEditorNode, OutputHandle);
@@ -1978,7 +1991,8 @@ namespace Metasound
 
 				// Add an editor input if none exist for a frontend input.
 				Inputs.Add(Graph->FindOrAddInput(NodeHandle));
-				FText NodeDisplayName = FGraphBuilder::GetDisplayName(*NodeHandle);
+				constexpr bool bIncludeNamespace = true;
+				FText NodeDisplayName = FGraphBuilder::GetDisplayName(*NodeHandle, bIncludeNamespace);
 				UE_LOG(LogMetasoundEditor, Verbose, TEXT("Synchronizing Inputs: Added missing input '%s'."), *NodeDisplayName.ToString());
 				bIsEditorGraphDirty = true;
 			}, EMetasoundFrontendClassType::Input);
@@ -1994,7 +2008,8 @@ namespace Metasound
 
 				// Add an editor output if none exist for a frontend output.
 				Outputs.Add(Graph->FindOrAddOutput(NodeHandle));
-				FText NodeDisplayName = FGraphBuilder::GetDisplayName(*NodeHandle);
+				constexpr bool bIncludeNamespace = true;
+				FText NodeDisplayName = FGraphBuilder::GetDisplayName(*NodeHandle, bIncludeNamespace);
 				UE_LOG(LogMetasoundEditor, Verbose, TEXT("Synchronizing Outputs: Added missing output '%s'."), *NodeDisplayName.ToString());
 				bIsEditorGraphDirty = true;
 			}, EMetasoundFrontendClassType::Output);
@@ -2037,7 +2052,8 @@ namespace Metasound
 						const FName NewDataType = InputHandle->GetDataType();
 						if (Input->TypeName != NewDataType)
 						{
-							FText NodeDisplayName = FGraphBuilder::GetDisplayName(*NodeHandle);
+							constexpr bool bIncludeNamespace = true;
+							FText NodeDisplayName = FGraphBuilder::GetDisplayName(*NodeHandle, bIncludeNamespace);
 							UE_LOG(LogMetasoundEditor, Verbose, TEXT("Synchronizing Input '%s': Updating DataType to '%s'."), *NodeDisplayName.ToString(), *NewDataType.ToString());
 
 							FMetasoundFrontendLiteral DefaultLiteral;
@@ -2075,7 +2091,8 @@ namespace Metasound
 						const FName NewDataType = OutputHandles[0]->GetDataType();
 						if (Output->TypeName != NewDataType)
 						{
-							FText NodeDisplayName = FGraphBuilder::GetDisplayName(*NodeHandle);
+							constexpr bool bIncludeNamespace = true;
+							FText NodeDisplayName = FGraphBuilder::GetDisplayName(*NodeHandle, bIncludeNamespace);
 							UE_LOG(LogMetasoundEditor, Verbose, TEXT("Synchronizing Output '%s': Updating DataType to '%s'."), *NodeDisplayName.ToString(), *NewDataType.ToString());
 
 							Output->ClassName = NodeHandle->GetClassMetadata().GetClassName();
diff --git a/Engine/Plugins/Runtime/Metasound/Source/MetasoundEditor/Private/MetasoundEditorGraphBuilder.h b/Engine/Plugins/Runtime/Metasound/Source/MetasoundEditor/Private/MetasoundEditorGraphBuilder.h
--- a/Engine/Plugins/Runtime/Metasound/Source/MetasoundEditor/Private/MetasoundEditorGraphBuilder.h
+++ b/Engine/Plugins/Runtime/Metasound/Source/MetasoundEditor/Private/MetasoundEditorGraphBuilder.h
@@ -75,7 +75,7 @@ namespace Metasound
 
 			// Returns a display name for a node. If the node has an empty or whitespace
 			// only DisplayName, then the NodeName is used. 
-			static FText GetDisplayName(const Frontend::INodeController& InFrontendNode);
+			static FText GetDisplayName(const Frontend::INodeController& InFrontendNode, bool bInIncludeNamespace);
 
 			// Returns a display name for an input. If the input has an empty or whitespace
 			// only DisplayName, then the VertexName is used. 
diff --git a/Engine/Plugins/Runtime/Metasound/Source/MetasoundEditor/Private/MetasoundEditorGraphInputNodes.cpp b/Engine/Plugins/Runtime/Metasound/Source/MetasoundEditor/Private/MetasoundEditorGraphInputNodes.cpp
--- a/Engine/Plugins/Runtime/Metasound/Source/MetasoundEditor/Private/MetasoundEditorGraphInputNodes.cpp
+++ b/Engine/Plugins/Runtime/Metasound/Source/MetasoundEditor/Private/MetasoundEditorGraphInputNodes.cpp
@@ -3,8 +3,7 @@
 
 #include "Algo/Count.h"
 #include "Algo/Transform.h"
-#include "AudioParameterInterface.h"
-#include "CoreMinimal.h"
+#include "AudioParameterControllerInterface.h"
 #include "EdGraph/EdGraphNode.h"
 #include "GraphEditorSettings.h"
 #include "MetasoundDataReference.h"
@@ -56,7 +55,7 @@ FMetasoundFrontendClassName UMetasoundEditorGraphInputNode::GetClassName() const
 	return Super::GetClassName();
 }
 
-void UMetasoundEditorGraphInputNode::UpdatePreviewInstance(const Metasound::FVertexName& InParameterName, TScriptInterface<IAudioParameterInterface>& InParameterInterface) const
+void UMetasoundEditorGraphInputNode::UpdatePreviewInstance(const Metasound::FVertexName& InParameterName, TScriptInterface<IAudioParameterControllerInterface>& InParameterInterface) const
 {
 	if (Input)
 	{
@@ -130,7 +129,7 @@ void UMetasoundEditorGraphInputBool::SetFromLiteral(const FMetasoundFrontendLite
 	Metasound::Editor::InputPrivate::ConvertLiteral<bool>(InLiteral, Default.Value);
 }
 
-void UMetasoundEditorGraphInputBool::UpdatePreviewInstance(const Metasound::FVertexName& InParameterName, TScriptInterface<IAudioParameterInterface>& InParameterInterface) const
+void UMetasoundEditorGraphInputBool::UpdatePreviewInstance(const Metasound::FVertexName& InParameterName, TScriptInterface<IAudioParameterControllerInterface>& InParameterInterface) const
 	{
 		InParameterInterface->SetBoolParameter(InParameterName, Default.Value);
 	}
@@ -155,7 +154,7 @@ void UMetasoundEditorGraphInputBoolArray::SetFromLiteral(const FMetasoundFronten
 		Metasound::Editor::InputPrivate::ConvertLiteralToArray<bool, FMetasoundEditorGraphInputBoolRef>(InLiteral, Default);
 	}
 
-void UMetasoundEditorGraphInputBoolArray::UpdatePreviewInstance(const Metasound::FVertexName& InParameterName, TScriptInterface<IAudioParameterInterface>& InParameterInterface) const
+void UMetasoundEditorGraphInputBoolArray::UpdatePreviewInstance(const Metasound::FVertexName& InParameterName, TScriptInterface<IAudioParameterControllerInterface>& InParameterInterface) const
 {
 	TArray<bool> BoolArray;
 	Algo::Transform(Default, BoolArray, [](const FMetasoundEditorGraphInputBoolRef& InValue) { return InValue.Value; });
@@ -179,7 +178,7 @@ void UMetasoundEditorGraphInputInt::SetFromLiteral(const FMetasoundFrontendLiter
 	Metasound::Editor::InputPrivate::ConvertLiteral<int32>(InLiteral, Default.Value);
 }
 
-void UMetasoundEditorGraphInputInt::UpdatePreviewInstance(const Metasound::FVertexName& InParameterName, TScriptInterface<IAudioParameterInterface>& InParameterInterface) const
+void UMetasoundEditorGraphInputInt::UpdatePreviewInstance(const Metasound::FVertexName& InParameterName, TScriptInterface<IAudioParameterControllerInterface>& InParameterInterface) const
 {
 	InParameterInterface->SetIntParameter(InParameterName, Default.Value);
 }
@@ -205,7 +204,7 @@ void UMetasoundEditorGraphInputIntArray::SetFromLiteral(const FMetasoundFrontend
 	Metasound::Editor::InputPrivate::ConvertLiteralToArray<int32, FMetasoundEditorGraphInputIntRef>(InLiteral, Default);
 }
 
-void UMetasoundEditorGraphInputIntArray::UpdatePreviewInstance(const Metasound::FVertexName& InParameterName, TScriptInterface<IAudioParameterInterface>& InParameterInterface) const
+void UMetasoundEditorGraphInputIntArray::UpdatePreviewInstance(const Metasound::FVertexName& InParameterName, TScriptInterface<IAudioParameterControllerInterface>& InParameterInterface) const
 {
 	TArray<int32> IntArray;
 	Algo::Transform(Default, IntArray, [](const FMetasoundEditorGraphInputIntRef& InValue) { return InValue.Value; });
@@ -229,7 +228,7 @@ void UMetasoundEditorGraphInputFloat::SetFromLiteral(const FMetasoundFrontendLit
 	Metasound::Editor::InputPrivate::ConvertLiteral<float>(InLiteral, Default);
 }
 
-void UMetasoundEditorGraphInputFloat::UpdatePreviewInstance(const Metasound::FVertexName& InParameterName, TScriptInterface<IAudioParameterInterface>& InParameterInterface) const
+void UMetasoundEditorGraphInputFloat::UpdatePreviewInstance(const Metasound::FVertexName& InParameterName, TScriptInterface<IAudioParameterControllerInterface>& InParameterInterface) const
 {
 	InParameterInterface->SetFloatParameter(InParameterName, Default);
 }
@@ -320,7 +319,7 @@ void UMetasoundEditorGraphInputFloatArray::SetFromLiteral(const FMetasoundFronte
 	Metasound::Editor::InputPrivate::ConvertLiteralToArray<float>(InLiteral, Default);
 }
 
-void UMetasoundEditorGraphInputFloatArray::UpdatePreviewInstance(const Metasound::FVertexName& InParameterName, TScriptInterface<IAudioParameterInterface>& InParameterInterface) const
+void UMetasoundEditorGraphInputFloatArray::UpdatePreviewInstance(const Metasound::FVertexName& InParameterName, TScriptInterface<IAudioParameterControllerInterface>& InParameterInterface) const
 {
 	InParameterInterface->SetFloatArrayParameter(InParameterName, Default);
 }
@@ -342,7 +341,7 @@ void UMetasoundEditorGraphInputString::SetFromLiteral(const FMetasoundFrontendLi
 	Metasound::Editor::InputPrivate::ConvertLiteral<FString>(InLiteral, Default);
 }
 
-void UMetasoundEditorGraphInputString::UpdatePreviewInstance(const Metasound::FVertexName& InParameterName, TScriptInterface<IAudioParameterInterface>& InParameterInterface) const
+void UMetasoundEditorGraphInputString::UpdatePreviewInstance(const Metasound::FVertexName& InParameterName, TScriptInterface<IAudioParameterControllerInterface>& InParameterInterface) const
 {
 	InParameterInterface->SetStringParameter(InParameterName, Default);
 }
@@ -364,7 +363,7 @@ void UMetasoundEditorGraphInputStringArray::SetFromLiteral(const FMetasoundFront
 	Metasound::Editor::InputPrivate::ConvertLiteralToArray<FString>(InLiteral, Default);
 }
 
-void UMetasoundEditorGraphInputStringArray::UpdatePreviewInstance(const Metasound::FVertexName& InParameterName, TScriptInterface<IAudioParameterInterface>& InParameterInterface) const
+void UMetasoundEditorGraphInputStringArray::UpdatePreviewInstance(const Metasound::FVertexName& InParameterName, TScriptInterface<IAudioParameterControllerInterface>& InParameterInterface) const
 {
 	InParameterInterface->SetStringArrayParameter(InParameterName, Default);
 }
@@ -386,7 +385,7 @@ void UMetasoundEditorGraphInputObject::SetFromLiteral(const FMetasoundFrontendLi
 	ensure(InLiteral.TryGet(Default.Object));
 }
 
-void UMetasoundEditorGraphInputObject::UpdatePreviewInstance(const Metasound::FVertexName& InParameterName, TScriptInterface<IAudioParameterInterface>& InParameterInterface) const
+void UMetasoundEditorGraphInputObject::UpdatePreviewInstance(const Metasound::FVertexName& InParameterName, TScriptInterface<IAudioParameterControllerInterface>& InParameterInterface) const
 {
 	// TODO. We need proxy object here safely.
 }
@@ -414,7 +413,7 @@ void UMetasoundEditorGraphInputObjectArray::SetFromLiteral(const FMetasoundFront
 	Algo::Transform(ObjectArray, Default, [](UObject* InValue) { return FMetasoundEditorGraphInputObjectRef { InValue }; });
 }
 
-void UMetasoundEditorGraphInputObjectArray::UpdatePreviewInstance(const Metasound::FVertexName& InParameterName, TScriptInterface<IAudioParameterInterface>& InParameterInterface) const
+void UMetasoundEditorGraphInputObjectArray::UpdatePreviewInstance(const Metasound::FVertexName& InParameterName, TScriptInterface<IAudioParameterControllerInterface>& InParameterInterface) const
 {
 	TArray<UObject*> ObjectArray;
 	Algo::Transform(Default, ObjectArray, [](const FMetasoundEditorGraphInputObjectRef& InValue) { return InValue.Object; });
diff --git a/Engine/Plugins/Runtime/Metasound/Source/MetasoundEditor/Private/MetasoundEditorGraphNode.cpp b/Engine/Plugins/Runtime/Metasound/Source/MetasoundEditor/Private/MetasoundEditorGraphNode.cpp
--- a/Engine/Plugins/Runtime/Metasound/Source/MetasoundEditor/Private/MetasoundEditorGraphNode.cpp
+++ b/Engine/Plugins/Runtime/Metasound/Source/MetasoundEditor/Private/MetasoundEditorGraphNode.cpp
@@ -461,7 +461,8 @@ FText UMetasoundEditorGraphNode::GetTooltipText() const
 
 FText UMetasoundEditorGraphNode::GetDisplayName() const
 {
-	return Metasound::Editor::FGraphBuilder::GetDisplayName(*GetConstNodeHandle());
+	constexpr bool bIncludeNamespace = true;
+	return Metasound::Editor::FGraphBuilder::GetDisplayName(*GetConstNodeHandle(), bIncludeNamespace);
 }
 
 FString UMetasoundEditorGraphNode::GetDocumentationExcerptName() const
diff --git a/Engine/Plugins/Runtime/Metasound/Source/MetasoundEditor/Private/MetasoundEditorGraphSchema.cpp b/Engine/Plugins/Runtime/Metasound/Source/MetasoundEditor/Private/MetasoundEditorGraphSchema.cpp
--- a/Engine/Plugins/Runtime/Metasound/Source/MetasoundEditor/Private/MetasoundEditorGraphSchema.cpp
+++ b/Engine/Plugins/Runtime/Metasound/Source/MetasoundEditor/Private/MetasoundEditorGraphSchema.cpp
@@ -174,8 +174,10 @@ namespace Metasound
 				{
 					if (!InQuery.Filter || InQuery.Filter(NodeHandle))
 					{
+						constexpr bool bIncludeNamespace = true;
+
 						const FText& GroupName = GetContextGroupDisplayName(InQuery.ContextGroup);
-						const FText NodeDisplayName = FGraphBuilder::GetDisplayName(*NodeHandle);
+						const FText NodeDisplayName = FGraphBuilder::GetDisplayName(*NodeHandle, bIncludeNamespace);
 						const FText Tooltip = FText::Format(InQuery.TooltipFormat, NodeDisplayName);
 						const FText DisplayName = FText::Format(InQuery.DisplayNameFormat, NodeDisplayName);
 						TSharedPtr<TAction> NewNodeAction = MakeShared<TAction>(GroupName, DisplayName, NodeHandle->GetID(), Tooltip, InQuery.ContextGroup);
@@ -1026,7 +1028,8 @@ FText UMetasoundEditorGraphSchema::GetPinDisplayName(const UEdGraphPin* Pin) con
 		case EMetasoundFrontendClassType::Input:
 		case EMetasoundFrontendClassType::Output:
 		{
-			return FGraphBuilder::GetDisplayName(*NodeHandle);
+			constexpr bool bIncludeNamespace = true;
+			return FGraphBuilder::GetDisplayName(*NodeHandle, bIncludeNamespace);
 		}
 
 		case EMetasoundFrontendClassType::Literal:
diff --git a/Engine/Plugins/Runtime/Metasound/Source/MetasoundEditor/Private/MetasoundEditorModule.cpp b/Engine/Plugins/Runtime/Metasound/Source/MetasoundEditor/Private/MetasoundEditorModule.cpp
--- a/Engine/Plugins/Runtime/Metasound/Source/MetasoundEditor/Private/MetasoundEditorModule.cpp
+++ b/Engine/Plugins/Runtime/Metasound/Source/MetasoundEditor/Private/MetasoundEditorModule.cpp
@@ -176,7 +176,7 @@ namespace Metasound
 					check(AssetSubsystem);
 
 					// Use the editor version of `RegisterGraphWithFrontend` so it re-registers any open MetaSound editors
-					AssetSubsystem->AddOrUpdateAsset(InAssetData, false /* bRegisterWithFrontend */);
+					AssetSubsystem->AddOrUpdateAsset(InAssetData);
 					
 					// Loading all assets necessary only in editor to register &
 					// populate potential graphs to reference in MetaSound editor.
@@ -230,7 +230,7 @@ namespace Metasound
 							check(AssetSubsystem);
 
 							// Use the editor version of UnregisterWithFrontend so it refreshes any open MetaSound editors
-							AssetSubsystem->RemoveAsset(*Pair.Key, false /* bUnregisterWithFrontend */);
+							AssetSubsystem->RemoveAsset(*Pair.Key);
 							FGraphBuilder::UnregisterGraphWithFrontend(*Pair.Key);
 						}
 					}
@@ -243,7 +243,7 @@ namespace Metasound
 							UMetaSoundAssetSubsystem* AssetSubsystem = GEngine->GetEngineSubsystem<UMetaSoundAssetSubsystem>();
 							check(AssetSubsystem);
 							// Use the editor version of RegisterWithFrontend so it refreshes any open MetaSound editors
-							AssetSubsystem->AddOrUpdateAsset(*Pair.Value, false /* bRegisterWithFrontend */);
+							AssetSubsystem->AddOrUpdateAsset(*Pair.Value);
 							FGraphBuilder::RegisterGraphWithFrontend(*Pair.Value);
 						}
 					}
@@ -276,7 +276,7 @@ namespace Metasound
 					check(AssetSubsystem);
 
 					// Use the editor version of UnregisterWithFrontend so it refreshes any open MetaSound editors
-					AssetSubsystem->RemoveAsset(InAssetData, false /* bUnregisterWithFrontend */);
+					AssetSubsystem->RemoveAsset(InAssetData);
 					if (UObject* AssetObject = InAssetData.GetAsset())
 					{
 						FGraphBuilder::UnregisterGraphWithFrontend(*AssetObject);
diff --git a/Engine/Plugins/Runtime/Metasound/Source/MetasoundEditor/Private/SMetasoundActionMenu.cpp b/Engine/Plugins/Runtime/Metasound/Source/MetasoundEditor/Private/SMetasoundActionMenu.cpp
--- a/Engine/Plugins/Runtime/Metasound/Source/MetasoundEditor/Private/SMetasoundActionMenu.cpp
+++ b/Engine/Plugins/Runtime/Metasound/Source/MetasoundEditor/Private/SMetasoundActionMenu.cpp
@@ -33,294 +33,278 @@ namespace Metasound
 {
 	namespace Editor
 	{
-		// Custom expander to specify our desired padding
-		class SMetasoundActionMenuExpanderArrow : public SExpanderArrow
+		void SMetasoundActionMenuExpanderArrow::Construct(const FArguments& InArgs, const FCustomExpanderData& ActionMenuData)
 		{
-			SLATE_BEGIN_ARGS(SMetasoundActionMenuExpanderArrow)
-			{
-			}
-
-			SLATE_ATTRIBUTE(float, IndentAmount)
-				SLATE_END_ARGS()
+			OwnerRowPtr = ActionMenuData.TableRow;
+			IndentAmount = InArgs._IndentAmount;
+			ActionPtr = ActionMenuData.RowAction;
 
-		public:
-			void Construct(const FArguments& InArgs, const FCustomExpanderData& ActionMenuData)
+			if (ActionPtr.IsValid())
 			{
-				OwnerRowPtr = ActionMenuData.TableRow;
-				IndentAmount = InArgs._IndentAmount;
-				ActionPtr = ActionMenuData.RowAction;
-
-				if (ActionPtr.IsValid())
-				{
-					ChildSlot
-						.Padding(TAttribute<FMargin>(this, &SMetasoundActionMenuExpanderArrow::GetCustomIndentPadding))
-						[
-							SNullWidget::NullWidget
-							// TODO: Add favorites support
-							// SNew(SMetasoundActionFavoriteToggle, ActionMenuData)
-						];
-				}
-				else
-				{
-					SExpanderArrow::FArguments SuperArgs;
-					SuperArgs._IndentAmount = InArgs._IndentAmount;
-
-					SExpanderArrow::Construct(SuperArgs, ActionMenuData.TableRow);
-				}
+				ChildSlot
+				.Padding(TAttribute<FMargin>(this, &SMetasoundActionMenuExpanderArrow::GetCustomIndentPadding))
+				[
+					SNullWidget::NullWidget
+					// TODO: Add favorites support
+					// SNew(SMetasoundActionFavoriteToggle, ActionMenuData)
+				];
 			}
-
-		private:
-			FMargin GetCustomIndentPadding() const
+			else
 			{
-				return SExpanderArrow::GetExpanderPadding();
-			}
-
-			TWeakPtr<FEdGraphSchemaAction> ActionPtr;
-		};
-	} // namespace Editor
-} // namespace Metasound
+				SExpanderArrow::FArguments SuperArgs;
+				SuperArgs._IndentAmount = InArgs._IndentAmount;
 
+				SExpanderArrow::Construct(SuperArgs, ActionMenuData.TableRow);
+			}
+		}
 
-SMetasoundActionMenu::~SMetasoundActionMenu()
-{
-	OnClosedCallback.ExecuteIfBound();
-	OnCloseReasonCallback.ExecuteIfBound(bActionExecuted, false, !DraggedFromPins.IsEmpty());
-}
+		FMargin SMetasoundActionMenuExpanderArrow::GetCustomIndentPadding() const
+		{
+			return SExpanderArrow::GetExpanderPadding();
+		}
 
-void SMetasoundActionMenu::Construct(const FArguments& InArgs)
-{
-	using namespace Metasound::Editor;
-	using namespace Metasound::Frontend;
-
-	Graph = InArgs._Graph;
-	DraggedFromPins = InArgs._DraggedFromPins;
-	NewNodePosition = InArgs._NewNodePosition;
-	OnClosedCallback = InArgs._OnClosedCallback;
-	bAutoExpandActionMenu = InArgs._AutoExpandActionMenu;
-	OnCloseReasonCallback = InArgs._OnCloseReason;
-
-	FSlateColor TypeColor;
-	const FSlateBrush* PinBrush = nullptr;
-	if (!DraggedFromPins.IsEmpty())
-	{
-		if (const UEdGraphPin* Pin = DraggedFromPins[0])
+		SMetasoundActionMenu::~SMetasoundActionMenu()
 		{
-			FDataTypeRegistryInfo RegistryInfo;
+			OnClosedCallback.ExecuteIfBound();
+			OnCloseReasonCallback.ExecuteIfBound(bActionExecuted, false, !DraggedFromPins.IsEmpty());
+		}
 
-			const IMetasoundEditorModule& EditorModule = FModuleManager::GetModuleChecked<IMetasoundEditorModule>("MetaSoundEditor");
-			if (Pin->Direction == EGPD_Input)
+		void SMetasoundActionMenu::Construct(const FArguments& InArgs)
+		{
+			using namespace Metasound::Editor;
+			using namespace Metasound::Frontend;
+
+			Graph = InArgs._Graph;
+			DraggedFromPins = InArgs._DraggedFromPins;
+			NewNodePosition = InArgs._NewNodePosition;
+			OnClosedCallback = InArgs._OnClosedCallback;
+			bAutoExpandActionMenu = InArgs._AutoExpandActionMenu;
+			OnCloseReasonCallback = InArgs._OnCloseReason;
+
+			FSlateColor TypeColor;
+			const FSlateBrush* PinBrush = nullptr;
+			if (!DraggedFromPins.IsEmpty())
 			{
-				FInputHandle InputHandle = FGraphBuilder::GetInputHandleFromPin(Pin);
-				IDataTypeRegistry::Get().GetDataTypeInfo(InputHandle->GetDataType(), RegistryInfo);
+				if (const UEdGraphPin* Pin = DraggedFromPins[0])
+				{
+					FDataTypeRegistryInfo RegistryInfo;
+
+					const IMetasoundEditorModule& EditorModule = FModuleManager::GetModuleChecked<IMetasoundEditorModule>("MetaSoundEditor");
+					if (Pin->Direction == EGPD_Input)
+					{
+						FInputHandle InputHandle = FGraphBuilder::GetInputHandleFromPin(Pin);
+						IDataTypeRegistry::Get().GetDataTypeInfo(InputHandle->GetDataType(), RegistryInfo);
+					}
+					else
+					{
+						FOutputHandle OutputHandle = FGraphBuilder::GetOutputHandleFromPin(Pin);
+						IDataTypeRegistry::Get().GetDataTypeInfo(OutputHandle->GetDataType(), RegistryInfo);
+					}
+
+					TypeColor = FGraphBuilder::GetPinCategoryColor(Pin->PinType);
+
+					if (RegistryInfo.IsArrayType())
+					{
+						PinBrush = FEditorStyle::GetBrush("Graph.ArrayPin.Connected");
+					}
+					else if (Pin->PinType.PinCategory == FGraphBuilder::PinCategoryTrigger)
+					{
+						if (const ISlateStyle* MetasoundStyle = FSlateStyleRegistry::FindSlateStyle("MetaSoundStyle"))
+						{
+							PinBrush = MetasoundStyle->GetBrush(TEXT("MetasoundEditor.Graph.TriggerPin.Connected"));
+						}
+					}
+
+					if (!PinBrush)
+					{
+						PinBrush = FEditorStyle::GetBrush("Graph.Pin.Connected");
+					}
+				}
 			}
-			else
+
+			SBorder::Construct(SBorder::FArguments()
+				.BorderImage(FEditorStyle::GetBrush("Menu.Background"))
+				.Padding(5)
+				[
+					SNew(SBox)
+					.WidthOverride(400)
+					.HeightOverride(400)
+					[
+						SNew(SVerticalBox)
+						+SVerticalBox::Slot()
+						.AutoHeight()
+						.Padding(2, 2, 2, 5)
+						[
+							SNew(SHorizontalBox)
+							+SHorizontalBox::Slot()
+							.AutoWidth()
+							.VAlign(VAlign_Center)
+							.Padding(0, 0, 5, 0)
+							[
+								SNew(SImage)
+								.ColorAndOpacity(TypeColor)
+								.Visibility_Lambda([this]()
+								{
+									return DraggedFromPins.IsEmpty() ? EVisibility::Hidden : EVisibility::Visible;
+								})
+								.Image(PinBrush)
+							]
+							+SHorizontalBox::Slot()
+							.AutoWidth()
+							.VAlign(VAlign_Center)
+							[
+								SNew(STextBlock)
+								.Text_Lambda([this]()
+								{
+									if (DraggedFromPins.IsEmpty())
+									{
+										return LOCTEXT("ContextText", "All MetaSound Node Classes");
+									}
+
+									UEdGraphPin* Pin = DraggedFromPins[0];
+
+									if (DraggedFromPins[0]->Direction == EGPD_Input)
+									{
+										FInputHandle InputHandle = FGraphBuilder::GetInputHandleFromPin(Pin);
+										return FText::Format(LOCTEXT("ContextTypeFilteredText_Output", "Classes with output of type '{0}'"), FText::FromName(InputHandle->GetDataType()));
+									}
+									else
+									{
+										FOutputHandle OutputHandle = FGraphBuilder::GetOutputHandleFromPin(Pin);
+										return FText::Format(LOCTEXT("ContextTypeFilteredText_Input", "Classes with input of type '{0}'"), FText::FromName(OutputHandle->GetDataType()));
+									}
+								})
+								// TODO: Move to Metasound Style
+								.Font(FEditorStyle::GetFontStyle("BlueprintEditor.ActionMenu.ContextDescriptionFont"))
+								.ToolTip(IDocumentation::Get()->CreateToolTip(
+									LOCTEXT("ActionMenuContextTextTooltip", "Describes the current context of the action list"),
+									nullptr,
+									TEXT("Shared/Editors/MetasoundEditor"),
+									TEXT("MetasoundActionMenuContextText")))
+								.WrapTextAt(280)
+							]
+						]
+						+SVerticalBox::Slot()
+						[
+							SAssignNew(GraphActionMenu, SGraphActionMenu)
+								.OnActionSelected(this, &SMetasoundActionMenu::OnActionSelected)
+								.OnCreateWidgetForAction(SGraphActionMenu::FOnCreateWidgetForAction::CreateSP(this, &SMetasoundActionMenu::OnCreateWidgetForAction))
+								.OnCollectAllActions(this, &SMetasoundActionMenu::CollectAllActions)
+								.OnCreateCustomRowExpander_Lambda([](const FCustomExpanderData& InCustomExpanderData)
+								{
+									return SNew(SMetasoundActionMenuExpanderArrow, InCustomExpanderData);
+								})
+								.DraggedFromPins(DraggedFromPins)
+								.GraphObj(Graph)
+								.AlphaSortItems(true)
+								.bAllowPreselectedItemActivation(true)
+						]
+					]
+				]
+			);
+		}
+
+		void SMetasoundActionMenu::CollectAllActions(FGraphActionListBuilderBase& OutAllActions)
+		{
+			if (!Graph)
 			{
-				FOutputHandle OutputHandle = FGraphBuilder::GetOutputHandleFromPin(Pin);
-				IDataTypeRegistry::Get().GetDataTypeInfo(OutputHandle->GetDataType(), RegistryInfo);
+				return;
 			}
 
-			TypeColor = FGraphBuilder::GetPinCategoryColor(Pin->PinType);
-
-			if (RegistryInfo.IsArrayType())
+			FGraphContextMenuBuilder MenuBuilder(Graph);
+			if (!DraggedFromPins.IsEmpty())
 			{
-				PinBrush = FEditorStyle::GetBrush("Graph.ArrayPin.Connected");
+				MenuBuilder.FromPin = DraggedFromPins[0];
 			}
-			else if (Pin->PinType.PinCategory == FGraphBuilder::PinCategoryTrigger)
+
+			// Cannot call GetGraphContextActions() during serialization and GC due to its use of FindObject()
+			if(!GIsSavingPackage && !IsGarbageCollecting())
 			{
-				if (const ISlateStyle* MetasoundStyle = FSlateStyleRegistry::FindSlateStyle("MetaSoundStyle"))
+				if (const UMetasoundEditorGraphSchema* Schema = GetDefault<UMetasoundEditorGraphSchema>())
 				{
-					PinBrush = MetasoundStyle->GetBrush(TEXT("MetasoundEditor.Graph.TriggerPin.Connected"));
+					Schema->GetGraphContextActions(MenuBuilder);
 				}
 			}
 
-			if (!PinBrush)
-			{
-				PinBrush = FEditorStyle::GetBrush("Graph.Pin.Connected");
-			}
+			OutAllActions.Append(MenuBuilder);
 		}
-	}
-
-	SBorder::Construct(SBorder::FArguments()
-		.BorderImage(FEditorStyle::GetBrush("Menu.Background"))
-		.Padding(5)
-		[
-			SNew(SBox)
-			.WidthOverride(400)
-			.HeightOverride(400)
-			[
-				SNew(SVerticalBox)
-				+SVerticalBox::Slot()
-				.AutoHeight()
-				.Padding(2, 2, 2, 5)
-				[
-					SNew(SHorizontalBox)
-					+SHorizontalBox::Slot()
-					.AutoWidth()
-					.VAlign(VAlign_Center)
-					.Padding(0, 0, 5, 0)
-					[
-						SNew(SImage)
-						.ColorAndOpacity(TypeColor)
-						.Visibility_Lambda([this]()
-						{
-							return DraggedFromPins.IsEmpty() ? EVisibility::Hidden : EVisibility::Visible;
-						})
-						.Image(PinBrush)
-					]
-					+SHorizontalBox::Slot()
-					.AutoWidth()
-					.VAlign(VAlign_Center)
-					[
-						SNew(STextBlock)
-						.Text_Lambda([this]()
-						{
-							if (DraggedFromPins.IsEmpty())
-							{
-								return LOCTEXT("ContextText", "All MetaSound Node Classes");
-							}
-
-							UEdGraphPin* Pin = DraggedFromPins[0];
-
-							if (DraggedFromPins[0]->Direction == EGPD_Input)
-							{
-								FInputHandle InputHandle = FGraphBuilder::GetInputHandleFromPin(Pin);
-								return FText::Format(LOCTEXT("ContextTypeFilteredText_Output", "Classes with output of type '{0}'"), FText::FromName(InputHandle->GetDataType()));
-							}
-							else
-							{
-								FOutputHandle OutputHandle = FGraphBuilder::GetOutputHandleFromPin(Pin);
-								return FText::Format(LOCTEXT("ContextTypeFilteredText_Input", "Classes with input of type '{0}'"), FText::FromName(OutputHandle->GetDataType()));
-							}
-						})
-						// TODO: Move to Metasound Style
-						.Font(FEditorStyle::GetFontStyle("BlueprintEditor.ActionMenu.ContextDescriptionFont"))
-						.ToolTip(IDocumentation::Get()->CreateToolTip(
-							LOCTEXT("ActionMenuContextTextTooltip", "Describes the current context of the action list"),
-							nullptr,
-							TEXT("Shared/Editors/MetasoundEditor"),
-							TEXT("MetasoundActionMenuContextText")))
-						.WrapTextAt(280)
-					]
-				]
-				+SVerticalBox::Slot()
-				[
-					SAssignNew(GraphActionMenu, SGraphActionMenu)
-						.OnActionSelected(this, &SMetasoundActionMenu::OnActionSelected)
-						.OnCreateWidgetForAction(SGraphActionMenu::FOnCreateWidgetForAction::CreateSP(this, &SMetasoundActionMenu::OnCreateWidgetForAction))
-						.OnCollectAllActions(this, &SMetasoundActionMenu::CollectAllActions)
-						.OnCreateCustomRowExpander_Lambda([](const FCustomExpanderData& InCustomExpanderData)
-						{
-							return SNew(SMetasoundActionMenuExpanderArrow, InCustomExpanderData);
-						})
-						.DraggedFromPins(DraggedFromPins)
-						.GraphObj(Graph)
-						.AlphaSortItems(true)
-						.bAllowPreselectedItemActivation(true)
-				]
-			]
-		]
-	);
-}
-
-void SMetasoundActionMenu::CollectAllActions(FGraphActionListBuilderBase& OutAllActions)
-{
-	if (!Graph)
-	{
-		return;
-	}
-
-	FGraphContextMenuBuilder MenuBuilder(Graph);
-	if (!DraggedFromPins.IsEmpty())
-	{
-		MenuBuilder.FromPin = DraggedFromPins[0];
-	}
 
-	// Cannot call GetGraphContextActions() during serialization and GC due to its use of FindObject()
-	if(!GIsSavingPackage && !IsGarbageCollecting())
-	{
-		if (const UMetasoundEditorGraphSchema* Schema = GetDefault<UMetasoundEditorGraphSchema>())
+		TSharedRef<SEditableTextBox> SMetasoundActionMenu::GetFilterTextBox()
 		{
-			Schema->GetGraphContextActions(MenuBuilder);
+			return GraphActionMenu->GetFilterTextBox();
 		}
-	}
-
-	OutAllActions.Append(MenuBuilder);
-}
 
-TSharedRef<SEditableTextBox> SMetasoundActionMenu::GetFilterTextBox()
-{
-	return GraphActionMenu->GetFilterTextBox();
-}
+		TSharedRef<SWidget> SMetasoundActionMenu::OnCreateWidgetForAction(FCreateWidgetForActionData* const InCreateData)
+		{
+			using namespace Metasound;
+			using namespace Metasound::Frontend;
 
-TSharedRef<SWidget> SMetasoundActionMenu::OnCreateWidgetForAction(FCreateWidgetForActionData* const InCreateData)
-{
-	using namespace Metasound;
-	using namespace Metasound::Frontend;
+			check(InCreateData);
+			InCreateData->bHandleMouseButtonDown = false;
 
-	check(InCreateData);
-	InCreateData->bHandleMouseButtonDown = false;
+			const FSlateBrush* IconBrush = nullptr;
+			FLinearColor IconColor;
+			TSharedPtr<FMetasoundGraphSchemaAction> Action = StaticCastSharedPtr<FMetasoundGraphSchemaAction>(InCreateData->Action);
+			if (Action.IsValid())
+			{
+				IconBrush = Action->GetIconBrush();
+				IconColor = Action->GetIconColor();
+			}
 
-	const FSlateBrush* IconBrush = nullptr;
-	FLinearColor IconColor;
-	TSharedPtr<FMetasoundGraphSchemaAction> Action = StaticCastSharedPtr<FMetasoundGraphSchemaAction>(InCreateData->Action);
-	if (Action.IsValid())
-	{
-		IconBrush = Action->GetIconBrush();
-		IconColor = Action->GetIconColor();
-	}
+			TSharedPtr<SHorizontalBox> WidgetBox = SNew(SHorizontalBox);
+			if (IconBrush)
+			{
+				WidgetBox->AddSlot()
+					.AutoWidth()
+					.VAlign(VAlign_Center)
+					.Padding(5, 0, 0, 0)
+					[
+						SNew(SImage)
+						.ColorAndOpacity(IconColor)
+						.Image(IconBrush)
+					];
+			}
 
-	TSharedPtr<SHorizontalBox> WidgetBox = SNew(SHorizontalBox);
-	if (IconBrush)
-	{
-		WidgetBox->AddSlot()
-			.AutoWidth()
-			.VAlign(VAlign_Center)
-			.Padding(5, 0, 0, 0)
-			[
-				SNew(SImage)
-				.ColorAndOpacity(IconColor)
-				.Image(IconBrush)
-			];
-	}
-
-	WidgetBox->AddSlot()
-		.AutoWidth()
-		.VAlign(VAlign_Center)
-		.Padding(0, 0, 0, 0)
-		[
-			SNew(SGraphPaletteItem, InCreateData)
-		];
-
-	return WidgetBox->AsShared();
-}
-
-void SMetasoundActionMenu::OnActionSelected(const TArray<TSharedPtr<FEdGraphSchemaAction>>& SelectedAction, ESelectInfo::Type InSelectionType)
-{
-	if (!Graph)
-	{
-		return;
-	}
+			WidgetBox->AddSlot()
+				.AutoWidth()
+				.VAlign(VAlign_Center)
+				.Padding(0, 0, 0, 0)
+				[
+					SNew(SGraphPaletteItem, InCreateData)
+				];
 
-	if (InSelectionType != ESelectInfo::OnMouseClick  && InSelectionType != ESelectInfo::OnKeyPress && !SelectedAction.IsEmpty())
-	{
-		return;
-	}
+			return WidgetBox->AsShared();
+		}
 
-	for (const TSharedPtr<FEdGraphSchemaAction>& Action : SelectedAction)
-	{
-		if (Action.IsValid() && Graph)
+		void SMetasoundActionMenu::OnActionSelected(const TArray<TSharedPtr<FEdGraphSchemaAction>>& SelectedAction, ESelectInfo::Type InSelectionType)
 		{
-			if (!bActionExecuted && (Action->GetTypeId() != FEdGraphSchemaAction_Dummy::StaticGetTypeId()))
+			if (!Graph)
 			{
-				FSlateApplication::Get().DismissAllMenus();
-				bActionExecuted = true;
+				return;
 			}
 
-			if (UEdGraphNode* ResultNode = Action->PerformAction(Graph, DraggedFromPins, NewNodePosition))
+			if (InSelectionType != ESelectInfo::OnMouseClick  && InSelectionType != ESelectInfo::OnKeyPress && !SelectedAction.IsEmpty())
 			{
-				NewNodePosition += Metasound::Frontend::DisplayStyle::NodeLayout::DefaultOffsetX;
+				return;
+			}
+
+			for (const TSharedPtr<FEdGraphSchemaAction>& Action : SelectedAction)
+			{
+				if (Action.IsValid() && Graph)
+				{
+					if (!bActionExecuted && (Action->GetTypeId() != FEdGraphSchemaAction_Dummy::StaticGetTypeId()))
+					{
+						FSlateApplication::Get().DismissAllMenus();
+						bActionExecuted = true;
+					}
+
+					if (UEdGraphNode* ResultNode = Action->PerformAction(Graph, DraggedFromPins, NewNodePosition))
+					{
+						NewNodePosition += Metasound::Frontend::DisplayStyle::NodeLayout::DefaultOffsetX;
+					}
+				}
 			}
 		}
-	}
-}
+	} // namespace Editor
+} // namespace Metasound
 #undef LOCTEXT_NAMESPACE
diff --git a/Engine/Plugins/Runtime/Metasound/Source/MetasoundEditor/Private/SMetasoundActionMenu.h b/Engine/Plugins/Runtime/Metasound/Source/MetasoundEditor/Private/SMetasoundActionMenu.h
--- a/Engine/Plugins/Runtime/Metasound/Source/MetasoundEditor/Private/SMetasoundActionMenu.h
+++ b/Engine/Plugins/Runtime/Metasound/Source/MetasoundEditor/Private/SMetasoundActionMenu.h
@@ -19,51 +19,75 @@ class SGraphActionMenu;
 class UEdGraph;
 struct FCustomExpanderData;
 
-
-class SMetasoundActionMenu : public SBorder
+namespace Metasound
 {
-public:
-	/** Delegate for the OnCloseReason event which is always raised when the SMetasoundActionMenu closes */
-	DECLARE_DELEGATE_ThreeParams(FClosedReason, bool /*bActionExecuted*/, bool /*bContextSensitiveChecked*/, bool /*bGraphPinContext*/);
-
-	SLATE_BEGIN_ARGS(SMetasoundActionMenu)
-		: _Graph(static_cast<UEdGraph*>(nullptr))
-		, _NewNodePosition(FVector2D::ZeroVector)
-		, _AutoExpandActionMenu(false)
+	namespace Editor
 	{
-	}
-		SLATE_ARGUMENT(UEdGraph*, Graph)
-		SLATE_ARGUMENT(FVector2D, NewNodePosition)
-		SLATE_ARGUMENT(TArray<UEdGraphPin*>, DraggedFromPins)
-		SLATE_ARGUMENT(SGraphEditor::FActionMenuClosed, OnClosedCallback)
-		SLATE_ARGUMENT(bool, AutoExpandActionMenu)
-		SLATE_EVENT(FClosedReason, OnCloseReason)
-	SLATE_END_ARGS()
+		// Custom expander to specify our desired padding
+		class SMetasoundActionMenuExpanderArrow : public SExpanderArrow
+		{
+			SLATE_BEGIN_ARGS(SMetasoundActionMenuExpanderArrow)
+			{
+			}
+
+			SLATE_ATTRIBUTE(float, IndentAmount)
+				SLATE_END_ARGS()
+
+		public:
+			void Construct(const FArguments& InArgs, const FCustomExpanderData& ActionMenuData);
+
+		private:
+			FMargin GetCustomIndentPadding() const;
+
+			TWeakPtr<FEdGraphSchemaAction> ActionPtr;
+		};
+
+		class SMetasoundActionMenu : public SBorder
+		{
+		public:
+			/** Delegate for the OnCloseReason event which is always raised when the SMetasoundActionMenu closes */
+			DECLARE_DELEGATE_ThreeParams(FClosedReason, bool /*bActionExecuted*/, bool /*bContextSensitiveChecked*/, bool /*bGraphPinContext*/);
+
+			SLATE_BEGIN_ARGS(SMetasoundActionMenu)
+				: _Graph(static_cast<UEdGraph*>(nullptr))
+				, _NewNodePosition(FVector2D::ZeroVector)
+				, _AutoExpandActionMenu(false)
+			{
+			}
+				SLATE_ARGUMENT(UEdGraph*, Graph)
+				SLATE_ARGUMENT(FVector2D, NewNodePosition)
+				SLATE_ARGUMENT(TArray<UEdGraphPin*>, DraggedFromPins)
+				SLATE_ARGUMENT(SGraphEditor::FActionMenuClosed, OnClosedCallback)
+				SLATE_ARGUMENT(bool, AutoExpandActionMenu)
+				SLATE_EVENT(FClosedReason, OnCloseReason)
+			SLATE_END_ARGS()
 
-	void Construct(const FArguments& InArgs);
+			void Construct(const FArguments& InArgs);
 
-	~SMetasoundActionMenu();
+			~SMetasoundActionMenu();
 
-	TSharedRef<SEditableTextBox> GetFilterTextBox();
+			TSharedRef<SEditableTextBox> GetFilterTextBox();
 
-protected:
-	void OnActionSelected(const TArray<TSharedPtr<FEdGraphSchemaAction>>& SelectedAction, ESelectInfo::Type InSelectionType);
+		protected:
+			void OnActionSelected(const TArray<TSharedPtr<FEdGraphSchemaAction>>& SelectedAction, ESelectInfo::Type InSelectionType);
 
-	TSharedRef<SWidget> OnCreateWidgetForAction(struct FCreateWidgetForActionData* const InCreateData);
+			TSharedRef<SWidget> OnCreateWidgetForAction(struct FCreateWidgetForActionData* const InCreateData);
 
-	/** Callback used to populate all actions list in SGraphActionMenu */
-	void CollectAllActions(FGraphActionListBuilderBase& OutAllActions);
+			/** Callback used to populate all actions list in SGraphActionMenu */
+			void CollectAllActions(FGraphActionListBuilderBase& OutAllActions);
 
-private:
-	UEdGraph* Graph = nullptr;
-	bool bAutoExpandActionMenu = false;
-	bool bActionExecuted = false;
+		private:
+			UEdGraph* Graph = nullptr;
+			bool bAutoExpandActionMenu = false;
+			bool bActionExecuted = false;
 
-	TArray<UEdGraphPin*> DraggedFromPins;
-	FVector2D NewNodePosition = FVector2D::ZeroVector;
+			TArray<UEdGraphPin*> DraggedFromPins;
+			FVector2D NewNodePosition = FVector2D::ZeroVector;
 
-	SGraphEditor::FActionMenuClosed OnClosedCallback;
-	FClosedReason OnCloseReasonCallback;
+			SGraphEditor::FActionMenuClosed OnClosedCallback;
+			FClosedReason OnCloseReasonCallback;
 
-	TSharedPtr<SGraphActionMenu> GraphActionMenu;
-};
+			TSharedPtr<SGraphActionMenu> GraphActionMenu;
+		};
+	} // namespace Editor
+} // namespace Metasound
diff --git a/Engine/Plugins/Runtime/Metasound/Source/MetasoundEditor/Private/SMetasoundGraphNode.cpp b/Engine/Plugins/Runtime/Metasound/Source/MetasoundEditor/Private/SMetasoundGraphNode.cpp
--- a/Engine/Plugins/Runtime/Metasound/Source/MetasoundEditor/Private/SMetasoundGraphNode.cpp
+++ b/Engine/Plugins/Runtime/Metasound/Source/MetasoundEditor/Private/SMetasoundGraphNode.cpp
@@ -1,7 +1,7 @@
 // Copyright Epic Games, Inc. All Rights Reserved.
 #include "SMetasoundGraphNode.h"
 
-#include "AudioParameterInterface.h"
+#include "AudioParameterControllerInterface.h"
 #include "Components/AudioComponent.h"
 #include "EditorStyleSet.h"
 #include "GraphEditorSettings.h"
diff --git a/Engine/Plugins/Runtime/Metasound/Source/MetasoundEditor/Public/MetasoundEditorGraph.h b/Engine/Plugins/Runtime/Metasound/Source/MetasoundEditor/Public/MetasoundEditorGraph.h
--- a/Engine/Plugins/Runtime/Metasound/Source/MetasoundEditor/Public/MetasoundEditorGraph.h
+++ b/Engine/Plugins/Runtime/Metasound/Source/MetasoundEditor/Public/MetasoundEditorGraph.h
@@ -3,7 +3,7 @@
 
 #include "CoreMinimal.h"
 
-#include "AudioParameterInterface.h"
+#include "AudioParameterControllerInterface.h"
 #include "MetasoundEditor.h"
 #include "MetasoundEditorGraphNode.h"
 #include "MetasoundFrontendController.h"
@@ -41,7 +41,7 @@ class METASOUNDEDITOR_API UMetasoundEditorGraphInputLiteral : public UObject
 	GENERATED_BODY()
 
 public:
-	virtual void UpdatePreviewInstance(const Metasound::FVertexName& InParameterName, TScriptInterface<IAudioParameterInterface>& InParameterInterface) const
+	virtual void UpdatePreviewInstance(const Metasound::FVertexName& InParameterName, TScriptInterface<IAudioParameterControllerInterface>& InParameterInterface) const
 	{
 	}
 
diff --git a/Engine/Plugins/Runtime/Metasound/Source/MetasoundEditor/Public/MetasoundEditorGraphInputNodes.h b/Engine/Plugins/Runtime/Metasound/Source/MetasoundEditor/Public/MetasoundEditorGraphInputNodes.h
--- a/Engine/Plugins/Runtime/Metasound/Source/MetasoundEditor/Public/MetasoundEditorGraphInputNodes.h
+++ b/Engine/Plugins/Runtime/Metasound/Source/MetasoundEditor/Public/MetasoundEditorGraphInputNodes.h
@@ -3,9 +3,8 @@
 
 #include "Algo/Count.h"
 #include "Algo/Transform.h"
-#include "AudioParameterInterface.h"
+#include "AudioParameterControllerInterface.h"
 #include "Components/Widget.h"
-#include "CoreMinimal.h"
 #include "Delegates/Delegate.h"
 #include "EdGraph/EdGraphNode.h"
 #include "GraphEditorSettings.h"
@@ -43,7 +42,7 @@ public:
 	TObjectPtr<UMetasoundEditorGraphInput> Input;
 
 public:
-	virtual void UpdatePreviewInstance(const Metasound::FVertexName& InParameterName, TScriptInterface<IAudioParameterInterface>& InParameterInterface) const;
+	virtual void UpdatePreviewInstance(const Metasound::FVertexName& InParameterName, TScriptInterface<IAudioParameterControllerInterface>& InParameterInterface) const;
 
 	virtual FMetasoundFrontendClassName GetClassName() const;
 
@@ -86,7 +85,7 @@ public:
 	virtual FMetasoundFrontendLiteral GetDefault() const override;
 	virtual EMetasoundFrontendLiteralType GetLiteralType() const override;
 	virtual void SetFromLiteral(const FMetasoundFrontendLiteral& InLiteral) override;
-	virtual void UpdatePreviewInstance(const Metasound::FVertexName& InParameterName, TScriptInterface<IAudioParameterInterface>& InParameterInterface) const override;
+	virtual void UpdatePreviewInstance(const Metasound::FVertexName& InParameterName, TScriptInterface<IAudioParameterControllerInterface>& InParameterInterface) const override;
 };
 
 UCLASS(MinimalAPI)
@@ -103,7 +102,7 @@ public:
 	virtual FMetasoundFrontendLiteral GetDefault() const override;
 	virtual EMetasoundFrontendLiteralType GetLiteralType() const override;
 	virtual void SetFromLiteral(const FMetasoundFrontendLiteral& InLiteral) override;
-	virtual void UpdatePreviewInstance(const Metasound::FVertexName& InParameterName, TScriptInterface<IAudioParameterInterface>& InParameterInterface) const override;
+	virtual void UpdatePreviewInstance(const Metasound::FVertexName& InParameterName, TScriptInterface<IAudioParameterControllerInterface>& InParameterInterface) const override;
 };
 
 // Broken out to be able to customize and swap enum behavior for basic integer literal behavior
@@ -130,7 +129,7 @@ public:
 	virtual FMetasoundFrontendLiteral GetDefault() const override;
 	virtual EMetasoundFrontendLiteralType GetLiteralType() const override;
 	virtual void SetFromLiteral(const FMetasoundFrontendLiteral& InLiteral) override;
-	virtual void UpdatePreviewInstance(const Metasound::FVertexName& InParameterName, TScriptInterface<IAudioParameterInterface>& InParameterInterface) const override;
+	virtual void UpdatePreviewInstance(const Metasound::FVertexName& InParameterName, TScriptInterface<IAudioParameterControllerInterface>& InParameterInterface) const override;
 };
 
 UCLASS(MinimalAPI)
@@ -147,7 +146,7 @@ public:
 	virtual FMetasoundFrontendLiteral GetDefault() const override;
 	virtual EMetasoundFrontendLiteralType GetLiteralType() const override;
 	virtual void SetFromLiteral(const FMetasoundFrontendLiteral& InLiteral) override;
-	virtual void UpdatePreviewInstance(const Metasound::FVertexName& InParameterName, TScriptInterface<IAudioParameterInterface>& InParameterInterface) const override;
+	virtual void UpdatePreviewInstance(const Metasound::FVertexName& InParameterName, TScriptInterface<IAudioParameterControllerInterface>& InParameterInterface) const override;
 };
 
 // For input widget
@@ -205,7 +204,7 @@ public:
 	virtual FMetasoundFrontendLiteral GetDefault() const override;
 	virtual EMetasoundFrontendLiteralType GetLiteralType() const override;
 	virtual void SetFromLiteral(const FMetasoundFrontendLiteral& InLiteral) override;
-	virtual void UpdatePreviewInstance(const Metasound::FVertexName& InParameterName, TScriptInterface<IAudioParameterInterface>& InParameterInterface) const override;
+	virtual void UpdatePreviewInstance(const Metasound::FVertexName& InParameterName, TScriptInterface<IAudioParameterControllerInterface>& InParameterInterface) const override;
 	virtual void PostEditChangeProperty(FPropertyChangedEvent& PropertyChangedEvent) override;
 
 	void SetDefault(const float InDefault);
@@ -227,7 +226,7 @@ public:
 	virtual FMetasoundFrontendLiteral GetDefault() const override;
 	virtual EMetasoundFrontendLiteralType GetLiteralType() const override;
 	virtual void SetFromLiteral(const FMetasoundFrontendLiteral& InLiteral) override;
-	virtual void UpdatePreviewInstance(const Metasound::FVertexName& InParameterName, TScriptInterface<IAudioParameterInterface>& InParameterInterface) const override;
+	virtual void UpdatePreviewInstance(const Metasound::FVertexName& InParameterName, TScriptInterface<IAudioParameterControllerInterface>& InParameterInterface) const override;
 };
 
 UCLASS(MinimalAPI)
@@ -244,7 +243,7 @@ public:
 	virtual FMetasoundFrontendLiteral GetDefault() const override;
 	virtual EMetasoundFrontendLiteralType GetLiteralType() const override;
 	virtual void SetFromLiteral(const FMetasoundFrontendLiteral& InLiteral) override;
-	virtual void UpdatePreviewInstance(const Metasound::FVertexName& InParameterName, TScriptInterface<IAudioParameterInterface>& InParameterInterface) const override;
+	virtual void UpdatePreviewInstance(const Metasound::FVertexName& InParameterName, TScriptInterface<IAudioParameterControllerInterface>& InParameterInterface) const override;
 };
 
 UCLASS(MinimalAPI)
@@ -261,7 +260,7 @@ public:
 	virtual FMetasoundFrontendLiteral GetDefault() const override;
 	virtual EMetasoundFrontendLiteralType GetLiteralType() const override;
 	virtual void SetFromLiteral(const FMetasoundFrontendLiteral& InLiteral) override;
-	virtual void UpdatePreviewInstance(const Metasound::FVertexName& InParameterName, TScriptInterface<IAudioParameterInterface>& InParameterInterface) const override;
+	virtual void UpdatePreviewInstance(const Metasound::FVertexName& InParameterName, TScriptInterface<IAudioParameterControllerInterface>& InParameterInterface) const override;
 };
 
 // Broken out to be able to customize and swap AllowedClass based on provided object proxy
@@ -288,7 +287,7 @@ public:
 	virtual FMetasoundFrontendLiteral GetDefault() const override;
 	virtual EMetasoundFrontendLiteralType GetLiteralType() const override;
 	virtual void SetFromLiteral(const FMetasoundFrontendLiteral& InLiteral) override;
-	virtual void UpdatePreviewInstance(const Metasound::FVertexName& InParameterName, TScriptInterface<IAudioParameterInterface>& InParameterInterface) const override;
+	virtual void UpdatePreviewInstance(const Metasound::FVertexName& InParameterName, TScriptInterface<IAudioParameterControllerInterface>& InParameterInterface) const override;
 };
 
 UCLASS(MinimalAPI)
@@ -305,5 +304,5 @@ public:
 	virtual FMetasoundFrontendLiteral GetDefault() const override;
 	virtual EMetasoundFrontendLiteralType GetLiteralType() const override;
 	virtual void SetFromLiteral(const FMetasoundFrontendLiteral& InLiteral) override;
-	virtual void UpdatePreviewInstance(const Metasound::FVertexName& InParameterName, TScriptInterface<IAudioParameterInterface>& InParameterInterface) const override;
+	virtual void UpdatePreviewInstance(const Metasound::FVertexName& InParameterName, TScriptInterface<IAudioParameterControllerInterface>& InParameterInterface) const override;
 };
\ No newline at end of file
diff --git a/Engine/Plugins/Runtime/Metasound/Source/MetasoundEngine/Private/Metasound.cpp b/Engine/Plugins/Runtime/Metasound/Source/MetasoundEngine/Private/Metasound.cpp
--- a/Engine/Plugins/Runtime/Metasound/Source/MetasoundEngine/Private/Metasound.cpp
+++ b/Engine/Plugins/Runtime/Metasound/Source/MetasoundEngine/Private/Metasound.cpp
@@ -135,10 +135,4 @@ const TSet<FSoftObjectPath>& UMetaSound::GetReferencedAssetClassCache() const
 {
 	return ReferenceAssetClassCache;
 }
-
-const FMetasoundFrontendVersion& UMetaSound::GetDefaultArchetypeVersion() const
-{
-	static const FMetasoundFrontendVersion DefaultArchetypeVersion = Metasound::Engine::MetasoundV1_0::GetVersion();
-	return DefaultArchetypeVersion;
-}
 #undef LOCTEXT_NAMESPACE // MetaSound
diff --git a/Engine/Plugins/Runtime/Metasound/Source/MetasoundEngine/Private/MetasoundAssetSubsystem.cpp b/Engine/Plugins/Runtime/Metasound/Source/MetasoundEngine/Private/MetasoundAssetSubsystem.cpp
--- a/Engine/Plugins/Runtime/Metasound/Source/MetasoundEngine/Private/MetasoundAssetSubsystem.cpp
+++ b/Engine/Plugins/Runtime/Metasound/Source/MetasoundEngine/Private/MetasoundAssetSubsystem.cpp
@@ -105,8 +105,7 @@ void UMetaSoundAssetSubsystem::PostInitAssetScan()
 	{
 		SearchAndIterateDirectoryAssets(Settings->DirectoriesToRegister, [this](const FAssetData& AssetData)
 		{
-			constexpr bool bRegisterWithFrontend = false;
-			AddOrUpdateAsset(AssetData, bRegisterWithFrontend);
+			AddOrUpdateAsset(AssetData);
 		});
 	}
 }
@@ -116,15 +115,13 @@ void UMetaSoundAssetSubsystem::AddAssetReferences(FMetasoundAssetBase& InAssetBa
 	using namespace Metasound;
 	using namespace Metasound::Frontend;
 
-	constexpr bool bRegisterWithFrontend = false;
-
 	const FNodeClassInfo AssetClassInfo = InAssetBase.GetAssetClassInfo();
 	const FNodeRegistryKey AssetClassKey = NodeRegistryKey::CreateKey(AssetClassInfo);
 
 	if (!ContainsKey(AssetClassKey))
 	{
-		AddOrUpdateAsset(*InAssetBase.GetOwningAsset(), bRegisterWithFrontend);
-		UE_LOG(LogMetaSound, Log, TEXT("Adding asset '%s' to MetaSoundAsset registry."), *InAssetBase.GetOwningAssetName());
+		AddOrUpdateAsset(*InAssetBase.GetOwningAsset());
+		UE_LOG(LogMetaSound, Verbose, TEXT("Adding asset '%s' to MetaSoundAsset registry."), *InAssetBase.GetOwningAssetName());
 	}
 
 	bool bLoadFromPathCache = false;
@@ -167,14 +164,14 @@ void UMetaSoundAssetSubsystem::AddAssetReferences(FMetasoundAssetBase& InAssetBa
 				UObject* MetaSoundObject = MetaSoundAsset->GetOwningAsset();
 				if (ensureAlways(MetaSoundObject))
 				{
-					AddOrUpdateAsset(*MetaSoundObject, bRegisterWithFrontend);
+					AddOrUpdateAsset(*MetaSoundObject);
 				}
 			}
 		}
 	}
 }
 
-void UMetaSoundAssetSubsystem::AddOrUpdateAsset(UObject& InObject, bool bInRegisterWithFrontend)
+void UMetaSoundAssetSubsystem::AddOrUpdateAsset(UObject& InObject)
 {
 	using namespace Metasound;
 	using namespace Metasound::AssetSubsystemPrivate;
@@ -185,17 +182,12 @@ void UMetaSoundAssetSubsystem::AddOrUpdateAsset(UObject& InObject, bool bInRegis
 	FMetasoundAssetBase* MetaSoundAsset = Metasound::IMetasoundUObjectRegistry::Get().GetObjectAsAssetBase(&InObject);
 	check(MetaSoundAsset);
 
-	if (bInRegisterWithFrontend)
-	{
-		MetaSoundAsset->RegisterGraphWithFrontend();
-	}
-
 	FNodeClassInfo ClassInfo = MetaSoundAsset->GetAssetClassInfo();
 	const FNodeRegistryKey RegistryKey = NodeRegistryKey::CreateKey(ClassInfo);
 	PathMap.FindOrAdd(RegistryKey) = InObject.GetPathName();
 }
 
-void UMetaSoundAssetSubsystem::AddOrUpdateAsset(const FAssetData& InAssetData, bool bInRegisterWithFrontend)
+void UMetaSoundAssetSubsystem::AddOrUpdateAsset(const FAssetData& InAssetData)
 {
 	using namespace Metasound;
 	using namespace Metasound::AssetSubsystemPrivate;
@@ -205,7 +197,7 @@ void UMetaSoundAssetSubsystem::AddOrUpdateAsset(const FAssetData& InAssetData, b
 
 	FNodeClassInfo ClassInfo;
 	bool bClassInfoFound = GetAssetClassInfo(InAssetData, ClassInfo);
-	if (!bClassInfoFound || bInRegisterWithFrontend)
+	if (!bClassInfoFound)
 	{
 		UObject* Object = nullptr;
 
@@ -213,22 +205,17 @@ void UMetaSoundAssetSubsystem::AddOrUpdateAsset(const FAssetData& InAssetData, b
 		if (InAssetData.IsAssetLoaded())
 		{
 			Object = Path.ResolveObject();
+			UE_LOG(LogMetaSound, Verbose, TEXT("Adding loaded asset '%s' to MetaSoundAsset registry."), *Object->GetName());
 		}
 		else
 		{
-			if (!bInRegisterWithFrontend)
-			{
-				UE_LOG(LogMetaSound, Warning,
-					TEXT("Failed to find serialized MetaSound asset registry data for asset '%s'. "
-						"Forcing synchronous load which increases load times. Re-save asset to avoid this."),
-					*InAssetData.ObjectPath.ToString());
-			}
 			Object = Path.TryLoad();
+			UE_LOG(LogMetaSound, Verbose, TEXT("Loaded asset '%s' and adding to MetaSoundAsset registry."), *Object->GetName());
 		}
 
 		if (Object)
 		{
-			AddOrUpdateAsset(*Object, bInRegisterWithFrontend);
+			AddOrUpdateAsset(*Object);
 			return;
 		}
 	}
@@ -361,7 +348,7 @@ FMetasoundAssetBase* UMetaSoundAssetSubsystem::TryLoadAsset(const FSoftObjectPat
 	return Metasound::IMetasoundUObjectRegistry::Get().GetObjectAsAssetBase(InObjectPath.TryLoad());
 }
 
-void UMetaSoundAssetSubsystem::RemoveAsset(UObject& InObject, bool bInUnregisterWithFrontend)
+void UMetaSoundAssetSubsystem::RemoveAsset(UObject& InObject)
 {
 	using namespace Metasound;
 	using namespace Metasound::Frontend;
@@ -371,15 +358,10 @@ void UMetaSoundAssetSubsystem::RemoveAsset(UObject& InObject, bool bInUnregister
 		const FNodeClassInfo ClassInfo = MetaSoundAsset->GetAssetClassInfo();
 		FNodeRegistryKey RegistryKey = FMetasoundFrontendRegistryContainer::Get()->GetRegistryKey(ClassInfo);
 		PathMap.Remove(RegistryKey);
-
-		if (bInUnregisterWithFrontend)
-		{
-			MetaSoundAsset->UnregisterGraphWithFrontend();
-		}
 	}
 }
 
-void UMetaSoundAssetSubsystem::RemoveAsset(const FAssetData& InAssetData, bool bInUnregisterWithFrontend)
+void UMetaSoundAssetSubsystem::RemoveAsset(const FAssetData& InAssetData)
 {
 	using namespace Metasound;
 	using namespace Metasound::Frontend;
@@ -389,23 +371,32 @@ void UMetaSoundAssetSubsystem::RemoveAsset(const FAssetData& InAssetData, bool b
 	{
 		FNodeRegistryKey RegistryKey = FMetasoundFrontendRegistryContainer::Get()->GetRegistryKey(ClassInfo);
 		PathMap.Remove(RegistryKey);
-
-		UObject* Object = InAssetData.GetAsset();
-		if (Object && bInUnregisterWithFrontend)
-		{
-			FMetasoundAssetBase* MetaSoundAsset = Metasound::IMetasoundUObjectRegistry::Get().GetObjectAsAssetBase(Object);
-			check(MetaSoundAsset);
-
-			MetaSoundAsset->UnregisterGraphWithFrontend();
-		}
 	}
 }
 
 void UMetaSoundAssetSubsystem::RenameAsset(const FAssetData& InAssetData, bool bInReregisterWithFrontend)
 {
-	RemoveAsset(InAssetData, bInReregisterWithFrontend);
-	SynchronizeAssetClassDisplayName(InAssetData);
-	AddOrUpdateAsset(InAssetData, bInReregisterWithFrontend);
+	auto PerformRename = [this, &InAssetData]()
+	{
+		RemoveAsset(InAssetData);
+
+		SynchronizeAssetClassDisplayName(InAssetData);
+		AddOrUpdateAsset(InAssetData);
+	};
+
+	if (bInReregisterWithFrontend)
+	{
+		FMetasoundAssetBase* MetaSoundAsset = Metasound::IMetasoundUObjectRegistry::Get().GetObjectAsAssetBase(InAssetData.GetAsset());
+		check(MetaSoundAsset);
+
+		MetaSoundAsset->UnregisterGraphWithFrontend();
+		PerformRename();
+		MetaSoundAsset->RegisterGraphWithFrontend();
+	}
+	else
+	{
+		PerformRename();
+	}
 }
 
 void UMetaSoundAssetSubsystem::SynchronizeAssetClassDisplayName(const FAssetData& InAssetData)
@@ -441,21 +432,16 @@ void UMetaSoundAssetSubsystem::SearchAndIterateDirectoryAssets(const TArray<FDir
 		Rules.AssetScanPaths.Add(*Path.Path);
 	}
 
-	Rules.AssetBaseClass = UMetaSound::StaticClass();
-	TArray<FAssetData> MetaSoundAssets;
-	AssetManager.SearchAssetRegistryPaths(MetaSoundAssets, Rules);
-	for (const FAssetData& AssetData : MetaSoundAssets)
-	{
-		InFunction(AssetData);
-	}
-
-	Rules.AssetBaseClass = UMetaSoundSource::StaticClass();
-	TArray<FAssetData> MetaSoundSourceAssets;
-	AssetManager.SearchAssetRegistryPaths(MetaSoundSourceAssets, Rules);
-	for (const FAssetData& AssetData : MetaSoundSourceAssets)
+	Metasound::IMetasoundUObjectRegistry::Get().IterateRegisteredUClasses([&](UClass& RegisteredClass)
 	{
-		InFunction(AssetData);
-	}
+		Rules.AssetBaseClass = &RegisteredClass;
+		TArray<FAssetData> MetaSoundAssets;
+		AssetManager.SearchAssetRegistryPaths(MetaSoundAssets, Rules);
+		for (const FAssetData& AssetData : MetaSoundAssets)
+		{
+			InFunction(AssetData);
+		}
+	});
 }
 
 void UMetaSoundAssetSubsystem::RegisterAssetClassesInDirectories(const TArray<FMetaSoundAssetDirectory>& InDirectories)
@@ -465,8 +451,10 @@ void UMetaSoundAssetSubsystem::RegisterAssetClassesInDirectories(const TArray<FM
 
 	SearchAndIterateDirectoryAssets(Directories, [this](const FAssetData& AssetData)
 	{
-		constexpr bool bRegisterWithFrontend = true;
-		AddOrUpdateAsset(AssetData, bRegisterWithFrontend);
+		AddOrUpdateAsset(AssetData);
+		FMetasoundAssetBase* MetaSoundAsset = Metasound::IMetasoundUObjectRegistry::Get().GetObjectAsAssetBase(AssetData.GetAsset());
+		check(MetaSoundAsset);
+		MetaSoundAsset->RegisterGraphWithFrontend();
 	});
 }
 
@@ -477,7 +465,30 @@ void UMetaSoundAssetSubsystem::UnregisterAssetClassesInDirectories(const TArray<
 
 	SearchAndIterateDirectoryAssets(Directories, [this](const FAssetData& AssetData)
 	{
-		constexpr bool bUnregisterWithFrontend = true;
-		RemoveAsset(AssetData, bUnregisterWithFrontend);
+		using namespace Metasound;
+		using namespace Metasound::Frontend;
+
+		if (AssetData.IsAssetLoaded())
+		{
+			FMetasoundAssetBase* MetaSoundAsset = Metasound::IMetasoundUObjectRegistry::Get().GetObjectAsAssetBase(AssetData.GetAsset());
+			check(MetaSoundAsset);
+			MetaSoundAsset->UnregisterGraphWithFrontend();
+
+			RemoveAsset(AssetData);
+		}
+		else
+		{
+			FNodeClassInfo AssetClassInfo;
+			if (ensureAlways(AssetSubsystemPrivate::GetAssetClassInfo(AssetData, AssetClassInfo)))
+			{
+				const FNodeRegistryKey RegistryKey = NodeRegistryKey::CreateKey(AssetClassInfo);
+				const bool bIsRegistered = FMetasoundFrontendRegistryContainer::Get()->IsNodeRegistered(RegistryKey);
+				if (bIsRegistered)
+				{
+					FMetasoundFrontendRegistryContainer::Get()->UnregisterNode(RegistryKey);
+					PathMap.Remove(RegistryKey);
+				}
+			}
+		}
 	});
 }
diff --git a/Engine/Plugins/Runtime/Metasound/Source/MetasoundEngine/Private/MetasoundEngineArchetypes.cpp b/Engine/Plugins/Runtime/Metasound/Source/MetasoundEngine/Private/MetasoundEngineArchetypes.cpp
--- a/Engine/Plugins/Runtime/Metasound/Source/MetasoundEngine/Private/MetasoundEngineArchetypes.cpp
+++ b/Engine/Plugins/Runtime/Metasound/Source/MetasoundEngine/Private/MetasoundEngineArchetypes.cpp
@@ -2,11 +2,13 @@
 
 #include "MetasoundEngineArchetypes.h"
 
+#include "IAudioParameterInterfaceRegistry.h"
 #include "MetasoundAudioFormats.h"
 #include "MetasoundDataReference.h"
 #include "MetasoundFrontendArchetypeRegistry.h"
 #include "MetasoundFrontendTransform.h"
 #include "MetasoundInterface.h"
+#include "MetasoundOutputFormatInterfaces.h"
 #include "MetasoundRouter.h"
 #include "MetasoundSource.h"
 #include "MetasoundUObjectRegistry.h"
@@ -19,78 +21,84 @@ namespace Metasound
 {
 	namespace Engine
 	{
-		// MetasoundV1_0 is a metasound without any required inputs or outputs.
-		namespace MetasoundV1_0
+		namespace ArchetypesPrivate
 		{
-			FMetasoundFrontendVersion GetVersion()
+			auto MatchMemberNamesIgnoreSpaces = [] (FName InNameA, FName InNameB)
 			{
-				static const FName VersionName = FName("MetaSound");
-				return FMetasoundFrontendVersion{ VersionName, FMetasoundFrontendVersionNumber{1, 0} };
-			}
+				FName ParamA;
+				FName ParamB;
+				FName Namespace;
+				Audio::FParameterPath::SplitName(InNameA, Namespace, ParamA);
+				Audio::FParameterPath::SplitName(InNameB, Namespace, ParamB);
 
-			FMetasoundFrontendInterface GetInterface()
-			{
-				FMetasoundFrontendInterface Interface;
-				Interface.Version = GetVersion();
+				auto GetStringNoWhitespace = [](FName InName) { return InName.ToString().Replace(TEXT(" "), TEXT("")); };
+				return GetStringNoWhitespace(ParamA) == GetStringNoWhitespace(ParamB);
+			};
 
-				return Interface;
+			const FVertexName& GetOnPlayInputName()
+			{
+				static const FVertexName OnPlayInputName = "On Play";
+				return OnPlayInputName;
 			}
-		}
 
-		namespace MetasoundSourceMonoV1_0
-		{
-			FMetasoundFrontendVersion GetVersion()
+			const FVertexName& GetIsFinishedOutputName()
 			{
-				static const FName VersionName = FName("MonoSource");
-				return FMetasoundFrontendVersion{ VersionName, FMetasoundFrontendVersionNumber{1, 0} };
-
+				static const FVertexName OnFinishedOutputName = "On Finished";
+				return OnFinishedOutputName;
 			}
 
-			const FVertexName& GetOnPlayInputName()
+			const FVertexName& GetAudioDeviceIDVariableName()
 			{
-				static const FVertexName TriggerInputName = TEXT("On Play");
-				return TriggerInputName;
+				static const FVertexName AudioDeviceIDVarName = "AudioDeviceID";
+				return AudioDeviceIDVarName;
 			}
 
-			const FVertexName& GetAudioOutputName()
+			const FVertexName& GetGraphName()
 			{
-				static const FVertexName AudioOutputName = TEXT("Generated Audio");
-				return AudioOutputName;
+				static const FVertexName SoundGraphName = "GraphName";
+				return SoundGraphName;
 			}
+		}
 
-			const FVertexName& GetIsFinishedOutputName()
+		// MetasoundV1_0 is a metasound without any required inputs or outputs.
+		namespace MetasoundV1_0
+		{
+			const FMetasoundFrontendVersion& GetVersion()
 			{
-				static const FVertexName OnFinishedOutputName = TEXT("On Finished");
-				return OnFinishedOutputName;
+				static const FMetasoundFrontendVersion Version { "MetaSound", FMetasoundFrontendVersionNumber {1, 0} };
+				return Version;
 			}
 
-			const FVertexName& GetAudioDeviceIDVariableName()
+			FMetasoundFrontendInterface GetInterface()
 			{
-				static const FVertexName AudioDeviceIDVarName = TEXT("AudioDeviceID");
-				return AudioDeviceIDVarName;
+				FMetasoundFrontendInterface Interface;
+				Interface.Version = GetVersion();
+				return Interface;
 			}
+		}
 
-			const FVertexName& GetSoundUniqueIdName()
+		namespace MetasoundOutputFormatMonoV1_0
+		{
+			const FMetasoundFrontendVersion& GetVersion()
 			{
-				static const FVertexName SoundUniqueIdVarName = TEXT("SoundUniqueId");
-				return SoundUniqueIdVarName;
+				static const FMetasoundFrontendVersion Version { "MonoSource", FMetasoundFrontendVersionNumber {1, 0} };
+				return Version;
 			}
 
-			const FVertexName& GetIsPreviewSoundName()
+			const FVertexName& GetAudioOutputName()
 			{
-				static const FVertexName SoundIsPreviewSoundName = TEXT("IsPreviewSound");
-				return SoundIsPreviewSoundName;
+				static const FVertexName AudioOutputName = "Generated Audio";
+				return AudioOutputName;
 			}
 
 			FMetasoundFrontendInterface GetInterface()
 			{
 				FMetasoundFrontendInterface Interface;
-				Interface.Version.Name = "MonoSource";
-				Interface.Version.Number = FMetasoundFrontendVersionNumber{1, 0};
+				Interface.Version = GetVersion();
 				
 				// Inputs
 				FMetasoundFrontendClassVertex OnPlayTrigger;
-				OnPlayTrigger.Name = GetOnPlayInputName();
+				OnPlayTrigger.Name = ArchetypesPrivate::GetOnPlayInputName();
 
 				OnPlayTrigger.TypeName = GetMetasoundDataTypeName<FTrigger>();
 				OnPlayTrigger.Metadata.DisplayName = LOCTEXT("OnPlay", "On Play");
@@ -101,7 +109,7 @@ namespace Metasound
 
 				// Outputs 
 				FMetasoundFrontendClassVertex OnFinished;
-				OnFinished.Name = GetIsFinishedOutputName();
+				OnFinished.Name = ArchetypesPrivate::GetIsFinishedOutputName();
 
 				OnFinished.TypeName = GetMetasoundDataTypeName<FTrigger>();
 				OnFinished.Metadata.DisplayName = LOCTEXT("OnFinished", "On Finished");
@@ -121,7 +129,7 @@ namespace Metasound
 
 				// Environment 
 				FMetasoundFrontendClassEnvironmentVariable AudioDeviceID;
-				AudioDeviceID.Name = GetAudioDeviceIDVariableName();
+				AudioDeviceID.Name = ArchetypesPrivate::GetAudioDeviceIDVariableName();
 				AudioDeviceID.Metadata.Description = LOCTEXT("AudioDeviceIDToolTip", "Audio Device ID");
 
 				Interface.Environment.Add(AudioDeviceID);
@@ -130,56 +138,20 @@ namespace Metasound
 			}
 		}
 
-		namespace MetasoundSourceStereoV1_0
+		namespace MetasoundOutputFormatStereoV1_0
 		{
-			FMetasoundFrontendVersion GetVersion()
+			const FMetasoundFrontendVersion& GetVersion()
 			{
-				static const FName VersionName = FName("StereoSource");
-				return FMetasoundFrontendVersion{ VersionName, FMetasoundFrontendVersionNumber{1, 0} };
-			}
-
-			const FVertexName& GetOnPlayInputName()
-			{
-				static const FVertexName TriggerInputName = TEXT("On Play");
-				return TriggerInputName;
+				static const FMetasoundFrontendVersion Version { "StereoSource", FMetasoundFrontendVersionNumber { 1, 0 } };
+				return Version;
 			}
 
 			const FVertexName& GetAudioOutputName()
 			{
-				static const FVertexName AudioOutputName = TEXT("Generated Audio");
+				static const FVertexName AudioOutputName = "Generated Audio";
 				return AudioOutputName;
 			}
 
-			const FVertexName& GetIsFinishedOutputName()
-			{
-				static const FVertexName OnFinishedOutputName = TEXT("On Finished");
-				return OnFinishedOutputName;
-			}
-
-			const FVertexName& GetAudioDeviceIDVariableName()
-			{
-				static const FVertexName AudioDeviceIDVarName = TEXT("AudioDeviceID");
-				return AudioDeviceIDVarName;
-			}
-
-			const FVertexName& GetSoundUniqueIdName()
-			{
-				static const FVertexName SoundUniqueIdVarName = TEXT("SoundUniqueId");
-				return SoundUniqueIdVarName;
-			}
-
-			const FVertexName& GetIsPreviewSoundName()
-			{
-				static const FVertexName SoundIsPreviewSoundName = TEXT("IsPreviewSound");
-				return SoundIsPreviewSoundName;
-			}
-
-			const FVertexName& GetGraphName()
-			{
-				static const FVertexName SoundGraphName = TEXT("GraphName");
-				return SoundGraphName;
-			}
-
 			FMetasoundFrontendInterface GetInterface()
 			{
 				FMetasoundFrontendInterface Interface;
@@ -187,7 +159,7 @@ namespace Metasound
 				
 				// Inputs
 				FMetasoundFrontendClassVertex OnPlayTrigger;
-				OnPlayTrigger.Name = GetOnPlayInputName();
+				OnPlayTrigger.Name = ArchetypesPrivate::GetOnPlayInputName();
 
 				OnPlayTrigger.TypeName = GetMetasoundDataTypeName<FTrigger>();
 				OnPlayTrigger.Metadata.DisplayName = LOCTEXT("OnPlay", "On Play");
@@ -198,7 +170,7 @@ namespace Metasound
 
 				// Outputs
 				FMetasoundFrontendClassVertex OnFinished;
-				OnFinished.Name = GetIsFinishedOutputName();
+				OnFinished.Name = ArchetypesPrivate::GetIsFinishedOutputName();
 
 				OnFinished.TypeName = GetMetasoundDataTypeName<FTrigger>();
 				OnFinished.Metadata.DisplayName = LOCTEXT("OnFinished", "On Finished");
@@ -218,7 +190,7 @@ namespace Metasound
 
 				// Environment
 				FMetasoundFrontendClassEnvironmentVariable AudioDeviceID;
-				AudioDeviceID.Name = GetAudioDeviceIDVariableName();
+				AudioDeviceID.Name = ArchetypesPrivate::GetAudioDeviceIDVariableName();
 				AudioDeviceID.Metadata.Description = LOCTEXT("AudioDeviceIDToolTip", "Audio Device ID");
 
 				Interface.Environment.Add(AudioDeviceID);
@@ -227,62 +199,20 @@ namespace Metasound
 			}
 		}
 
-		namespace MetasoundSourceMonoV1_1
+		namespace MetasoundOutputFormatMonoV1_1
 		{
-			FMetasoundFrontendVersion GetVersion()
+			const FMetasoundFrontendVersion& GetVersion()
 			{
-				static const FName VersionName = FName("MonoSource");
-				return FMetasoundFrontendVersion{ VersionName, FMetasoundFrontendVersionNumber{1, 1} };
-			}
-
-			const FVertexName& GetOnPlayInputName()
-			{
-				static const FVertexName TriggerInputName = TEXT("On Play");
-				return TriggerInputName;
+				static const FMetasoundFrontendVersion Version { "MonoSource", FMetasoundFrontendVersionNumber { 1, 1} };
+				return Version;
 			}
 
 			const FVertexName& GetAudioOutputName()
 			{
-				static const FVertexName AudioOutputName = TEXT("Audio:0");
+				static const FVertexName AudioOutputName = "Audio:0";
 				return AudioOutputName;
 			}
 
-			const FVertexName& GetIsFinishedOutputName()
-			{
-				static const FVertexName OnFinishedOutputName = TEXT("On Finished");
-				return OnFinishedOutputName;
-			}
-
-			const FVertexName& GetAudioDeviceIDVariableName()
-			{
-				static const FVertexName AudioDeviceIDVarName = TEXT("AudioDeviceID");
-				return AudioDeviceIDVarName;
-			}
-
-			const FVertexName& GetSoundUniqueIdName()
-			{
-				static const FVertexName SoundUniqueIdVarName = TEXT("SoundUniqueId");
-				return SoundUniqueIdVarName;
-			}
-
-			const FVertexName& GetIsPreviewSoundName()
-			{
-				static const FVertexName SoundIsPreviewSoundName = TEXT("IsPreviewSound");
-				return SoundIsPreviewSoundName;
-			}
-
-			const FVertexName& GetGraphName()
-			{
-				static const FVertexName SoundGraphName = TEXT("GraphName");
-				return SoundGraphName;
-			}
-
-			const FVertexName& GetInstanceIDName()
-			{
-				static const FVertexName& Name = FMetaSoundParameterTransmitter::GetInstanceIDEnvironmentVariableName();
-				return Name;
-			}
-
 			FMetasoundFrontendClassVertex GetClassAudioOutput()
 			{
 				FMetasoundFrontendClassVertex GeneratedAudio;
@@ -299,10 +229,10 @@ namespace Metasound
 			{
 				FMetasoundFrontendInterface Interface;
 				Interface.Version = GetVersion();
-				
+
 				// Inputs
 				FMetasoundFrontendClassVertex OnPlayTrigger;
-				OnPlayTrigger.Name = GetOnPlayInputName();
+				OnPlayTrigger.Name = ArchetypesPrivate::GetOnPlayInputName();
 
 				OnPlayTrigger.TypeName = GetMetasoundDataTypeName<FTrigger>();
 				OnPlayTrigger.Metadata.DisplayName = LOCTEXT("OnPlay", "On Play");
@@ -313,7 +243,7 @@ namespace Metasound
 
 				// Outputs
 				FMetasoundFrontendClassVertex OnFinished;
-				OnFinished.Name = GetIsFinishedOutputName();
+				OnFinished.Name = ArchetypesPrivate::GetIsFinishedOutputName();
 
 				OnFinished.TypeName = GetMetasoundDataTypeName<FTrigger>();
 				OnFinished.Metadata.DisplayName = LOCTEXT("OnFinished", "On Finished");
@@ -327,7 +257,7 @@ namespace Metasound
 
 				// Environment
 				FMetasoundFrontendClassEnvironmentVariable AudioDeviceID;
-				AudioDeviceID.Name = GetAudioDeviceIDVariableName();
+				AudioDeviceID.Name = ArchetypesPrivate::GetAudioDeviceIDVariableName();
 				AudioDeviceID.Metadata.Description = LOCTEXT("AudioDeviceIDToolTip", "Audio Device ID");
 
 				Interface.Environment.Add(AudioDeviceID);
@@ -335,11 +265,7 @@ namespace Metasound
 				return Interface;
 			}
 
-			// Update from a MetasoundSourceMonoV1_0 to MetasoundSourceMonoV1_1
-			class FUpdateInterface : public Frontend::IDocumentTransform
-			{
-			public:
-				virtual bool Transform(Frontend::FDocumentHandle InDocument) const override
+			bool FUpdateInterface::Transform(Frontend::FDocumentHandle InDocument) const
 				{
 					// Swap FMonoAudioFormat output node to an FAudioBuffer output node.
 					using namespace Frontend;
@@ -350,10 +276,10 @@ namespace Metasound
 						return false;
 					}
 
-					InDocument->RemoveInterfaceVersion(MetasoundSourceMonoV1_0::GetInterface().Version);
-					InDocument->AddInterfaceVersion(MetasoundSourceMonoV1_1::GetInterface().Version);
+					InDocument->RemoveInterfaceVersion(MetasoundOutputFormatMonoV1_0::GetInterface().Version);
+					InDocument->AddInterfaceVersion(MetasoundOutputFormatMonoV1_1::GetInterface().Version);
 
-					FNodeHandle MonoFormatOutput = Graph->GetOutputNodeWithName(MetasoundSourceMonoV1_0::GetAudioOutputName());
+					FNodeHandle MonoFormatOutput = Graph->GetOutputNodeWithName(MetasoundOutputFormatMonoV1_0::GetAudioOutputName());
 					FVector2D MonoFormatLocation;
 
 					FOutputHandle OutputToReconnect = IOutputController::GetInvalidHandle();
@@ -372,7 +298,7 @@ namespace Metasound
 							OutputToReconnect = Inputs[0]->GetConnectedOutput();
 						}
 
-						Graph->RemoveOutputVertex(MetasoundSourceMonoV1_0::GetAudioOutputName());
+						Graph->RemoveOutputVertex(MetasoundOutputFormatMonoV1_0::GetAudioOutputName());
 					}
 
 					// Create output
@@ -396,71 +322,34 @@ namespace Metasound
 
 					return true;
 				}
-			};
 		}
 
-		namespace MetasoundSourceStereoV1_1
+		namespace MetasoundOutputFormatStereoV1_1
 		{
-			FMetasoundFrontendVersion GetVersion()
+			const FMetasoundFrontendVersion& GetVersion()
 			{
-				static const FName VersionName = FName("StereoSource");
-				return FMetasoundFrontendVersion{ VersionName, FMetasoundFrontendVersionNumber{1, 1} };
+				static const FMetasoundFrontendVersion Version { "StereoSource", FMetasoundFrontendVersionNumber { 1, 1 } };
+				return Version;
 			}
 
 			const FVertexName& GetOnPlayInputName()
 			{
-				static const FVertexName TriggerInputName = TEXT("On Play");
+				static const FVertexName TriggerInputName = "On Play";
 				return TriggerInputName;
 			}
 
 			const FVertexName& GetLeftAudioOutputName()
 			{
-				static const FVertexName AudioOutputName = TEXT("Audio:0");
+				static const FVertexName AudioOutputName = "Audio:0";
 				return AudioOutputName;
 			}
 
 			const FVertexName& GetRightAudioOutputName()
 			{
-				static const FVertexName AudioOutputName = TEXT("Audio:1");
+				static const FVertexName AudioOutputName = "Audio:1";
 				return AudioOutputName;
 			}
 
-			const FVertexName& GetIsFinishedOutputName()
-			{
-				static const FVertexName OnFinishedOutputName = TEXT("On Finished");
-				return OnFinishedOutputName;
-			}
-
-			const FVertexName& GetAudioDeviceIDVariableName()
-			{
-				static const FVertexName AudioDeviceIDVarName = TEXT("AudioDeviceID");
-				return AudioDeviceIDVarName;
-			}
-
-			const FVertexName& GetSoundUniqueIdName()
-			{
-				static const FVertexName SoundUniqueIdVarName = TEXT("SoundUniqueId");
-				return SoundUniqueIdVarName;
-			}
-
-			const FVertexName& GetIsPreviewSoundName()
-			{
-				static const FVertexName SoundIsPreviewSoundName = TEXT("IsPreviewSound");
-				return SoundIsPreviewSoundName;
-			}
-
-			const FVertexName& GetGraphName()
-			{
-				static const FVertexName SoundGraphName = TEXT("GraphName");
-				return SoundGraphName;
-			}
-
-			const FVertexName& GetInstanceIDName()
-			{
-				static const FVertexName& Name = FMetaSoundParameterTransmitter::GetInstanceIDEnvironmentVariableName();
-				return Name;
-			}
-
 			FMetasoundFrontendClassVertex GetClassLeftAudioOutput()
 			{
 				FMetasoundFrontendClassVertex GeneratedLeftAudio;
@@ -492,7 +381,7 @@ namespace Metasound
 				
 				// Inputs
 				FMetasoundFrontendClassVertex OnPlayTrigger;
-				OnPlayTrigger.Name = GetOnPlayInputName();
+				OnPlayTrigger.Name = ArchetypesPrivate::GetOnPlayInputName();
 
 				OnPlayTrigger.TypeName = GetMetasoundDataTypeName<FTrigger>();
 				OnPlayTrigger.Metadata.DisplayName = LOCTEXT("OnPlay", "On Play");
@@ -503,7 +392,7 @@ namespace Metasound
 
 				// Outputs
 				FMetasoundFrontendClassVertex OnFinished;
-				OnFinished.Name = GetIsFinishedOutputName();
+				OnFinished.Name = ArchetypesPrivate::GetIsFinishedOutputName();
 
 				OnFinished.TypeName = GetMetasoundDataTypeName<FTrigger>();
 				OnFinished.Metadata.DisplayName = LOCTEXT("OnFinished", "On Finished");
@@ -520,18 +409,14 @@ namespace Metasound
 
 				// Environment
 				FMetasoundFrontendClassEnvironmentVariable AudioDeviceID;
-				AudioDeviceID.Name = GetAudioDeviceIDVariableName();
+				AudioDeviceID.Name = ArchetypesPrivate::GetAudioDeviceIDVariableName();
 				AudioDeviceID.Metadata.Description = LOCTEXT("AudioDeviceIDToolTip", "Audio Device ID");
 				Interface.Environment.Add(AudioDeviceID);
 
 				return Interface;
 			};
 
-			// Update from MetasoundSourceStereoV1_0 to MetasoundSourceStereoV1_1
-			class FUpdateInterface : public Frontend::IDocumentTransform
-			{
-			public:
-				virtual bool Transform(Frontend::FDocumentHandle InDocument) const override
+			bool FUpdateInterface::Transform(Frontend::FDocumentHandle InDocument) const
 				{
 					using namespace Frontend;
 
@@ -542,10 +427,10 @@ namespace Metasound
 						return false;
 					}
 
-					InDocument->RemoveInterfaceVersion(MetasoundSourceStereoV1_0::GetInterface().Version);
-					InDocument->AddInterfaceVersion(MetasoundSourceStereoV1_1::GetInterface().Version);
+					InDocument->RemoveInterfaceVersion(MetasoundOutputFormatStereoV1_0::GetInterface().Version);
+					InDocument->AddInterfaceVersion(MetasoundOutputFormatStereoV1_1::GetInterface().Version);
 
-					FNodeHandle StereoFormatOutput = Graph->GetOutputNodeWithName(MetasoundSourceStereoV1_0::GetAudioOutputName());
+					FNodeHandle StereoFormatOutput = Graph->GetOutputNodeWithName(MetasoundOutputFormatStereoV1_0::GetAudioOutputName());
 					FOutputHandle LeftOutputToReconnect = IOutputController::GetInvalidHandle();
 					FOutputHandle RightOutputToReconnect = IOutputController::GetInvalidHandle();
 
@@ -559,13 +444,13 @@ namespace Metasound
 							StereoFormatLocation = Location.Value;
 						}
 
-						FInputHandle LeftInput = StereoFormatOutput->GetInputWithVertexName(TEXT("Left"));
+						FInputHandle LeftInput = StereoFormatOutput->GetInputWithVertexName("Left");
 						LeftOutputToReconnect = LeftInput->GetConnectedOutput();
 
-						FInputHandle RightInput = StereoFormatOutput->GetInputWithVertexName(TEXT("Right"));
+						FInputHandle RightInput = StereoFormatOutput->GetInputWithVertexName("Right");
 						RightOutputToReconnect = RightInput->GetConnectedOutput();
 
-						Graph->RemoveOutputVertex(MetasoundSourceStereoV1_0::GetAudioOutputName());
+						Graph->RemoveOutputVertex(MetasoundOutputFormatStereoV1_0::GetAudioOutputName());
 					}
 
 					FNodeHandle LeftBufferOutput = Graph->AddOutputVertex(GetClassLeftAudioOutput());
@@ -604,108 +489,82 @@ namespace Metasound
 
 					return true;
 				}
-			};
 		}
 
-		namespace MetasoundSource
+		namespace MetasoundOutputFormatMonoV1_2
 		{
-			namespace CurrentMonoVersion = MetasoundSourceMonoV1_1;
-			namespace CurrentStereoVersion = MetasoundSourceStereoV1_1;
-
-			const FVertexName& GetOnPlayInputName()
+			const FMetasoundFrontendVersion& GetVersion()
 			{
-				check(CurrentMonoVersion::GetOnPlayInputName() == CurrentStereoVersion::GetOnPlayInputName());
-				return CurrentMonoVersion::GetOnPlayInputName();
+				const static FMetasoundFrontendVersion Version = { "MonoSource", FMetasoundFrontendVersionNumber { 1, 2 } };
+				return Version;
 			}
 
-			const FVertexName& GetIsFinishedOutputName()
-			{
-				check(CurrentMonoVersion::GetIsFinishedOutputName() == CurrentStereoVersion::GetIsFinishedOutputName());
-				return CurrentMonoVersion::GetIsFinishedOutputName();
-			}
-
-			const FVertexName& GetAudioDeviceIDVariableName()
+			FMetasoundFrontendInterface GetInterface()
 			{
-				check(CurrentMonoVersion::GetAudioDeviceIDVariableName() == CurrentStereoVersion::GetAudioDeviceIDVariableName());
-				return CurrentMonoVersion::GetAudioDeviceIDVariableName();
+				FMetasoundFrontendInterface Interface;
+				Interface.Version = GetVersion();
+				return Interface;
 			}
 
-			const FVertexName& GetSoundUniqueIdName()
+			bool FUpdateInterface::Transform(Frontend::FDocumentHandle InDocument) const
 			{
-				check(CurrentMonoVersion::GetSoundUniqueIdName() == CurrentStereoVersion::GetSoundUniqueIdName());
-				return CurrentMonoVersion::GetSoundUniqueIdName();
-			}
+				using namespace ArchetypesPrivate;
+				using namespace Frontend;
 
-			const FVertexName& GetIsPreviewSoundName()
-			{
-				check(CurrentMonoVersion::GetIsPreviewSoundName() == CurrentStereoVersion::GetIsPreviewSoundName());
-				return CurrentMonoVersion::GetIsPreviewSoundName();
-			}
+				const TArray<FMetasoundFrontendVersion> InterfacesToRemove
+				{
+					MetasoundOutputFormatMonoV1_1::GetVersion()
+				};
 
-			const FVertexName& GetInstanceIDName()
-			{
-				check(CurrentMonoVersion::GetInstanceIDName() == CurrentStereoVersion::GetInstanceIDName());
-				return CurrentMonoVersion::GetInstanceIDName();
-			}
+				const TArray<FMetasoundFrontendVersion> InterfacesToAdd
+				{
+					SourceInterface::GetVersion(),
+					OutputFormatMonoInterface::GetVersion()
+				};
 
-			const FVertexName& GetGraphName()
-			{
-				check(CurrentMonoVersion::GetGraphName() == CurrentStereoVersion::GetGraphName());
-				return CurrentMonoVersion::GetGraphName();
+				FModifyRootGraphInterfaces InterfaceTransform(InterfacesToRemove, InterfacesToAdd);
+				InterfaceTransform.SetNamePairingFunction(MatchMemberNamesIgnoreSpaces);
+				return InterfaceTransform.Transform(InDocument);
 			}
 		}
 
-		namespace MetasoundSourceMono
+		namespace MetasoundOutputFormatStereoV1_2
 		{
-			namespace Current = MetasoundSourceMonoV1_1;
-
-			FMetasoundFrontendVersion GetVersion()
-			{
-				return Current::GetVersion();
-			}
-
-			const FVertexName& GetAudioOutputName()
+			const FMetasoundFrontendVersion& GetVersion()
 			{
-				return Current::GetAudioOutputName();
+				const static FMetasoundFrontendVersion Version = { "StereoSource", FMetasoundFrontendVersionNumber { 1, 2 } };
+				return Version;
 			}
-		}
 
-		namespace MetasoundSourceStereo
-		{
-			namespace Current = MetasoundSourceStereoV1_1;
-
-			FMetasoundFrontendVersion GetVersion()
+			FMetasoundFrontendInterface GetInterface()
 			{
-				return Current::GetVersion();
+				FMetasoundFrontendInterface Interface;
+				Interface.Version = GetVersion();
+				return Interface;
 			}
 
-			const FVertexName& GetLeftAudioOutputName()
+			bool FUpdateInterface::Transform(Frontend::FDocumentHandle InDocument) const
 			{
-				return Current::GetLeftAudioOutputName();
-			}
+				using namespace ArchetypesPrivate;
+				using namespace Frontend;
 
-			const FVertexName& GetRightAudioOutputName()
-			{
-				return Current::GetRightAudioOutputName();
-			}
-		}
+				const TArray<FMetasoundFrontendVersion> InterfacesToRemove
+				{
+					MetasoundOutputFormatStereoV1_1::GetVersion()
+				};
 
-		void RegisterInternalInterfaces()
-		{
-			{
-				constexpr bool bIsDefault = false;
-				RegisterInterface<UMetaSoundSource>(MetasoundSourceStereoV1_1::GetInterface(), MakeUnique<MetasoundSourceStereoV1_1::FUpdateInterface>(), bIsDefault, IDataReference::RouterName);
-				RegisterInterface<UMetaSoundSource>(MetasoundSourceStereoV1_0::GetInterface(), nullptr, bIsDefault, IDataReference::RouterName);
-			}
+				const TArray<FMetasoundFrontendVersion> InterfacesToAdd
+				{
+					SourceInterface::GetVersion(),
+					OutputFormatStereoInterface::GetVersion()
+				};
 
-			{
-				constexpr bool bIsDefault = true;
-				RegisterInterface<UMetaSound>(MetasoundV1_0::GetInterface(), nullptr, bIsDefault, IDataReference::RouterName);
-				RegisterInterface<UMetaSoundSource>(MetasoundSourceMonoV1_0::GetInterface(), nullptr, bIsDefault, IDataReference::RouterName);
-				RegisterInterface<UMetaSoundSource>(MetasoundSourceMonoV1_1::GetInterface(), MakeUnique<MetasoundSourceMonoV1_1::FUpdateInterface>(), bIsDefault, IDataReference::RouterName);
+				FModifyRootGraphInterfaces InterfaceTransform(InterfacesToRemove, InterfacesToAdd);
+				InterfaceTransform.SetNamePairingFunction(MatchMemberNamesIgnoreSpaces);
+				return InterfaceTransform.Transform(InDocument);
 			}
 		}
-	}
-}
+	} // namespace Engine
+} // namespace Metasound
 
 #undef LOCTEXT_NAMESPACE // MetasoundEngineArchetypes
diff --git a/Engine/Plugins/Runtime/Metasound/Source/MetasoundEngine/Private/MetasoundEngineArchetypes.h b/Engine/Plugins/Runtime/Metasound/Source/MetasoundEngine/Private/MetasoundEngineArchetypes.h
--- a/Engine/Plugins/Runtime/Metasound/Source/MetasoundEngine/Private/MetasoundEngineArchetypes.h
+++ b/Engine/Plugins/Runtime/Metasound/Source/MetasoundEngine/Private/MetasoundEngineArchetypes.h
@@ -1,103 +1,96 @@
 // Copyright Epic Games, Inc. All Rights Reserved.
-
 #pragma once
 
-#include "CoreMinimal.h"
 #include "MetasoundFrontendDocument.h"
+#include "MetasoundFrontendTransform.h"
 #include "MetasoundVertex.h"
 
+// DEPRECATED: Archetypes that preceded the interface system.  Only remain
+// for purposes of upgrading content generated prior to 5.0 release.
+
 namespace Metasound
 {
 	namespace Engine
 	{
-		// Base metasound without any required inputs or outputs
+		// Base interface without any required inputs or outputs (TODO: Version & Remove)
 		namespace MetasoundV1_0
 		{
-			FMetasoundFrontendVersion GetVersion();
+			const FMetasoundFrontendVersion& GetVersion();
+			FMetasoundFrontendInterface GetInterface();
 		}
 
-		// V1.0 of a Metasound Mono Source. Uses FMonoAudioFormat as output.
-		namespace MetasoundSourceMonoV1_0
+		// V1.0 of a Metasound Mono output format. Uses FMonoAudioFormat as output.
+		namespace MetasoundOutputFormatMonoV1_0
 		{
-			FMetasoundFrontendVersion GetVersion();
-			const FVertexName& GetOnPlayInputName();
+			const FMetasoundFrontendVersion& GetVersion();
 			const FVertexName& GetAudioOutputName();
-			const FVertexName& GetIsFinishedOutputName();
-			const FVertexName& GetAudioDeviceIDVariableName();
-			const FVertexName& GetSoundUniqueIdName();
-			const FVertexName& GetIsPreviewSoundName();
+			FMetasoundFrontendInterface GetInterface();
 		}
 
-		// V1.0 of a Metasound Stereo Source. Uses FStereoAudioFormat as output.
-		namespace MetasoundSourceStereoV1_0
+		// V1.0 of a Metasound Stereo output format. Uses FStereoAudioFormat as output.
+		namespace MetasoundOutputFormatStereoV1_0
 		{
-			FMetasoundFrontendVersion GetVersion();
-			const FVertexName& GetOnPlayInputName();
+			const FMetasoundFrontendVersion& GetVersion();
 			const FVertexName& GetAudioOutputName();
-			const FVertexName& GetIsFinishedOutputName();
-			const FVertexName& GetAudioDeviceIDVariableName();
-			const FVertexName& GetSoundUniqueIdName();
-			const FVertexName& GetIsPreviewSoundName();
-			const FVertexName& GetGraphName();
+			FMetasoundFrontendInterface GetInterface();
 		}
 
-		// V1.1 of a Metasound Mono source. Uses FAudioBuffer as output.
-		namespace MetasoundSourceMonoV1_1
+		// V1.1 of a Metasound Mono output format. Uses FAudioBuffer as output.
+		namespace MetasoundOutputFormatMonoV1_1
 		{
-			FMetasoundFrontendVersion GetVersion();
-			const FVertexName& GetOnPlayInputName();
+			const FMetasoundFrontendVersion& GetVersion();
 			const FVertexName& GetAudioOutputName();
-			const FVertexName& GetIsFinishedOutputName();
-			const FVertexName& GetAudioDeviceIDVariableName();
-			const FVertexName& GetSoundUniqueIdName();
-			const FVertexName& GetIsPreviewSoundName();
-			const FVertexName& GetInstanceIDName();
-			const FVertexName& GetGraphName();
-		}
+			FMetasoundFrontendInterface GetInterface();
 
-		// V1.1 of a Metasound Stereo source. Uses two FAudioBuffers as outputs.
-		namespace MetasoundSourceStereoV1_1
-		{
-			FMetasoundFrontendVersion GetVersion();
-			const FVertexName& GetOnPlayInputName();
-			const FVertexName& GetLeftAudioOutputName();
-			const FVertexName& GetRightAudioOutputName();
-			const FVertexName& GetIsFinishedOutputName();
-			const FVertexName& GetAudioDeviceIDVariableName();
-			const FVertexName& GetSoundUniqueIdName();
-			const FVertexName& GetIsPreviewSoundName();
-			const FVertexName& GetInstanceIDName();
-			const FVertexName& GetGraphName();
+			class FUpdateInterface : public Frontend::IDocumentTransform
+			{
+			public:
+				virtual bool Transform(Frontend::FDocumentHandle InDocument) const override;
+			};
 		}
 
-		// Current version of Metasound Source.
-		namespace MetasoundSource
+		// V1.2 of a Metasound Mono output format.
+		// Deprecate path from MetasoundOutputFormatStereoV1_1 to SourceInterface & respective OutputFormatMonoInterface
+		namespace MetasoundOutputFormatMonoV1_2
 		{
-			const FVertexName& GetOnPlayInputName();
-			const FVertexName& GetIsFinishedOutputName();
-			const FVertexName& GetAudioDeviceIDVariableName();
-			const FVertexName& GetSoundUniqueIdName();
-			const FVertexName& GetIsPreviewSoundName();
-			const FVertexName& GetInstanceIDName();
-			const FVertexName& GetGraphName();
-		}
+			const FMetasoundFrontendVersion& GetVersion();
+			FMetasoundFrontendInterface GetInterface();
 
-		// Current version of MetasoundSourceMono
-		namespace MetasoundSourceMono
-		{
-			FMetasoundFrontendVersion GetVersion();
-			const FVertexName& GetAudioOutputName();
+			class FUpdateInterface : public Frontend::IDocumentTransform
+			{
+			public:
+				virtual bool Transform(Frontend::FDocumentHandle InDocument) const override;
+			};
 		}
 
-		// Current version of MetasoundSourceStereo
-		namespace MetasoundSourceStereo
+		// V1.1 of a Metasound Stereo output format. Uses two FAudioBuffers as outputs.
+		namespace MetasoundOutputFormatStereoV1_1
 		{
-			FMetasoundFrontendVersion GetVersion();
+			const FMetasoundFrontendVersion& GetVersion();
 			const FVertexName& GetLeftAudioOutputName();
 			const FVertexName& GetRightAudioOutputName();
+			FMetasoundFrontendInterface GetInterface();
+
+			class FUpdateInterface : public Frontend::IDocumentTransform
+			{
+			public:
+				virtual bool Transform(Frontend::FDocumentHandle InDocument) const override;
+			};
 		}
 
-		// Register all interfaces defined in MetaSoundEngine.
-		void RegisterInternalInterfaces();
+		// V1.2 of a Metasound Stereo output format. Uses two FAudioBuffers as outputs.
+		// Deprecate path from MetasoundOutputFormatStereoV1_1 to SourceInterface & respective OutputFormatStereoInterface
+		namespace MetasoundOutputFormatStereoV1_2
+		{
+			const FMetasoundFrontendVersion& GetVersion();
+			FMetasoundFrontendInterface GetInterface();
+
+			// Update from MetasoundOutputFormatStereoV1_0 to MetasoundOutputFormatStereoV1_1
+			class FUpdateInterface : public Frontend::IDocumentTransform
+			{
+			public:
+				virtual bool Transform(Frontend::FDocumentHandle InDocument) const override;
+			};
+		}
 	}
 }
diff --git a/Engine/Plugins/Runtime/Metasound/Source/MetasoundEngine/Private/MetasoundEngineModule.cpp b/Engine/Plugins/Runtime/Metasound/Source/MetasoundEngine/Private/MetasoundEngineModule.cpp
--- a/Engine/Plugins/Runtime/Metasound/Source/MetasoundEngine/Private/MetasoundEngineModule.cpp
+++ b/Engine/Plugins/Runtime/Metasound/Source/MetasoundEngine/Private/MetasoundEngineModule.cpp
@@ -26,8 +26,7 @@ class FMetasoundEngineModule : public IMetasoundEngineModule
 		FModuleManager::Get().LoadModuleChecked("AudioCodecEngine");
 
 		// Register interfaces
-		Metasound::Engine::RegisterInternalInterfaces();
-		Metasound::Engine::RegisterExternalInterfaces();
+		Metasound::Engine::RegisterInterfaces();
 
 		// Flush node registration queue
 		FMetasoundFrontendRegistryContainer::Get()->RegisterPendingNodes();
diff --git a/Engine/Plugins/Runtime/Metasound/Source/MetasoundEngine/Private/MetasoundInterface.cpp b/Engine/Plugins/Runtime/Metasound/Source/MetasoundEngine/Private/MetasoundInterface.cpp
--- a/Engine/Plugins/Runtime/Metasound/Source/MetasoundEngine/Private/MetasoundInterface.cpp
+++ b/Engine/Plugins/Runtime/Metasound/Source/MetasoundEngine/Private/MetasoundInterface.cpp
@@ -1,87 +1,223 @@
 // Copyright Epic Games, Inc. All Rights Reserved.
 #include "MetasoundInterface.h"
 
+#include "Algo/Transform.h"
+#include "AudioParameterControllerInterface.h"
+#include "IAudioParameterInterfaceRegistry.h"
 #include "IAudioParameterTransmitter.h"
+#include "MetasoundEngineArchetypes.h"
+#include "MetasoundFrontendDataTypeRegistry.h"
+#include "MetasoundFrontendDocument.h"
+#include "MetasoundFrontendTransform.h"
+#include "MetasoundLog.h"
+#include "MetasoundOutputFormatInterfaces.h"
+#include "MetasoundParameterTransmitter.h"
+#include "MetasoundSource.h"
+#include "MetasoundSourceInterface.h"
+#include "MetasoundUObjectRegistry.h"
+#include "Templates/SharedPointer.h"
+#include "Templates/UniquePtr.h"
+#include "UObject/Class.h"
+#include "UObject/NoExportTypes.h"
 
 
 namespace Metasound
 {
 	namespace Engine
 	{
-		void RegisterExternalInterfaces()
+		struct FInterfaceRegistryOptions
 		{
-			using namespace Audio;
+			bool bIsDefault = false;
+			FName InputSystemName;
+			FName UClassName;
+		};
 
-			auto RegisterExternalInterface = [](FGeneratorInterfacePtr Interface)
+		// Entry for registered interface.
+		class FInterfaceRegistryEntry : public Frontend::IInterfaceRegistryEntry
+		{
+		public:
+			FInterfaceRegistryEntry(
+				const FMetasoundFrontendInterface& InInterface,
+				TUniquePtr<Frontend::IDocumentTransform>&& InUpdateTransform,
+				FInterfaceRegistryOptions&& InOptions
+			)
+				: Interface(InInterface)
+				, UpdateTransform(MoveTemp(InUpdateTransform))
+				, Options(MoveTemp(InOptions))
+			{
+			}
+
+			virtual bool UClassIsSupported(FName InUClassName) const override
 			{
-				auto ResolveMemberDataType = [](FName DataType, EAudioParameterType ParamType)
+				if (Options.UClassName.IsNone())
 				{
-					if (!DataType.IsNone())
-					{
-						const bool bIsRegisteredType = Frontend::IDataTypeRegistry::Get().IsRegistered(DataType);
-						if (ensureAlwaysMsgf(bIsRegisteredType, TEXT("Attempting to register Interface member with unregistered DataType '%s'."), *DataType.ToString()))
-						{
-							return DataType;
-						}
-					}
+					return true;
+				}
 
-					return Frontend::ConvertParameterToDataType(ParamType);
-				};
+				// TODO: Support child asset class types.
+				return Options.UClassName == InUClassName;
+			}
 
-				FMetasoundFrontendInterface FrontendInterface;
-				FrontendInterface.Version = { Interface->Name, FMetasoundFrontendVersionNumber { Interface->Version.Major, Interface->Version.Minor } };
+			virtual bool IsDefault() const override
+			{
+				return Options.bIsDefault;
+			}
 
-				Algo::Transform(Interface->Inputs, FrontendInterface.Inputs, [&](const FGeneratorInterface::FInput& Input)
-				{
-					FMetasoundFrontendClassInput ClassInput;
-					ClassInput.Name = IGeneratorInterfaceRegistry::GetMemberFullName(Interface->Name, Input.InitValue.ParamName);
-					ClassInput.DefaultLiteral = FMetasoundFrontendLiteral(Input.InitValue);
-					ClassInput.TypeName = ResolveMemberDataType(Input.DataType, Input.InitValue.ParamType);
-					ClassInput.Metadata.DisplayName = Input.DisplayName;
-					ClassInput.Metadata.Description = Input.Description;
-					ClassInput.VertexID = FGuid::NewGuid();
-
-					return ClassInput;
-				});
-
-				Algo::Transform(Interface->Outputs, FrontendInterface.Outputs, [&](const FGeneratorInterface::FOutput& Output)
-				{
-					FMetasoundFrontendClassOutput ClassOutput;
-					ClassOutput.Name = IGeneratorInterfaceRegistry::GetMemberFullName(Interface->Name, Output.ParamName);
-					ClassOutput.TypeName = ResolveMemberDataType(Output.DataType, Output.ParamType);
-					ClassOutput.Metadata.DisplayName = Output.DisplayName;
-					ClassOutput.Metadata.Description = Output.Description;
-					ClassOutput.VertexID = FGuid::NewGuid();
+			virtual FName GetRouterName() const override
+			{
+				return Options.InputSystemName;
+			}
 
-					return ClassOutput;
-				});
+			virtual const FMetasoundFrontendInterface& GetInterface() const override
+			{
+				return Interface;
+			}
 
-				Algo::Transform(Interface->Environment, FrontendInterface.Environment, [&](const FGeneratorInterface::FEnvironmentVariable& Environment)
+			virtual bool UpdateRootGraphInterface(Frontend::FDocumentHandle InDocument) const override
+			{
+				if (UpdateTransform.IsValid())
 				{
-					FMetasoundFrontendClassEnvironmentVariable EnvironmentVariable;
-					EnvironmentVariable.Name = IGeneratorInterfaceRegistry::GetMemberFullName(Interface->Name, Environment.ParamName);
-					EnvironmentVariable.TypeName = ResolveMemberDataType(Environment.DataType, Environment.ParamType);
-					EnvironmentVariable.Metadata.DisplayName = Environment.DisplayName;
-					EnvironmentVariable.Metadata.Description = Environment.Description;
-
-					return EnvironmentVariable;
-				});
-
-				// For now, only support interfaces for UMetaSoundSource type
-				const UClass* SourceClass = UMetaSoundSource::StaticClass();
-				check(SourceClass);
-				if (ensureAlways(SourceClass->IsChildOf(Interface->Type)))
+					return UpdateTransform->Transform(InDocument);
+				}
+				return false;
+			}
+
+		private:
+			FMetasoundFrontendInterface Interface;
+			TUniquePtr<Frontend::IDocumentTransform> UpdateTransform;
+			FInterfaceRegistryOptions Options;
+		};
+
+		template <typename UClassType>
+		void RegisterInterface(const FMetasoundFrontendInterface& InInterface, TUniquePtr<Frontend::IDocumentTransform>&& InUpdateTransform, bool bInIsDefault, FName InRouterName)
+		{
+			using namespace Frontend;
+
+			FInterfaceRegistryOptions Options
+			{
+				bInIsDefault,
+				InRouterName,
+				UClassType::StaticClass()->GetFName()
+			};
+
+			IMetasoundUObjectRegistry::Get().RegisterUClassInterface(MakeUnique<TMetasoundUObjectRegistryEntry<UClassType>>(InInterface.Version));
+			IInterfaceRegistry::Get().RegisterInterface(MakeUnique<FInterfaceRegistryEntry>(InInterface, MoveTemp(InUpdateTransform), MoveTemp(Options)));
+		}
+
+		void RegisterInterface(Audio::FParameterInterfacePtr Interface, TUniquePtr<Frontend::IDocumentTransform>&& InUpdateTransform, bool bInIsDefault, FName InRouterName)
+		{
+			using namespace Frontend;
+
+			auto ResolveMemberDataType = [](FName DataType, EAudioParameterType ParamType)
+			{
+				if (!DataType.IsNone())
 				{
-					// For now, all external interfaces use the ParameterTransmitter and cannot be applied by default to new
-					// MetaSound objects. Could pipe either of these through in the future if required/desired.
-					constexpr bool bIsDefault = false;
-					RegisterInterface<UMetaSoundSource>(FrontendInterface, nullptr, bIsDefault, Audio::IParameterTransmitter::RouterName);
-					UE_LOG(LogMetaSound, Verbose, TEXT("Interface '%s' registered for asset type '%s'."), *Interface->Name.ToString(), *SourceClass->GetName());
+					const bool bIsRegisteredType = Frontend::IDataTypeRegistry::Get().IsRegistered(DataType);
+					if (ensureAlwaysMsgf(bIsRegisteredType, TEXT("Attempting to register Interface member with unregistered DataType '%s'."), *DataType.ToString()))
+					{
+						return DataType;
+					}
 				}
+
+				return ConvertParameterToDataType(ParamType);
+			};
+
+			FMetasoundFrontendInterface FrontendInterface;
+			FrontendInterface.Version = { Interface->GetName(), FMetasoundFrontendVersionNumber { Interface->GetVersion().Major, Interface->GetVersion().Minor } };
+
+			Algo::Transform(Interface->GetInputs(), FrontendInterface.Inputs, [&](const Audio::FParameterInterface::FInput& Input)
+			{
+				FMetasoundFrontendClassInput ClassInput;
+				ClassInput.Name = Input.InitValue.ParamName;
+				ClassInput.DefaultLiteral = FMetasoundFrontendLiteral(Input.InitValue);
+				ClassInput.TypeName = ResolveMemberDataType(Input.DataType, Input.InitValue.ParamType);
+				ClassInput.Metadata.DisplayName = Input.DisplayName;
+				ClassInput.Metadata.Description = Input.Description;
+				ClassInput.VertexID = FGuid::NewGuid();
+
+				return ClassInput;
+			});
+
+			Algo::Transform(Interface->GetOutputs(), FrontendInterface.Outputs, [&](const Audio::FParameterInterface::FOutput& Output)
+			{
+				FMetasoundFrontendClassOutput ClassOutput;
+				ClassOutput.Name = Output.ParamName;
+				ClassOutput.TypeName = ResolveMemberDataType(Output.DataType, Output.ParamType);
+				ClassOutput.Metadata.DisplayName = Output.DisplayName;
+				ClassOutput.Metadata.Description = Output.Description;
+				ClassOutput.VertexID = FGuid::NewGuid();
+
+				return ClassOutput;
+			});
+
+			Algo::Transform(Interface->GetEnvironment(), FrontendInterface.Environment, [&](const Audio::FParameterInterface::FEnvironmentVariable& Environment)
+			{
+				FMetasoundFrontendClassEnvironmentVariable EnvironmentVariable;
+				EnvironmentVariable.Name = Environment.ParamName;
+
+				// Disabled as it isn't used to infer type when getting/setting at a lower level.
+				// TODO: Either remove type info for environment variables all together or enforce type.
+				// EnvironmentVariable.TypeName = ResolveMemberDataType(Environment.DataType, Environment.ParamType);
+
+				EnvironmentVariable.Metadata.DisplayName = Environment.DisplayName;
+				EnvironmentVariable.Metadata.Description = Environment.Description;
+
+				return EnvironmentVariable;
+			});
+
+			const UClass* SourceClass = UMetaSoundSource::StaticClass();
+			check(SourceClass);
+			const bool bInterfaceSupportsSource = SourceClass->IsChildOf(&Interface->GetType());
+			if (ensureAlwaysMsgf(bInterfaceSupportsSource, TEXT("Interfaces defined using Audio::FParameterInterface currently only supported by UMetaSoundSource asset type.")))
+			{
+				RegisterInterface<UMetaSoundSource>(FrontendInterface, MoveTemp(InUpdateTransform), bInIsDefault, InRouterName);
+				UE_LOG(LogMetaSound, Verbose, TEXT("Interface '%s' registered for asset type '%s'."), *Interface->GetName().ToString(), *SourceClass->GetName());
+			}
+		}
+
+		void RegisterInterfaces()
+		{
+			using namespace Frontend;
+
+			// Register Default Internal Interfaces (Not managed directly by end-user & added by default when creating new MetaSound assets).
+			{
+				constexpr bool bIsDefault = true;
+				RegisterInterface<UMetaSound>(MetasoundV1_0::GetInterface(), nullptr, bIsDefault, IDataReference::RouterName);
+
+				RegisterInterface(SourceInterface::CreateInterface(*UMetaSoundSource::StaticClass()), nullptr, bIsDefault, IDataReference::RouterName);
+				RegisterInterface(OutputFormatMonoInterface::CreateInterface(*UMetaSoundSource::StaticClass()), nullptr, bIsDefault, IDataReference::RouterName);
+			}
+
+			// Register Non-Default Internal Interfaces (Not managed directly by end-user & not added by default when creating new MetaSound assets).
+			{
+				constexpr bool bIsDefault = false;
+
+				// Set default interface with unset version to use base UMetaSound class implementation (legacy requirement for 5.0 alpha).
+				RegisterInterface<UMetaSound>(FMetasoundFrontendInterface(), nullptr, bIsDefault, IDataReference::RouterName);
+
+				RegisterInterface<UMetaSoundSource>(MetasoundOutputFormatStereoV1_0::GetInterface(), nullptr, bIsDefault, IDataReference::RouterName);
+				RegisterInterface<UMetaSoundSource>(MetasoundOutputFormatStereoV1_1::GetInterface(), MakeUnique<MetasoundOutputFormatStereoV1_1::FUpdateInterface>(), bIsDefault, IDataReference::RouterName);
+				RegisterInterface<UMetaSoundSource>(MetasoundOutputFormatStereoV1_2::GetInterface(), MakeUnique<MetasoundOutputFormatStereoV1_2::FUpdateInterface>(), bIsDefault, IDataReference::RouterName);
+
+				RegisterInterface<UMetaSoundSource>(MetasoundOutputFormatMonoV1_0::GetInterface(), nullptr, bIsDefault, IDataReference::RouterName);
+				RegisterInterface<UMetaSoundSource>(MetasoundOutputFormatMonoV1_1::GetInterface(), MakeUnique<MetasoundOutputFormatMonoV1_1::FUpdateInterface>(), bIsDefault, IDataReference::RouterName);
+				RegisterInterface<UMetaSoundSource>(MetasoundOutputFormatMonoV1_2::GetInterface(), MakeUnique<MetasoundOutputFormatMonoV1_2::FUpdateInterface>(), bIsDefault, IDataReference::RouterName);
+
+				RegisterInterface(OutputFormatStereoInterface::CreateInterface(*UMetaSoundSource::StaticClass()), nullptr, bIsDefault, IDataReference::RouterName);
+			}
+
+			// Register External Interfaces (Interfaces defined externally & can be managed directly by end-user).
+			auto RegisterExternalInterface = [](Audio::FParameterInterfacePtr Interface)
+			{
+				// Currently, no externally defined interfaces can be added as default for protection against undesired
+				// interfaces automatically being added when creating a new MetaSound asset through the editor.
+				constexpr bool bIsDefault = false;
+				TUniquePtr<Frontend::IDocumentTransform> NullTransform;
+				RegisterInterface(Interface, MoveTemp(NullTransform), bIsDefault, Audio::IParameterTransmitter::RouterName);
 			};
 
-			IGeneratorInterfaceRegistry::Get().IterateInterfaces(RegisterExternalInterface);
-			IGeneratorInterfaceRegistry::Get().OnRegistration(MoveTemp(RegisterExternalInterface));
+			Audio::IAudioParameterInterfaceRegistry::Get().IterateInterfaces(RegisterExternalInterface);
+			Audio::IAudioParameterInterfaceRegistry::Get().OnRegistration(MoveTemp(RegisterExternalInterface));
 		}
 	} // namespace Engine
 } // namespace Metasound
diff --git a/Engine/Plugins/Runtime/Metasound/Source/MetasoundEngine/Private/MetasoundInterface.h b/Engine/Plugins/Runtime/Metasound/Source/MetasoundEngine/Private/MetasoundInterface.h
--- a/Engine/Plugins/Runtime/Metasound/Source/MetasoundEngine/Private/MetasoundInterface.h
+++ b/Engine/Plugins/Runtime/Metasound/Source/MetasoundEngine/Private/MetasoundInterface.h
@@ -1,102 +1,11 @@
 // Copyright Epic Games, Inc. All Rights Reserved.
 #pragma once
 
-#include "Algo/Transform.h"
-#include "AudioParameterInterface.h"
-#include "IAudioGeneratorInterfaceRegistry.h"
-#include "MetasoundFrontendDataTypeRegistry.h"
-#include "MetasoundFrontendDocument.h"
-#include "MetasoundFrontendTransform.h"
-#include "MetasoundLog.h"
-#include "MetasoundParameterTransmitter.h"
-#include "MetasoundSource.h"
-#include "MetasoundUObjectRegistry.h"
-#include "Templates/UniquePtr.h"
-#include "UObject/Class.h"
-#include "UObject/NoExportTypes.h"
-
 
 namespace Metasound
 {
 	namespace Engine
 	{
-		void RegisterExternalInterfaces();
-
-		struct FInterfaceRegistryOptions
-		{
-			bool bIsDefault = false;
-			FName InputSystemName;
-			FName UClassName;
-		};
-
-		// Entry for registered interface.
-		class FInterfaceRegistryEntry : public Frontend::IInterfaceRegistryEntry
-		{
-		public:
-			FInterfaceRegistryEntry(
-				const FMetasoundFrontendInterface& InInterface,
-				TUniquePtr<Frontend::IDocumentTransform>&& InUpdateTransform,
-				FInterfaceRegistryOptions&& InOptions
-			)
-				: Interface(InInterface)
-				, UpdateTransform(MoveTemp(InUpdateTransform))
-				, Options(MoveTemp(InOptions))
-			{
-			}
-
-			virtual bool UClassIsSupported(FName InUClassName) const override
-			{
-				if (Options.UClassName.IsNone())
-				{
-					return true;
-				}
-
-				// TODO: Support child asset class types.
-				return Options.UClassName == InUClassName;
-			}
-
-			virtual bool IsDefault() const override
-			{
-				return Options.bIsDefault;
-			}
-
-			virtual FName GetRouterName() const override
-			{
-				return Options.InputSystemName;
-			}
-
-			virtual const FMetasoundFrontendInterface& GetInterface() const override
-			{
-				return Interface;
-			}
-
-			virtual bool UpdateRootGraphInterface(Frontend::FDocumentHandle InDocument) const override
-			{
-				if (UpdateTransform.IsValid())
-				{
-					return UpdateTransform->Transform(InDocument);
-				}
-				return false;
-			}
-
-			private:
-				FMetasoundFrontendInterface Interface;
-				TUniquePtr<Frontend::IDocumentTransform> UpdateTransform;
-				FInterfaceRegistryOptions Options;
-		};
-
-		template <typename UClassType>
-		void RegisterInterface(const FMetasoundFrontendInterface& InInterface, TUniquePtr<Frontend::IDocumentTransform>&& UpdateTransform, bool bInIsDefault, FName InRouterName)
-		{
-			FInterfaceRegistryOptions Options
-			{
-				bInIsDefault,
-				InRouterName,
-				UClassType::StaticClass()->GetFName()
-			};
-
-			IMetasoundUObjectRegistry::Get().RegisterUClassInterface(MakeUnique<TMetasoundUObjectRegistryEntry<UClassType>>(InInterface.Version));
-			Frontend::IInterfaceRegistry::Get().RegisterInterface(MakeUnique<FInterfaceRegistryEntry>(InInterface, MoveTemp(UpdateTransform), MoveTemp(Options)));
-		}
+		void RegisterInterfaces();
 	} // namespace Engine
 } // namespace Metasound
diff --git a/Engine/Plugins/Runtime/Metasound/Source/MetasoundEngine/Private/MetasoundSource.cpp b/Engine/Plugins/Runtime/Metasound/Source/MetasoundEngine/Private/MetasoundSource.cpp
--- a/Engine/Plugins/Runtime/Metasound/Source/MetasoundEngine/Private/MetasoundSource.cpp
+++ b/Engine/Plugins/Runtime/Metasound/Source/MetasoundEngine/Private/MetasoundSource.cpp
@@ -3,7 +3,7 @@
 
 #include "Algo/Transform.h"
 #include "AssetRegistryModule.h"
-#include "IAudioGeneratorInterfaceRegistry.h"
+#include "IAudioParameterInterfaceRegistry.h"
 #include "Internationalization/Text.h"
 #include "MetasoundAssetBase.h"
 #include "MetasoundAudioFormats.h"
@@ -19,9 +19,11 @@
 #include "MetasoundGenerator.h"
 #include "MetasoundLog.h"
 #include "MetasoundOperatorSettings.h"
+#include "MetasoundOutputFormatInterfaces.h"
 #include "MetasoundParameterTransmitter.h"
 #include "MetasoundPrimitives.h"
 #include "MetasoundReceiveNode.h"
+#include "MetasoundSourceInterface.h"
 #include "MetasoundTrace.h"
 #include "MetasoundTrigger.h"
 #include "MetasoundUObjectRegistry.h"
@@ -31,7 +33,7 @@
 #include "EdGraph/EdGraph.h"
 #endif // WITH_EDITORONLY_DATA
 
-#define LOCTEXT_NAMESPACE "MetaSoundSource"
+#define LOCTEXT_NAMESPACE "MetaSound"
 
 namespace MetaSoundSourcePrivate
 {
@@ -39,25 +41,28 @@ namespace MetaSoundSourcePrivate
 
 	const FFormatOutputVertexKeyMap& GetFormatOutputVertexKeys()
 	{
-		using namespace Metasound;
-		using namespace Metasound::Engine;
+		auto CreateVertexKeyMap = []()
+		{
+			using namespace Metasound::Frontend;
 
-		static const FFormatOutputVertexKeyMap Map
-		(
+			return FFormatOutputVertexKeyMap
 			{
 				{
 					EMetasoundSourceAudioFormat::Mono,
-					{ MetasoundSourceMono::GetAudioOutputName() }
+					{
+						OutputFormatMonoInterface::Outputs::MonoOut
+					}
 				},
 				{
 					EMetasoundSourceAudioFormat::Stereo,
 					{
-						MetasoundSourceStereo::GetLeftAudioOutputName(),
-						MetasoundSourceStereo::GetRightAudioOutputName()
+						OutputFormatStereoInterface::Outputs::LeftOut,
+						OutputFormatStereoInterface::Outputs::RightOut,
 					}
 				}
-			}
-		);
+			};
+		};
+		static const FFormatOutputVertexKeyMap Map = CreateVertexKeyMap();
 		return Map;
 	}
 }
@@ -122,8 +127,8 @@ void UMetaSoundSource::PostEditChangeProperty(FPropertyChangedEvent& InEvent)
 			case EMetasoundSourceAudioFormat::Mono:
 			{
 				bDidModifyDocument = FModifyRootGraphInterfaces(
-					{ MetasoundSourceStereo::GetVersion() },
-					{ MetasoundSourceMono::GetVersion() }
+					{ OutputFormatStereoInterface::GetVersion() },
+					{ OutputFormatMonoInterface::GetVersion() }
 				).Transform(GetDocumentHandle());
 			}
 			break;
@@ -131,8 +136,8 @@ void UMetaSoundSource::PostEditChangeProperty(FPropertyChangedEvent& InEvent)
 			case EMetasoundSourceAudioFormat::Stereo:
 			{
 				bDidModifyDocument = FModifyRootGraphInterfaces(
-					{ MetasoundSourceMono::GetVersion() },
-					{ MetasoundSourceStereo::GetVersion() }
+					{ OutputFormatMonoInterface::GetVersion() },
+					{ OutputFormatStereoInterface::GetVersion() }
 				).Transform(GetDocumentHandle());
 			}
 			break;
@@ -162,8 +167,9 @@ void UMetaSoundSource::PostEditChangeProperty(FPropertyChangedEvent& InEvent)
 
 bool UMetaSoundSource::ConformObjectDataToInterfaces()
 {
-	const TArray<FMetasoundFrontendVersion>& InterfaceVersions = GetDocumentHandle()->GetInterfaceVersions();
-	if (InterfaceVersions.Contains(Metasound::Engine::MetasoundSourceMono::GetVersion()))
+	using namespace Metasound::Frontend;
+
+	if (IsInterfaceDeclared(OutputFormatMonoInterface::GetVersion()))
 	{
 		if (OutputFormat != EMetasoundSourceAudioFormat::Mono || NumChannels != 1)
 		{
@@ -173,7 +179,7 @@ bool UMetaSoundSource::ConformObjectDataToInterfaces()
 		}
 	}
 
-	if (InterfaceVersions.Contains(Metasound::Engine::MetasoundSourceStereo::GetVersion()))
+	if (IsInterfaceDeclared(OutputFormatStereoInterface::GetVersion()))
 	{
 		if (OutputFormat != EMetasoundSourceAudioFormat::Stereo || NumChannels != 2)
 		{
@@ -426,27 +432,16 @@ bool UMetaSoundSource::SupportsSubtitles() const
 	return Super::SupportsSubtitles();
 }
 
-const FMetasoundFrontendVersion& UMetaSoundSource::GetDefaultArchetypeVersion() const
-{
-	static const FMetasoundFrontendVersion DefaultVersion = Metasound::Engine::MetasoundSourceMono::GetVersion();
-	return DefaultVersion;
-}
-
 float UMetaSoundSource::GetDuration()
 {
 	// eh? this is kind of a weird field anyways.
 	return Super::GetDuration();
 }
 
-bool UMetaSoundSource::ImplementsGeneratorInterface(Audio::FGeneratorInterfacePtr InInterface) const
+bool UMetaSoundSource::ImplementsParameterInterface(Audio::FParameterInterfacePtr InInterface) const
 {
-	const FMetasoundFrontendVersionNumber Number = FMetasoundFrontendVersionNumber{ InInterface->Version.Major, InInterface->Version.Minor };
-	auto ContainsInterface = [Name = InInterface->Name, &Number](FMetasoundFrontendVersion& InVersion)
-	{
-		return InVersion.Name == Name && InVersion.Number == Number;
-	};
-
-	return GetDocumentChecked().InterfaceVersions.ContainsByPredicate(ContainsInterface);
+	const FMetasoundFrontendVersion Version { InInterface->GetName(), { InInterface->GetVersion().Major, InInterface->GetVersion().Minor } };
+	return GetDocumentChecked().Interfaces.Contains(Version);
 }
 
 ISoundGeneratorPtr UMetaSoundSource::CreateSoundGenerator(const FSoundGeneratorInitParams& InParams)
@@ -477,9 +472,7 @@ ISoundGeneratorPtr UMetaSoundSource::CreateSoundGenerator(const FSoundGeneratorI
 		MetasoundGraph,
 		Environment,
 		GetName(),
-		GetAudioOutputVertexKeys(),
-		MetasoundSource::GetOnPlayInputName(),
-		MetasoundSource::GetIsFinishedOutputName()
+		GetAudioOutputVertexKeys()
 	};
 
 	return ISoundGeneratorPtr(new FMetasoundGenerator(MoveTemp(InitParams)));
@@ -725,6 +718,7 @@ Metasound::FMetasoundEnvironment UMetaSoundSource::CreateEnvironment() const
 {
 	using namespace Metasound;
 	using namespace Metasound::Engine;
+	using namespace Metasound::Frontend;
 
 	FMetasoundEnvironment Environment;
 	// Add audio device ID to environment.
@@ -747,8 +741,8 @@ Metasound::FMetasoundEnvironment UMetaSoundSource::CreateEnvironment() const
 		}
 	}
 
-	Environment.SetValue<Audio::FDeviceId>(MetasoundSource::GetAudioDeviceIDVariableName(), AudioDeviceID);
-	Environment.SetValue<uint32>(MetasoundSource::GetSoundUniqueIdName(), GetUniqueID());
+	Environment.SetValue<Audio::FDeviceId>(SourceInterface::Environment::DeviceID, AudioDeviceID);
+	Environment.SetValue<uint32>(SourceInterface::Environment::SoundUniqueID, GetUniqueID());
 
 	return Environment;
 }
@@ -757,13 +751,14 @@ Metasound::FMetasoundEnvironment UMetaSoundSource::CreateEnvironment(const FSoun
 {
 	using namespace Metasound;
 	using namespace Metasound::Engine;
+	using namespace Metasound::Frontend;
 
 	FMetasoundEnvironment Environment = CreateEnvironment();
-	Environment.SetValue<bool>(MetasoundSource::GetIsPreviewSoundName(), InParams.bIsPreviewSound);
-	Environment.SetValue<uint64>(MetasoundSource::GetInstanceIDName(), InParams.InstanceID);
-	
+	Environment.SetValue<bool>(SourceInterface::Environment::IsPreview, InParams.bIsPreviewSound);
+	Environment.SetValue<uint64>(SourceInterface::Environment::TransmitterID, InParams.InstanceID);
+
 #if WITH_METASOUND_DEBUG_ENVIRONMENT
-	Environment.SetValue<FString>(MetasoundSource::GetGraphName(), GetFullName());
+	Environment.SetValue<FString>(SourceInterface::Environment::GraphName, GetFullName());
 #endif // WITH_METASOUND_DEBUG_ENVIRONMENT
 
 	return Environment;
@@ -773,9 +768,10 @@ Metasound::FMetasoundEnvironment UMetaSoundSource::CreateEnvironment(const Audio
 {
 	using namespace Metasound;
 	using namespace Metasound::Engine;
+	using namespace Metasound::Frontend;
 
 	FMetasoundEnvironment Environment = CreateEnvironment();
-	Environment.SetValue<uint64>(MetasoundSource::GetInstanceIDName(), InParams.InstanceID);
+	Environment.SetValue<uint64>(SourceInterface::Environment::TransmitterID, InParams.InstanceID);
 
 	return Environment;
 }
@@ -796,5 +792,4 @@ const TArray<Metasound::FVertexName>& UMetaSoundSource::GetAudioOutputVertexKeys
 		return Empty;
 	}
 }
-
-#undef LOCTEXT_NAMESPACE // MetaSoundSource
+#undef LOCTEXT_NAMESPACE // MetaSound
diff --git a/Engine/Plugins/Runtime/Metasound/Source/MetasoundEngine/Private/MetasoundUObjectRegistry.cpp b/Engine/Plugins/Runtime/Metasound/Source/MetasoundEngine/Private/MetasoundUObjectRegistry.cpp
--- a/Engine/Plugins/Runtime/Metasound/Source/MetasoundEngine/Private/MetasoundUObjectRegistry.cpp
+++ b/Engine/Plugins/Runtime/Metasound/Source/MetasoundEngine/Private/MetasoundUObjectRegistry.cpp
@@ -24,14 +24,7 @@ namespace Metasound
 	class FMetasoundUObjectRegistry : public IMetasoundUObjectRegistry
 	{
 		public:
-			FMetasoundUObjectRegistry()
-			{
-				// Set default version to use base UMetaSound class implementation
-				FMetasoundFrontendVersion DefaultVersion;
-
-				using FRegistryEntryType = TMetasoundUObjectRegistryEntry<UMetaSound>;
-				RegisterUClassInterface(MakeUnique<FRegistryEntryType>(DefaultVersion));
-			}
+			FMetasoundUObjectRegistry() = default;
 
 			void RegisterUClassInterface(TUniquePtr<IMetasoundUObjectRegistryEntry>&& InEntry) override
 			{
@@ -79,7 +72,7 @@ namespace Metasound
 			{
 				TArray<const IMetasoundUObjectRegistryEntry*> AllInterfaceEntries;
 
-				for (const FMetasoundFrontendVersion& InterfaceVersion : InDocument.InterfaceVersions)
+				for (const FMetasoundFrontendVersion& InterfaceVersion : InDocument.Interfaces)
 				{
 					TArray<const IMetasoundUObjectRegistryEntry*> EntriesForInterface;
 					EntriesByInterface.MultiFind(Frontend::GetInterfaceRegistryKey(InterfaceVersion), EntriesForInterface);
@@ -122,6 +115,20 @@ namespace Metasound
 				return nullptr;
 			}
 
+			void IterateRegisteredUClasses(TFunctionRef<void(UClass&)> InFunc) const override
+			{
+				for (const IMetasoundUObjectRegistryEntry* Entry : Entries)
+				{
+					if (Entry)
+					{
+						if (UClass* Class = Entry->GetUClass())
+						{
+							InFunc(*Class);
+						}
+					}
+				}
+			}
+
 		private:
 			UObject* NewObject(const IMetasoundUObjectRegistryEntry& InEntry, const FMetasoundFrontendDocument& InDocument, const FString& InPath) const
 			{
diff --git a/Engine/Plugins/Runtime/Metasound/Source/MetasoundEngine/Public/Metasound.h b/Engine/Plugins/Runtime/Metasound/Source/MetasoundEngine/Public/Metasound.h
--- a/Engine/Plugins/Runtime/Metasound/Source/MetasoundEngine/Public/Metasound.h
+++ b/Engine/Plugins/Runtime/Metasound/Source/MetasoundEngine/Public/Metasound.h
@@ -230,7 +230,6 @@ public:
 	// Returns Asset Metadata associated with this MetaSound
 	virtual Metasound::Frontend::FNodeClassInfo GetAssetClassInfo() const override;
 
-	virtual const FMetasoundFrontendVersion& GetDefaultArchetypeVersion() const override;
 
 	virtual bool ConformObjectDataToInterfaces() override { return false; }
 
diff --git a/Engine/Plugins/Runtime/Metasound/Source/MetasoundEngine/Public/MetasoundAssetSubsystem.h b/Engine/Plugins/Runtime/Metasound/Source/MetasoundEngine/Public/MetasoundAssetSubsystem.h
--- a/Engine/Plugins/Runtime/Metasound/Source/MetasoundEngine/Public/MetasoundAssetSubsystem.h
+++ b/Engine/Plugins/Runtime/Metasound/Source/MetasoundEngine/Public/MetasoundAssetSubsystem.h
@@ -34,13 +34,13 @@ class METASOUNDENGINE_API UMetaSoundAssetSubsystem : public UEngineSubsystem, pu
 public:
 	virtual void Initialize(FSubsystemCollectionBase& InCollection) override;
 
-	void AddOrUpdateAsset(const FAssetData& InAssetData, bool bInRegisterWithFrontend = true);
-	void RemoveAsset(UObject& InObject, bool bInUnregisterWithFrontend = true);
-	void RemoveAsset(const FAssetData& InAssetData, bool bInUnregisterWithFrontend = true);
+	void AddOrUpdateAsset(const FAssetData& InAssetData);
+	void RemoveAsset(UObject& InObject);
+	void RemoveAsset(const FAssetData& InAssetData);
 	void RenameAsset(const FAssetData& InAssetData, bool bInReregisterWithFrontend = true);
 
 	virtual void AddAssetReferences(FMetasoundAssetBase& InAssetBase) override;
-	virtual void AddOrUpdateAsset(UObject& InObject, bool bInRegisterWithFrontend = true) override;
+	virtual void AddOrUpdateAsset(UObject& InObject) override;
 	virtual bool CanAutoUpdate(const FMetasoundFrontendClassName& InClassName) const override;
 	virtual bool ContainsKey(const Metasound::Frontend::FNodeRegistryKey& InRegistryKey) const override;
 	virtual const FSoftObjectPath* FindObjectPathFromKey(const Metasound::Frontend::FNodeRegistryKey& RegistryKey) const override;
diff --git a/Engine/Plugins/Runtime/Metasound/Source/MetasoundEngine/Public/MetasoundSource.h b/Engine/Plugins/Runtime/Metasound/Source/MetasoundEngine/Public/MetasoundSource.h
--- a/Engine/Plugins/Runtime/Metasound/Source/MetasoundEngine/Public/MetasoundSource.h
+++ b/Engine/Plugins/Runtime/Metasound/Source/MetasoundEngine/Public/MetasoundSource.h
@@ -145,7 +145,6 @@ public:
 	// Returns Asset Metadata associated with this MetaSoundSource
 	virtual Metasound::Frontend::FNodeClassInfo GetAssetClassInfo() const override;
 
-	virtual const FMetasoundFrontendVersion& GetDefaultArchetypeVersion() const override;
 
 	virtual bool ConformObjectDataToInterfaces() override;
 
@@ -165,7 +164,7 @@ public:
 	virtual bool IsPlayable() const override;
 	virtual bool SupportsSubtitles() const override;
 	virtual float GetDuration() override;
-	virtual bool ImplementsGeneratorInterface(Audio::FGeneratorInterfacePtr InInterface) const override;
+	virtual bool ImplementsParameterInterface(Audio::FParameterInterfacePtr InInterface) const override;
 	virtual ISoundGeneratorPtr CreateSoundGenerator(const FSoundGeneratorInitParams& InParams) override;
 	virtual TUniquePtr<Audio::IParameterTransmitter> CreateParameterTransmitter(Audio::FParameterTransmitterInitParams&& InParams) const override;
 	virtual bool IsParameterValid(const FAudioParameter& InParameter) const override;
diff --git a/Engine/Plugins/Runtime/Metasound/Source/MetasoundEngine/Public/MetasoundUObjectRegistry.h b/Engine/Plugins/Runtime/Metasound/Source/MetasoundEngine/Public/MetasoundUObjectRegistry.h
--- a/Engine/Plugins/Runtime/Metasound/Source/MetasoundEngine/Public/MetasoundUObjectRegistry.h
+++ b/Engine/Plugins/Runtime/Metasound/Source/MetasoundEngine/Public/MetasoundUObjectRegistry.h
@@ -7,6 +7,7 @@
 #include "MetasoundAssetBase.h"
 #include "Subsystems/EngineSubsystem.h"
 #include "UObject/Object.h"
+#include "Templates/Function.h"
 
 
 // Forward Declarations
@@ -163,6 +164,9 @@ namespace Metasound
 			 */
 			virtual UObject* NewObject(UClass* InClass, const FMetasoundFrontendDocument& InDocument, const FString& InPath) const = 0;
 
+			/** Iterate all registered UClasses that serve as MetaSound assets.*/
+			virtual void IterateRegisteredUClasses(TFunctionRef<void(UClass&)> InFunc) const = 0;
+
 			/** Returns true if the InObject is of a class or child class which is registered. */
 			virtual bool IsRegisteredClass(UObject* InObject) const = 0;
 
diff --git a/Engine/Plugins/Runtime/Metasound/Source/MetasoundFrontend/Private/MetasoundArchetype.cpp b/Engine/Plugins/Runtime/Metasound/Source/MetasoundFrontend/Private/MetasoundArchetype.cpp
--- a/Engine/Plugins/Runtime/Metasound/Source/MetasoundFrontend/Private/MetasoundArchetype.cpp
+++ b/Engine/Plugins/Runtime/Metasound/Source/MetasoundFrontend/Private/MetasoundArchetype.cpp
@@ -177,10 +177,12 @@ namespace Metasound
 
 		bool DeclaredInterfaceVersionsMatch(const FMetasoundFrontendDocument& InDocumentA, const FMetasoundFrontendDocument& InDocumentB)
 		{
-			using namespace MetasoundInterfaceIntrinsics;
+			if (InDocumentA.Interfaces.Num() != InDocumentB.Interfaces.Num())
+			{
+				return false;
+			}
 
-			auto IsLessThanVersion = [](const FMetasoundFrontendVersion& A, const FMetasoundFrontendVersion& B) { return A < B; };
-			return IsSetEquivalent(InDocumentA.InterfaceVersions, InDocumentB.InterfaceVersions, IsLessThanVersion);
+			return InDocumentA.Interfaces.Includes(InDocumentB.Interfaces);
 		}
 
 		int32 InputOutputDifferenceCount(const FMetasoundFrontendClass& InClass, const FMetasoundFrontendInterface& InInterface)
diff --git a/Engine/Plugins/Runtime/Metasound/Source/MetasoundFrontend/Private/MetasoundAssetBase.cpp b/Engine/Plugins/Runtime/Metasound/Source/MetasoundFrontend/Private/MetasoundAssetBase.cpp
--- a/Engine/Plugins/Runtime/Metasound/Source/MetasoundFrontend/Private/MetasoundAssetBase.cpp
+++ b/Engine/Plugins/Runtime/Metasound/Source/MetasoundFrontend/Private/MetasoundAssetBase.cpp
@@ -77,6 +77,7 @@ void FMetasoundAssetBase::RegisterGraphWithFrontend(Metasound::Frontend::FMetaSo
 	if (InRegistrationOptions.bRegisterDependencies)
 	{
 		// Must be called in case register is called prior to asset scan being completed
+		// or in certain cases where asset is being loaded for first time.
 		IMetaSoundAssetManager::GetChecked().AddAssetReferences(*this);
 	}
 
@@ -281,7 +282,7 @@ bool FMetasoundAssetBase::GetDeclaredInterfaces(TArray<const Metasound::Frontend
 	{
 		bool bInterfacesFound = true;
 
-		Algo::Transform(Document->InterfaceVersions, OutInterfaces, [&](const FMetasoundFrontendVersion& Version)
+		Algo::Transform(Document->Interfaces, OutInterfaces, [&](const FMetasoundFrontendVersion& Version)
 		{
 			const FInterfaceRegistryKey InterfaceKey = GetInterfaceRegistryKey(Version);
 			const IInterfaceRegistryEntry* RegistryEntry = IInterfaceRegistry::Get().FindInterfaceRegistryEntry(InterfaceKey);
@@ -300,12 +301,9 @@ bool FMetasoundAssetBase::GetDeclaredInterfaces(TArray<const Metasound::Frontend
 	return false;
 }
 
-bool FMetasoundAssetBase::IsInterfaceDeclared(FName InName) const
+bool FMetasoundAssetBase::IsInterfaceDeclared(const FMetasoundFrontendVersion& InVersion) const
 {
-	return Algo::AnyOf(GetDocumentChecked().InterfaceVersions, [&InName](const FMetasoundFrontendVersion& FrontendVersion)
-	{
-		return FrontendVersion.Name == InName;
-	});
+	return GetDocumentChecked().Interfaces.Contains(InVersion);
 }
 
 void FMetasoundAssetBase::SetDocument(const FMetasoundFrontendDocument& InDocument)
@@ -365,18 +363,10 @@ bool FMetasoundAssetBase::VersionAsset()
 		return false;
 	}
 
-	// Data migration for 5.0 Early Access data. ArchetypeVersion can be removed post 5.0 release.
-	bool bDidEdit = false;
-	if (Doc->ArchetypeVersion.IsValid())
-	{
-		Doc->InterfaceVersions.Add(Doc->ArchetypeVersion);
-		Doc->ArchetypeVersion = FMetasoundFrontendVersion::GetInvalid();
-		bDidEdit = true;
-	}
-
-	FDocumentHandle DocHandle = GetDocumentHandle();
+	bool bDidEdit = Doc->VersionInterfaces();
 
 	// Version Document Model
+	FDocumentHandle DocHandle = GetDocumentHandle();
 	{
 		bDidEdit |= FVersionDocument(AssetName, AssetPath).Transform(DocHandle);
 	}
@@ -384,7 +374,8 @@ bool FMetasoundAssetBase::VersionAsset()
 	// Version Interfaces
 	{
 		bool bInterfaceUpdated = false;
-		for (const FMetasoundFrontendVersion& Version : Doc->InterfaceVersions)
+		const TArray<FMetasoundFrontendVersion> Versions = Doc->Interfaces.Array();
+		for (const FMetasoundFrontendVersion& Version : Versions)
 		{
 			bInterfaceUpdated |= FUpdateRootGraphInterface(Version).Transform(DocHandle);
 		}
diff --git a/Engine/Plugins/Runtime/Metasound/Source/MetasoundFrontend/Private/MetasoundFrontendDocument.cpp b/Engine/Plugins/Runtime/Metasound/Source/MetasoundFrontend/Private/MetasoundFrontendDocument.cpp
--- a/Engine/Plugins/Runtime/Metasound/Source/MetasoundFrontend/Private/MetasoundFrontendDocument.cpp
+++ b/Engine/Plugins/Runtime/Metasound/Source/MetasoundFrontend/Private/MetasoundFrontendDocument.cpp
@@ -3,6 +3,7 @@
 #include "MetasoundFrontendDocument.h"
 
 #include "Algo/Transform.h"
+#include "IAudioParameterInterfaceRegistry.h"
 #include "MetasoundFrontend.h"
 #include "MetasoundFrontendRegistries.h"
 #include "MetasoundLog.h"
@@ -75,6 +76,11 @@ bool FMetasoundFrontendVertex::IsFunctionalEquivalent(const FMetasoundFrontendVe
 	return (InLHS.Name == InRHS.Name) && (InLHS.TypeName == InRHS.TypeName);
 }
 
+void FMetasoundFrontendClassVertex::SplitName(FName& OutNamespace, FName& OutParameterName) const
+{
+	Audio::FParameterPath::SplitName(Name, OutNamespace, OutParameterName);
+}
+
 bool FMetasoundFrontendClassVertex::IsFunctionalEquivalent(const FMetasoundFrontendClassVertex& InLHS, const FMetasoundFrontendClassVertex& InRHS)
 {
 	return FMetasoundFrontendVertex::IsFunctionalEquivalent(InLHS, InRHS);
diff --git a/Engine/Plugins/Runtime/Metasound/Source/MetasoundFrontend/Private/MetasoundFrontendDocumentController.cpp b/Engine/Plugins/Runtime/Metasound/Source/MetasoundFrontend/Private/MetasoundFrontendDocumentController.cpp
--- a/Engine/Plugins/Runtime/Metasound/Source/MetasoundFrontend/Private/MetasoundFrontendDocumentController.cpp
+++ b/Engine/Plugins/Runtime/Metasound/Source/MetasoundFrontend/Private/MetasoundFrontendDocumentController.cpp
@@ -117,22 +117,22 @@ namespace Metasound
 			return false;
 		}
 
-		const TArray<FMetasoundFrontendVersion>& FDocumentController::GetInterfaceVersions() const
+		const TSet<FMetasoundFrontendVersion>& FDocumentController::GetInterfaceVersions() const
 		{
 			if (const FMetasoundFrontendDocument* Doc = DocumentPtr.Get())
 			{
-				return Doc->InterfaceVersions;
+				return Doc->Interfaces;
 			}
 
-			static const TArray<FMetasoundFrontendVersion> EmptyArray;
-			return EmptyArray;
+			static const TSet<FMetasoundFrontendVersion> EmptySet;
+			return EmptySet;
 		}
 
 		void FDocumentController::AddInterfaceVersion(const FMetasoundFrontendVersion& InVersion)
 		{
 			if (FMetasoundFrontendDocument* Doc = DocumentPtr.Get())
 			{
-				Doc->InterfaceVersions.Add(InVersion);
+				Doc->Interfaces.Add(InVersion);
 			}
 		}
 
@@ -140,8 +140,7 @@ namespace Metasound
 		{
 			if (FMetasoundFrontendDocument* Doc = DocumentPtr.Get())
 			{
-				constexpr bool bAllowShrinking = false;
-				Doc->InterfaceVersions.RemoveAllSwap([&](const FMetasoundFrontendVersion& InterfaceVersion) { return InterfaceVersion == InVersion; }, bAllowShrinking);
+				Doc->Interfaces.Remove(InVersion);
 			}
 		}
 
@@ -149,7 +148,7 @@ namespace Metasound
 		{
 			if (FMetasoundFrontendDocument* Doc = DocumentPtr.Get())
 			{
-				Doc->InterfaceVersions.Reset();
+				Doc->Interfaces.Reset();
 			}
 		}
 
diff --git a/Engine/Plugins/Runtime/Metasound/Source/MetasoundFrontend/Private/MetasoundFrontendDocumentController.h b/Engine/Plugins/Runtime/Metasound/Source/MetasoundFrontend/Private/MetasoundFrontendDocumentController.h
--- a/Engine/Plugins/Runtime/Metasound/Source/MetasoundFrontend/Private/MetasoundFrontendDocumentController.h
+++ b/Engine/Plugins/Runtime/Metasound/Source/MetasoundFrontend/Private/MetasoundFrontendDocumentController.h
@@ -51,7 +51,7 @@ namespace Metasound
 
 			FGraphHandle AddDuplicateSubgraph(const IGraphController& InGraph) override;
 
-			const TArray<FMetasoundFrontendVersion>& GetInterfaceVersions() const override;
+			const TSet<FMetasoundFrontendVersion>& GetInterfaceVersions() const override;
 			void AddInterfaceVersion(const FMetasoundFrontendVersion& InVersion) override;
 			void RemoveInterfaceVersion(const FMetasoundFrontendVersion& InVersion) override;
 			void ClearInterfaceVersions() override;
diff --git a/Engine/Plugins/Runtime/Metasound/Source/MetasoundFrontend/Private/MetasoundFrontendInvalidController.h b/Engine/Plugins/Runtime/Metasound/Source/MetasoundFrontend/Private/MetasoundFrontendInvalidController.h
--- a/Engine/Plugins/Runtime/Metasound/Source/MetasoundFrontend/Private/MetasoundFrontendInvalidController.h
+++ b/Engine/Plugins/Runtime/Metasound/Source/MetasoundFrontend/Private/MetasoundFrontendInvalidController.h
@@ -3,6 +3,7 @@
 #pragma once
 
 #include "Containers/Array.h"
+#include "Containers/Set.h"
 #include "MetasoundAssetBase.h"
 #include "MetasoundBuilderInterface.h"
 #include "MetasoundFrontendController.h"
@@ -411,7 +412,7 @@ namespace Metasound
 				virtual FConstClassAccessPtr FindOrAddClass(const FMetasoundFrontendClassMetadata& InMetadata) override{ return FConstClassAccessPtr(); }
 				virtual FGraphHandle AddDuplicateSubgraph(const IGraphController& InGraph) override { return IGraphController::GetInvalidHandle(); }
 
-				virtual const TArray<FMetasoundFrontendVersion>& GetInterfaceVersions() const override { static const TArray<FMetasoundFrontendVersion> EmptyArray; return EmptyArray; }
+				virtual const TSet<FMetasoundFrontendVersion>& GetInterfaceVersions() const override { static const TSet<FMetasoundFrontendVersion> EmptyArray; return EmptyArray; }
 				virtual void AddInterfaceVersion(const FMetasoundFrontendVersion& InVersion) override { }
 				virtual void RemoveInterfaceVersion(const FMetasoundFrontendVersion& InVersion) override { }
 				virtual void ClearInterfaceVersions() override { }
diff --git a/Engine/Plugins/Runtime/Metasound/Source/MetasoundFrontend/Private/MetasoundFrontendTransform.cpp b/Engine/Plugins/Runtime/Metasound/Source/MetasoundFrontend/Private/MetasoundFrontendTransform.cpp
--- a/Engine/Plugins/Runtime/Metasound/Source/MetasoundFrontend/Private/MetasoundFrontendTransform.cpp
+++ b/Engine/Plugins/Runtime/Metasound/Source/MetasoundFrontend/Private/MetasoundFrontendTransform.cpp
@@ -54,8 +54,21 @@ namespace Metasound
 			bSetDefaultNodeLocations = bInSetDefaultNodeLocations;
 		}
 
-		void FModifyRootGraphInterfaces::Init()
+		void FModifyRootGraphInterfaces::SetNamePairingFunction(const TFunction<bool(FName, FName)>& InNamePairingFunction)
 		{
+			// Reinit required to rebuild list of pairs
+			Init(&InNamePairingFunction);
+		}
+
+		void FModifyRootGraphInterfaces::Init(const TFunction<bool(FName, FName)>* InNamePairingFunction)
+		{
+			InputsToRemove.Reset();
+			InputsToAdd.Reset();
+			OutputsToRemove.Reset();
+			OutputsToAdd.Reset();
+			PairedInputs.Reset();
+			PairedOutputs.Reset();
+
 			for (const FMetasoundFrontendInterface& FromInterface : InterfacesToRemove)
 			{
 				InputsToRemove.Append(FromInterface.Inputs);
@@ -72,9 +85,25 @@ namespace Metasound
 			{
 				const FMetasoundFrontendClassVertex& VertexToAdd = InputsToAdd[AddIndex];
 
-				int32 RemoveIndex = InputsToRemove.IndexOfByPredicate([&](const FMetasoundFrontendClassVertex& VertexToRemove)
+				const int32 RemoveIndex = InputsToRemove.IndexOfByPredicate([&](const FMetasoundFrontendClassVertex& VertexToRemove)
 				{
-					return FMetasoundFrontendClassVertex::IsFunctionalEquivalent(VertexToAdd, VertexToRemove);
+					if (VertexToAdd.TypeName != VertexToRemove.TypeName)
+					{
+						return false;
+					}
+
+					if (InNamePairingFunction && *InNamePairingFunction)
+					{
+						return (*InNamePairingFunction)(VertexToAdd.Name, VertexToRemove.Name);
+					}
+
+					FName ParamA;
+					FName ParamB;
+					FName Namespace;
+					VertexToAdd.SplitName(Namespace, ParamA);
+					VertexToRemove.SplitName(Namespace, ParamB);
+
+					return ParamA == ParamB;
 				});
 
 				if (INDEX_NONE != RemoveIndex)
@@ -90,10 +119,26 @@ namespace Metasound
 			{
 				const FMetasoundFrontendClassVertex& VertexToAdd = OutputsToAdd[AddIndex];
 
-				int32 RemoveIndex = OutputsToRemove.IndexOfByPredicate([&](const FMetasoundFrontendClassVertex& VertexToRemove)
+				const int32 RemoveIndex = OutputsToRemove.IndexOfByPredicate([&](const FMetasoundFrontendClassVertex& VertexToRemove)
+				{
+					if (VertexToAdd.TypeName != VertexToRemove.TypeName)
+					{
+						return false;
+					}
+
+					if (InNamePairingFunction && *InNamePairingFunction)
 					{
-						return FMetasoundFrontendClassVertex::IsFunctionalEquivalent(VertexToAdd, VertexToRemove);
-					});
+						return (*InNamePairingFunction)(VertexToAdd.Name, VertexToRemove.Name);
+					}
+
+					FName ParamA;
+					FName ParamB;
+					FName Namespace;
+					VertexToAdd.SplitName(Namespace, ParamA);
+					VertexToRemove.SplitName(Namespace, ParamB);
+
+					return ParamA == ParamB;
+				});
 
 				if (INDEX_NONE != RemoveIndex)
 				{
@@ -392,7 +437,7 @@ namespace Metasound
 			if (LastVersionUpdated)
 			{
 				UE_LOG(LogMetaSound, Display, TEXT("Asset '%s' interface '%s' updated: '%s' --> '%s'"),
-					*InDocument->GetMetadata().Version.ToString(),
+					*InDocument->GetRootGraphClass().Metadata.GetDisplayName().ToString(),
 					*InterfaceVersion.Name.ToString(),
 					*InterfaceVersion.Number.ToString(),
 					*LastVersionUpdated->ToString());
@@ -528,7 +573,7 @@ namespace Metasound
 			// Ensure preset interfaces match those found in referenced graph.  Referenced graph is assumed to be
 			// well-formed (i.e. all inputs/outputs/environment variables declared by interfaces are present, and
 			// of proper name & data type).
-			const TArray<FMetasoundFrontendVersion>& RefInterfaceVersions = ReferencedDocument->GetInterfaceVersions();
+			const TSet<FMetasoundFrontendVersion>& RefInterfaceVersions = ReferencedDocument->GetInterfaceVersions();
 			for (const FMetasoundFrontendVersion& Version : RefInterfaceVersions)
 			{
 				InDocument->AddInterfaceVersion(Version);
@@ -870,12 +915,12 @@ namespace Metasound
 				check(InDocument->GetMetadata().Version.Number.Major == 1);
 				check(InDocument->GetMetadata().Version.Number.Minor == 3);
 
-				const TArray<FMetasoundFrontendVersion>& InterfaceVersions = InDocument->GetInterfaceVersions();
+				const TSet<FMetasoundFrontendVersion>& Interfaces = InDocument->GetInterfaceVersions();
 
-				// Version 1.3 did not have an "InterfaceVersion" property on the 
+				// Version 1.3 did not have an "InterfaceVersion" property on the
 				// document, so any document that is being updated should start off
 				// with an "Invalid" interface version.
-				if (ensure(InterfaceVersions.IsEmpty()))
+				if (ensure(Interfaces.IsEmpty()))
 				{
 					constexpr bool bIncludeDeprecatedInterfaces = true;
 					TArray<FMetasoundFrontendInterface> AllInterfaces = ISearchEngine::Get().FindAllInterfaces(bIncludeDeprecatedInterfaces);
diff --git a/Engine/Plugins/Runtime/Metasound/Source/MetasoundFrontend/Private/MetasoundOutputFormatInterfaces.cpp b/Engine/Plugins/Runtime/Metasound/Source/MetasoundFrontend/Private/MetasoundOutputFormatInterfaces.cpp
new file mode 100644
--- /dev/null
+++ b/Engine/Plugins/Runtime/Metasound/Source/MetasoundFrontend/Private/MetasoundOutputFormatInterfaces.cpp
@@ -0,0 +1,100 @@
+// Copyright Epic Games, Inc. All Rights Reserved.
+#include "MetasoundOutputFormatInterfaces.h"
+
+#include "IAudioParameterInterfaceRegistry.h"
+#include "MetasoundAudioBuffer.h"
+#include "MetasoundTrigger.h"
+#include "Templates/SharedPointer.h"
+#include "UObject/Class.h"
+
+
+#define LOCTEXT_NAMESPACE "MetasoundFrontend"
+namespace Metasound
+{
+	namespace Frontend
+	{
+#define AUDIO_PARAMETER_INTERFACE_NAMESPACE "UE.OutputFormat.Mono"
+		namespace OutputFormatMonoInterface
+		{
+			const FMetasoundFrontendVersion& GetVersion()
+			{
+				static const FMetasoundFrontendVersion Version = { AUDIO_PARAMETER_INTERFACE_NAMESPACE, { 1, 0 } };
+				return Version;
+			}
+
+			namespace Outputs
+			{
+				const FName MonoOut = AUDIO_PARAMETER_INTERFACE_MEMBER_DEFINE("Audio:0");
+			}
+
+			Audio::FParameterInterfacePtr CreateInterface(const UClass& InUClass)
+			{
+				struct FInterface : public Audio::FParameterInterface
+				{
+					FInterface(const UClass& InAssetClass)
+						: FParameterInterface(OutputFormatMonoInterface::GetVersion().Name, OutputFormatMonoInterface::GetVersion().Number.ToInterfaceVersion(), InAssetClass)
+					{
+						Outputs =
+						{
+							{
+								LOCTEXT("GeneratedAudioDisplayName", "Out Mono"),
+								LOCTEXT("GeneratedAudioDescription", "The resulting mono output from this source."),
+								GetMetasoundDataTypeName<FAudioBuffer>(),
+								Outputs::MonoOut
+							}
+						};
+					}
+				};
+
+				return MakeShared<FInterface>(InUClass);
+			}
+		} // namespace OutputFormatMonoInterface
+#undef AUDIO_PARAMETER_INTERFACE_NAMESPACE
+
+#define AUDIO_PARAMETER_INTERFACE_NAMESPACE "UE.OutputFormat.Stereo"
+		namespace OutputFormatStereoInterface
+		{
+			const FMetasoundFrontendVersion& GetVersion()
+			{
+				static const FMetasoundFrontendVersion Version = { AUDIO_PARAMETER_INTERFACE_NAMESPACE, { 1, 0 } };
+				return Version;
+			}
+
+			namespace Outputs
+			{
+				const FName LeftOut = AUDIO_PARAMETER_INTERFACE_MEMBER_DEFINE("Audio:0");
+				const FName RightOut = AUDIO_PARAMETER_INTERFACE_MEMBER_DEFINE("Audio:1");
+			}
+
+			Audio::FParameterInterfacePtr CreateInterface(const UClass& InUClass)
+			{
+				struct FInterface : public Audio::FParameterInterface
+				{
+					FInterface(const UClass& InAssetClass)
+						: FParameterInterface(OutputFormatStereoInterface::GetVersion().Name, OutputFormatStereoInterface::GetVersion().Number.ToInterfaceVersion(), InAssetClass)
+					{
+						Outputs =
+						{
+							{
+								LOCTEXT("GeneratedLeftDisplayName", "Out Left"),
+								LOCTEXT("GeneratedLeftDescription", "The resulting left channel output audio from this source."),
+								GetMetasoundDataTypeName<FAudioBuffer>(),
+								Outputs::LeftOut
+							},
+							{
+								LOCTEXT("GeneratedLeftDisplayName", "Out Right"),
+								LOCTEXT("GeneratedLeftDescription", "The resulting right channel output audio from this source."),
+								GetMetasoundDataTypeName<FAudioBuffer>(),
+								Outputs::RightOut
+							}
+						};
+					}
+				};
+
+				return MakeShared<FInterface>(InUClass);
+			}
+		} // namespace OutputFormatStereoInterface
+#undef AUDIO_PARAMETER_INTERFACE_NAMESPACE
+	}
+}
+#undef LOCTEXT_NAMESPACE
\ No newline at end of file
diff --git a/Engine/Plugins/Runtime/Metasound/Source/MetasoundFrontend/Private/MetasoundParameterTransmitter.cpp b/Engine/Plugins/Runtime/Metasound/Source/MetasoundFrontend/Private/MetasoundParameterTransmitter.cpp
--- a/Engine/Plugins/Runtime/Metasound/Source/MetasoundFrontend/Private/MetasoundParameterTransmitter.cpp
+++ b/Engine/Plugins/Runtime/Metasound/Source/MetasoundFrontend/Private/MetasoundParameterTransmitter.cpp
@@ -2,8 +2,9 @@
 
 #include "MetasoundParameterTransmitter.h"
 
-#include "IAudioGeneratorInterfaceRegistry.h"
+#include "IAudioParameterInterfaceRegistry.h"
 #include "MetasoundLog.h"
+#include "MetasoundSourceInterface.h"
 
 
 namespace Metasound
@@ -124,22 +125,16 @@ namespace Metasound
 					return FName();
 			}
 		}
-	}
-
-	const FVertexName& FMetaSoundParameterTransmitter::GetInstanceIDEnvironmentVariableName()
-	{
-		static const FVertexName VarName = "TransmitterInstanceID";
-		return VarName;
-	}
+	} // namespace Frontend
 
 	FSendAddress FMetaSoundParameterTransmitter::CreateSendAddressFromEnvironment(const FMetasoundEnvironment& InEnvironment, const FVertexName& InVertexName, const FName& InTypeName)
 	{
-		const FVertexName IDVarName = GetInstanceIDEnvironmentVariableName();
-		uint64 InstanceID = -1;
+		using namespace Frontend;
 
-		if (ensure(InEnvironment.Contains<uint64>(IDVarName)))
+		uint64 InstanceID = -1;
+		if (ensure(InEnvironment.Contains<uint64>(SourceInterface::Environment::TransmitterID)))
 		{
-			InstanceID = InEnvironment.GetValue<uint64>(IDVarName);
+			InstanceID = InEnvironment.GetValue<uint64>(SourceInterface::Environment::TransmitterID);
 		}
 
 		return CreateSendAddressFromInstanceID(InstanceID, InVertexName, InTypeName);
@@ -180,13 +175,6 @@ namespace Metasound
 		return SetParameterWithLiteral(ParamName, Frontend::ConvertParameterToLiteral(MoveTemp(InParameter)));
 	}
 
-	bool FMetaSoundParameterTransmitter::SetParameter(FName InInterfaceName, FAudioParameter&& InParameter)
-	{
-		InParameter.ParamName = Audio::IGeneratorInterfaceRegistry::GetMemberFullName(InInterfaceName, InParameter.ParamName);
-		const FName ParamName = InParameter.ParamName;
-		return SetParameterWithLiteral(ParamName, Frontend::ConvertParameterToLiteral(MoveTemp(InParameter)));
-	}
-
 	bool FMetaSoundParameterTransmitter::SetParameterWithLiteral(FName InParameterName, const FLiteral& InLiteral)
 	{
 		if (ISender* Sender = FindSender(InParameterName))
diff --git a/Engine/Plugins/Runtime/Metasound/Source/MetasoundFrontend/Private/MetasoundSourceInterface.cpp b/Engine/Plugins/Runtime/Metasound/Source/MetasoundFrontend/Private/MetasoundSourceInterface.cpp
new file mode 100644
--- /dev/null
+++ b/Engine/Plugins/Runtime/Metasound/Source/MetasoundFrontend/Private/MetasoundSourceInterface.cpp
@@ -0,0 +1,113 @@
+// Copyright Epic Games, Inc. All Rights Reserved.
+#include "MetasoundSourceInterface.h"
+
+#include "MetasoundDataReference.h"
+#include "MetasoundPrimitives.h"
+#include "MetasoundTrigger.h"
+#include "UObject/Class.h"
+#include "Templates/SharedPointer.h"
+
+
+#define LOCTEXT_NAMESPACE "Metasound"
+#define AUDIO_PARAMETER_INTERFACE_NAMESPACE "UE.Source"
+namespace Metasound
+{
+	namespace Frontend
+	{
+		namespace SourceInterface
+		{
+			const FMetasoundFrontendVersion& GetVersion()
+			{
+				static const FMetasoundFrontendVersion Version = { AUDIO_PARAMETER_INTERFACE_NAMESPACE, { 1, 0 }};
+				return Version;
+			}
+
+			namespace Inputs
+			{
+				const FName OnPlay = AUDIO_PARAMETER_INTERFACE_MEMBER_DEFINE("OnPlay");
+			}
+
+			namespace Outputs
+			{
+				const FName OnFinished = AUDIO_PARAMETER_INTERFACE_MEMBER_DEFINE("OnFinished");
+			}
+
+			namespace Environment
+			{
+				const FName DeviceID = AUDIO_PARAMETER_INTERFACE_MEMBER_DEFINE("AudioDeviceID");
+				const FName GraphName = AUDIO_PARAMETER_INTERFACE_MEMBER_DEFINE("GraphName");
+				const FName IsPreview = AUDIO_PARAMETER_INTERFACE_MEMBER_DEFINE("IsPreviewSound");
+				const FName SoundUniqueID = AUDIO_PARAMETER_INTERFACE_MEMBER_DEFINE("SoundUniqueID");
+				const FName TransmitterID = AUDIO_PARAMETER_INTERFACE_MEMBER_DEFINE("TransmitterID");
+			}
+
+			Audio::FParameterInterfacePtr CreateInterface(const UClass& InUClass)
+			{
+				struct FInterface : public Audio::FParameterInterface
+				{
+					FInterface(const UClass& InAssetClass)
+						: FParameterInterface(SourceInterface::GetVersion().Name, SourceInterface::GetVersion().Number.ToInterfaceVersion(), InAssetClass)
+					{
+						Inputs =
+						{
+							{
+								LOCTEXT("OnPlay", "On Play"),
+								LOCTEXT("OnPlayDescription", "Trigger executed when the source is played."),
+								GetMetasoundDataTypeName<FTrigger>(),
+								{ Inputs::OnPlay, false }
+							}
+						};
+
+						Outputs =
+						{
+							{
+								LOCTEXT("OnFinished", "On Finished"),
+								LOCTEXT("OnFinishedDescription", "Trigger executed to initiate stopping the source."),
+								GetMetasoundDataTypeName<FTrigger>(),
+								Outputs::OnFinished
+							}
+						};
+
+						Environment =
+						{
+							{
+								LOCTEXT("AudioDeviceIDDisplayName", "Audio Device ID"),
+								LOCTEXT("AudioDeviceIDDescription", "ID of AudioDevice source is played from."),
+								FName(), // TODO: Align environment data types with environment (ex. this is actually set/get as a uint32)
+								Environment::DeviceID
+							},
+							{
+								LOCTEXT("GraphNameDisplayName", "Graph Name"),
+								LOCTEXT("AudioDeviceIDDescription", "Name of source graph (for debugging/logging)."),
+								GetMetasoundDataTypeName<FString>(),
+								Environment::GraphName
+							},
+							{
+								LOCTEXT("IsPreviewSoundDisplayName", "Is Preview Sound"),
+								LOCTEXT("IsPreviewSoundDescription", "Whether source is being played as a previewed sound."),
+								GetMetasoundDataTypeName<bool>(),
+								Environment::IsPreview
+							},
+							{
+								LOCTEXT("TransmitterIDDisplayName", "Transmitter ID"),
+								LOCTEXT("TransmitterIDDescription", "ID used by Transmission System to generate a unique send address for each source instance."),
+								FName(), // TODO: Align environment data types with environment (ex. this is actually set/get as a uint64)
+								Environment::TransmitterID
+							},
+							{
+								LOCTEXT("AudioDeviceIDDisplayName", "Sound Unique ID"),
+								LOCTEXT("AudioDeviceIDDescription", "ID of unique source instance."),
+								FName(), // TODO: Align environment data types with environment (ex. this is actually set/get as a uint32)
+								Environment::SoundUniqueID
+							}
+						};
+					}
+				};
+
+				return MakeShared<FInterface>(InUClass);
+			}
+		} // namespace SourceInterface
+	} // namespace Frontend
+} // namespace Metasound
+#undef AUDIO_PARAMETER_INTERFACE_NAMESPACE
+#define LOCTEXT_NAMESPACE "Metasound"
diff --git a/Engine/Plugins/Runtime/Metasound/Source/MetasoundFrontend/Public/MetasoundArrayRandomNode.h b/Engine/Plugins/Runtime/Metasound/Source/MetasoundFrontend/Public/MetasoundArrayRandomNode.h
--- a/Engine/Plugins/Runtime/Metasound/Source/MetasoundFrontend/Public/MetasoundArrayRandomNode.h
+++ b/Engine/Plugins/Runtime/Metasound/Source/MetasoundFrontend/Public/MetasoundArrayRandomNode.h
@@ -8,6 +8,7 @@
 #include "MetasoundArrayNodes.h"
 #include "MetasoundExecutableOperator.h"
 #include "MetasoundFacade.h"
+#include "MetasoundSourceInterface.h"
 #include "MetasoundNodeInterface.h"
 #include "MetasoundTrigger.h"
 
@@ -190,10 +191,12 @@ namespace Metasound
 			, TriggerOnNext(FTriggerWriteRef::CreateNew(InParams.OperatorSettings))
 			, TriggerOnReset(FTriggerWriteRef::CreateNew(InParams.OperatorSettings))
 			, OutValue(TDataWriteReferenceFactory<ElementType>::CreateAny(InParams.OperatorSettings))
+		{
+			using namespace Frontend;
+
 #if WITH_METASOUND_DEBUG_ENVIRONMENT
-			, GraphName(*InParams.Environment.GetValue<FString>("GraphName"))
+			GraphName = *InParams.Environment.GetValue<FString>(SourceInterface::Environment::GraphName);
 #endif // WITH_METASOUND_DEBUG_ENVIRONMENT
-		{
 
 			// Check to see if this is a global shuffler or a local one. 
 			// Global shuffler will use a namespace to opt into it.
@@ -210,10 +213,10 @@ namespace Metasound
 				if (*bEnableSharedState)
 				{
 					// Get the environment variable for the unique ID of the sound
-					SharedStateUniqueId = InParams.Environment.GetValue<uint32>(TEXT("SoundUniqueId"));
+					SharedStateUniqueId = InParams.Environment.GetValue<uint32>(SourceInterface::Environment::SoundUniqueID);
 					check(SharedStateUniqueId != INDEX_NONE);
 
-					bIsPreviewSound = InParams.Environment.GetValue<bool>(TEXT("IsPreviewSound"));
+					bIsPreviewSound = InParams.Environment.GetValue<bool>(SourceInterface::Environment::IsPreview);
 
 					FSharedStateRandomGetManager& RGM = FSharedStateRandomGetManager::Get();
 
diff --git a/Engine/Plugins/Runtime/Metasound/Source/MetasoundFrontend/Public/MetasoundArrayShuffleNode.h b/Engine/Plugins/Runtime/Metasound/Source/MetasoundFrontend/Public/MetasoundArrayShuffleNode.h
--- a/Engine/Plugins/Runtime/Metasound/Source/MetasoundFrontend/Public/MetasoundArrayShuffleNode.h
+++ b/Engine/Plugins/Runtime/Metasound/Source/MetasoundFrontend/Public/MetasoundArrayShuffleNode.h
@@ -13,6 +13,7 @@
 #include "MetasoundNodeInterface.h"
 #include "MetasoundOperatorInterface.h"
 #include "MetasoundPrimitives.h"
+#include "MetasoundSourceInterface.h"
 #include "MetasoundTrigger.h"
 #include "MetasoundVertex.h"
 #include "Misc/ScopeLock.h"
@@ -222,6 +223,8 @@ namespace Metasound
 			, TriggerOnReset(FTriggerWriteRef::CreateNew(InParams.OperatorSettings))
 			, OutValue(TDataWriteReferenceFactory<ElementType>::CreateAny(InParams.OperatorSettings))
 		{
+			using namespace Frontend;
+
 			// Check to see if this is a global shuffler or a local one. 
 			// Global shuffler will use a namespace to opt into it.
 			PrevSeedValue = *SeedValue;
@@ -234,7 +237,7 @@ namespace Metasound
 				if (*bEnableSharedState)
 				{
 					// Get the environment variable for the unique ID of the sound
-					SharedStateUniqueId = InParams.Environment.GetValue<uint32>(TEXT("SoundUniqueId"));
+					SharedStateUniqueId = InParams.Environment.GetValue<uint32>(SourceInterface::Environment::SoundUniqueID);
 					check(SharedStateUniqueId != INDEX_NONE);
 
 					FSharedStateShuffleManager& SM = FSharedStateShuffleManager::Get();
diff --git a/Engine/Plugins/Runtime/Metasound/Source/MetasoundFrontend/Public/MetasoundAssetBase.h b/Engine/Plugins/Runtime/Metasound/Source/MetasoundFrontend/Public/MetasoundAssetBase.h
--- a/Engine/Plugins/Runtime/Metasound/Source/MetasoundFrontend/Public/MetasoundAssetBase.h
+++ b/Engine/Plugins/Runtime/Metasound/Source/MetasoundFrontend/Public/MetasoundAssetBase.h
@@ -75,13 +75,12 @@ public:
 	virtual void SetMetadata(FMetasoundFrontendClassMetadata& InMetadata);
 
 	// Returns a default interface required when a given asset is created.
-	virtual const FMetasoundFrontendVersion& GetDefaultArchetypeVersion() const = 0;
 
 	// Returns the interface entries declared by the given asset's document from the InterfaceRegistry.
 	bool GetDeclaredInterfaces(TArray<const Metasound::Frontend::IInterfaceRegistryEntry*>& OutInterfaces) const;
 
-	// Returns whether an interface with the given name is declared by the given asset's document.
-	bool IsInterfaceDeclared(FName InName) const;
+	// Returns whether an interface with the given version is declared by the given asset's document.
+	bool IsInterfaceDeclared(const FMetasoundFrontendVersion& InVersion) const;
 
 	// Gets the asset class info.
 	virtual Metasound::Frontend::FNodeClassInfo GetAssetClassInfo() const = 0;
diff --git a/Engine/Plugins/Runtime/Metasound/Source/MetasoundFrontend/Public/MetasoundAssetManager.h b/Engine/Plugins/Runtime/Metasound/Source/MetasoundFrontend/Public/MetasoundAssetManager.h
--- a/Engine/Plugins/Runtime/Metasound/Source/MetasoundFrontend/Public/MetasoundAssetManager.h
+++ b/Engine/Plugins/Runtime/Metasound/Source/MetasoundFrontend/Public/MetasoundAssetManager.h
@@ -80,7 +80,7 @@ namespace Metasound
 			virtual void AddAssetReferences(FMetasoundAssetBase& InAssetBase) = 0;
 
 			// Add or Update a MetaSound Asset's entry data
-			virtual void AddOrUpdateAsset(UObject& InObject, bool bInRegisterWithFrontend = true) = 0;
+			virtual void AddOrUpdateAsset(UObject& InObject) = 0;
 
 			// Whether or not the class is eligible for auto-update
 			virtual bool CanAutoUpdate(const FMetasoundFrontendClassName& InClassName) const = 0;
diff --git a/Engine/Plugins/Runtime/Metasound/Source/MetasoundFrontend/Public/MetasoundFrontendController.h b/Engine/Plugins/Runtime/Metasound/Source/MetasoundFrontend/Public/MetasoundFrontendController.h
--- a/Engine/Plugins/Runtime/Metasound/Source/MetasoundFrontend/Public/MetasoundFrontendController.h
+++ b/Engine/Plugins/Runtime/Metasound/Source/MetasoundFrontend/Public/MetasoundFrontendController.h
@@ -986,7 +986,7 @@ namespace Metasound
 			 */
 			virtual FGraphHandle AddDuplicateSubgraph(const IGraphController& InGraph) = 0;
 
-			virtual const TArray<FMetasoundFrontendVersion>& GetInterfaceVersions() const = 0;
+			virtual const TSet<FMetasoundFrontendVersion>& GetInterfaceVersions() const = 0;
 			virtual void AddInterfaceVersion(const FMetasoundFrontendVersion& InVersion) = 0;
 			virtual void RemoveInterfaceVersion(const FMetasoundFrontendVersion& InVersion) = 0;
 			virtual void ClearInterfaceVersions() = 0;
diff --git a/Engine/Plugins/Runtime/Metasound/Source/MetasoundFrontend/Public/MetasoundFrontendDocument.h b/Engine/Plugins/Runtime/Metasound/Source/MetasoundFrontend/Public/MetasoundFrontendDocument.h
--- a/Engine/Plugins/Runtime/Metasound/Source/MetasoundFrontend/Public/MetasoundFrontendDocument.h
+++ b/Engine/Plugins/Runtime/Metasound/Source/MetasoundFrontend/Public/MetasoundFrontendDocument.h
@@ -10,6 +10,8 @@
 #include "MetasoundNodeInterface.h"
 #include "MetasoundVertex.h"
 #include "Misc/Guid.h"
+#include "IAudioParameterInterfaceRegistry.h"
+#include "Templates/TypeHash.h"
 
 #include "MetasoundFrontendDocument.generated.h"
 
@@ -77,11 +79,11 @@ struct METASOUNDFRONTEND_API FMetasoundFrontendVersionNumber
 	GENERATED_BODY()
 
 	// Major version number.
-	UPROPERTY(EditAnywhere, Category = General)
+	UPROPERTY(VisibleAnywhere, Category = General)
 	int32 Major = 1;
 
 	// Minor version number.
-	UPROPERTY(EditAnywhere, Category = General)
+	UPROPERTY(VisibleAnywhere, Category = General)
 	int32 Minor = 0;
 
 	static const FMetasoundFrontendVersionNumber& GetInvalid()
@@ -95,6 +97,11 @@ struct METASOUNDFRONTEND_API FMetasoundFrontendVersionNumber
 		return *this != GetInvalid();
 	}
 
+	Audio::FParameterInterface::FVersion ToInterfaceVersion() const
+	{
+		return Audio::FParameterInterface::FVersion { Major, Minor };
+	}
+
 	friend bool operator==(const FMetasoundFrontendVersionNumber& InLHS, const FMetasoundFrontendVersionNumber& InRHS)
 	{
 		return InLHS.Major == InRHS.Major && InLHS.Minor == InRHS.Minor;
@@ -151,6 +158,11 @@ struct METASOUNDFRONTEND_API FMetasoundFrontendVersionNumber
 	}
 };
 
+FORCEINLINE uint32 GetTypeHash(const FMetasoundFrontendVersionNumber& InNumber)
+{
+	return HashCombineFast(GetTypeHash(InNumber.Major), GetTypeHash(InNumber.Minor));
+}
+
 // General purpose version info for Metasound Frontend objects.
 USTRUCT()
 struct METASOUNDFRONTEND_API FMetasoundFrontendVersion
@@ -158,11 +170,11 @@ struct METASOUNDFRONTEND_API FMetasoundFrontendVersion
 	GENERATED_BODY()
 
 	// Name of version.
-	UPROPERTY(EditAnywhere, Category = CustomView)
+	UPROPERTY(VisibleAnywhere, Category = CustomView)
 	FName Name;
 
 	// Version number.
-	UPROPERTY(EditAnywhere, Category = CustomView)
+	UPROPERTY(VisibleAnywhere, Category = CustomView)
 	FMetasoundFrontendVersionNumber Number;
 
 	FString ToString() const;
@@ -222,6 +234,10 @@ struct METASOUNDFRONTEND_API FMetasoundFrontendVersion
 	}
 };
 
+FORCEINLINE uint32 GetTypeHash(const FMetasoundFrontendVersion& InVersion)
+{
+	return HashCombineFast(GetTypeHash(InVersion.Name), GetTypeHash(InVersion.Number));
+}
 
 // An FMetasoundFrontendVertex provides a named connection point of a node.
 USTRUCT() 
@@ -236,7 +252,7 @@ struct METASOUNDFRONTEND_API FMetasoundFrontendVertex
 	FName Name;
 
 	// Data type name of the vertex.
-	UPROPERTY(EditAnywhere, Category = Parameters)
+	UPROPERTY(VisibleAnywhere, Category = Parameters)
 	FName TypeName;
 
 	// ID of vertex
@@ -254,10 +270,10 @@ struct FMetasoundFrontendVertexLiteral
 	GENERATED_BODY()
 
 	// ID of vertex.
-	UPROPERTY(EditAnywhere, Category = Parameters)
+	UPROPERTY(VisibleAnywhere, Category = Parameters)
 	FGuid VertexID = Metasound::FrontendInvalidID;
 
-	// Value to use when constructing input. 
+	// Value to use when constructing input.
 	UPROPERTY(EditAnywhere, Category = Parameters)
 	FMetasoundFrontendLiteral Value;
 };
@@ -574,6 +590,9 @@ struct METASOUNDFRONTEND_API FMetasoundFrontendClassVertex : public FMetasoundFr
 	UPROPERTY(EditAnywhere, Category = CustomView)
 	FMetasoundFrontendVertexMetadata Metadata;
 
+	// Splits name into namespace & parameter name
+	void SplitName(FName& OutNamespace, FName& OutParameterName) const;
+
 	static bool IsFunctionalEquivalent(const FMetasoundFrontendClassVertex& InLHS, const FMetasoundFrontendClassVertex& InRHS);
 };
 
@@ -1138,11 +1157,9 @@ struct METASOUNDFRONTEND_API FMetasoundFrontendDocument
 	UPROPERTY()
 	FMetasoundFrontendEditorData EditorData;
 
-	UPROPERTY(meta = (DeprecatedProperty, DeprecationMessage = "5.0 - ArchetypeVersion has been migrated to InterfaceVersions array."))
-	FMetasoundFrontendVersion ArchetypeVersion;
-
-	UPROPERTY(EditAnywhere, Category = CustomView, meta = (DisplayName = "Interfaces"))
-	TArray<FMetasoundFrontendVersion> InterfaceVersions;
+public:
+	UPROPERTY(VisibleAnywhere, Category = CustomView)
+	TSet<FMetasoundFrontendVersion> Interfaces;
 
 	UPROPERTY(EditAnywhere, Category = CustomView)
 	FMetasoundFrontendGraphClass RootGraph;
@@ -1152,4 +1169,33 @@ struct METASOUNDFRONTEND_API FMetasoundFrontendDocument
 
 	UPROPERTY()
 	TArray<FMetasoundFrontendClass> Dependencies;
+
+private:
+	UPROPERTY(meta = (DeprecatedProperty, DeprecationMessage = "5.0 - ArchetypeVersion has been migrated to InterfaceVersions array."))
+	FMetasoundFrontendVersion ArchetypeVersion;
+
+	UPROPERTY(meta = (DeprecatedProperty, DeprecationMessage = "5.0 - InterfaceVersions has been migrated to Interfaces set."))
+	TArray<FMetasoundFrontendVersion> InterfaceVersions;
+
+public:
+	// Data migration for 5.0 Early Access data. ArchetypeVersion/InterfaceVersions properties can be removed post 5.0 release
+	// and this fix-up can be removed post 5.0 release.
+	bool VersionInterfaces()
+	{
+		bool bDidEdit = false;
+		if (ArchetypeVersion.IsValid())
+		{
+			Interfaces.Add(ArchetypeVersion);
+			ArchetypeVersion = FMetasoundFrontendVersion::GetInvalid();
+			bDidEdit = true;
+		}
+		if (!InterfaceVersions.IsEmpty())
+		{
+			Interfaces.Append(InterfaceVersions);
+			InterfaceVersions.Reset();
+			bDidEdit = true;
+		}
+
+		return bDidEdit;
+	}
 };
diff --git a/Engine/Plugins/Runtime/Metasound/Source/MetasoundFrontend/Public/MetasoundFrontendLiteral.h b/Engine/Plugins/Runtime/Metasound/Source/MetasoundFrontend/Public/MetasoundFrontendLiteral.h
--- a/Engine/Plugins/Runtime/Metasound/Source/MetasoundFrontend/Public/MetasoundFrontendLiteral.h
+++ b/Engine/Plugins/Runtime/Metasound/Source/MetasoundFrontend/Public/MetasoundFrontendLiteral.h
@@ -2,7 +2,7 @@
 
 #pragma once
 
-#include "AudioParameterInterface.h"
+#include "AudioParameterControllerInterface.h"
 #include "CoreMinimal.h"
 #include "MetasoundLiteral.h"
 
diff --git a/Engine/Plugins/Runtime/Metasound/Source/MetasoundFrontend/Public/MetasoundFrontendTransform.h b/Engine/Plugins/Runtime/Metasound/Source/MetasoundFrontend/Public/MetasoundFrontendTransform.h
--- a/Engine/Plugins/Runtime/Metasound/Source/MetasoundFrontend/Public/MetasoundFrontendTransform.h
+++ b/Engine/Plugins/Runtime/Metasound/Source/MetasoundFrontend/Public/MetasoundFrontendTransform.h
@@ -7,6 +7,7 @@
 #include "MetasoundFrontendController.h"
 #include "MetasoundFrontendDocument.h"
 #include "MetasoundVertex.h"
+#include "Templates/Function.h"
 
 namespace Metasound
 {
@@ -37,11 +38,19 @@ namespace Metasound
 			FModifyRootGraphInterfaces(const TArray<FMetasoundFrontendInterface>& InInterfacesToRemove, const TArray<FMetasoundFrontendInterface>& InInterfacesToAdd);
 			FModifyRootGraphInterfaces(const TArray<FMetasoundFrontendVersion>& InInterfaceVersionsToRemove, const TArray<FMetasoundFrontendVersion>& InInterfaceVersionsToAdd);
 
+			// Whether or not to propagate node locations to new members. Setting to false
+			// results in members not having a default physical location in the editor graph.
 			void SetDefaultNodeLocations(bool bInSetDefaultNodeLocations);
+
+			// Override function used to match removed members with added members, allowing
+			// transform to preserve connections made between removed interface members & new interface members
+			// that may be related but not be named the same.
+			void SetNamePairingFunction(const TFunction<bool(FName, FName)>& InNamePairingFunction);
+
 			bool Transform(FDocumentHandle InDocument) const override;
 
 		private:
-			void Init();
+			void Init(const TFunction<bool(FName, FName)>* InNamePairingFunction = nullptr);
 
 			bool bSetDefaultNodeLocations = true;
 
diff --git a/Engine/Plugins/Runtime/Metasound/Source/MetasoundFrontend/Public/MetasoundOutputFormatInterfaces.h b/Engine/Plugins/Runtime/Metasound/Source/MetasoundFrontend/Public/MetasoundOutputFormatInterfaces.h
new file mode 100644
--- /dev/null
+++ b/Engine/Plugins/Runtime/Metasound/Source/MetasoundFrontend/Public/MetasoundOutputFormatInterfaces.h
@@ -0,0 +1,38 @@
+// Copyright Epic Games, Inc. All Rights Reserved.
+#pragma once
+
+#include "IAudioParameterInterfaceRegistry.h"
+#include "MetasoundFrontendDocument.h"
+#include "UObject/NameTypes.h"
+
+// Forward Declarations
+class UClass;
+
+namespace Metasound
+{
+	namespace Frontend
+	{
+		namespace OutputFormatMonoInterface
+		{
+			namespace Outputs
+			{
+				METASOUNDFRONTEND_API const extern FName MonoOut;
+			} // namespace Outputs
+
+			METASOUNDFRONTEND_API const FMetasoundFrontendVersion& GetVersion();
+			METASOUNDFRONTEND_API Audio::FParameterInterfacePtr CreateInterface(const UClass& InClass);
+		} // namespace OutputFormatMonoInterface
+
+		namespace OutputFormatStereoInterface
+		{
+			namespace Outputs
+			{
+				METASOUNDFRONTEND_API const extern FName LeftOut;
+				METASOUNDFRONTEND_API const extern FName RightOut;
+			} // namespace Outputs
+
+			METASOUNDFRONTEND_API const FMetasoundFrontendVersion& GetVersion();
+			METASOUNDFRONTEND_API Audio::FParameterInterfacePtr CreateInterface(const UClass& InClass);
+		} // namespace OutputFormatMonoInterface
+	} // namespace Frontend
+} // namespace Metasound
diff --git a/Engine/Plugins/Runtime/Metasound/Source/MetasoundFrontend/Public/MetasoundParameterTransmitter.h b/Engine/Plugins/Runtime/Metasound/Source/MetasoundFrontend/Public/MetasoundParameterTransmitter.h
--- a/Engine/Plugins/Runtime/Metasound/Source/MetasoundFrontend/Public/MetasoundParameterTransmitter.h
+++ b/Engine/Plugins/Runtime/Metasound/Source/MetasoundFrontend/Public/MetasoundParameterTransmitter.h
@@ -1,7 +1,7 @@
 // Copyright Epic Games, Inc. All Rights Reserved.
 #pragma once
 
-#include "AudioParameterInterface.h"
+#include "AudioParameterControllerInterface.h"
 #include "IAudioParameterTransmitter.h"
 #include "MetasoundDataReference.h"
 #include "MetasoundFrontendLiteral.h"
@@ -71,9 +71,6 @@ namespace Metasound
 
 		};
 
-		/** Returns the MetaSound environment variable name which contains the instance ID. */
-		static const FVertexName& GetInstanceIDEnvironmentVariableName();
-
 		/** Creates a unique send address using the given MetaSound environment. */
 		static FSendAddress CreateSendAddressFromEnvironment(const FMetasoundEnvironment& InEnvironment, const FVertexName& InVertexName, const FName& InTypeName);
 		
@@ -93,7 +90,6 @@ namespace Metasound
 		 * @param InParameter - Parameter to set.
 		 */
 		bool SetParameter(FAudioParameter&& InParameter) override;
-		bool SetParameter(FName InInterfaceName, FAudioParameter&& InParameter) override;
 
 		/** Set a parameter using a literal.
 		 *
diff --git a/Engine/Plugins/Runtime/Metasound/Source/MetasoundFrontend/Public/MetasoundSourceInterface.h b/Engine/Plugins/Runtime/Metasound/Source/MetasoundFrontend/Public/MetasoundSourceInterface.h
new file mode 100644
--- /dev/null
+++ b/Engine/Plugins/Runtime/Metasound/Source/MetasoundFrontend/Public/MetasoundSourceInterface.h
@@ -0,0 +1,41 @@
+// Copyright Epic Games, Inc. All Rights Reserved.
+#pragma once
+
+#include "IAudioParameterInterfaceRegistry.h"
+#include "MetasoundFrontendDocument.h"
+#include "UObject/NameTypes.h"
+
+
+// Forward Declarations
+struct FMetasoundFrontendVersion;
+
+namespace Metasound
+{
+	namespace Frontend
+	{
+		namespace SourceInterface
+		{
+			namespace Inputs
+			{
+				METASOUNDFRONTEND_API const extern FName OnPlay;
+			}
+
+			namespace Outputs
+			{
+				METASOUNDFRONTEND_API const extern FName OnFinished;
+			}
+
+			namespace Environment
+			{
+				METASOUNDFRONTEND_API const extern FName DeviceID;
+				METASOUNDFRONTEND_API const extern FName GraphName;
+				METASOUNDFRONTEND_API const extern FName IsPreview;
+				METASOUNDFRONTEND_API const extern FName SoundUniqueID;
+				METASOUNDFRONTEND_API const extern FName TransmitterID;
+			}
+
+			METASOUNDFRONTEND_API const FMetasoundFrontendVersion& GetVersion();
+			METASOUNDFRONTEND_API Audio::FParameterInterfacePtr CreateInterface(const UClass& InClass);
+		}
+	} // namespace Frontend
+} // namespace Metasound
diff --git a/Engine/Plugins/Runtime/Metasound/Source/MetasoundGenerator/MetasoundGenerator.Build.cs b/Engine/Plugins/Runtime/Metasound/Source/MetasoundGenerator/MetasoundGenerator.Build.cs
--- a/Engine/Plugins/Runtime/Metasound/Source/MetasoundGenerator/MetasoundGenerator.Build.cs
+++ b/Engine/Plugins/Runtime/Metasound/Source/MetasoundGenerator/MetasoundGenerator.Build.cs
@@ -14,6 +14,7 @@ namespace UnrealBuildTool.Rules
 					"SignalProcessing",
 					"CoreUObject",
 					"Engine",
+					"AudioExtensions",
 					"AudioMixerCore"
 				}
 			);
diff --git a/Engine/Plugins/Runtime/Metasound/Source/MetasoundGenerator/Private/MetasoundGenerator.cpp b/Engine/Plugins/Runtime/Metasound/Source/MetasoundGenerator/Private/MetasoundGenerator.cpp
--- a/Engine/Plugins/Runtime/Metasound/Source/MetasoundGenerator/Private/MetasoundGenerator.cpp
+++ b/Engine/Plugins/Runtime/Metasound/Source/MetasoundGenerator/Private/MetasoundGenerator.cpp
@@ -1,12 +1,14 @@
 // Copyright Epic Games, Inc. All Rights Reserved.
 #include "MetasoundGenerator.h"
 
+#include "AudioParameter.h"
 #include "DSP/Dsp.h"
 #include "MetasoundGraph.h"
 #include "MetasoundInputNode.h"
 #include "MetasoundOperatorBuilder.h"
 #include "MetasoundOperatorInterface.h"
 #include "MetasoundOutputNode.h"
+#include "MetasoundSourceInterface.h"
 #include "MetasoundTrace.h"
 #include "MetasoundTrigger.h"
 
@@ -35,8 +37,6 @@ namespace Metasound
 		Environment = {};
 		MetaSoundName = {};
 		AudioOutputNames = {};
-		OnPlayInputName = {};
-		IsFinishedOutputName = {};
 	}
 
 	FAsyncMetaSoundBuilder::FAsyncMetaSoundBuilder(FMetasoundGenerator* InGenerator, FMetasoundGeneratorInitParams&& InInitParams, bool bInTriggerGenerator)
@@ -48,6 +48,9 @@ namespace Metasound
 
 	void FAsyncMetaSoundBuilder::DoWork()
 	{
+		using namespace Audio;
+		using namespace Frontend;
+
 		METASOUND_TRACE_CPUPROFILER_EVENT_SCOPE_TEXT(*FString::Printf(TEXT("AsyncMetaSoundBuilder::DoWork %s"), *InitParams.MetaSoundName));
 
 		// Create an instance of the new graph
@@ -84,8 +87,8 @@ namespace Metasound
 
 			// References must be cached before moving the operator to the InitParams
 			FDataReferenceCollection Inputs = GraphOperator->GetInputs();
-			FTriggerWriteRef PlayTrigger = Inputs.GetDataWriteReferenceOrConstruct<FTrigger>(InitParams.OnPlayInputName, InitParams.OperatorSettings, false);
-			FTriggerReadRef FinishTrigger = Outputs.GetDataReadReferenceOrConstruct<FTrigger>(InitParams.IsFinishedOutputName, InitParams.OperatorSettings, false);
+			FTriggerWriteRef PlayTrigger = Inputs.GetDataWriteReferenceOrConstruct<FTrigger>(SourceInterface::Inputs::OnPlay, InitParams.OperatorSettings, false);
+			FTriggerReadRef FinishTrigger = Outputs.GetDataReadReferenceOrConstruct<FTrigger>(SourceInterface::Outputs::OnFinished, InitParams.OperatorSettings, false);
 
 			FMetasoundGeneratorData GeneratorData
 			{
diff --git a/Engine/Plugins/Runtime/Metasound/Source/MetasoundGenerator/Public/MetasoundGenerator.h b/Engine/Plugins/Runtime/Metasound/Source/MetasoundGenerator/Public/MetasoundGenerator.h
--- a/Engine/Plugins/Runtime/Metasound/Source/MetasoundGenerator/Public/MetasoundGenerator.h
+++ b/Engine/Plugins/Runtime/Metasound/Source/MetasoundGenerator/Public/MetasoundGenerator.h
@@ -25,8 +25,6 @@ namespace Metasound
 		FMetasoundEnvironment Environment;
 		FString MetaSoundName;
 		TArray<FVertexName> AudioOutputNames;
-		FVertexName OnPlayInputName;
-		FVertexName IsFinishedOutputName;
 
 		void Release();
 	};
diff --git a/Engine/Plugins/Runtime/Metasound/Source/MetasoundGraphCore/Public/MetasoundEnvironment.h b/Engine/Plugins/Runtime/Metasound/Source/MetasoundGraphCore/Public/MetasoundEnvironment.h
--- a/Engine/Plugins/Runtime/Metasound/Source/MetasoundGraphCore/Public/MetasoundEnvironment.h
+++ b/Engine/Plugins/Runtime/Metasound/Source/MetasoundGraphCore/Public/MetasoundEnvironment.h
@@ -3,8 +3,9 @@
 #pragma once
 
 #include "Containers/SortedMap.h"
-#include "CoreMinimal.h"
+#include "IAudioParameterInterfaceRegistry.h"
 #include "Misc/Build.h"
+
 #include <type_traits>
 
 #ifndef WITH_METASOUND_DEBUG_ENVIRONMENT
@@ -72,8 +73,8 @@ namespace Metasound
 	private:
 		// Only allow TMetasoundEnvironmentVariable to implement this class
 		// by making constructor private so that casting can be effectively checked. 
-		template<typename Type> 
-        friend class TMetasoundEnvironmentVariable;
+		template<typename Type>
+		friend class TMetasoundEnvironmentVariable;
 
 		IMetasoundEnvironmentVariable() = default;
 
@@ -93,7 +94,7 @@ namespace Metasound
 	template<typename Type>
 	class TMetasoundEnvironmentVariable : public IMetasoundEnvironmentVariable
 	{
-		// Make sure that the `Type` is not a reference or CV qualified. 
+		// Make sure that the `Type` is not a reference or CV qualified.
 		static_assert(std::is_same<Type, typename std::decay<Type>::type>::value, "Type must not change when decayed");
 
 		// Make sure the `Type` has appropriate construction and copy characteristics. 
@@ -178,7 +179,7 @@ namespace Metasound
 	}
 
 	/** FMetasoundEnvironment contains a set of TMetasoundEnvironmentVariables requiring
-	 * that each environment variable has a unique name. 
+	 * that each environment variable has a unique name.
 	 */
 	class METASOUNDGRAPHCORE_API FMetasoundEnvironment
 	{
@@ -193,8 +194,8 @@ namespace Metasound
 		/** Returns true if the environment variable with the given name contains
 		 * the data of the same type as `VarType`
 		 *
-		 * @param InVariableName - Name of environment variable. 
-		 * @tparam VarType - Type of the underlying data stored in the environment variable. 
+		 * @param InVariableName - Name of environment variable.
+		 * @tparam VarType - Type of the underlying data stored in the environment variable.
 		 */
 		template<typename VarType>
 		bool IsType(const FName& InVariableName) const
@@ -206,7 +207,7 @@ namespace Metasound
 		 * and the type `VarType`.
 		 *
 		 * @param InVariableName - Name of environment variable. 
-		 * @tparam VarType - Type of the underlying data stored in the environment variable. 
+		 * @tparam VarType - Type of the underlying data stored in the environment variable.
 		 */
 		template<typename VarType>
 		bool Contains(const FName& InVariableName) const
@@ -218,10 +219,10 @@ namespace Metasound
 			return false;
 		}
 
-		/** Returns the environment variable data. 
+		/** Returns the environment variable data.
 		 *
-		 * @param InVariableName - Name of environment variable. 
-		 * @tparam VarType - Type of the underlying data stored in the environment variable. 
+		 * @param InVariableName - Name of environment variable.
+		 * @tparam VarType - Type of the underlying data stored in the environment variable.
 		 *
 		 * @return copy of the underlying data.
 		 */
@@ -239,8 +240,8 @@ namespace Metasound
 
 		/** Sets the environment variable data
 		 *
-		 * @param InVariableName - Name of environment variable. 
-		 * @tparam VarType - Type of the underlying data stored in the environment variable. 
+		 * @param InVariableName - Name of environment variable.
+		 * @tparam VarType - Type of the underlying data stored in the environment variable.
 		 *
 		 * @return const ref to the underlying data.
 		 */
diff --git a/Engine/Source/Runtime/AudioExtensions/Private/AudioParameter.cpp b/Engine/Source/Runtime/AudioExtensions/Private/AudioParameter.cpp
new file mode 100644
--- /dev/null
+++ b/Engine/Source/Runtime/AudioExtensions/Private/AudioParameter.cpp
@@ -0,0 +1,232 @@
+// Copyright Epic Games, Inc. All Rights Reserved.
+#include "AudioParameter.h"
+
+
+namespace Audio
+{
+	namespace ParameterPrivate
+	{
+		template <typename T>
+		void SetOrMergeArray(const TArray<T>& InArray, TArray<T>& OutArray, bool bInMerge)
+		{
+			if (bInMerge)
+			{
+				OutArray.Append(InArray);
+			}
+			else
+			{
+				OutArray = InArray;
+			}
+		}
+	} // ParameterPrivate
+
+	const FString FParameterPath::NamespaceDelimiter = TEXT(AUDIO_PARAMETER_NAMESPACE_PATH_DELIMITER);
+
+	FName FParameterPath::CombineNames(FName InLeft, FName InRight)
+	{
+		if (InLeft.IsNone())
+		{
+			return InRight;
+		}
+
+		const FString FullName = FString::Join(TArray<FString>({ *InLeft.ToString(), *InRight.ToString() }), *NamespaceDelimiter);
+		return FName(FullName);
+	}
+
+	void FParameterPath::SplitName(FName InFullName, FName& OutNamespace, FName& OutParameterName)
+	{
+		FString FullName = InFullName.ToString();
+		const int32 IndexOfDelim = FullName.Find(NamespaceDelimiter, ESearchCase::IgnoreCase, ESearchDir::FromEnd);
+		if (IndexOfDelim != INDEX_NONE)
+		{
+			OutNamespace = FName(*FullName.Left(IndexOfDelim));
+			OutParameterName = FName(*FullName.RightChop(IndexOfDelim + 1));
+		}
+		else
+		{
+			OutNamespace = FName();
+			OutParameterName = InFullName;
+		}
+	}
+} // namespace Audio
+
+void FAudioParameter::Merge(const FAudioParameter& InParameter, bool bInTakeName, bool bInTakeType, bool bInMergeArrayTypes)
+{
+	if (bInTakeName)
+	{
+		ParamName = InParameter.ParamName;
+	}
+
+	if (bInTakeType)
+	{
+		ParamType = InParameter.ParamType;
+	}
+
+	switch (InParameter.ParamType)
+	{
+		case EAudioParameterType::Boolean:
+		{
+			BoolParam = InParameter.BoolParam;
+		}
+		break;
+
+		case EAudioParameterType::BooleanArray:
+		{
+			Audio::ParameterPrivate::SetOrMergeArray(InParameter.ArrayBoolParam, ArrayBoolParam, bInMergeArrayTypes);
+		}
+		break;
+
+		case EAudioParameterType::Float:
+		{
+			FloatParam = InParameter.FloatParam;
+		}
+		break;
+
+		case EAudioParameterType::FloatArray:
+		{
+			Audio::ParameterPrivate::SetOrMergeArray(InParameter.ArrayFloatParam, ArrayFloatParam, bInMergeArrayTypes);
+		}
+		break;
+
+		case EAudioParameterType::Integer:
+		case EAudioParameterType::NoneArray:
+		{
+			if (bInMergeArrayTypes)
+			{
+				IntParam += InParameter.IntParam;
+			}
+			else
+			{
+				IntParam = InParameter.IntParam;
+			}
+		}
+		break;
+
+		case EAudioParameterType::IntegerArray:
+		{
+			Audio::ParameterPrivate::SetOrMergeArray(InParameter.ArrayIntParam, ArrayIntParam, bInMergeArrayTypes);
+		}
+		break;
+
+		case EAudioParameterType::None:
+		{
+			FloatParam = InParameter.FloatParam;
+			BoolParam = InParameter.BoolParam;
+			IntParam = InParameter.IntParam;
+			ObjectParam = InParameter.ObjectParam;
+			StringParam = InParameter.StringParam;
+
+			Audio::ParameterPrivate::SetOrMergeArray(InParameter.ArrayBoolParam, ArrayBoolParam, bInMergeArrayTypes);
+			Audio::ParameterPrivate::SetOrMergeArray(InParameter.ArrayFloatParam, ArrayFloatParam, bInMergeArrayTypes);
+			Audio::ParameterPrivate::SetOrMergeArray(InParameter.ArrayIntParam, ArrayIntParam, bInMergeArrayTypes);
+			Audio::ParameterPrivate::SetOrMergeArray(InParameter.ArrayObjectParam, ArrayObjectParam, bInMergeArrayTypes);
+			Audio::ParameterPrivate::SetOrMergeArray(InParameter.ArrayStringParam, ArrayStringParam, bInMergeArrayTypes);
+
+			if (!bInMergeArrayTypes)
+			{
+				ObjectProxies.Reset();
+			}
+
+			for (const Audio::IProxyDataPtr& ProxyPtr : InParameter.ObjectProxies)
+			{
+				if (ensure(ProxyPtr.IsValid()))
+				{
+					ObjectProxies.Emplace(ProxyPtr->Clone());
+				}
+			}
+		}
+		break;
+
+		case EAudioParameterType::Object:
+		{
+			ObjectParam = InParameter.ObjectParam;
+
+			ObjectProxies.Reset();
+			for (const Audio::IProxyDataPtr& ProxyPtr : InParameter.ObjectProxies)
+			{
+				if (ensure(ProxyPtr.IsValid()))
+				{
+					ObjectProxies.Emplace(ProxyPtr->Clone());
+				}
+			}
+		}
+		break;
+
+		case EAudioParameterType::ObjectArray:
+		{
+			Audio::ParameterPrivate::SetOrMergeArray(InParameter.ArrayObjectParam, ArrayObjectParam, bInMergeArrayTypes);
+
+			if (!bInMergeArrayTypes)
+			{
+				ObjectProxies.Reset();
+			}
+
+			for (const Audio::IProxyDataPtr& ProxyPtr : InParameter.ObjectProxies)
+			{
+				if (ensure(ProxyPtr.IsValid()))
+				{
+					ObjectProxies.Emplace(ProxyPtr->Clone());
+				}
+			}
+		}
+		break;
+
+		case EAudioParameterType::String:
+		{
+			StringParam = InParameter.StringParam;
+		}
+		break;
+
+		case EAudioParameterType::StringArray:
+		{
+			Audio::ParameterPrivate::SetOrMergeArray(InParameter.ArrayStringParam, ArrayStringParam, bInMergeArrayTypes);
+		}
+		break;
+
+		default:
+			break;
+	}
+}
+
+void FAudioParameter::Merge(TArray<FAudioParameter>&& InParams, TArray<FAudioParameter>& OutParams)
+{
+	if (InParams.IsEmpty())
+	{
+		return;
+	}
+
+	if (OutParams.IsEmpty())
+	{
+		OutParams.Append(MoveTemp(InParams));
+		return;
+	}
+
+	auto SortParamsPredicate = [](const FAudioParameter& A, const FAudioParameter& B) { return A.ParamName.FastLess(B.ParamName); };
+
+	InParams.Sort(SortParamsPredicate);
+	OutParams.Sort(SortParamsPredicate);
+
+	for (int32 i = OutParams.Num() - 1; i >= 0; --i)
+	{
+		while (!InParams.IsEmpty())
+		{
+			FAudioParameter& OutParam = OutParams[i];
+			if (InParams.Last().ParamName.FastLess(OutParam.ParamName))
+			{
+				break;
+			}
+
+			constexpr bool bAllowShrinking = false;
+			FAudioParameter NewParam = InParams.Pop(bAllowShrinking);
+			if (NewParam.ParamName == OutParam.ParamName)
+			{
+				NewParam.Merge(OutParam);
+				OutParam = MoveTemp(NewParam);
+			}
+			else
+			{
+				OutParams.Emplace(MoveTemp(NewParam));
+			}
+		}
+	}
+}
diff --git a/Engine/Source/Runtime/AudioExtensions/Private/AudioParameterControllerInterface.cpp b/Engine/Source/Runtime/AudioExtensions/Private/AudioParameterControllerInterface.cpp
new file mode 100644
--- /dev/null
+++ b/Engine/Source/Runtime/AudioExtensions/Private/AudioParameterControllerInterface.cpp
@@ -0,0 +1,9 @@
+// Copyright Epic Games, Inc. All Rights Reserved.
+
+#include "AudioParameterControllerInterface.h"
+
+
+UAudioParameterControllerInterface::UAudioParameterControllerInterface(FObjectInitializer const& InObjectInitializer)
+	: UInterface(InObjectInitializer)
+{
+}
diff --git a/Engine/Source/Runtime/AudioExtensions/Private/AudioParameterInterface.cpp b/Engine/Source/Runtime/AudioExtensions/Private/AudioParameterInterface.cpp
deleted file mode 100644
diff --git a/Engine/Source/Runtime/AudioExtensions/Private/IAudioGeneratorInterfaceRegistry.cpp b/Engine/Source/Runtime/AudioExtensions/Private/IAudioGeneratorInterfaceRegistry.cpp
deleted file mode 100644
diff --git a/Engine/Source/Runtime/AudioExtensions/Private/IAudioParameterInterfaceRegistry.cpp b/Engine/Source/Runtime/AudioExtensions/Private/IAudioParameterInterfaceRegistry.cpp
new file mode 100644
--- /dev/null
+++ b/Engine/Source/Runtime/AudioExtensions/Private/IAudioParameterInterfaceRegistry.cpp
@@ -0,0 +1,48 @@
+// Copyright Epic Games, Inc. All Rights Reserved.
+#include "IAudioParameterInterfaceRegistry.h"
+
+
+namespace Audio
+{
+	namespace ParameterInterfaceRegistryPrivate
+	{
+		class FParameterInterfaceRegistry : public IAudioParameterInterfaceRegistry
+		{
+		public:
+			virtual void IterateInterfaces(TFunction<void(FParameterInterfacePtr)> InFunction) const override
+			{
+				for (const FParameterInterfacePtr& InterfacePtr : Interfaces)
+				{
+					InFunction(InterfacePtr);
+				}
+			}
+
+			virtual void RegisterInterface(FParameterInterfacePtr InInterface) override
+			{
+				Interfaces.Add(InInterface);
+				if (RegistrationFunction)
+				{
+					RegistrationFunction(InInterface);
+				}
+			}
+
+			virtual void OnRegistration(TUniqueFunction<void(FParameterInterfacePtr)>&& InFunction) override
+			{
+				RegistrationFunction = MoveTemp(InFunction);
+			}
+		};
+	} // namespace ParameterInterfaceRegistryPrivate
+
+	TUniquePtr<IAudioParameterInterfaceRegistry> IAudioParameterInterfaceRegistry::Instance;
+
+	IAudioParameterInterfaceRegistry& IAudioParameterInterfaceRegistry::Get()
+	{
+		using namespace ParameterInterfaceRegistryPrivate;
+
+		if (!Instance.IsValid())
+		{
+			Instance = MakeUnique<FParameterInterfaceRegistry>();
+		}
+		return *Instance;
+	}
+} // namespace Audio
diff --git a/Engine/Source/Runtime/AudioExtensions/Private/IAudioParameterTransmitter.cpp b/Engine/Source/Runtime/AudioExtensions/Private/IAudioParameterTransmitter.cpp
--- a/Engine/Source/Runtime/AudioExtensions/Private/IAudioParameterTransmitter.cpp
+++ b/Engine/Source/Runtime/AudioExtensions/Private/IAudioParameterTransmitter.cpp
@@ -1,7 +1,7 @@
 // Copyright Epic Games, Inc. All Rights Reserved.
 #include "IAudioParameterTransmitter.h"
+#include "UObject/Object.h"
 
-#include "AudioParameterInterface.h"
 
 namespace Audio
 {
diff --git a/Engine/Source/Runtime/AudioExtensions/Public/AudioParameter.h b/Engine/Source/Runtime/AudioExtensions/Public/AudioParameter.h
new file mode 100644
--- /dev/null
+++ b/Engine/Source/Runtime/AudioExtensions/Public/AudioParameter.h
@@ -0,0 +1,392 @@
+// Copyright Epic Games, Inc. All Rights Reserved.
+#pragma once
+
+#include "Containers/Array.h"
+#include "IAudioProxyInitializer.h"
+#include "UObject/Interface.h"
+#include "UObject/NameTypes.h"
+#include "UObject/Object.h"
+
+#include "AudioParameter.generated.h"
+
+
+#define AUDIO_PARAMETER_NAMESPACE_PATH_DELIMITER "."
+
+// Convenience macro for statically declaring an interface member's FName. AUDIO_PARAMETER_INTERFACE_NAMESPACE must be defined.
+#define AUDIO_PARAMETER_INTERFACE_MEMBER_DEFINE(Name) (AUDIO_PARAMETER_INTERFACE_NAMESPACE AUDIO_PARAMETER_NAMESPACE_PATH_DELIMITER Name)
+
+namespace Audio
+{
+	struct AUDIOEXTENSIONS_API FParameterPath
+	{
+		static const FString NamespaceDelimiter;
+
+		// Combines names using the namespace delimiter
+		static FName CombineNames(FName InLeft, FName InRight);
+
+		// Splits name into namespace & parameter name
+		static void SplitName(FName InFullName, FName& OutNamespace, FName& OutParameterName);
+	};
+}
+
+UENUM()
+enum class EAudioParameterType : uint8
+{
+	// 'Default' results in behavior that is resolved
+	// based on the system interpreting it.  To support
+	// legacy implementation, SoundCues cache all typed values
+	// associated with a given parameter name, whereas MetaSounds
+	// only use the last typed value set.
+	None UMETA(DisplayName = "Default"),
+
+	// Boolean value
+	Boolean,
+
+	// Integer value
+	Integer,
+
+	// Float value
+	Float,
+
+	// String value (not supported by legacy SoundCue system)
+	String,
+
+	// Object value (types other than SoundWave not supported by legacy SoundCue system)
+	Object,
+
+	// Array of default initialized values (not supported by legacy SoundCue system)
+	// Hidden for now as no parameter types exist that support default construction
+	NoneArray UMETA(Hidden, DisplayName = "Default (Array)"),
+
+	// Array of boolean values (not supported by legacy SoundCue system)
+	BooleanArray UMETA(DisplayName = "Boolean (Array)"),
+
+	// Array of integer values (not supported by legacy SoundCue system)
+	IntegerArray UMETA(DisplayName = "Integer (Array)"),
+
+	// Array of float values (not supported by legacy SoundCue system)
+	FloatArray UMETA(DisplayName = "Float (Array)"),
+
+	// Array of string values (not supported by legacy SoundCue system)
+	StringArray UMETA(DisplayName = "String (Array)"),
+
+	// Array of object values (not supported by legacy SoundCue system)
+	ObjectArray UMETA(DisplayName = "Object (Array)"),
+
+	COUNT UMETA(Hidden)
+};
+
+
+USTRUCT(BlueprintType)
+struct AUDIOEXTENSIONS_API FAudioParameter
+{
+	GENERATED_USTRUCT_BODY()
+
+	FAudioParameter() = default;
+
+	FAudioParameter(FName InName)
+		: ParamName(InName)
+	{
+	}
+
+	FAudioParameter(FName InName, float InValue)
+		: ParamName(InName)
+		, FloatParam(InValue)
+		, ParamType(EAudioParameterType::Float)
+	{
+	}
+
+	FAudioParameter(FName InName, bool InValue)
+		: ParamName(InName)
+		, BoolParam(InValue)
+		, ParamType(EAudioParameterType::Boolean)
+	{
+	}
+
+	FAudioParameter(FName InName, int32 InValue)
+		: ParamName(InName)
+		, IntParam(InValue)
+		, ParamType(EAudioParameterType::Integer)
+	{
+	}
+
+	FAudioParameter(FName InName, UObject* InValue)
+		: ParamName(InName)
+		, ObjectParam(InValue)
+		, ParamType(EAudioParameterType::Object)
+	{
+	}
+
+	FAudioParameter(FName InName, const FString& InValue)
+		: ParamName(InName)
+		, StringParam(InValue)
+		, ParamType(EAudioParameterType::String)
+	{
+	}
+
+	FAudioParameter(FName InName, const TArray<float>& InValue)
+		: ParamName(InName)
+		, ArrayFloatParam(InValue)
+		, ParamType(EAudioParameterType::FloatArray)
+	{
+	}
+
+	FAudioParameter(FName InName, TArray<float>&& InValue)
+		: ParamName(InName)
+		, ArrayFloatParam(MoveTemp(InValue))
+		, ParamType(EAudioParameterType::FloatArray)
+	{
+	}
+
+	FAudioParameter(FName InName, const TArray<bool>& InValue)
+		: ParamName(InName)
+		, ArrayBoolParam(InValue)
+		, ParamType(EAudioParameterType::BooleanArray)
+	{
+	}
+
+	FAudioParameter(FName InName, TArray<bool>&& InValue)
+		: ParamName(InName)
+		, ArrayBoolParam(MoveTemp(InValue))
+		, ParamType(EAudioParameterType::BooleanArray)
+	{
+	}
+
+	FAudioParameter(FName InName, const TArray<int32>& InValue)
+		: ParamName(InName)
+		, ArrayIntParam(InValue)
+		, ParamType(EAudioParameterType::IntegerArray)
+	{
+	}
+
+	FAudioParameter(FName InName, TArray<int32>&& InValue)
+		: ParamName(InName)
+		, ArrayIntParam(MoveTemp(InValue))
+		, ParamType(EAudioParameterType::IntegerArray)
+	{
+	}
+
+	FAudioParameter(FName InName, const TArray<UObject*>& InValue)
+		: ParamName(InName)
+		, ArrayObjectParam(InValue)
+		, ParamType(EAudioParameterType::ObjectArray)
+	{
+	}
+
+	FAudioParameter(FName InName, TArray<UObject*>&& InValue)
+		: ParamName(InName)
+		, ArrayObjectParam(MoveTemp(InValue))
+		, ParamType(EAudioParameterType::ObjectArray)
+	{
+	}
+
+	FAudioParameter(FName InName, const TArray<FString>& InValue)
+		: ParamName(InName)
+		, ArrayStringParam(InValue)
+		, ParamType(EAudioParameterType::StringArray)
+	{
+	}
+
+	FAudioParameter(FName InName, TArray<FString>&& InValue)
+		: ParamName(InName)
+		, ArrayStringParam(MoveTemp(InValue))
+		, ParamType(EAudioParameterType::StringArray)
+	{
+	}
+
+	// Static function to avoid int32 constructor collision
+	static FAudioParameter CreateDefaultArray(FName InName, int32 InNum)
+	{
+		FAudioParameter NewParam(InName, InNum);
+		NewParam.ParamType = EAudioParameterType::NoneArray;
+		return NewParam;
+	}
+
+	// Sets values specified by type field of the given parameter on this parameter. If the provided parameter is set to type 'None', takes all values of the given parameter.
+	// bInTakeName = Take the name of the provided parameter.
+	// bInTakeType = Take the type of the provided parameter.
+	// bInMergeArrayTypes - Appends array(s) for specified type if true, else swaps the local value with that of the provided parameter if false.
+	void Merge(const FAudioParameter& InParameter, bool bInTakeName = true, bool bInTakeType = true, bool bInMergeArrayTypes = false);
+
+	// Moves InParams to OutParams that are not already included. For backward compatibility (i.e. SoundCues),
+	// if a param is already in OutParams, attempts to merge param values together, but assigns the param the
+	// incoming param's type. Currently existing OutParam values left if new value of the same type is provided
+	// by InParam.
+	static void Merge(TArray<FAudioParameter>&& InParams, TArray<FAudioParameter>& OutParams);
+
+	FAudioParameter(FAudioParameter&& InParameter)
+		: ParamName(MoveTemp(InParameter.ParamName))
+		, FloatParam(MoveTemp(InParameter.FloatParam))
+		, BoolParam(MoveTemp(InParameter.BoolParam))
+		, IntParam(MoveTemp(InParameter.IntParam))
+		, ObjectParam(MoveTemp(InParameter.ObjectParam))
+		, StringParam(MoveTemp(InParameter.StringParam))
+		, ArrayFloatParam(MoveTemp(InParameter.ArrayFloatParam))
+		, ArrayBoolParam(MoveTemp(InParameter.ArrayBoolParam))
+		, ArrayIntParam(MoveTemp(InParameter.ArrayIntParam))
+		, ArrayObjectParam(MoveTemp(InParameter.ArrayObjectParam))
+		, ArrayStringParam(MoveTemp(InParameter.ArrayStringParam))
+		, ParamType(MoveTemp(InParameter.ParamType))
+		, ObjectProxies(MoveTemp(InParameter.ObjectProxies))
+	{
+	}
+
+	FAudioParameter(const FAudioParameter& InParameter)
+		: ParamName(InParameter.ParamName)
+		, FloatParam(InParameter.FloatParam)
+		, BoolParam(InParameter.BoolParam)
+		, IntParam(InParameter.IntParam)
+		, ObjectParam(InParameter.ObjectParam)
+		, StringParam(InParameter.StringParam)
+		, ArrayFloatParam(InParameter.ArrayFloatParam)
+		, ArrayBoolParam(InParameter.ArrayBoolParam)
+		, ArrayIntParam(InParameter.ArrayIntParam)
+		, ArrayObjectParam(InParameter.ArrayObjectParam)
+		, ArrayStringParam(InParameter.ArrayStringParam)
+		, ParamType(InParameter.ParamType)
+	{
+		ObjectProxies.Reset();
+		for (const Audio::IProxyDataPtr& DataPtr : InParameter.ObjectProxies)
+		{
+			ObjectProxies.Emplace(DataPtr->Clone());
+		}
+	}
+
+	FAudioParameter& operator=(const FAudioParameter& InParameter)
+	{
+		ParamName = InParameter.ParamName;
+		FloatParam = InParameter.FloatParam;
+		BoolParam = InParameter.BoolParam;
+		IntParam = InParameter.IntParam;
+		ObjectParam = InParameter.ObjectParam;
+		StringParam = InParameter.StringParam;
+		ArrayFloatParam = InParameter.ArrayFloatParam;
+		ArrayBoolParam = InParameter.ArrayBoolParam;
+		ArrayIntParam = InParameter.ArrayIntParam;
+		ArrayObjectParam = InParameter.ArrayObjectParam;
+		ArrayStringParam = InParameter.ArrayStringParam;
+		ParamType = InParameter.ParamType;
+
+		ObjectProxies.Reset();
+		for (const Audio::IProxyDataPtr& DataPtr : InParameter.ObjectProxies)
+		{
+			ObjectProxies.Emplace(DataPtr->Clone());
+		}
+
+		return *this;
+	}
+
+	FAudioParameter& operator=(FAudioParameter&& InParameter)
+	{
+		ParamName = MoveTemp(InParameter.ParamName);
+		FloatParam = MoveTemp(InParameter.FloatParam);
+		BoolParam = MoveTemp(InParameter.BoolParam);
+		IntParam = MoveTemp(InParameter.IntParam);
+		ObjectParam = MoveTemp(InParameter.ObjectParam);
+		StringParam = MoveTemp(InParameter.StringParam);
+		ArrayFloatParam = MoveTemp(InParameter.ArrayFloatParam);
+		ArrayBoolParam = MoveTemp(InParameter.ArrayBoolParam);
+		ArrayIntParam = MoveTemp(InParameter.ArrayIntParam);
+		ArrayObjectParam = MoveTemp(InParameter.ArrayObjectParam);
+		ArrayStringParam = MoveTemp(InParameter.ArrayStringParam);
+		ParamType = MoveTemp(InParameter.ParamType);
+
+		ObjectProxies = MoveTemp(InParameter.ObjectProxies);
+
+		return *this;
+	}
+
+	// Name of the parameter
+	UPROPERTY(EditAnywhere, BlueprintReadWrite, meta = (DisplayName="Name"), Category = AudioParameter)
+	FName ParamName;
+
+	// Float value of parameter
+	UPROPERTY(EditAnywhere, BlueprintReadWrite, meta = (DisplayName = "Value (Float)", DisplayAfter = "ParamType", EditConditionHides, EditCondition = "ParamType == EAudioParameterType::None || ParamType == EAudioParameterType::Float"), Category = AudioParameter)
+	float FloatParam = 0.f;
+
+	// Boolean value of parameter
+	UPROPERTY(EditAnywhere, BlueprintReadWrite, meta = (DisplayName = "Value (Bool)", DisplayAfter = "ParamType", EditConditionHides, EditCondition = "ParamType == EAudioParameterType::None || ParamType == EAudioParameterType::Boolean"), Category = AudioParameter)
+	bool BoolParam = false;
+
+	// Integer value of parameter. If set to 'Default Construct', value is number of array items to construct.
+	UPROPERTY(EditAnywhere, BlueprintReadWrite, meta = (DisplayName = "Value (Int)", DisplayAfter = "ParamType", EditConditionHides, EditCondition = "ParamType == EAudioParameterType::None || ParamType == EAudioParameterType::Integer || ParamType == EAudioParameterType::NoneArray"), Category = AudioParameter)
+	int32 IntParam = 0;
+
+	// Object value of parameter
+	UPROPERTY(EditAnywhere, BlueprintReadWrite, meta = (DisplayName = "Value (Object)", DisplayAfter = "ParamType", EditConditionHides, EditCondition = "ParamType == EAudioParameterType::None || ParamType == EAudioParameterType::Object"), Category = AudioParameter)
+	TObjectPtr<UObject> ObjectParam = nullptr;
+
+	// String value of parameter
+	UPROPERTY(EditAnywhere, BlueprintReadWrite, meta = (DisplayName = "Value (String)", DisplayAfter = "ParamType", EditConditionHides, EditCondition = "ParamType == EAudioParameterType::String"), Category = AudioParameter)
+	FString StringParam;
+
+	// Array Float value of parameter
+	UPROPERTY(EditAnywhere, BlueprintReadWrite, meta = (DisplayName = "Value (Float Array)", DisplayAfter = "ParamType", EditConditionHides, EditCondition = "ParamType == EAudioParameterType::FloatArray"), Category = AudioParameter)
+	TArray<float> ArrayFloatParam;
+
+	// Boolean value of parameter
+	UPROPERTY(EditAnywhere, BlueprintReadWrite, meta = (DisplayName = "Value (Bool Array)", DisplayAfter = "ParamType", EditConditionHides, EditCondition = "ParamType == EAudioParameterType::BooleanArray"), Category = AudioParameter)
+	TArray<bool> ArrayBoolParam;
+
+	// Integer value of parameter
+	UPROPERTY(EditAnywhere, BlueprintReadWrite, meta = (DisplayName = "Value (Int Array)", DisplayAfter = "ParamType", EditConditionHides, EditCondition = "ParamType == EAudioParameterType::IntegerArray"), Category = AudioParameter)
+	TArray<int32> ArrayIntParam;
+
+	// Object value of parameter
+	UPROPERTY(EditAnywhere, BlueprintReadWrite, meta = (DisplayName = "Value (Object Array)", DisplayAfter = "ParamType", EditConditionHides, EditCondition = "ParamType == EAudioParameterType::ObjectArray"), Category = AudioParameter)
+	TArray<TObjectPtr<UObject>> ArrayObjectParam;
+
+	// String value of parameter
+	UPROPERTY(EditAnywhere, BlueprintReadWrite, meta = (DisplayName = "Value (String Array)", DisplayAfter = "ParamType", EditConditionHides, EditCondition = "ParamType == EAudioParameterType::StringArray"), Category = AudioParameter)
+	TArray<FString> ArrayStringParam;
+
+	UPROPERTY(EditAnywhere, BlueprintReadWrite, meta = (DisplayName = "Type"), Category = AudioParameter)
+	EAudioParameterType ParamType = EAudioParameterType::None;
+
+	// Object proxies to be generated when parameter is passed to the AudioThread to represent ObjectParam/ArrayObjectParam safely
+	TArray<Audio::IProxyDataPtr> ObjectProxies;
+
+	// Common find algorithm for default/legacy parameter system
+	static const FAudioParameter* FindParam(const TArray<FAudioParameter>& InParams, FName InParamName)
+	{
+		if (!InParamName.IsNone())
+		{
+			for (const FAudioParameter& ExistingParam : InParams)
+			{
+				if (ExistingParam.ParamName == InParamName)
+				{
+					return &ExistingParam;
+				}
+			}
+		}
+
+		return nullptr;
+	}
+
+	// Common find & add algorithm for default/legacy parameter system.
+	static FAudioParameter* FindOrAddParam(TArray<FAudioParameter>& OutParams, FName InParamName)
+	{
+		FAudioParameter* Param = nullptr;
+		if (InParamName.IsNone())
+		{
+			return Param;
+		}
+
+		for (FAudioParameter& ExistingParam : OutParams)
+		{
+			if (ExistingParam.ParamName == InParamName)
+			{
+				Param = &ExistingParam;
+				break;
+			}
+		}
+
+		if (!Param)
+		{
+			Param = &OutParams.AddDefaulted_GetRef();
+			Param->ParamName = InParamName;
+		}
+
+		return Param;
+	}
+};
\ No newline at end of file
diff --git a/Engine/Source/Runtime/AudioExtensions/Public/AudioParameterControllerInterface.h b/Engine/Source/Runtime/AudioExtensions/Public/AudioParameterControllerInterface.h
new file mode 100644
--- /dev/null
+++ b/Engine/Source/Runtime/AudioExtensions/Public/AudioParameterControllerInterface.h
@@ -0,0 +1,93 @@
+// Copyright Epic Games, Inc. All Rights Reserved.
+
+#pragma once
+
+#include "AudioParameter.h"
+#include "Containers/Array.h"
+#include "IAudioProxyInitializer.h"
+#include "UObject/Interface.h"
+#include "UObject/NameTypes.h"
+#include "UObject/Object.h"
+
+#include "AudioParameterControllerInterface.generated.h"
+
+
+UINTERFACE(BlueprintType, meta = (CannotImplementInterfaceInBlueprint))
+class AUDIOEXTENSIONS_API UAudioParameterControllerInterface : public UInterface
+{
+	GENERATED_UINTERFACE_BODY()
+};
+
+// Base interface for any object implementing a
+class AUDIOEXTENSIONS_API IAudioParameterControllerInterface : public IInterface
+{
+	GENERATED_IINTERFACE_BODY()
+
+public:
+	// Resets all parameters to their original values.
+	UFUNCTION(BlueprintCallable, Category = "Audio|Parameter")
+	virtual void ResetParameters() = 0;
+
+	// Triggers a named trigger
+	UFUNCTION(BlueprintCallable, meta = (DisplayName = "Set Trigger Parameter"), Category = "Audio|Parameter")
+	virtual void SetTriggerParameter(FName InName) = 0;
+
+	// Sets a named Boolean
+	UFUNCTION(BlueprintCallable, meta = (DisplayName = "Set Boolean Parameter"), Category = "Audio|Parameter")
+	virtual void SetBoolParameter(FName InName, bool InBool) = 0;
+	
+	// Sets a named Boolean Array
+	UFUNCTION(BlueprintCallable, meta = (DisplayName = "Set Boolean Array Parameter"), Category = "Audio|Parameter")
+	virtual void SetBoolArrayParameter(FName InName, const TArray<bool>& InValue) = 0;
+
+	// Sets a named Int32
+	UFUNCTION(BlueprintCallable, meta = (DisplayName = "Set Integer Parameter"), Category = "Audio|Parameter")
+	virtual void SetIntParameter(FName InName, int32 InInt) = 0;
+
+	// Sets a named Int32 Array
+	UFUNCTION(BlueprintCallable, meta = (DisplayName = "Set Integer Array Parameter"), Category = "Audio|Parameter")
+	virtual void SetIntArrayParameter(FName InName, const TArray<int32>& InValue) = 0;
+
+	// Sets a named Float
+	UFUNCTION(BlueprintCallable, meta = (DisplayName = "Set Float Parameter"), Category = "Audio|Parameter")
+	virtual void SetFloatParameter(FName InName, float InFloat) = 0;
+	
+	// Sets a named Float Array
+	UFUNCTION(BlueprintCallable, meta = (DisplayName = "Set Float Array Parameter"), Category = "Audio|Parameter")
+	virtual void SetFloatArrayParameter(FName InName, const TArray<float>& InValue) = 0;
+	
+	// Sets a named String
+	UFUNCTION(BlueprintCallable, meta = (DisplayName = "Set String Parameter"), Category = "Audio|Parameter")
+	virtual void SetStringParameter(FName InName, const FString& InValue) = 0;
+
+	// Sets a named String Array
+	UFUNCTION(BlueprintCallable, meta = (DisplayName = "Set String Array Parameter"), Category = "Audio|Parameter")
+	virtual void SetStringArrayParameter(FName InName, const TArray<FString>& InValue) = 0;
+
+	// Sets a named UObject
+	UFUNCTION(BlueprintCallable, meta = (DisplayName = "Set Object Parameter"), Category = "Audio|Parameter")
+	virtual void SetObjectParameter(FName InName, UObject* InValue) = 0;
+
+	// Sets a named UObject Array
+	UFUNCTION(BlueprintCallable, meta = (DisplayName = "Set Object Array Parameter"), Category = "Audio|Parameter")
+	virtual void SetObjectArrayParameter(FName InName, const TArray<UObject*>& InValue) = 0;
+
+	// Sets a named parameter to the given parameter structure value
+	virtual void SetParameter(FAudioParameter&& InValue) = 0;
+
+	// Sets an array of parameters as a batch
+	virtual void SetParameters(TArray<FAudioParameter>&& InValues) = 0;
+
+	// Template Specialization for non-script clients.
+	template<typename DataType> void SetParameter(FName InName, DataType&&) = delete;
+	template<> void SetParameter(FName InName, bool&& InBool) { SetBoolParameter(InName, InBool); }
+	template<> void SetParameter(FName InName, float&& InFloat) { SetFloatParameter(InName, InFloat); }
+	template<> void SetParameter(FName InName, int32&& InInteger) { SetIntParameter(InName, InInteger); }
+	template<> void SetParameter(FName InName, FString&& InString) { SetStringParameter(InName, InString); }
+	template<> void SetParameter(FName InName, UObject*&& InObject) { SetObjectParameter(InName, InObject); }
+	template<> void SetParameter(FName InName, TArray<bool>&& InBools) { SetBoolArrayParameter(InName, InBools); }
+	template<> void SetParameter(FName InName, TArray<float>&& InFloats) { SetFloatArrayParameter(InName, InFloats); }
+	template<> void SetParameter(FName InName, TArray<int32>&& InIntegers) { SetIntArrayParameter(InName, InIntegers); }
+	template<> void SetParameter(FName InName, TArray<FString>&& InStrings) { SetStringArrayParameter(InName, InStrings); }
+	template<> void SetParameter(FName InName, TArray<UObject*>&& InObjects) { SetObjectArrayParameter(InName, InObjects); }
+};
diff --git a/Engine/Source/Runtime/AudioExtensions/Public/AudioParameterInterface.h b/Engine/Source/Runtime/AudioExtensions/Public/AudioParameterInterface.h
deleted file mode 100644
diff --git a/Engine/Source/Runtime/AudioExtensions/Public/IAudioGeneratorInterfaceRegistry.h b/Engine/Source/Runtime/AudioExtensions/Public/IAudioGeneratorInterfaceRegistry.h
deleted file mode 100644
diff --git a/Engine/Source/Runtime/AudioExtensions/Public/IAudioParameterInterfaceRegistry.h b/Engine/Source/Runtime/AudioExtensions/Public/IAudioParameterInterfaceRegistry.h
new file mode 100644
--- /dev/null
+++ b/Engine/Source/Runtime/AudioExtensions/Public/IAudioParameterInterfaceRegistry.h
@@ -0,0 +1,126 @@
+// Copyright Epic Games, Inc. All Rights Reserved.
+#pragma once
+
+#include "AudioParameter.h"
+#include "AudioParameterControllerInterface.h"
+#include "Containers/Array.h"
+#include "Containers/Map.h"
+#include "Templates/Function.h"
+#include "UObject/Class.h"
+
+
+namespace Audio
+{
+	// Forward Declarations
+	class IAudioParameterInterfaceRegistry;
+
+	struct AUDIOEXTENSIONS_API FParameterInterface
+	{
+		struct FVersion
+		{
+			const int32 Major = 1;
+			const int32 Minor = 0;
+		};
+
+		struct FInput
+		{
+			const FText DisplayName;
+			const FText Description;
+			const FName DataType;
+
+			const FAudioParameter InitValue;
+		};
+
+		struct FOutput
+		{
+			const FText DisplayName;
+			const FText Description;
+			const FName DataType;
+
+			const FName ParamName;
+			const EAudioParameterType ParamType = EAudioParameterType::None;
+		};
+
+		struct FEnvironmentVariable
+		{
+			const FText DisplayName;
+			const FText Description;
+			const FName DataType;
+
+			const FName ParamName;
+			const EAudioParameterType ParamType = EAudioParameterType::None;
+		};
+
+		FParameterInterface() = default;
+		FParameterInterface(
+			FName InName,
+			const FVersion& InVersion,
+			const UClass& InType
+		)
+			: NamePrivate(InName)
+			, VersionPrivate(InVersion)
+			, Type(&InType)
+		{
+		}
+
+		FName GetName() const
+		{
+			return NamePrivate;
+		}
+
+		const FVersion& GetVersion() const
+		{
+			return VersionPrivate;
+		}
+
+		const UClass& GetType() const
+		{
+			check(Type);
+			return *Type;
+		}
+
+		const TArray<FInput>& GetInputs() const
+		{
+			return Inputs;
+		}
+
+		const TArray<FOutput>& GetOutputs() const
+		{
+			return Outputs;
+		}
+
+		const TArray<FEnvironmentVariable>& GetEnvironment() const
+		{
+			return Environment;
+		}
+
+	private:
+		FName NamePrivate;
+		FVersion VersionPrivate;
+		const UClass* Type = nullptr;
+
+	protected:
+		TArray<FInput> Inputs;
+		TArray<FOutput> Outputs;
+		TArray<FEnvironmentVariable> Environment;
+	};
+	using FParameterInterfacePtr = TSharedPtr<FParameterInterface>;
+
+	class AUDIOEXTENSIONS_API IAudioParameterInterfaceRegistry
+	{
+		static TUniquePtr<IAudioParameterInterfaceRegistry> Instance;
+
+	public:
+		static IAudioParameterInterfaceRegistry& Get();
+
+		virtual ~IAudioParameterInterfaceRegistry() = default;
+
+		virtual void IterateInterfaces(TFunction<void(FParameterInterfacePtr)> InFunction) const = 0;
+		virtual void OnRegistration(TUniqueFunction<void(FParameterInterfacePtr)>&& InFunction) = 0;
+		virtual void RegisterInterface(FParameterInterfacePtr InInterface) = 0;
+
+	protected:
+		TSet<FParameterInterfacePtr> Interfaces;
+		TUniqueFunction<void(FParameterInterfacePtr)> RegistrationFunction;
+	};
+} // namespace Audio
diff --git a/Engine/Source/Runtime/AudioExtensions/Public/IAudioParameterTransmitter.h b/Engine/Source/Runtime/AudioExtensions/Public/IAudioParameterTransmitter.h
--- a/Engine/Source/Runtime/AudioExtensions/Public/IAudioParameterTransmitter.h
+++ b/Engine/Source/Runtime/AudioExtensions/Public/IAudioParameterTransmitter.h
@@ -1,12 +1,15 @@
 // Copyright Epic Games, Inc. All Rights Reserved.
 #pragma once
 
-#include "AudioParameterInterface.h"
-#include "CoreMinimal.h"
+#include "AudioParameterControllerInterface.h"
+#include "Containers/Array.h"
 #include "IAudioProxyInitializer.h"
 #include "Templates/UniquePtr.h"
 
 
+// Forward Declarations
+class UObject;
+
 namespace Audio
 {
 	/** Data passed to CreateParameterTransmitter. */
@@ -52,7 +55,6 @@ namespace Audio
 
 			// Parameter Setters
 			virtual bool SetParameter(FAudioParameter&& InValue) = 0;
-			virtual bool SetParameter(FName InInterfaceName, FAudioParameter&& InValue) = 0;
 
 			// Create a copy of the instance transmitter.
 			virtual TUniquePtr<Audio::IParameterTransmitter> Clone() const = 0;
diff --git a/Engine/Source/Runtime/Engine/Classes/Audio/ActorSoundParameterInterface.h b/Engine/Source/Runtime/Engine/Classes/Audio/ActorSoundParameterInterface.h
--- a/Engine/Source/Runtime/Engine/Classes/Audio/ActorSoundParameterInterface.h
+++ b/Engine/Source/Runtime/Engine/Classes/Audio/ActorSoundParameterInterface.h
@@ -1,8 +1,9 @@
 // Copyright Epic Games, Inc. All Rights Reserved.
 #pragma once
 
-#include "AudioParameterInterface.h"
-#include "CoreMinimal.h"
+#include "AudioParameterControllerInterface.h"
+#include "Containers/Array.h"
+#include "GameFramework/Actor.h"
 #include "UObject/Object.h"
 #include "UObject/Interface.h"
 
diff --git a/Engine/Source/Runtime/Engine/Classes/Audio/SoundGeneratorParameterInterface.h b/Engine/Source/Runtime/Engine/Classes/Audio/SoundGeneratorParameterInterface.h
deleted file mode 100644
diff --git a/Engine/Source/Runtime/Engine/Classes/Components/AudioComponent.h b/Engine/Source/Runtime/Engine/Classes/Components/AudioComponent.h
--- a/Engine/Source/Runtime/Engine/Classes/Components/AudioComponent.h
+++ b/Engine/Source/Runtime/Engine/Classes/Components/AudioComponent.h
@@ -1,8 +1,7 @@
 // Copyright Epic Games, Inc. All Rights Reserved.
 #pragma once
 
-#include "Audio/SoundGeneratorParameterInterface.h"
-#include "AudioParameterInterface.h"
+#include "Audio/SoundParameterControllerInterface.h"
 #include "Components/SceneComponent.h"
 #include "CoreMinimal.h"
 #include "Engine/EngineTypes.h"
@@ -161,7 +160,7 @@ class ENGINE_API UInitialActiveSoundParams : public UObject
  * @see USoundBase
  */
 UCLASS(ClassGroup=(Audio, Common), HideCategories=(Object, ActorComponent, Physics, Rendering, Mobility, LOD), ShowCategories=Trigger, meta=(BlueprintSpawnableComponent))
-class ENGINE_API UAudioComponent : public USceneComponent, public ISoundGeneratorParameterInterface, public FQuartzTickableObject
+class ENGINE_API UAudioComponent : public USceneComponent, public ISoundParameterControllerInterface, public FQuartzTickableObject
 {
 	GENERATED_UCLASS_BODY()
 
@@ -178,7 +177,7 @@ class ENGINE_API UAudioComponent : public USceneComponent, public ISoundGenerato
 
 	/** Array of transient parameters for this AudioComponent instance. Not serialized and can be set by code or BP.
 	  * Changes to this array directly will not be forwarded to the sound if the component is actively playing.
-	  * This should be done via the 'SetParameterX' calls implemented by the ISoundGeneratorParameterInterface.
+	  * This should be done via the 'SetParameterX' calls implemented by the ISoundParameterControllerInterface.
 	  * Instance parameter values superseded the parameters set by the actor interface & the component's default
 	  * parameters.
 	  */
@@ -763,7 +762,7 @@ public:
 
 	void SetSourceEffectChain(USoundEffectSourcePresetChain* InSourceEffectChain);
 
-	/** SoundGeneratorParameterInterface Implementation */
+	/** SoundParameterControllerInterface Implementation */
 	FAudioDevice* GetAudioDevice() const override;
 	TArray<FAudioParameter>& GetInstanceParameters() override { return InstanceParameters; }
 	uint64 GetInstanceOwnerID() const override { return AudioComponentID; }
diff --git a/Engine/Source/Runtime/Engine/Classes/Sound/SoundAttenuation.h b/Engine/Source/Runtime/Engine/Classes/Sound/SoundAttenuation.h
--- a/Engine/Source/Runtime/Engine/Classes/Sound/SoundAttenuation.h
+++ b/Engine/Source/Runtime/Engine/Classes/Sound/SoundAttenuation.h
@@ -4,7 +4,7 @@
 
 #include "Engine/Attenuation.h"
 #include "IAudioExtensionPlugin.h"
-#include "IAudioGeneratorInterfaceRegistry.h"
+#include "IAudioParameterInterfaceRegistry.h"
 
 #include "SoundAttenuation.generated.h"
 
@@ -492,31 +492,28 @@ class USoundAttenuation : public UObject
 
 namespace Audio
 {
-	struct ENGINE_API FAttenuationInterface : public FGeneratorInterface
+	namespace AttenuationInterface
 	{
-		static const FName Name;
+		ENGINE_API const extern FName Name;
 
-		struct FInputs
+		namespace Inputs
 		{
-			static const FName Distance;
-		};
+			ENGINE_API const extern FName Distance;
+		} // namespace Inputs
 
-		FAttenuationInterface();
-	};
+		ENGINE_API Audio::FParameterInterfacePtr GetInterface();
+	} // namespace AttenuationInterface
 
-	struct ENGINE_API FSpatializationInterface : public FGeneratorInterface
+	namespace SpatializationInterface
 	{
-		static const FName Name;
+		ENGINE_API const extern FName Name;
 
-		struct FInputs
+		namespace Inputs
 		{
-			static const FName Azimuth;
-			static const FName Elevation;
-		};
-
-		FSpatializationInterface();
-	};
+			ENGINE_API const extern FName Azimuth;
+			ENGINE_API const extern FName Elevation;
+		} // namespace Inputs
 
-	ENGINE_API FGeneratorInterfacePtr GetAttenuationInterface();
-	ENGINE_API FGeneratorInterfacePtr GetSpatializationInterface();
-} // namespace Audio
\ No newline at end of file
+		ENGINE_API Audio::FParameterInterfacePtr GetInterface();
+	} // namespace SpatializationInterface
+} // namespace Audio
diff --git a/Engine/Source/Runtime/Engine/Classes/Sound/SoundBase.h b/Engine/Source/Runtime/Engine/Classes/Sound/SoundBase.h
--- a/Engine/Source/Runtime/Engine/Classes/Sound/SoundBase.h
+++ b/Engine/Source/Runtime/Engine/Classes/Sound/SoundBase.h
@@ -5,8 +5,6 @@
 #include "Audio.h"
 #include "CoreMinimal.h"
 #include "IAudioExtensionPlugin.h"
-#include "IAudioGeneratorInterfaceRegistry.h"
-#include "IAudioParameterTransmitter.h"
 #include "Sound/AudioSettings.h"
 #include "Sound/SoundClass.h"
 #include "SoundConcurrency.h"
@@ -26,6 +24,7 @@
 namespace Audio
 {
 	class IParameterTransmitter;
+	struct FParameterTransmitterInitParams;
 } // namespace Audio
 
 class USoundEffectSourcePreset;
@@ -297,7 +296,7 @@ public:
 	virtual void InitResources() { }
 
 	/** Whether or not the given sound is a generator and implements an interface with the given name. */
-	virtual bool ImplementsGeneratorInterface(Audio::FGeneratorInterfacePtr InInterface) const { return false; }
+	virtual bool ImplementsParameterInterface(Audio::FParameterInterfacePtr InParameterInterface) const { return false; }
 
 	/** Creates a sound generator instance from this sound base. Return true if this is being implemented by a subclass. Sound generators procedurally generate audio in the audio render thread. */
 	virtual ISoundGeneratorPtr CreateSoundGenerator(const FSoundGeneratorInitParams& InParams) { return nullptr; }
diff --git a/Engine/Source/Runtime/Engine/Private/ActiveSound.cpp b/Engine/Source/Runtime/Engine/Private/ActiveSound.cpp
--- a/Engine/Source/Runtime/Engine/Private/ActiveSound.cpp
+++ b/Engine/Source/Runtime/Engine/Private/ActiveSound.cpp
@@ -562,9 +562,9 @@ bool FActiveSound::GetConcurrencyFadeDuration(float& OutFadeDuration) const
 	return true;
 }
 
-void FActiveSound::UpdateGeneratorParameters(const TArray<FListener>& InListeners)
+void FActiveSound::UpdateInterfaceParameters(const TArray<FListener>& InListeners)
 {
-	TRACE_CPUPROFILER_EVENT_SCOPE(FActiveSound::UpdateGeneratorParameters);
+	TRACE_CPUPROFILER_EVENT_SCOPE(FActiveSound::UpdateInterfaceParameters);
 
 	using namespace Audio;
 
@@ -583,11 +583,11 @@ void FActiveSound::UpdateGeneratorParameters(const TArray<FListener>& InListener
 		return;
 	}
 
-	FGeneratorInterfacePtr AttenuationInterface = GetAttenuationInterface();
-	FGeneratorInterfacePtr SpatializationInterface = GetSpatializationInterface();
+	FParameterInterfacePtr AttenuationInterface = AttenuationInterface::GetInterface();
+	FParameterInterfacePtr SpatializationInterface = SpatializationInterface::GetInterface();
 
-	const bool bImplementsAttenuation = Sound->ImplementsGeneratorInterface(AttenuationInterface);
-	const bool bImplementsSpatialization = Sound->ImplementsGeneratorInterface(SpatializationInterface);
+	const bool bImplementsAttenuation = Sound->ImplementsParameterInterface(AttenuationInterface);
+	const bool bImplementsSpatialization = Sound->ImplementsParameterInterface(SpatializationInterface);
 
 	if (!bImplementsAttenuation && !bImplementsSpatialization)
 	{
@@ -599,7 +599,7 @@ void FActiveSound::UpdateGeneratorParameters(const TArray<FListener>& InListener
 	if (bImplementsAttenuation)
 	{
 		const float Distance = SourceDirection.Size();
-		InstanceTransmitter->SetParameter(FAttenuationInterface::Name, { FAttenuationInterface::FInputs::Distance, Distance });
+		InstanceTransmitter->SetParameter({ AttenuationInterface::Inputs::Distance, Distance });
 	}
 
 	if (bImplementsSpatialization)
@@ -608,8 +608,8 @@ void FActiveSound::UpdateGeneratorParameters(const TArray<FListener>& InListener
 		const FVector2D SourceAzimuthAndElevation = FMath::GetAzimuthAndElevation(SourceDirectionNormal, FVector::ForwardVector, FVector::RightVector, FVector::UpVector);
 		const float Azimuth = FMath::RadiansToDegrees(SourceAzimuthAndElevation.X);
 		const float Elevation = FMath::RadiansToDegrees(SourceAzimuthAndElevation.Y);
-		InstanceTransmitter->SetParameter(FSpatializationInterface::Name, { FSpatializationInterface::FInputs::Azimuth, Azimuth });
-		InstanceTransmitter->SetParameter(FSpatializationInterface::Name, { FSpatializationInterface::FInputs::Elevation, Elevation });
+		InstanceTransmitter->SetParameter({ SpatializationInterface::Inputs::Azimuth, Azimuth });
+		InstanceTransmitter->SetParameter({ SpatializationInterface::Inputs::Elevation, Elevation });
 	}
 }
 
diff --git a/Engine/Source/Runtime/Engine/Private/Audio/SoundGeneratorParameterInterface.cpp b/Engine/Source/Runtime/Engine/Private/Audio/SoundGeneratorParameterInterface.cpp
deleted file mode 100644
diff --git a/Engine/Source/Runtime/Engine/Private/Audio/SoundParameterControllerInterface.cpp b/Engine/Source/Runtime/Engine/Private/Audio/SoundParameterControllerInterface.cpp
new file mode 100644
--- /dev/null
+++ b/Engine/Source/Runtime/Engine/Private/Audio/SoundParameterControllerInterface.cpp
@@ -0,0 +1,194 @@
+// Copyright Epic Games, Inc. All Rights Reserved.
+#include "Audio/SoundParameterControllerInterface.h"
+
+#include "ActiveSound.h"
+#include "Audio.h"
+#include "AudioDevice.h"
+#include "AudioThread.h"
+#include "IAudioExtensionPlugin.h"
+
+
+namespace SoundParameterControllerInterfacePrivate
+{
+	static const FName ProxyFeatureName("SoundParameterControllerInterface");
+} // namespace SoundParameterControllerInterfacePrivate
+
+USoundParameterControllerInterface::USoundParameterControllerInterface(FObjectInitializer const& InObjectInitializer)
+	: Super(InObjectInitializer)
+{
+}
+
+void ISoundParameterControllerInterface::ResetParameters()
+{
+	if (FAudioDevice* AudioDevice = GetAudioDevice())
+	{
+		if (IsPlaying() && !GetDisableParameterUpdatesWhilePlaying())
+		{
+			DECLARE_CYCLE_STAT(TEXT("FAudioThreadTask.SoundParameterController.ResetParameters"), STAT_AudioResetParameters, STATGROUP_AudioThreadCommands);
+			AudioDevice->SendCommandToActiveSounds(GetInstanceOwnerID(), [] (FActiveSound& ActiveSound)
+			{
+				if (Audio::IParameterTransmitter* Transmitter = ActiveSound.GetTransmitter())
+				{
+					Transmitter->Reset();
+				}
+			}, GET_STATID(STAT_AudioResetParameters));
+		}
+	}
+}
+
+void ISoundParameterControllerInterface::SetTriggerParameter(FName InName)
+{
+	// Trigger is just a (true) bool param currently.
+	SetParameterInternal(FAudioParameter(InName, true));
+}
+
+void ISoundParameterControllerInterface::SetBoolParameter(FName InName, bool InValue)
+{
+	SetParameterInternal(FAudioParameter(InName, InValue));
+}
+
+void ISoundParameterControllerInterface::SetBoolArrayParameter(FName InName, const TArray<bool>& InValue)
+{
+	SetParameterInternal(FAudioParameter(InName, InValue));
+}
+
+void ISoundParameterControllerInterface::SetIntParameter(FName InName, int32 InValue)
+{
+	SetParameterInternal(FAudioParameter(InName, InValue));
+}
+
+void ISoundParameterControllerInterface::SetIntArrayParameter(FName InName, const TArray<int32>& InValue)
+{
+	SetParameterInternal(FAudioParameter(InName, InValue));
+}
+
+void ISoundParameterControllerInterface::SetFloatParameter(FName InName, float InValue)
+{
+	SetParameterInternal(FAudioParameter(InName, InValue));
+}
+
+void ISoundParameterControllerInterface::SetFloatArrayParameter(FName InName, const TArray<float>& InValue)
+{
+	SetParameterInternal(FAudioParameter(InName, InValue));
+}
+
+void ISoundParameterControllerInterface::SetStringParameter(FName InName, const FString& InValue)
+{
+	SetParameterInternal(FAudioParameter(InName, InValue));
+}
+
+void ISoundParameterControllerInterface::SetStringArrayParameter(FName InName, const TArray<FString>& InValue)
+{
+	SetParameterInternal(FAudioParameter(InName, InValue));
+}
+
+void ISoundParameterControllerInterface::SetObjectParameter(FName InName, UObject* InValue)
+{
+	SetParameterInternal(FAudioParameter(InName, InValue));
+}
+
+void ISoundParameterControllerInterface::SetObjectArrayParameter(FName InName, const TArray<UObject*>& InValue)
+{
+	SetParameterInternal(FAudioParameter(InName, InValue));
+}
+
+void ISoundParameterControllerInterface::SetParameter(FAudioParameter&& InValue)
+{
+	SetParameterInternal(MoveTemp(InValue));
+}
+
+void ISoundParameterControllerInterface::SetParameters(TArray<FAudioParameter>&& InValues)
+{
+	for (const FAudioParameter& Value : InValues)
+	{
+		TArray<FAudioParameter>& InstanceParameters = GetInstanceParameters();
+		if (FAudioParameter* CurrentParam = FAudioParameter::FindOrAddParam(InstanceParameters, Value.ParamName))
+		{
+			CurrentParam->Merge(Value, false /* bInTakeName */);
+		}
+	}
+
+	if (IsPlaying() && !GetDisableParameterUpdatesWhilePlaying())
+	{
+		if (FAudioDevice* AudioDevice = GetAudioDevice())
+		{
+			TArray<FAudioParameter> ParamsToSet;
+			if (USoundBase* Sound = GetSound())
+			{
+				Sound->InitParameters(InValues, SoundParameterControllerInterfacePrivate::ProxyFeatureName);
+			}
+
+			ParamsToSet = MoveTemp(InValues);
+
+			DECLARE_CYCLE_STAT(TEXT("FAudioThreadTask.SoundParameterController.SetParameters"), STAT_AudioSetParameters, STATGROUP_AudioThreadCommands);
+			AudioDevice->SendCommandToActiveSounds(GetInstanceOwnerID(), [AudioDevice, Params = MoveTemp(ParamsToSet)](FActiveSound& ActiveSound)
+			{
+				if (Audio::IParameterTransmitter* Transmitter = ActiveSound.GetTransmitter())
+				{
+					for (const FAudioParameter& Param : Params)
+					{
+						const FName ParamName = Param.ParamName;
+						if (!ParamName.IsNone())
+						{
+							// Must be copied as original version must be preserved in case command is called on multiple ActiveSounds.
+							FAudioParameter TempParam = Param;
+							if (!Transmitter->SetParameter(MoveTemp(TempParam)))
+							{
+								UE_LOG(LogAudio, Warning, TEXT("Failed to set parameter '%s'"), *ParamName.ToString());
+							}
+						}
+					}
+				}
+			}, GET_STATID(STAT_AudioSetParameters));
+		}
+	}
+}
+
+void ISoundParameterControllerInterface::SetParameterInternal(FAudioParameter&& InParam)
+{
+	if (InParam.ParamName.IsNone())
+	{
+		return;
+	}
+
+	TArray<FAudioParameter>& InstanceParameters = GetInstanceParameters();
+	if (FAudioParameter* CurrentParam = FAudioParameter::FindOrAddParam(InstanceParameters, InParam.ParamName))
+	{
+		CurrentParam->Merge(InParam, false /* bInTakeName */);
+	}
+
+	if (IsPlaying() && !GetDisableParameterUpdatesWhilePlaying())
+	{
+		if (FAudioDevice* AudioDevice = GetAudioDevice())
+		{
+			FAudioParameter ParamToSet;
+			if (USoundBase* Sound = GetSound())
+			{
+				TArray<FAudioParameter> Params = { MoveTemp(InParam) };
+				Sound->InitParameters(Params, SoundParameterControllerInterfacePrivate::ProxyFeatureName);
+				ParamToSet = MoveTemp(Params[0]);
+			}
+			else
+			{
+				ParamToSet = MoveTemp(InParam);
+			}
+
+			DECLARE_CYCLE_STAT(TEXT("FAudioThreadTask.SoundParameterController.SetParameter"), STAT_AudioSetParameter, STATGROUP_AudioThreadCommands);
+
+			AudioDevice->SendCommandToActiveSounds(GetInstanceOwnerID(), [AudioDevice, Param = MoveTemp(ParamToSet)](FActiveSound& ActiveSound)
+			{
+				if (Audio::IParameterTransmitter* Transmitter = ActiveSound.GetTransmitter())
+				{
+					const FName ParamName = Param.ParamName;
+
+					// Must be copied as original version must be preserved in case command is called on multiple ActiveSounds.
+					FAudioParameter TempParam = Param;
+					if (!Transmitter->SetParameter(MoveTemp(TempParam)))
+					{
+						UE_LOG(LogAudio, Warning, TEXT("Failed to set parameter '%s'"), *ParamName.ToString());
+					}
+				}
+			}, GET_STATID(STAT_AudioSetParameter));
+		}
+	}
+}
diff --git a/Engine/Source/Runtime/Engine/Private/AudioDevice.cpp b/Engine/Source/Runtime/Engine/Private/AudioDevice.cpp
--- a/Engine/Source/Runtime/Engine/Private/AudioDevice.cpp
+++ b/Engine/Source/Runtime/Engine/Private/AudioDevice.cpp
@@ -3887,7 +3887,7 @@ int32 FAudioDevice::GetSortedActiveWaveInstances(TArray<FWaveInstance*>& WaveIns
 						UsedDeltaTime = 0.0f;
 					}
 
-					ActiveSound->UpdateGeneratorParameters(Listeners);
+					ActiveSound->UpdateInterfaceParameters(Listeners);
 					ActiveSound->UpdateWaveInstances(WaveInstances, UsedDeltaTime);
 				}
 			}
diff --git a/Engine/Source/Runtime/Engine/Private/AudioSettings.cpp b/Engine/Source/Runtime/Engine/Private/AudioSettings.cpp
--- a/Engine/Source/Runtime/Engine/Private/AudioSettings.cpp
+++ b/Engine/Source/Runtime/Engine/Private/AudioSettings.cpp
@@ -4,7 +4,7 @@
 #include "AudioDevice.h"
 #include "AudioDeviceManager.h"
 #include "AudioMixerDevice.h"
-#include "IAudioGeneratorInterfaceRegistry.h"
+#include "IAudioParameterInterfaceRegistry.h"
 #include "Misc/ConfigCacheIni.h"
 #include "Misc/Paths.h"
 #include "Sound/SoundAttenuation.h"
@@ -260,9 +260,9 @@ void UAudioSettings::LoadDefaultObjects()
 		UE_LOG(LogAudio, Display, TEXT("No default SoundConcurrencyObject specified (or failed to load)."));
 	}
 
-	Audio::IGeneratorInterfaceRegistry& InterfaceRegistry = Audio::IGeneratorInterfaceRegistry::Get();
-	InterfaceRegistry.RegisterInterface(Audio::GetAttenuationInterface());
-	InterfaceRegistry.RegisterInterface(Audio::GetSpatializationInterface());
+	Audio::IAudioParameterInterfaceRegistry& InterfaceRegistry = Audio::IAudioParameterInterfaceRegistry::Get();
+	InterfaceRegistry.RegisterInterface(Audio::AttenuationInterface::GetInterface());
+	InterfaceRegistry.RegisterInterface(Audio::SpatializationInterface::GetInterface());
 }
 
 USoundClass* UAudioSettings::GetDefaultMediaSoundClass() const
diff --git a/Engine/Source/Runtime/Engine/Private/Components/AudioComponent.cpp b/Engine/Source/Runtime/Engine/Private/Components/AudioComponent.cpp
--- a/Engine/Source/Runtime/Engine/Private/Components/AudioComponent.cpp
+++ b/Engine/Source/Runtime/Engine/Private/Components/AudioComponent.cpp
@@ -4,7 +4,6 @@
 #include "ActiveSound.h"
 #include "Audio.h"
 #include "Audio/ActorSoundParameterInterface.h"
-#include "Audio/SoundGeneratorParameterInterface.h"
 #include "AudioDevice.h"
 #include "AudioThread.h"
 #include "Components/BillboardComponent.h"
@@ -936,7 +935,7 @@ void UAudioComponent::AdjustVolumeInternal(float AdjustVolumeDuration, float Adj
 void UAudioComponent::ResetParameters()
 {
 	InstanceParameters.Reset();
-	ISoundGeneratorParameterInterface::ResetParameters();
+	ISoundParameterControllerInterface::ResetParameters();
 }
 
 void UAudioComponent::Stop()
diff --git a/Engine/Source/Runtime/Engine/Private/SoundAttenuation.cpp b/Engine/Source/Runtime/Engine/Private/SoundAttenuation.cpp
--- a/Engine/Source/Runtime/Engine/Private/SoundAttenuation.cpp
+++ b/Engine/Source/Runtime/Engine/Private/SoundAttenuation.cpp
@@ -6,6 +6,7 @@
 #include "EngineDefines.h"
 #include "Internationalization/Internationalization.h"
 #include "Sound/SoundBase.h"
+#include "Templates/SharedPointer.h"
 #include "UObject/AnimPhysObjectVersion.h"
 
 /*-----------------------------------------------------------------------------
@@ -144,69 +145,94 @@ USoundAttenuation::USoundAttenuation(const FObjectInitializer& ObjectInitializer
 {
 }
 
-#define LOCTEXT_NAMESPACE "AudioGeneratorInterface"
+#define LOCTEXT_NAMESPACE "AudioParameterInterface"
+#define AUDIO_PARAMETER_INTERFACE_NAMESPACE "UE.Attenuation"
 namespace Audio
 {
-	const FName FAttenuationInterface::Name = "UE.Attenuation";
-	const FName FAttenuationInterface::FInputs::Distance = "Distance";
-
-	FAttenuationInterface::FAttenuationInterface()
-		: FGeneratorInterface(Name, USoundBase::StaticClass())
+	namespace AttenuationInterface
 	{
-		Inputs =
-		{
-			{
-				FText(),
-				NSLOCTEXT("AudioGeneratorInterface_Attenuation", "DistanceDescription", "Distance between listener and sound location in game units."),
-				FName(),
-				{ FInputs::Distance, 0.0f }
-			}
-		};
-	}
+		const FName Name = AUDIO_PARAMETER_INTERFACE_NAMESPACE;
 
-	const FName FSpatializationInterface::Name = "UE.Spatialization";
-	const FName FSpatializationInterface::FInputs::Azimuth = "Azimuth";
-	const FName FSpatializationInterface::FInputs::Elevation = "Elevation";
+		namespace Inputs
+		{
+			const FName Distance = AUDIO_PARAMETER_INTERFACE_MEMBER_DEFINE("Distance");
+		} // namespace Inputs
 
-	FSpatializationInterface::FSpatializationInterface()
-		: FGeneratorInterface(Name, USoundBase::StaticClass())
-	{
-		Inputs =
+		Audio::FParameterInterfacePtr GetInterface()
 		{
+			struct FInterface : public Audio::FParameterInterface
 			{
-				FText(),
-				NSLOCTEXT("Spatialization", "AzimuthDescription", "Horizontal angle between listener forward and sound location in degrees."),
-				FName(),
-				{ FInputs::Azimuth, 0.0f }
-			},
+				FInterface()
+					: FParameterInterface(AttenuationInterface::Name, { 1, 0 }, *USoundBase::StaticClass())
+				{
+					Inputs =
+					{
+						{
+							FText(),
+							NSLOCTEXT("AudioGeneratorInterface_Attenuation", "DistanceDescription", "Distance between listener and sound location in game units."),
+							FName(),
+							{ Inputs::Distance, 0.0f }
+						}
+					};
+				}
+			};
+
+			static FParameterInterfacePtr InterfacePtr;
+			if (!InterfacePtr.IsValid())
 			{
-				FText(),
-				NSLOCTEXT("Spatialization", "ElevationDescription", "Vertical angle between listener forward and sound location in degrees."),
-				FName(),
-				{ FInputs::Elevation, 0.0f }
+				InterfacePtr = MakeShared<FInterface>();
 			}
-		};
-	}
 
-	FGeneratorInterfacePtr GetAttenuationInterface()
-	{
-		static FGeneratorInterfacePtr InterfacePtr;
-		if (!InterfacePtr.IsValid())
-		{
-			InterfacePtr = MakeShared<FAttenuationInterface>();
+			return InterfacePtr;
 		}
-		return InterfacePtr;
-	}
+	} // namespace AttenuationInterface
+#undef AUDIO_PARAMETER_INTERFACE_NAMESPACE
 
-	FGeneratorInterfacePtr GetSpatializationInterface()
+#define AUDIO_PARAMETER_INTERFACE_NAMESPACE "UE.Spatialization"
+	namespace SpatializationInterface
 	{
-		static FGeneratorInterfacePtr InterfacePtr;
-		if (!InterfacePtr.IsValid())
+		const FName Name = AUDIO_PARAMETER_INTERFACE_NAMESPACE;
+
+		namespace Inputs
 		{
-			InterfacePtr = MakeShared<FSpatializationInterface>();
-		}
+			const FName Azimuth = AUDIO_PARAMETER_INTERFACE_MEMBER_DEFINE("Azimuth");
+			const FName Elevation = AUDIO_PARAMETER_INTERFACE_MEMBER_DEFINE("Elevation");
+		} // namespace Inputs
 
-		return InterfacePtr;
-	}
+		Audio::FParameterInterfacePtr GetInterface()
+		{
+			struct FInterface : public Audio::FParameterInterface
+			{
+				FInterface()
+					: FParameterInterface(SpatializationInterface::Name, { 1, 0 }, *USoundBase::StaticClass())
+				{
+					Inputs =
+					{
+						{
+							FText(),
+							NSLOCTEXT("Spatialization", "AzimuthDescription", "Horizontal angle between listener forward and sound location in degrees."),
+							FName(),
+							{ Inputs::Azimuth, 0.0f }
+						},
+						{
+							FText(),
+							NSLOCTEXT("Spatialization", "ElevationDescription", "Vertical angle between listener forward and sound location in degrees."),
+							FName(),
+							{ Inputs::Elevation, 0.0f }
+						}
+					};
+				}
+			};
+
+			static FParameterInterfacePtr InterfacePtr;
+			if (!InterfacePtr.IsValid())
+			{
+				InterfacePtr = MakeShared<FInterface>();
+			}
+
+			return InterfacePtr;
+		}
+	} // namespace SpatializationInterface
+#undef AUDIO_PARAMETER_INTERFACE_NAMESPACE
 } // namespace Audio
-#undef LOCTEXT_NAMESPACE
+#undef LOCTEXT_NAMESPACE
\ No newline at end of file
diff --git a/Engine/Source/Runtime/Engine/Private/SoundBase.cpp b/Engine/Source/Runtime/Engine/Private/SoundBase.cpp
--- a/Engine/Source/Runtime/Engine/Private/SoundBase.cpp
+++ b/Engine/Source/Runtime/Engine/Private/SoundBase.cpp
@@ -2,13 +2,14 @@
 #include "Sound/SoundBase.h"
 
 #include "AudioDevice.h"
+#include "AudioParameter.h"
+#include "Engine/AssetUserData.h"
 #include "EngineDefines.h"
-#include "IAudioGeneratorInterfaceRegistry.h"
+#include "IAudioParameterInterfaceRegistry.h"
 #include "IAudioParameterTransmitter.h"
 #include "Sound/AudioSettings.h"
 #include "Sound/SoundClass.h"
 #include "Sound/SoundSubmix.h"
-#include "Engine/AssetUserData.h"
 
 
 USoundBase::USoundBase(const FObjectInitializer& ObjectInitializer)
@@ -394,13 +395,6 @@ TUniquePtr<Audio::IParameterTransmitter> USoundBase::CreateParameterTransmitter(
 			return false;
 		}
 
-		bool SetParameter(FName InInterfaceName, FAudioParameter&& InValue) override
-		{
-			InValue.ParamName = Audio::IGeneratorInterfaceRegistry::GetMemberFullName(InInterfaceName, InValue.ParamName);
-			const FName ParamName = InValue.ParamName;
-			return SetParameter(MoveTemp(InValue));
-		}
-
 		uint64 GetInstanceID() const override
 		{
 			return InstanceID;
diff --git a/Engine/Source/Runtime/Engine/Public/ActiveSound.h b/Engine/Source/Runtime/Engine/Public/ActiveSound.h
--- a/Engine/Source/Runtime/Engine/Public/ActiveSound.h
+++ b/Engine/Source/Runtime/Engine/Public/ActiveSound.h
@@ -663,7 +663,7 @@ public:
 	FColor DebugColor;
 #endif // ENABLE_AUDIO_DEBUG
 
-	void UpdateGeneratorParameters(const TArray<FListener>& InListeners);
+	void UpdateInterfaceParameters(const TArray<FListener>& InListeners);
 
 	// Updates the wave instances to be played.
 	void UpdateWaveInstances(TArray<FWaveInstance*> &OutWaveInstances, const float DeltaTime);
diff --git a/Engine/Source/Runtime/Engine/Public/Audio/SoundParameterControllerInterface.h b/Engine/Source/Runtime/Engine/Public/Audio/SoundParameterControllerInterface.h
new file mode 100644
--- /dev/null
+++ b/Engine/Source/Runtime/Engine/Public/Audio/SoundParameterControllerInterface.h
@@ -0,0 +1,64 @@
+// Copyright Epic Games, Inc. All Rights Reserved.
+#pragma once
+
+#include "AudioParameterControllerInterface.h"
+#include "UObject/Object.h"
+
+#include "SoundParameterControllerInterface.generated.h"
+
+
+// Forward Declarations
+class FAudioDevice;
+class USoundBase;
+struct FActiveSound;
+
+UINTERFACE(BlueprintType, meta = (CannotImplementInterfaceInBlueprint))
+class ENGINE_API USoundParameterControllerInterface : public UAudioParameterControllerInterface
+{
+	GENERATED_UINTERFACE_BODY()
+};
+
+// UObject interface for all object types that are controlling parameter values sent
+// to sound instances (i.e. sources)
+class ENGINE_API ISoundParameterControllerInterface : public IAudioParameterControllerInterface
+{
+	GENERATED_IINTERFACE_BODY()
+
+public:
+	// IAudioParameterControllerInterface
+	void ResetParameters() override;
+
+	void SetTriggerParameter(FName InName) override;
+	void SetBoolParameter(FName InName, bool InBool) override;
+	void SetBoolArrayParameter(FName InName, const TArray<bool>& InValue) override;
+	void SetIntParameter(FName InName, int32 InInt) override;
+	void SetIntArrayParameter(FName InName, const TArray<int32>& InValue) override;
+	void SetFloatParameter(FName InName, float InFloat) override;
+	void SetFloatArrayParameter(FName InName, const TArray<float>& InValue) override;
+	void SetStringParameter(FName InName, const FString& InValue) override;
+	void SetStringArrayParameter(FName InName, const TArray<FString>& InValue) override;
+	void SetObjectParameter(FName InName, UObject* InValue) override;
+	void SetObjectArrayParameter(FName InName, const TArray<UObject*>& InValue) override;
+
+	void SetParameter(FAudioParameter&& InValue) override;
+	void SetParameters(TArray<FAudioParameter>&& InValues) override;
+
+	/** Returns the active audio device to use for this component based on whether or not the component is playing in a world. */
+	virtual FAudioDevice* GetAudioDevice() const = 0;
+
+	/** Returns the id of the sound owner's instance associated with the parameter interface. */
+	virtual uint64 GetInstanceOwnerID() const = 0;
+
+	/** Returns the Game Thread copy of parameters to modify in place. */
+	virtual TArray<FAudioParameter>& GetInstanceParameters() = 0;
+
+	/** Returns the USoundBase used to initialize instance parameters to update. */
+	virtual USoundBase* GetSound() = 0;
+
+	virtual bool IsPlaying() const = 0;
+
+	virtual bool GetDisableParameterUpdatesWhilePlaying() const = 0;
+
+private:
+	void SetParameterInternal(FAudioParameter&& InValue);
+};
diff --git a/Engine/Source/Runtime/MovieSceneTracks/Private/Evaluation/MovieSceneAudioTemplate.cpp b/Engine/Source/Runtime/MovieSceneTracks/Private/Evaluation/MovieSceneAudioTemplate.cpp
--- a/Engine/Source/Runtime/MovieSceneTracks/Private/Evaluation/MovieSceneAudioTemplate.cpp
+++ b/Engine/Source/Runtime/MovieSceneTracks/Private/Evaluation/MovieSceneAudioTemplate.cpp
@@ -382,7 +382,7 @@ struct FAudioSectionExecutionToken : IMovieSceneExecutionToken
 
 	// Helper template to pair channel evaluation and parameter application.
 	template<typename ChannelType, typename ValueType>
-	void EvaluateAllAndSetParameters(IAudioParameterInterface& InParamaterInterface, const FFrameTime& InTime) const
+	void EvaluateAllAndSetParameters(IAudioParameterControllerInterface& InParamaterInterface, const FFrameTime& InTime) const
 	{
 		AudioSection->ForEachInput([&InParamaterInterface, &InTime](FName InName, const ChannelType& InChannel)
 		{
