EngineTest fix-ups for interface changes
#rnx

diff --git a/Engine/Plugins/Runtime/Metasound/Source/MetasoundEngineTest/MetasoundEngineTest.Build.cs b/Engine/Plugins/Runtime/Metasound/Source/MetasoundEngineTest/MetasoundEngineTest.Build.cs
--- a/Engine/Plugins/Runtime/Metasound/Source/MetasoundEngineTest/MetasoundEngineTest.Build.cs
+++ b/Engine/Plugins/Runtime/Metasound/Source/MetasoundEngineTest/MetasoundEngineTest.Build.cs
@@ -16,6 +16,7 @@ namespace UnrealBuildTool.Rules
 					"Core",
 					"CoreUObject",
 					"Engine",
+					"AudioExtensions",
 					"MetasoundGraphCore",
 					"MetasoundGenerator",
 					"MetasoundFrontend",
diff --git a/Engine/Plugins/Runtime/Metasound/Source/MetasoundEngineTest/Private/EngineTestMetasoundSource.cpp b/Engine/Plugins/Runtime/Metasound/Source/MetasoundEngineTest/Private/EngineTestMetasoundSource.cpp
--- a/Engine/Plugins/Runtime/Metasound/Source/MetasoundEngineTest/Private/EngineTestMetasoundSource.cpp
+++ b/Engine/Plugins/Runtime/Metasound/Source/MetasoundEngineTest/Private/EngineTestMetasoundSource.cpp
@@ -2,10 +2,13 @@
 
 #include "AudioDevice.h"
 #include "Components/AudioComponent.h"
+#include "IAudioParameterInterfaceRegistry.h"
 #include "Interfaces/IPluginManager.h"
 #include "MetasoundFrontendController.h"
 #include "MetasoundFrontendSearchEngine.h"
+#include "MetasoundOutputFormatInterfaces.h"
 #include "MetasoundSource.h"
+#include "MetasoundSourceInterface.h"
 #include "Misc/AutomationTest.h"
 #include "Tests/AutomationCommon.h"
 
@@ -60,58 +63,46 @@ namespace EngineTestMetasoundSourcePrivate
 
 	FMetasoundFrontendDocument CreateMetaSoundMonoSourceDocument()
 	{
+		using namespace Audio;
 		using namespace Metasound;
 		using namespace Metasound::Frontend;
 
 		FMetasoundFrontendDocument Document;
-		Document.InterfaceVersions.Add(FMetasoundFrontendVersion{"MonoSource", {1, 1}});
 
 		Document.RootGraph.Metadata.SetClassName(FMetasoundFrontendClassName { "Namespace", "Unit Test Node", *LexToString(FGuid::NewGuid()) });
 		Document.RootGraph.Metadata.SetType(EMetasoundFrontendClassType::Graph);
 
 		FDocumentHandle DocumentHandle = IDocumentController::CreateDocumentHandle(Document);
-		Metasound::Frontend::FGraphHandle RootGraph = DocumentHandle->GetRootGraph();
+		FGraphHandle RootGraph = DocumentHandle->GetRootGraph();
 		check(RootGraph->IsValid());
 
-		// Input on play
-		FMetasoundFrontendClassInput OnPlayInput;
-		OnPlayInput.Name = TEXT("On Play");
-		OnPlayInput.TypeName = GetMetasoundDataTypeName<FTrigger>();
-		OnPlayInput.VertexID = FGuid::NewGuid();
+		// Add default source & mono interface members (OnPlay, OnFinished & Mono Out)
+		FModifyRootGraphInterfaces InterfaceTransform({ SourceInterface::GetVersion(), OutputFormatMonoInterface::GetVersion() }, { });
+		InterfaceTransform.Transform(DocumentHandle);
 
-		FNodeHandle OnPlayInputNode = RootGraph->AddInputVertex(OnPlayInput);
-		check(OnPlayInputNode->IsValid());
+		// Input on Play
+		FNodeHandle OnPlayOutputNode = RootGraph->GetOutputNodeWithName(SourceInterface::Inputs::OnPlay);
+		check(OnPlayOutputNode->IsValid());
 
 		// Input Frequency
 		FMetasoundFrontendClassInput FrequencyInput;
-		FrequencyInput.Name = TEXT("Frequency");
+		FrequencyInput.Name = "Frequency";
 		FrequencyInput.TypeName = GetMetasoundDataTypeName<float>();
 		FrequencyInput.VertexID = FGuid::NewGuid();
 		FrequencyInput.DefaultLiteral.Set(100.f);
-
 		FNodeHandle FrequencyInputNode = RootGraph->AddInputVertex(FrequencyInput);
 		check(FrequencyInputNode->IsValid());
 
-
-		// Output Is Finished
-		FMetasoundFrontendClassOutput IsFinishedOutput;
-		IsFinishedOutput.Name = TEXT("On Finished");
-		IsFinishedOutput.TypeName = Metasound::GetMetasoundDataTypeName<FTrigger>();
-		IsFinishedOutput.VertexID = FGuid::NewGuid();
-		FNodeHandle IsFinishedOutputNode = RootGraph->AddOutputVertex(IsFinishedOutput);
-		check(IsFinishedOutputNode->IsValid());
+		// Output On Finished
+		FNodeHandle OnFinishedOutputNode = RootGraph->GetOutputNodeWithName(SourceInterface::Outputs::OnFinished);
+		check(OnFinishedOutputNode->IsValid());
 
 		// Output Audio
-		FMetasoundFrontendClassOutput AudioOutput;
-		AudioOutput.Name = TEXT("Audio:0");
-		AudioOutput.TypeName = Metasound::GetMetasoundDataTypeName<FAudioBuffer>();
-		AudioOutput.VertexID = FGuid::NewGuid();
-		
-		FNodeHandle OutputAudioNode = RootGraph->AddOutputVertex(AudioOutput);
-		check(OutputAudioNode->IsValid());
+		FNodeHandle AudioOutputNode = RootGraph->GetOutputNodeWithName(OutputFormatMonoInterface::Outputs::MonoOut);
+		check(AudioOutputNode->IsValid());
 
-		// osc node 
-		FNodeHandle OscNode = AddNode(*RootGraph, {TEXT("UE"), TEXT("Sine"), TEXT("Audio")});
+		// osc node
+		FNodeHandle OscNode = AddNode(*RootGraph, { "UE", "Sine", "Audio" });
 
 		// Make connections:
 
@@ -122,7 +113,7 @@ namespace EngineTestMetasoundSourcePrivate
 
 		// oscillator to output
 		OutputToConnect = OscNode->GetOutputWithVertexName("Audio");
-		InputToConnect = OutputAudioNode->GetInputWithVertexName("Audio:0");
+		InputToConnect = AudioOutputNode->GetInputWithVertexName(OutputFormatMonoInterface::Outputs::MonoOut);
 		ensure(InputToConnect->Connect(*OutputToConnect));
 
 		return Document;
