- Fix bug with literal box offset when not visible
#rnx

diff --git a/Engine/Plugins/Runtime/Metasound/Source/MetasoundEditor/Private/SMetasoundGraphPin.h b/Engine/Plugins/Runtime/Metasound/Source/MetasoundEditor/Private/SMetasoundGraphPin.h
--- a/Engine/Plugins/Runtime/Metasound/Source/MetasoundEditor/Private/SMetasoundGraphPin.h
+++ b/Engine/Plugins/Runtime/Metasound/Source/MetasoundEditor/Private/SMetasoundGraphPin.h
@@ -115,7 +115,7 @@ namespace Metasound
 								}
 							}
 
-							return EVisibility::Hidden;
+							return EVisibility::Collapsed;
 						}))
 						.OnClicked(FOnClicked::CreateLambda([this]()
 						{
