- Fix bug with variable types being selectible in DataType picker
#rnx

diff --git a/Engine/Plugins/Runtime/Metasound/Source/MetasoundEditor/Private/MetasoundNodeDetailCustomization.cpp b/Engine/Plugins/Runtime/Metasound/Source/MetasoundEditor/Private/MetasoundNodeDetailCustomization.cpp
--- a/Engine/Plugins/Runtime/Metasound/Source/MetasoundEditor/Private/MetasoundNodeDetailCustomization.cpp
+++ b/Engine/Plugins/Runtime/Metasound/Source/MetasoundEditor/Private/MetasoundNodeDetailCustomization.cpp
@@ -603,7 +603,6 @@ namespace Metasound
 			DetailLayoutBuilder = &InDetailLayout;
 			GraphMember = InGraphMember;
 
-			IDetailCategoryBuilder& CategoryBuilder = InDetailLayout.EditCategory("General");
 			FString CurrentTypeName = GraphMember->TypeName.ToString();
 
 			bool bCurrentTypeIsArray = CurrentTypeName.EndsWith(MemberCustomizationPrivate::ArrayIdentifier);
@@ -620,33 +619,30 @@ namespace Metasound
 			const bool bIsArrayTypeRegistered = EditorModule.IsRegisteredDataType(ArrayType);
 			const bool bIsArrayTypeRegisteredHidden = MemberCustomizationPrivate::HiddenInputTypeNames.Contains(ArrayType);
 
-			TArray<FName> DataTypeNames;
+			TArray<FName> BaseDataTypes;
 			EditorModule.IterateDataTypes([&](const FEditorDataType& EditorDataType)
 			{
 				const FName& TypeName = EditorDataType.RegistryInfo.DataTypeName;
-				FString TypeNameString = TypeName.ToString();
 
-				// Array types are handled separately via checkbox
-				if (TypeNameString.EndsWith(MemberCustomizationPrivate::ArrayIdentifier))
+				// Hide the type from the combo selector if any of the following is true;
+				const bool bIsArrayType = EditorDataType.RegistryInfo.IsArrayType();
+				const bool bIsVariable = EditorDataType.RegistryInfo.bIsVariable;
+				const bool bIsHiddenType = MemberCustomizationPrivate::HiddenInputTypeNames.Contains(TypeName);
+				const bool bHideBaseType = bIsArrayType || bIsVariable || bIsHiddenType;
+				if (!bHideBaseType)
 				{
-					return;
-				}
-
-				// Hidden input types should be omitted from the drop down.
-				if (!MemberCustomizationPrivate::HiddenInputTypeNames.Contains(EditorDataType.RegistryInfo.DataTypeName))
-				{
-					DataTypeNames.Add(TypeName);
+					BaseDataTypes.Add(TypeName);
 				}
 			});
 
-			DataTypeNames.Sort([](const FName& DataTypeNameL, const FName& DataTypeNameR)
+			BaseDataTypes.Sort([](const FName& DataTypeNameL, const FName& DataTypeNameR)
 			{
 				return DataTypeNameL.LexicalLess(DataTypeNameR);
 			});
 
-			Algo::Transform(DataTypeNames, ComboOptions, [](const FName& Name) { return MakeShared<FString>(Name.ToString()); });
+			Algo::Transform(BaseDataTypes, ComboOptions, [](const FName& Name) { return MakeShared<FString>(Name.ToString()); });
 
-			CategoryBuilder.AddCustomRow(InRowName)
+			InDetailLayout.EditCategory("General").AddCustomRow(InRowName)
 			.IsEnabled(bIsEnabled)
 			.NameContent()
 			[
