Move EditorModule to be public header to load from other modules

diff --git a/Engine/Plugins/Runtime/Metasound/Source/MetasoundEditor/Private/MetasoundEditorModule.h b/Engine/Plugins/Runtime/Metasound/Source/MetasoundEditor/Private/MetasoundEditorModule.h
deleted file mode 100644
diff --git a/Engine/Plugins/Runtime/Metasound/Source/MetasoundEditor/Public/MetasoundEditorModule.h b/Engine/Plugins/Runtime/Metasound/Source/MetasoundEditor/Public/MetasoundEditorModule.h
new file mode 100644
--- /dev/null
+++ b/Engine/Plugins/Runtime/Metasound/Source/MetasoundEditor/Public/MetasoundEditorModule.h
@@ -0,0 +1,81 @@
+// Copyright Epic Games, Inc. All Rights Reserved.
+#pragma once
+
+#include "CoreMinimal.h"
+#include "DetailCategoryBuilder.h"
+#include "EdGraph/EdGraphNode.h"
+#include "EdGraph/EdGraphPin.h"
+#include "MetasoundFrontendDataTypeRegistry.h"
+#include "MetasoundFrontendLiteral.h"
+#include "MetasoundFrontendRegistries.h"
+#include "Modules/ModuleInterface.h"
+#include "PropertyHandle.h"
+#include "Templates/Function.h"
+
+
+class UMetasoundEditorGraph;
+class UMetasoundEditorGraphInputLiteral;
+
+DECLARE_LOG_CATEGORY_EXTERN(LogMetasoundEditor, Log, All);
+
+
+namespace Metasound
+{
+	namespace Editor
+	{
+		using FDataTypeRegistryInfo = Frontend::FDataTypeRegistryInfo;
+
+		struct FEditorDataType
+		{
+			FEdGraphPinType PinType;
+			FDataTypeRegistryInfo RegistryInfo;
+
+			FEditorDataType(FEdGraphPinType&& InPinType, FDataTypeRegistryInfo&& InRegistryInfo)
+				: PinType(MoveTemp(InPinType))
+				, RegistryInfo(InRegistryInfo)
+			{
+			}
+		};
+
+		class METASOUNDEDITOR_API IMetaSoundInputLiteralCustomization
+		{
+		public:
+			virtual ~IMetaSoundInputLiteralCustomization() = default;
+
+			virtual void CustomizeLiteral(UMetasoundEditorGraphInputLiteral& InLiteral, TSharedPtr<IPropertyHandle> InDefaultValueHandle) = 0;
+		};
+
+		class METASOUNDEDITOR_API IMetaSoundInputLiteralCustomizationFactory
+		{
+		public:
+			virtual ~IMetaSoundInputLiteralCustomizationFactory() = default;
+
+			virtual TUniquePtr<IMetaSoundInputLiteralCustomization> CreateLiteralCustomization(IDetailCategoryBuilder& DefaultCategoryBuilder) const = 0;
+		};
+
+		class METASOUNDEDITOR_API IMetasoundEditorModule : public IModuleInterface
+		{
+		public:
+			// Whether or not the given proxy class has to be explicit (i.e.
+			// selectors do not support inherited types). By default, proxy
+			// classes support child classes & inheritance.
+			virtual bool IsExplicitProxyClass(const UClass& InClass) const = 0;
+
+			// Register proxy class as explicitly selectable.
+			// By default, proxy classes support child classes & inheritance.
+			virtual void RegisterExplicitProxyClass(const UClass& InClass) = 0;
+
+			virtual const FEditorDataType* FindDataType(FName InDataTypeName) const = 0;
+			virtual const FEditorDataType& FindDataTypeChecked(FName InDataTypeName) const = 0;
+			virtual bool IsMetaSoundAssetClass(const FName InClassName) const = 0;
+
+			virtual bool IsRegisteredDataType(FName InDataTypeName) const = 0;
+
+			virtual void IterateDataTypes(TUniqueFunction<void(const FEditorDataType&)> InDataTypeFunction) const = 0;
+
+			virtual TUniquePtr<IMetaSoundInputLiteralCustomization> CreateInputLiteralCustomization(UClass& InClass, IDetailCategoryBuilder& DefaultCategoryBuilder) const = 0;
+
+			virtual const TSubclassOf<UMetasoundEditorGraphInputLiteral> FindInputLiteralClass(EMetasoundFrontendLiteralType InLiteralType) const = 0;
+		};
+	} // namespace Editor
+} // namespace Metasound
