MetaSound conversion node improvements
- enable autoconverter nodes between enums and int32 (and vice versa)
- add background image and description with typenames for autoconverter nodes to match BP
- update conversion nodes appearance to match autoconversion nodes

diff --git a/Engine/Plugins/Runtime/Metasound/Source/MetasoundEditor/Private/MetasoundEditorGraphSchema.cpp b/Engine/Plugins/Runtime/Metasound/Source/MetasoundEditor/Private/MetasoundEditorGraphSchema.cpp
--- a/Engine/Plugins/Runtime/Metasound/Source/MetasoundEditor/Private/MetasoundEditorGraphSchema.cpp
+++ b/Engine/Plugins/Runtime/Metasound/Source/MetasoundEditor/Private/MetasoundEditorGraphSchema.cpp
@@ -1219,13 +1219,14 @@ void UMetasoundEditorGraphSchema::GetConversionActions(FGraphActionMenuBuilder&
 		const FText Tooltip = Metadata.GetAuthor().IsEmpty()
 			? Metadata.GetDescription()
 			: FText::Format(LOCTEXT("MetasoundTooltipAuthorFormat", "{0}\nAuthor: {1}"), Metadata.GetDescription(), Metadata.GetAuthor());
-
 		if (!Metadata.GetCategoryHierarchy().IsEmpty() && !Metadata.GetCategoryHierarchy()[0].CompareTo(NodeCategories::Conversions))
 		{
 			FText KeywordsText = FText::Join(SchemaPrivate::KeywordDelim, Metadata.GetKeywords());
+			const FText CategoryText = FText::Join(SchemaPrivate::CategoryDelim, Metadata.GetCategoryHierarchy());
+
 			TSharedPtr<FMetasoundGraphSchemaAction_NewNode> NewNodeAction = MakeShared<FMetasoundGraphSchemaAction_NewNode>
 			(
-				NodeCategories::Conversions,
+				CategoryText,
 				Metadata.GetDisplayName(),
 				Tooltip,
 				EPrimaryContextGroup::Conversions,
diff --git a/Engine/Plugins/Runtime/Metasound/Source/MetasoundEditor/Private/MetasoundEditorModule.cpp b/Engine/Plugins/Runtime/Metasound/Source/MetasoundEditor/Private/MetasoundEditorModule.cpp
--- a/Engine/Plugins/Runtime/Metasound/Source/MetasoundEditor/Private/MetasoundEditorModule.cpp
+++ b/Engine/Plugins/Runtime/Metasound/Source/MetasoundEditor/Private/MetasoundEditorModule.cpp
@@ -129,9 +129,9 @@ namespace Metasound
 					Set("MetasoundEditor.Graph.Node.Math.Modulo", new FSlateImageBrush(RootToContentDir(TEXT("/Graph/node_math_modulo_40x.png")), Icon40x40));
 					Set("MetasoundEditor.Graph.Node.Math.Multiply", new FSlateImageBrush(RootToContentDir(TEXT("/Graph/node_math_multiply_40x.png")), Icon40x40));
 					Set("MetasoundEditor.Graph.Node.Math.Subtract", new FSlateImageBrush(RootToContentDir(TEXT("/Graph/node_math_subtract_40x.png")), Icon40x40));
-					Set("MetasoundEditor.Graph.Node.Math.Modulo", new FSlateImageBrush(RootToContentDir(TEXT("/Graph/node_math_modulo_40x.png")), Icon40x40));
 					Set("MetasoundEditor.Graph.Node.Math.Power", new FSlateImageBrush(RootToContentDir(TEXT("/Graph/node_math_power_40x.png")), Icon40x40));
 					Set("MetasoundEditor.Graph.Node.Math.Logarithm", new FSlateImageBrush(RootToContentDir(TEXT("/Graph/node_math_logarithm_40x.png")), Icon40x40));
+					Set("MetasoundEditor.Graph.Node.Conversion", new FSlateImageBrush(RootToContentDir(TEXT("/Graph/node_conversion_40x.png")), Icon40x40));
 
 					// Analyzers
 					Set("MetasoundEditor.Analyzers.BackgroundColor", FLinearColor(0.0075f, 0.0075f, 0.0075, 1.0f));
diff --git a/Engine/Plugins/Runtime/Metasound/Source/MetasoundFrontend/Private/MetasoundFrontendNodesCategories.cpp b/Engine/Plugins/Runtime/Metasound/Source/MetasoundFrontend/Private/MetasoundFrontendNodesCategories.cpp
--- a/Engine/Plugins/Runtime/Metasound/Source/MetasoundFrontend/Private/MetasoundFrontendNodesCategories.cpp
+++ b/Engine/Plugins/Runtime/Metasound/Source/MetasoundFrontend/Private/MetasoundFrontendNodesCategories.cpp
@@ -11,6 +11,9 @@ namespace Metasound
 		// Base category for all conversion functions.
 		const FText Conversions  = LOCTEXT("Metasound_ConversionsCategory", "Conversions");
 
+		// Category for enum conversions. A subcategory of conversion functions to prevent cluttering up the action menu. 
+		const FText EnumConversions = LOCTEXT("Metasound_EnumConversionsCategory", "EnumConversions");
+
 		// Base category for all standard MetaSound functions.
 		const FText Functions = LOCTEXT("Metasound_FunctionsCategory", "Functions");
 
diff --git a/Engine/Plugins/Runtime/Metasound/Source/MetasoundFrontend/Public/MetasoundAutoConverterNode.h b/Engine/Plugins/Runtime/Metasound/Source/MetasoundFrontend/Public/MetasoundAutoConverterNode.h
--- a/Engine/Plugins/Runtime/Metasound/Source/MetasoundFrontend/Public/MetasoundAutoConverterNode.h
+++ b/Engine/Plugins/Runtime/Metasound/Source/MetasoundFrontend/Public/MetasoundAutoConverterNode.h
@@ -21,12 +21,28 @@
 
 namespace Metasound
 {
-	// This convenience node can be registered and will invoke static_cast<ToDataType>(FromDataType) every time it is executed.
+	// Determines whether an auto converter node will be registered to convert 
+	// between two types. 
+	template<typename TFromDataType, typename TToDataType>
+	struct TIsAutoConvertible
+	{
+		static constexpr bool bIsConvertible = std::is_convertible<TFromDataType, TToDataType>::value;
+
+		// Handle case of converting enums to/from integers.
+		static constexpr bool bIsIntToEnumConversion = std::is_same<int32, TFromDataType>::value && TEnumTraits<TToDataType>::bIsEnum;
+		static constexpr bool bIsEnumToIntConversion = TEnumTraits<TFromDataType>::bIsEnum && std::is_same<int32, TToDataType>::value;
+
+		static constexpr bool Value = bIsConvertible || bIsIntToEnumConversion || bIsEnumToIntConversion;
+	};
+
+	// This convenience node can be registered and will invoke static_cast<ToDataType>(FromDataType) every time it is executed, 
+	// with a special case for enum <-> int32 conversions. 
 	template<typename FromDataType, typename ToDataType>
 	class TAutoConverterNode : public FNode
 	{
-		static_assert(std::is_convertible<FromDataType, ToDataType>::value, "Tried to create an auto converter node between two types we can't static_cast between.");
-
+		static_assert(TIsAutoConvertible<FromDataType, ToDataType>::Value,
+		"Tried to create an auto converter node between two types we can't static_cast between.");
+		
 	public:
 		static const FVertexName& GetInputName()
 		{
@@ -42,12 +58,15 @@ namespace Metasound
 
 		static FVertexInterface DeclareVertexInterface()
 		{
+			static const FText InputDesc = FText::Format(LOCTEXT("AutoConvDisplayNamePatternFrom", "Input {0} value."), GetMetasoundDataTypeDisplayText<FromDataType>());
+			static const FText OutputDesc = FText::Format(LOCTEXT("AutoConvDisplayNamePatternTo", "Output {0} value."), GetMetasoundDataTypeDisplayText<ToDataType>());
+
 			return FVertexInterface(
 				FInputVertexInterface(
-					TInputDataVertexModel<FromDataType>(GetInputName(), FText::GetEmpty())
+					TInputDataVertexModel<FromDataType>(GetInputName(), InputDesc)
 				),
 				FOutputVertexInterface(
-					TOutputDataVertexModel<ToDataType>(GetOutputName(), FText::GetEmpty())
+					TOutputDataVertexModel<ToDataType>(GetOutputName(), OutputDesc)
 				)
 			);
 		}
@@ -58,26 +77,35 @@ namespace Metasound
 			{
 				FNodeDisplayStyle DisplayStyle;
 				DisplayStyle.bShowName = false;
+				DisplayStyle.ImageName = TEXT("MetasoundEditor.Graph.Node.Conversion");
 				DisplayStyle.bShowInputNames = false;
 				DisplayStyle.bShowOutputNames = false;
 
+				const FText FromTypeText = GetMetasoundDataTypeDisplayText<FromDataType>();
+				const FText ToTypeText = GetMetasoundDataTypeDisplayText<ToDataType>();
+
 				FNodeClassMetadata Info;
 				Info.ClassName = { TEXT("Convert"), GetMetasoundDataTypeName<ToDataType>(), GetMetasoundDataTypeName<FromDataType>() };
 				Info.MajorVersion = 1;
 				Info.MinorVersion = 0;
-				Info.DisplayName = FText::Format(LOCTEXT("Metasound_AutoConverterNodeDisplayNameFormat", "{0} to {1}"), FText::FromName(GetMetasoundDataTypeName<FromDataType>()), FText::FromName(GetMetasoundDataTypeName<ToDataType>()));
-				Info.Description = LOCTEXT("Metasound_ConverterNodeDescription", "Converts between two different data types.");
+				Info.DisplayName = FText::Format(LOCTEXT("Metasound_AutoConverterNodeDisplayNameFormat", "{0} to {1}"), FromTypeText, ToTypeText);
+				Info.Description = FText::Format(LOCTEXT("Metasound_AutoConverterNodeDescriptionNameFormat", "Converts from {0} to {1}."), FromTypeText, ToTypeText);
 				Info.Author = PluginAuthor;
 				Info.DisplayStyle = DisplayStyle;
 				Info.PromptIfMissing = PluginNodeMissingPrompt;
 				Info.DefaultInterface = DeclareVertexInterface();
 
 				Info.CategoryHierarchy.Emplace(NodeCategories::Conversions);
+				if (TEnumTraits<FromDataType>::bIsEnum || TEnumTraits<ToDataType>::bIsEnum)
+				{
+					Info.CategoryHierarchy.Emplace(NodeCategories::EnumConversions);
+				}
+				
 				Info.Keywords =
 				{
 					LOCTEXT("MetasoundConvertKeyword", "Convert"),
-					FText::FromName(GetMetasoundDataTypeName<FromDataType>()),
-					FText::FromName(GetMetasoundDataTypeName<ToDataType>())
+					GetMetasoundDataTypeDisplayText<FromDataType>(),
+					GetMetasoundDataTypeDisplayText<ToDataType>()
 				};
 
 				return Info;
@@ -90,42 +118,88 @@ namespace Metasound
 
 	private:
 		/** FConverterOperator converts from "FromDataType" to "ToDataType" using
-		 * a implicit conversion operators. 
+		 * a implicit conversion operators.
 		 */
 		class FConverterOperator : public TExecutableOperator<FConverterOperator>
 		{
-			public:
+		public:
 
-				FConverterOperator(TDataReadReference<FromDataType> InFromDataReference, TDataWriteReference<ToDataType> InToDataReference)
-					: FromData(InFromDataReference)
-					, ToData(InToDataReference) 
-				{
-				}
+			FConverterOperator(TDataReadReference<FromDataType> InFromDataReference, TDataWriteReference<ToDataType> InToDataReference)
+				: FromData(InFromDataReference)
+				, ToData(InToDataReference)
+			{
+			}
 
-				virtual ~FConverterOperator() {}
+			virtual ~FConverterOperator() {}
 
-				virtual FDataReferenceCollection GetInputs() const override
-				{
-					FDataReferenceCollection Inputs;
-					Inputs.AddDataReadReference<FromDataType>(GetInputName(), FromData);
-					return Inputs;
-				}
+			virtual FDataReferenceCollection GetInputs() const override
+			{
+				FDataReferenceCollection Inputs;
+				Inputs.AddDataReadReference<FromDataType>(GetInputName(), FromData);
+				return Inputs;
+			}
+
+			virtual FDataReferenceCollection GetOutputs() const override
+			{
+				FDataReferenceCollection Outputs;
+				Outputs.AddDataReadReference<ToDataType>(GetOutputName(), ToData);
+				return Outputs;
+			}
 
-				virtual FDataReferenceCollection GetOutputs() const override
+			void Execute()
+			{
+				// enum -> int32
+				if constexpr (TIsAutoConvertible<FromDataType, ToDataType>::bIsEnumToIntConversion)
 				{
-					FDataReferenceCollection Outputs;
-					Outputs.AddDataReadReference<ToDataType>(GetOutputName(), ToData);
-					return Outputs;
+					// Convert from enum wrapper to inner enum type, then to int
+					typename TEnumTraits<FromDataType>::InnerType InnerEnum = static_cast<typename TEnumTraits<FromDataType>::InnerType>(*FromData);
+					*ToData = static_cast<ToDataType>(InnerEnum);
 				}
+				// int32 -> enum
+				else if constexpr (TIsAutoConvertible<FromDataType, ToDataType>::bIsIntToEnumConversion)
+				{
+					const int32 FromInt = *FromData;
+					// Convert from int to inner enum type
+					typename TEnumTraits<ToDataType>::InnerType InnerEnum = static_cast<typename TEnumTraits<ToDataType>::InnerType>(FromInt);
 
-				void Execute()
+					// Update tracking for previous int value we tried to convert, used to prevent log spam if it's an invalid enum value
+					if (FromInt != PreviousIntValueForEnumConversion)
+					{
+						PreviousIntValueForEnumConversion = FromInt;
+						bHasLoggedInvalidEnum = false;
+					}
+
+					// If int value is invalid for this enum, return enum default value
+					TOptional<FName> EnumName = ToDataType::ToName(InnerEnum);
+					if (!EnumName.IsSet())
+					{
+						if (!bHasLoggedInvalidEnum)
+						{
+							UE_LOG(LogMetaSound, Warning, TEXT("Cannot convert int32 value '%d' to enum type '%s'. No valid corresponding enum value exists, so returning enum default value instead."), FromInt, *GetMetasoundDataTypeDisplayText<ToDataType>().ToString());
+							bHasLoggedInvalidEnum = true;
+						}
+						*ToData = static_cast<ToDataType>(TEnumTraits<ToDataType>::DefaultValue);
+					}
+					else
+					{
+						// Convert from inner enum type to int
+						*ToData = static_cast<ToDataType>(InnerEnum);
+					}
+				}
+				else
 				{
 					*ToData = static_cast<ToDataType>(*FromData);
 				}
+			}
 
 			private:
 				TDataReadReference<FromDataType> FromData;
 				TDataWriteReference<ToDataType> ToData;
+
+				// To prevent log spam, keep track of whether we've logged an invalid enum value being converted already
+				// and the previous int value (need both bool and int for the initial case)
+				bool bHasLoggedInvalidEnum = false;
+				int32 PreviousIntValueForEnumConversion = 0;
 		};
 
 		/** FConverterOperatorFactory creates an operator which converts from 
diff --git a/Engine/Plugins/Runtime/Metasound/Source/MetasoundFrontend/Public/MetasoundDataTypeRegistrationMacro.h b/Engine/Plugins/Runtime/Metasound/Source/MetasoundFrontend/Public/MetasoundDataTypeRegistrationMacro.h
--- a/Engine/Plugins/Runtime/Metasound/Source/MetasoundFrontend/Public/MetasoundDataTypeRegistrationMacro.h
+++ b/Engine/Plugins/Runtime/Metasound/Source/MetasoundFrontend/Public/MetasoundDataTypeRegistrationMacro.h
@@ -89,13 +89,6 @@ namespace Metasound
 			static constexpr bool Value = TIsTransmittable<TElementType>::Value;
 		};
 
-		// Determines whether an auto converter node will be registered to convert 
-		// between two types. 
-		template<typename TFromDataType, typename TToDataType>
-		struct TIsAutoConvertible
-		{
-			static constexpr bool Value = std::is_convertible<TFromDataType, TToDataType>::value;
-		};
 		// Returns the Array version of a literal type if it exists.
 		template<ELiteralType LiteralType>
 		struct TLiteralArrayEnum 
diff --git a/Engine/Plugins/Runtime/Metasound/Source/MetasoundFrontend/Public/MetasoundEnumRegistrationMacro.h b/Engine/Plugins/Runtime/Metasound/Source/MetasoundFrontend/Public/MetasoundEnumRegistrationMacro.h
--- a/Engine/Plugins/Runtime/Metasound/Source/MetasoundFrontend/Public/MetasoundEnumRegistrationMacro.h
+++ b/Engine/Plugins/Runtime/Metasound/Source/MetasoundFrontend/Public/MetasoundEnumRegistrationMacro.h
@@ -9,18 +9,18 @@
 
 namespace Metasound
 {
-	// Disable registration of converter nodes which convert to enums
+	// Enable registration of converter nodes which convert from int32 to enums
 	template<typename FromDataType, typename EnumType, EnumType DefaultValue>
 	struct TEnableAutoConverterNodeRegistration<FromDataType, TEnum<EnumType, DefaultValue>>
 	{
-		static constexpr bool Value = false;
+		static constexpr bool Value = std::is_same<int32, FromDataType>::value;
 	};
 
-	// Disable registration of converter nodes which convert from enums
+	// Enable registration of converter nodes which convert from enums to int32
 	template<typename ToDataType, typename EnumType, EnumType DefaultValue>
 	struct TEnableAutoConverterNodeRegistration<TEnum<EnumType, DefaultValue>, ToDataType>
 	{
-		static constexpr bool Value = false;
+		static constexpr bool Value = std::is_same<int32, ToDataType>::value;
 	};
 
 	// Disable arrays of enums
diff --git a/Engine/Plugins/Runtime/Metasound/Source/MetasoundFrontend/Public/MetasoundFrontendNodesCategories.h b/Engine/Plugins/Runtime/Metasound/Source/MetasoundFrontend/Public/MetasoundFrontendNodesCategories.h
--- a/Engine/Plugins/Runtime/Metasound/Source/MetasoundFrontend/Public/MetasoundFrontendNodesCategories.h
+++ b/Engine/Plugins/Runtime/Metasound/Source/MetasoundFrontend/Public/MetasoundFrontendNodesCategories.h
@@ -9,6 +9,7 @@ namespace Metasound
 	namespace NodeCategories
 	{
 		extern const FText METASOUNDFRONTEND_API Conversions;
+		extern const FText METASOUNDFRONTEND_API EnumConversions;
 		extern const FText METASOUNDFRONTEND_API Functions;
 		extern const FText METASOUNDFRONTEND_API Graphs;
 		extern const FText METASOUNDFRONTEND_API Inputs;
diff --git a/Engine/Plugins/Runtime/Metasound/Source/MetasoundStandardNodes/Private/MetasoundConversionNodes.cpp b/Engine/Plugins/Runtime/Metasound/Source/MetasoundStandardNodes/Private/MetasoundConversionNodes.cpp
--- a/Engine/Plugins/Runtime/Metasound/Source/MetasoundStandardNodes/Private/MetasoundConversionNodes.cpp
+++ b/Engine/Plugins/Runtime/Metasound/Source/MetasoundStandardNodes/Private/MetasoundConversionNodes.cpp
@@ -153,7 +153,7 @@ namespace Metasound
 			using namespace ConversionNodeVertexNames;
 
 			static FText InputDesc = FText::Format(LOCTEXT("ConvDisplayNamePatternFrom", "Input {0} value."), GetMetasoundDataTypeDisplayText<FromType>());
-			static FText OutputDesc = FText::Format(LOCTEXT("ConvDisplayNamePatternTo", "Ouput {0} value."), GetMetasoundDataTypeDisplayText<ToType>());
+			static FText OutputDesc = FText::Format(LOCTEXT("ConvDisplayNamePatternTo", "Output {0} value."), GetMetasoundDataTypeDisplayText<ToType>());
 
 			static const FVertexInterface DefaultInterface(
 				FInputVertexInterface(
@@ -171,6 +171,12 @@ namespace Metasound
 		{
 			auto InitNodeInfo = []() -> FNodeClassMetadata
 			{
+				FNodeDisplayStyle DisplayStyle;
+				DisplayStyle.bShowName = false;
+				DisplayStyle.ImageName = TEXT("MetasoundEditor.Graph.Node.Conversion");
+				DisplayStyle.bShowInputNames = false;
+				DisplayStyle.bShowOutputNames = false;
+
 				const FText& FromTypeText = GetMetasoundDataTypeDisplayText<FromType>();
 				const FText& ToTypeText = GetMetasoundDataTypeDisplayText<ToType>();
 
@@ -179,15 +185,15 @@ namespace Metasound
 
 				const FName ClassName = *FString::Format(TEXT("Conversion{0}To{1}"), { FromTypeString, ToTypeString });
 				const FText NodeDisplayName = FText::Format(LOCTEXT("ConverterNodeDisplayName", "{0} To {1}"), FromTypeText, ToTypeText);
-				const FText NodeDescription = FText::Format(LOCTEXT("ConverterNodeDesc", "Converts from {0} to {1}."), FromTypeText, ToTypeText);
 
 				FNodeClassMetadata Info;
 				Info.ClassName = { StandardNodes::Namespace, ClassName, "" };
 				Info.MajorVersion = 1;
 				Info.MinorVersion = 0;
 				Info.DisplayName = NodeDisplayName;
-				Info.Description = NodeDescription;
+				Info.Description = GetNodeDescription();
 				Info.Author = PluginAuthor;
+				Info.DisplayStyle = DisplayStyle;
 				Info.PromptIfMissing = PluginNodeMissingPrompt;
 				Info.DefaultInterface = GetDefaultInterface();
 				Info.CategoryHierarchy.Emplace(NodeCategories::Conversions);
@@ -215,6 +221,21 @@ namespace Metasound
 
 		virtual ~TConversionOperator() = default;
 
+		static FText GetNodeDescription() 
+		{
+			if constexpr (std::is_same<FromType, float>::value && std::is_same<ToType, FAudioBuffer>::value)
+			{
+				return LOCTEXT("FloatToAudioConverterDescription", "Converts from float to audio buffer with each sample set to the given float value.");
+			}
+			else if constexpr (std::is_same<FromType, FAudioBuffer>::value && std::is_same<ToType, float>::value)
+			{
+				return LOCTEXT("AudioToFloatConverterDescription", "Converts from audio buffer to float by averaging sample values over the buffer. ");
+			}
+			else
+			{
+				return FText::Format(LOCTEXT("ConverterNodeDesc", "Converts from {0} to {1}."), GetMetasoundDataTypeDisplayText<FromType>(), GetMetasoundDataTypeDisplayText<ToType>());
+			}
+		}
 
 		virtual FDataReferenceCollection GetInputs() const override
 		{
