Updates to Metasound Frontend Query system.
- Simplified and removed a bunch of code.
- Reworked queries to include the concept of a "Partition" mapped to a Key
- Added more unit tests to ensure it works correctly.
- Fixed some bugs around incremental updates.

diff --git a/Engine/Plugins/Runtime/Metasound/Source/MetasoundFrontend/Private/MetasoundFrontendQuery.cpp b/Engine/Plugins/Runtime/Metasound/Source/MetasoundFrontend/Private/MetasoundFrontendQuery.cpp
--- a/Engine/Plugins/Runtime/Metasound/Source/MetasoundFrontend/Private/MetasoundFrontendQuery.cpp
+++ b/Engine/Plugins/Runtime/Metasound/Source/MetasoundFrontend/Private/MetasoundFrontendQuery.cpp
@@ -2,11 +2,10 @@
 
 #include "MetasoundFrontendQuery.h"
 
-#include "Algo/IsSorted.h"
-#include "Algo/Sort.h"
-#include "Algo/Transform.h"
-#include "CoreMinimal.h"
+#include "Algo/AnyOf.h"
 #include "MetasoundFrontendDocument.h"
+#include "MetasoundTrace.h"
+#include "Misc/Guid.h"
 #include "Templates/TypeHash.h"
 
 namespace Metasound
@@ -23,9 +22,9 @@ namespace Metasound
 			{
 			}
 
-			void Stream(TArray<FFrontendQueryEntry>& OutEntries) override
+			void Stream(TArray<FFrontendQueryValue>& OutValues) override
 			{
-				Func(OutEntries);
+				Func(OutValues);
 			}
 
 			private:
@@ -61,7 +60,7 @@ namespace Metasound
 			{
 			}
 
-			FFrontendQueryEntry::FKey Map(const FFrontendQueryEntry& InEntry) const override
+			FFrontendQueryKey Map(const FFrontendQueryEntry& InEntry) const override
 			{
 				return Func(InEntry);
 			}
@@ -78,10 +77,9 @@ namespace Metasound
 			:	Func(InFunc)
 			{
 			}
-
-			void Reduce(FFrontendQueryEntry::FKey InKey, TArrayView<FFrontendQueryEntry * const>& InEntries, FReduceOutputView& OutResult) const override
+			void Reduce(const FFrontendQueryKey& InKey, FFrontendQueryPartition& InOutEntries) const override
 			{
-				return Func(InKey, InEntries, OutResult);
+				return Func(InKey, InOutEntries);
 			}
 
 			FReduceFunction Func;
@@ -159,553 +157,436 @@ namespace Metasound
 			FLimitFunction Func;
 		};
 
-		// Base implementation of a query step execution.
-		struct FStepExecuterBase : public FFrontendQueryStep::IStepExecuter
+		// Base step executer for all step types.
+		template<typename StepType>
+		struct TStepExecuterBase : public FFrontendQueryStep::IStepExecuter
 		{
+			using FIncremental = FFrontendQueryStep::FIncremental;
 
-			// Appends incremental results to the output results.
-			EResultModificationState Append(const FFrontendQuerySelection& InIncremental, FFrontendQuerySelection& InOutResult) const
-			{
-				EResultModificationState ResultState = EResultModificationState::Unmodified;
-
-				if (InIncremental.GetSelection().Num() > 0)
-				{
-					InOutResult.AppendToStorageAndSelection(InIncremental);
-					ResultState = EResultModificationState::Modified;
-				}
-
-				return ResultState;
-			}
-		};
-
-		// Implements execution of a source
-		struct FSourceStepExecuter : public FStepExecuterBase
-		{
-			FSourceStepExecuter(TUniquePtr<IFrontendQuerySource>&& InSource)
-			: Source(MoveTemp(InSource))
+			TStepExecuterBase(TUniquePtr<StepType>&& InStep)
+			:	Step(MoveTemp(InStep))
 			{
 			}
 
-			// Get the latest entries from the source
-			virtual EResultModificationState Increment(FFrontendQuerySelection& InOutResult) override
+			virtual void Merge(FIncremental& InOutIncremental, FFrontendQuerySelection& InOutSelection) const override
 			{
-				if (Source.IsValid())
-				{
-					TArray<FFrontendQueryEntry> NewEntries;
-					Source->Stream(NewEntries);
-
-					if (NewEntries.Num() > 0)
-					{
-						InOutResult.AppendToStorageAndSelection(NewEntries);
-						return EResultModificationState::Modified;
-					}
-				}
-
-				return EResultModificationState::Unmodified;
+				METASOUND_TRACE_CPUPROFILER_EVENT_SCOPE(MetaSound::BaseQueryStep::Merge);
+				MergePartitionCompositionIndependent(InOutIncremental, InOutSelection);
 			}
 
-			virtual EResultModificationState Merge(const FFrontendQuerySelection& InIncremental, FFrontendQuerySelection& InOutResult) override
-			{
-				return Append(InIncremental, InOutResult);
-			}
+			virtual bool IsDependentOnPartitionComposition() const override { return false; }
+			virtual bool CanProcessRemovals() const override { return true; }
+			virtual bool CanProduceEntries() const override { return false; }
 
-			virtual bool IsMergeRequiredForIncremental() const override
-			{
-				return false;
-			}
+		protected:
 
-			virtual EResultModificationState Execute(FFrontendQuerySelection& InOutResult) override
+			void MergePartitionCompositionIndependent(FIncremental& InOutIncremental, FFrontendQuerySelection& InOutSelection) const
 			{
-				Reset(); // Reset source to original state.
-				return Increment(InOutResult);
+				checkf(!IsDependentOnPartitionComposition(), TEXT("Incorrect merge function called on step which is dependent upon merge composition"));
+				Remove(InOutIncremental.ActiveRemovalKeys, InOutIncremental.ActiveRemovalSelection, InOutSelection);
+				Append(InOutIncremental.ActiveKeys, InOutIncremental.ActiveSelection, InOutSelection);
 			}
 
-			virtual void Reset() override
+			// Appends incremental results to the output results.
+			void Append(const TSet<FFrontendQueryKey>& InActiveKeys, const FFrontendQuerySelection& InIncrementalSelection, FFrontendQuerySelection& InOutSelection) const
 			{
-				if (Source.IsValid())
+				for (const FFrontendQueryKey& Key : InActiveKeys)
 				{
-					Source->Reset();
+					if (const FFrontendQueryPartition* Partition = InIncrementalSelection.Find(Key))
+					{
+						if (Partition->Num() > 0)
+						{
+							InOutSelection.FindOrAdd(Key).Append(*Partition);
+						}
+					}
 				}
 			}
 
-		private:
-
-			TUniquePtr<IFrontendQuerySource> Source;
-		};
-
-		template<typename StepType>
-		struct TStepExecuter : public FStepExecuterBase
-		{
-			TStepExecuter(TUniquePtr<StepType>&& InStep)
-			:	Step(MoveTemp(InStep))
+			void Remove(const TSet<FFrontendQueryKey>& InRemovalActiveKeys, const FFrontendQuerySelection& InRemovalSelection, FFrontendQuerySelection& InOutSelection) const
 			{
-			}
-
-			virtual void Reset() override { }
-
-			virtual bool IsMergeRequiredForIncremental() const override { return false; }
-
-		protected:
-
-			TUniquePtr<StepType> Step;
-		};
-
-		struct FStreamStepExecuter : TStepExecuter<IFrontendQueryStreamStep>
-		{
-			using TStepExecuter<IFrontendQueryStreamStep>::TStepExecuter;
-
-			virtual EResultModificationState Increment(FFrontendQuerySelection& InOutResult) override
-			{
-				EResultModificationState ResultState = EResultModificationState::Unmodified;
-
-				if (Step.IsValid())
+				for (const FFrontendQueryKey& Key : InRemovalActiveKeys)
 				{
-					TArray<FFrontendQueryEntry> NewEntries;
-					Step->Stream(NewEntries);
-
-					if (NewEntries.Num() > 0)
+					if (FFrontendQueryPartition* Entries = InOutSelection.Find(Key))
 					{
-						ResultState = EResultModificationState::Modified;
-						InOutResult.AppendToStorageAndSelection(NewEntries);
+						if (const FFrontendQueryPartition* EntriesToRemove = InRemovalSelection.Find(Key))
+						{
+							for (const FFrontendQueryEntry& EntryToRemove : *EntriesToRemove)
+							{
+								Entries->Remove(EntryToRemove);
+							}
+						}
+
+						if (Entries->Num() == 0)
+						{
+							InOutSelection.Remove(Key);
+						}
 					}
 				}
-
-				return ResultState;
 			}
 
-			virtual EResultModificationState Merge(const FFrontendQuerySelection& InIncremental, FFrontendQuerySelection& InOutResult) override
-			{
-				return Append(InIncremental, InOutResult);
-			}
-
-			virtual EResultModificationState Execute(FFrontendQuerySelection& InOutResult) override
-			{
-				return Increment(InOutResult);
-			}
+			TUniquePtr<StepType> Step;
 		};
 
-		struct FTransformStepExecuter : TStepExecuter<IFrontendQueryTransformStep>
+		struct FStreamStepExecuter : TStepExecuterBase<IFrontendQueryStreamStep>
 		{
-			using TStepExecuter<IFrontendQueryTransformStep>::TStepExecuter;
+			using TStepExecuterBase<IFrontendQueryStreamStep>::TStepExecuterBase;
 
-			virtual EResultModificationState Increment(FFrontendQuerySelection& InOutResult) override
+			virtual void Execute(TSet<FFrontendQueryKey>& InOutUpdatedKeys, FFrontendQuerySelection& InOutResult) const override
 			{
-				EResultModificationState ResultState = EResultModificationState::Unmodified;
+				METASOUND_TRACE_CPUPROFILER_EVENT_SCOPE(MetaSound::StreamQueryStep::Execute);
+				InOutUpdatedKeys.Reset();
 
 				if (Step.IsValid())
 				{
-					if (InOutResult.GetSelection().Num() > 0)
-					{
-						ResultState = EResultModificationState::Modified;
-					}
+					// Retrieve new values from step
+					TArray<FFrontendQueryValue> NewValues;
+					Step->Stream(NewValues);
 
-					for (FFrontendQueryEntry* Entry : InOutResult.GetSelection())
+					const int32 Num = NewValues.Num();
+					if (NewValues.Num() > 0)
 					{
-						if (nullptr != Entry)
+						FFrontendQueryKey NullKey;
+						InOutUpdatedKeys.Add(NullKey);
+						FFrontendQueryPartition& NullPartition = InOutResult.FindOrAdd(NullKey);
+
+						for (FFrontendQueryValue& Value : NewValues)
 						{
-							Step->Transform(Entry->Value);
+							FFrontendQueryEntry Entry;
+							Entry.ID = FGuid::NewGuid();
+							Entry.Value = MoveTemp(Value);
+							NullPartition.Add(MoveTemp(Entry));
 						}
 					}
 				}
-
-				return ResultState;
 			}
 
-			virtual EResultModificationState Merge(const FFrontendQuerySelection& InIncremental, FFrontendQuerySelection& InOutResult) override
-			{
-				return Append(InIncremental, InOutResult);
+			virtual bool CanProcessRemovals() const override 
+			{ 
+				// Cannot process removals since streams only add information.
+				return false; 
 			}
 
-			virtual EResultModificationState Execute(FFrontendQuerySelection& InOutResult) override
-			{
-				return Increment(InOutResult);
+			virtual bool CanProduceEntries() const override 
+			{ 
+				return true; 
 			}
 		};
 
-		struct FMapStepExecuter : TStepExecuter<IFrontendQueryMapStep>
+		struct FTransformStepExecuter : TStepExecuterBase<IFrontendQueryTransformStep>
 		{
-			using TStepExecuter<IFrontendQueryMapStep>::TStepExecuter;
+			using TStepExecuterBase<IFrontendQueryTransformStep>::TStepExecuterBase;
 
-			virtual EResultModificationState Increment(FFrontendQuerySelection& InOutResult) override
+			virtual void Execute(TSet<FFrontendQueryKey>& InOutUpdatedKeys, FFrontendQuerySelection& InOutResult) const override
 			{
-				EResultModificationState ResultState = EResultModificationState::Unmodified;
-
+				METASOUND_TRACE_CPUPROFILER_EVENT_SCOPE(MetaSound::TransformQueryStep::Execute);
 				if (Step.IsValid())
 				{
-					if (InOutResult.GetSelection().Num() > 0)
+					for (const FFrontendQueryKey& Key : InOutUpdatedKeys)
 					{
-						ResultState = EResultModificationState::Modified;
-					}
-
-					for (FFrontendQueryEntry* Entry : InOutResult.GetSelection())
-					{
-						if (nullptr != Entry)
+						if (FFrontendQueryPartition* Partition = InOutResult.Find(Key))
 						{
-							Entry->Key = Step->Map(*Entry);
+							if (Partition->Num() > 0)
+							{
+								for (FFrontendQueryEntry& Entry : *Partition)
+								{
+									Step->Transform(Entry.Value);
+								}
+							}
 						}
 					}
 				}
-
-				return ResultState;
-			}
-
-			virtual EResultModificationState Merge(const FFrontendQuerySelection& InIncremental, FFrontendQuerySelection& InOutResult) override
-			{
-				return Append(InIncremental, InOutResult);
-			}
-
-			virtual EResultModificationState Execute(FFrontendQuerySelection& InOutResult) override
-			{
-				return Increment(InOutResult);
 			}
 		};
 
-		struct FReduceStepExecuter : TStepExecuter<IFrontendQueryReduceStep>
+		struct FMapStepExecuter : TStepExecuterBase<IFrontendQueryMapStep>
 		{
-			using FKey = FFrontendQueryEntry::FKey;
-
-		private:
-			static const FKey InvalidKey;
-
-			static const FFrontendQueryEntry::FKey& GetKey(const FFrontendQueryEntry* InEntry)
-			{
-				if (nullptr != InEntry)
-				{
-					return InEntry->Key;
-				}
-				else
-				{
-					return InvalidKey;
-				}
-			}
+			using TStepExecuterBase<IFrontendQueryMapStep>::TStepExecuterBase;
 
-			// Reduce entries.
-			//
-			// @parma InKey - The key associated with InEntriesToReduce
-			// @parma InEntriesToReduce - Entries to either append or reduce.
-			// @param OutResult - The output query selection to append to.
-			void Reduce(const FFrontendQueryEntry::FKey& InKey, TArrayView<FFrontendQueryEntry* const> InEntriesToReduce, FFrontendQuerySelection& OutResult)
+			virtual void Execute(TSet<FFrontendQueryKey>& InOutUpdatedKeys, FFrontendQuerySelection& InOutResult) const override
 			{
-				check(Step.IsValid());
-
-				// Note: InEntriesToReduce is has valid storage in OutResult.
-				FFrontendQuerySelection::FReduceOutputView ReduceResult(InEntriesToReduce);
-				Step->Reduce(InKey, InEntriesToReduce, ReduceResult);
-				OutResult.AppendToStorageAndSelection(MoveTemp(ReduceResult));
-			}
-
-			// Iterate over sorted entries.
-			template<typename FuncType>
-			void IterSortedEntries(TArrayView<FFrontendQueryEntry* const> SortedSelection, FuncType InFunc)
-			{
-				// input selection must be sorted by key.
-				check(Algo::IsSortedBy(SortedSelection, FReduceStepExecuter::GetKey));
-
-				if (SortedSelection.Num() > 0)
+				METASOUND_TRACE_CPUPROFILER_EVENT_SCOPE(MetaSound::MapQueryStep::Execute);
+				if (Step.IsValid())
 				{
-					int32 StartIndex = 0;
-					
-					FFrontendQueryEntry::FKey CurrentKey = GetKey(SortedSelection[0]);
-
-					const int32 Num = SortedSelection.Num();
-					for (int32 EndIndex = 1; EndIndex < Num; EndIndex++)
+					if (InOutUpdatedKeys.Num() > 0)
 					{
-						FFrontendQueryEntry::FKey ThisKey = GetKey(SortedSelection[EndIndex]);
-						if (CurrentKey != ThisKey)
-						{
-							InFunc(CurrentKey, SortedSelection.Slice(StartIndex, EndIndex - StartIndex));
+						FFrontendQuerySelection Result;
 
-							CurrentKey = ThisKey;
-							StartIndex = EndIndex;
+						for (const FFrontendQueryKey& Key : InOutUpdatedKeys)
+						{
+							if (FFrontendQueryPartition* Partition = InOutResult.Find(Key))
+							{
+								// Map all entries associated with the key to a new key.
+								if (Partition->Num() > 0)
+								{
+									for (FFrontendQueryEntry& Entry : *Partition)
+									{
+										FFrontendQueryKey NewKey = Step->Map(Entry);
+										Result.FindOrAdd(NewKey).Add(Entry);
+									}
+								}
+
+								// Remove entries since they have been mapped to a new key
+								Partition->Reset();
+							}
 						}
-					}
 
-					InFunc(CurrentKey, SortedSelection.Slice(StartIndex, SortedSelection.Num() - StartIndex));
+						// Get the updated set of keys in the output.
+						InOutUpdatedKeys.Reset();
+						Result.GetKeys(InOutUpdatedKeys);
+
+						// Append the new values to the final result.
+						Append(InOutUpdatedKeys, Result, InOutResult);
+					}
 				}
 			}
+		};
 
-			TArray<FFrontendQueryEntry*> GetSortedSelection(TArrayView<FFrontendQueryEntry* const> InSelection)
-			{
-				TArray<FFrontendQueryEntry*> SortedSelection(InSelection.GetData(), InSelection.Num()); 
-				Algo::SortBy(SortedSelection, GetKey);
-				return SortedSelection;
-			}
+		struct FReduceStepExecuter : TStepExecuterBase<IFrontendQueryReduceStep>
+		{
+			using FKey = FFrontendQueryKey;
 
 		public:
 
-			using TStepExecuter<IFrontendQueryReduceStep>::TStepExecuter;
-
-			virtual EResultModificationState Increment(FFrontendQuerySelection& InOutResult) override
-			{
-				check(IsMergeRequiredForIncremental());
-				return Execute(InOutResult);
-			}
+			using TStepExecuterBase<IFrontendQueryReduceStep>::TStepExecuterBase;
 
-			virtual EResultModificationState Merge(const FFrontendQuerySelection& InIncremental, FFrontendQuerySelection& InOutResult) override
+			virtual void Merge(FIncremental& InOutIncremental, FFrontendQuerySelection& InOutSelection) const override
 			{
-				EResultModificationState ResultState = EResultModificationState::Unmodified;
-
+				METASOUND_TRACE_CPUPROFILER_EVENT_SCOPE(MetaSound::ReduceQueryStep::Merge);
+				Remove(InOutIncremental.ActiveRemovalKeys, InOutIncremental.ActiveRemovalSelection, InOutSelection);
+				
 				if (Step.IsValid())
 				{
-					if (InIncremental.GetSelection().Num() > 0)
+					for (const FFrontendQueryKey& Key : InOutIncremental.ActiveKeys)
 					{
-						TSet<FKey> NewKeys;
-
-						Algo::Transform(InIncremental.GetSelection(), NewKeys, GetKey);
-
-						InOutResult.AppendToStorageAndSelection(InIncremental);
-						TArray<FFrontendQueryEntry*> SortedSelection = GetSortedSelection(InOutResult.GetSelection());
-						InOutResult.ResetSelection();
-
-						IterSortedEntries(SortedSelection, [&](const FKey& InKey, TArrayView<FFrontendQueryEntry* const> InEntriesToReduce)
+						if (const FFrontendQueryPartition* NewPartition = InOutIncremental.ActiveSelection.Find(Key))
 						{
-							if (NewKeys.Contains(InKey))
-							{
-								// Some of the entries re new. Need to re-reduce.
-								this->Reduce(InKey, InEntriesToReduce, InOutResult);
-							}
-							else
+							if (NewPartition->Num() > 0)
 							{
-								// All these entries are owned by the result already.
-								InOutResult.AppendToSelection(InEntriesToReduce);
+								FFrontendQueryPartition& Partition = InOutSelection.FindOrAdd(Key);
+								FFrontendQueryPartition OriginalPartition = Partition;
+								Partition.Append(*NewPartition);
+
+								Step->Reduce(Key, Partition);
+
+								FFrontendQueryPartition Removed = OriginalPartition.Difference(Partition);
+								if (Removed.Num() > 0)
+								{
+									InOutIncremental.ActiveRemovalSelection.FindOrAdd(Key).Append(Removed);
+									InOutIncremental.ActiveRemovalKeys.Add(Key);
+								}
 							}
-						});
-
-						ResultState = EResultModificationState::Modified;
+						}
 					}
 				}
-
-				return ResultState;
 			}
 
-			virtual EResultModificationState Execute(FFrontendQuerySelection& InOutResult) override
+			virtual void Execute(TSet<FFrontendQueryKey>& InOutUpdatedKeys, FFrontendQuerySelection& InOutResult) const override
 			{
-				EResultModificationState ResultState = EResultModificationState::Unmodified;
-
+				METASOUND_TRACE_CPUPROFILER_EVENT_SCOPE(MetaSound::ReduceQueryStep::Execute);
 				if (Step.IsValid())
 				{
-					if (InOutResult.GetSelection().Num() > 0)
+					for (const FFrontendQueryKey& Key : InOutUpdatedKeys)
 					{
-						TArray<FFrontendQueryEntry*> SortedSelection = GetSortedSelection(InOutResult.GetSelection());
-						InOutResult.ResetSelection();
-
-						IterSortedEntries(SortedSelection, [&](const FKey& InKey, TArrayView<FFrontendQueryEntry* const> InEntriesToReduce)
+						if (FFrontendQueryPartition* Partition = InOutResult.Find(Key))
 						{
-							this->Reduce(InKey, InEntriesToReduce, InOutResult);
-						});
-
-						ResultState = EResultModificationState::Modified;
+							if (Partition->Num() > 0)
+							{
+								Step->Reduce(Key, *Partition);
+							}
+						}
 					}
 				}
-
-				return ResultState;
 			}
 
 
-			virtual bool IsMergeRequiredForIncremental() const override
+			virtual bool IsDependentOnPartitionComposition() const override
 			{
-				// Merge is required during an incremental update because
-				// existing results and new results may share the same key and
-				// need re-reducing.
+				// Results may change depending on what is in the partition.
 				return true;
 			}
 		};
 
-		const FReduceStepExecuter::FKey FReduceStepExecuter::InvalidKey;
-
-		struct FFilterStepExecuter : TStepExecuter<IFrontendQueryFilterStep>
+		struct FFilterStepExecuter : TStepExecuterBase<IFrontendQueryFilterStep>
 		{
-			using TStepExecuter<IFrontendQueryFilterStep>::TStepExecuter;
+			using TStepExecuterBase<IFrontendQueryFilterStep>::TStepExecuterBase;
 
-			virtual EResultModificationState Increment(FFrontendQuerySelection& InOutResult) override
+			virtual void Execute(TSet<FFrontendQueryKey>& InOutUpdatedKeys, FFrontendQuerySelection& InOutResult) const override
 			{
-				EResultModificationState ResultState = EResultModificationState::Unmodified;
-
+				METASOUND_TRACE_CPUPROFILER_EVENT_SCOPE(MetaSound::FilterQueryStep::Execute);
 				if (Step.IsValid())
 				{
-					if (InOutResult.GetSelection().Num() > 0)
+					for (const FFrontendQueryKey& Key : InOutUpdatedKeys)
 					{
-						ResultState = EResultModificationState::Modified;
-						auto FilterFunc = [&](const FFrontendQueryEntry* Entry)
+						if (FFrontendQueryPartition* Partition = InOutResult.Find(Key))
 						{
-							return (nullptr != Entry) && Step->Filter(*Entry);
-						};
-
-						InOutResult.FilterSelection(FilterFunc);
+							for (FFrontendQueryPartition::TIterator Iter = Partition->CreateIterator(); Iter; ++Iter)
+							{
+								const bool bKeepEntry = Step->Filter(*Iter);
+								if (!bKeepEntry)
+								{
+									Iter.RemoveCurrent();
+								}
+							}
+						}
 					}
 				}
-
-				return ResultState;
-			}
-
-			virtual EResultModificationState Merge(const FFrontendQuerySelection& InIncremental, FFrontendQuerySelection& InOutResult) override
-			{
-				return Append(InIncremental, InOutResult);
-			}
-
-			virtual EResultModificationState Execute(FFrontendQuerySelection& InOutResult) override
-			{
-				return Increment(InOutResult);
 			}
 		};
 
-		struct FScoreStepExecuter : TStepExecuter<IFrontendQueryScoreStep>
+		struct FScoreStepExecuter : TStepExecuterBase<IFrontendQueryScoreStep>
 		{
-			using TStepExecuter<IFrontendQueryScoreStep>::TStepExecuter;
+			using TStepExecuterBase<IFrontendQueryScoreStep>::TStepExecuterBase;
 
-			virtual EResultModificationState Increment(FFrontendQuerySelection& InOutResult) override
+			virtual void Execute(TSet<FFrontendQueryKey>& InOutUpdatedKeys, FFrontendQuerySelection& InOutResult) const override
 			{
-				EResultModificationState ResultState = EResultModificationState::Unmodified;
-
+				METASOUND_TRACE_CPUPROFILER_EVENT_SCOPE(MetaSound::ScoreQueryStep::Execute);
 				if (Step.IsValid())
 				{
-					if (InOutResult.GetSelection().Num() > 0)
-					{
-						ResultState = EResultModificationState::Modified;
-					}
-
-					for (FFrontendQueryEntry* Entry : InOutResult.GetSelection())
+					for (const FFrontendQueryKey& Key : InOutUpdatedKeys)
 					{
-						if (nullptr != Entry)
+						if (FFrontendQueryPartition* Partition = InOutResult.Find(Key))
 						{
-							Entry->Score = Step->Score(*Entry);
+							for (FFrontendQueryEntry& Entry : *Partition)
+							{
+								Entry.Score = Step->Score(Entry);
+							}
 						}
 					}
 				}
-
-				return ResultState;
-			}
-
-			virtual EResultModificationState Merge(const FFrontendQuerySelection& InIncremental, FFrontendQuerySelection& InOutResult) override
-			{
-				return Append(InIncremental, InOutResult);
-			}
-
-			virtual EResultModificationState Execute(FFrontendQuerySelection& InOutResult) override
-			{
-				return Increment(InOutResult);
 			}
 		};
 		
-		struct FSortStepExecuter : TStepExecuter<IFrontendQuerySortStep>
+		struct FSortStepExecuter : TStepExecuterBase<IFrontendQuerySortStep>
 		{
-			using TStepExecuter<IFrontendQuerySortStep>::TStepExecuter;
+			using TStepExecuterBase<IFrontendQuerySortStep>::TStepExecuterBase;
 
-			EResultModificationState Increment(FFrontendQuerySelection& InOutResult) override
+			virtual void Merge(FIncremental& InOutIncremental, FFrontendQuerySelection& InOutSelection) const override
 			{
-				check(IsMergeRequiredForIncremental());
-				return Execute(InOutResult);
-			}
-
-			virtual EResultModificationState Merge(const FFrontendQuerySelection& InIncremental, FFrontendQuerySelection& InOutResult) override
-			{
-				EResultModificationState ResultState = EResultModificationState::Unmodified;
+				METASOUND_TRACE_CPUPROFILER_EVENT_SCOPE(MetaSound::SortQueryStep::Merge);
+				Remove(InOutIncremental.ActiveRemovalKeys, InOutIncremental.ActiveRemovalSelection, InOutSelection);
 
 				if (Step.IsValid())
 				{
-					const bool bIncrementalHasNonEmptySelection = InIncremental.GetSelection().Num() > 0;
-
-					if (bIncrementalHasNonEmptySelection)
+					auto SortFunc =  [&](const FFrontendQueryEntry& InLHS, const FFrontendQueryEntry& InRHS)
 					{
-						// Append incremental data to this result
-						InOutResult.AppendToStorageAndSelection(InIncremental);
+						return Step->Sort(InLHS, InRHS);
+					};
 
-						InOutResult.GetSelection().Sort(
-							[&](const FFrontendQueryEntry& InLHS, const FFrontendQueryEntry& InRHS)
+					for (const FFrontendQueryKey& Key : InOutIncremental.ActiveKeys)
+					{
+						if (const FFrontendQueryPartition* NewPartition = InOutIncremental.ActiveSelection.Find(Key))
+						{
+							if (FFrontendQueryPartition* OrigPartition = InOutSelection.Find(Key))
 							{
-								return Step->Sort(InLHS, InRHS);
+								// Need to re-sort if merging two arrays
+								OrigPartition->Append(*NewPartition);
+								OrigPartition->Sort(SortFunc);
 							}
-						);
-
-						ResultState = EResultModificationState::Modified;
+							else
+							{
+								// New key entries are already sorted. 
+								InOutSelection.Add(Key, *NewPartition);
+							}
+						}
 					}
 				}
-
-				return ResultState;
 			}
 
-			virtual EResultModificationState Execute(FFrontendQuerySelection& InOutResult) override
+			virtual void Execute(TSet<FFrontendQueryKey>& InOutUpdatedKeys, FFrontendQuerySelection& InOutResult) const override
 			{
-				EResultModificationState ResultState = EResultModificationState::Unmodified;
-
+				METASOUND_TRACE_CPUPROFILER_EVENT_SCOPE(MetaSound::SortQueryStep::Execute);
 				if (Step.IsValid())
 				{
-					if (InOutResult.GetSelection().Num() > 0)
+					auto SortFunc = [&](const FFrontendQueryEntry& InLHS, const FFrontendQueryEntry& InRHS)
 					{
-						InOutResult.GetSelection().Sort(
-							[&](const FFrontendQueryEntry& InLHS, const FFrontendQueryEntry& InRHS)
-							{
-								return Step->Sort(InLHS, InRHS);
-							}
-						);
+						return Step->Sort(InLHS, InRHS);
+					};
 
-						ResultState = EResultModificationState::Modified;
+					for (const FFrontendQueryKey& Key : InOutUpdatedKeys)
+					{
+						if (FFrontendQueryPartition* OrigPartition = InOutResult.Find(Key))
+						{
+							OrigPartition->Sort(SortFunc);
+						}
 					}
 				}
-
-				return ResultState;
 			}
 
-			virtual bool IsMergeRequiredForIncremental() const 
+			virtual bool IsDependentOnPartitionComposition() const 
 			{
-				// Merging is required during incremental update because sort
-				// order may change after merge.
+				// Results may change depending on what is in the partition.
 				return true;
 			}
 		};
 
-		struct FLimitStepExecuter : TStepExecuter<IFrontendQueryLimitStep>
+		struct FLimitStepExecuter : TStepExecuterBase<IFrontendQueryLimitStep>
 		{
-			using TStepExecuter<IFrontendQueryLimitStep>::TStepExecuter;
-
-			EResultModificationState Increment(FFrontendQuerySelection& InOutResult) override
-			{
-				check(IsMergeRequiredForIncremental());
-				return Execute(InOutResult);
-			}
+			using TStepExecuterBase<IFrontendQueryLimitStep>::TStepExecuterBase;
 
-			virtual EResultModificationState Merge(const FFrontendQuerySelection& InIncremental, FFrontendQuerySelection& InOutResult) override
+			virtual void Merge(FIncremental& InOutIncremental, FFrontendQuerySelection& InOutSelection) const override
 			{
-				EResultModificationState ResultState = EResultModificationState::Unmodified;
+				METASOUND_TRACE_CPUPROFILER_EVENT_SCOPE(MetaSound::LimitQueryStep::Merge);
+				Remove(InOutIncremental.ActiveRemovalKeys, InOutIncremental.ActiveRemovalSelection, InOutSelection);
 
 				if (Step.IsValid())
 				{
-					int32 Limit = Step->Limit();
+					TSet<FFrontendQueryKey> UpdatedKeys;
 
-					if (InIncremental.GetSelection().Num() > 0)
-					{
-						ResultState = EResultModificationState::Modified;
-						InOutResult.AppendToStorageAndSelection(InIncremental);
+					const int32 Limit = Step->Limit();
 
-						if (InOutResult.GetSelection().Num() > Limit)
+					for (const FFrontendQueryKey& Key : InOutIncremental.ActiveKeys)
+					{
+						if (const FFrontendQueryPartition* NewPartition = InOutIncremental.ActiveSelection.Find(Key))
 						{
-							InOutResult.SetSelection(InOutResult.GetSelection().Slice(0, Limit));
+							FFrontendQueryPartition& ExistingPartition = InOutSelection.FindOrAdd(Key);
+
+							const int32 NumToAdd = FMath::Min(Limit - ExistingPartition.Num(), NewPartition->Num());
+							if (NumToAdd > 0)
+							{
+								UpdatedKeys.Add(Key);
+								for (FFrontendQueryPartition::TConstIterator Iter = NewPartition->CreateConstIterator(); Iter; ++Iter)
+								{
+									ExistingPartition.Add(*Iter);
+								}
+							}
 						}
 					}
 				}
-
-				return ResultState;
 			}
 
-			virtual EResultModificationState Execute(FFrontendQuerySelection& InOutResult) override
+			virtual void Execute(TSet<FFrontendQueryKey>& InOutUpdatedKeys, FFrontendQuerySelection& InOutResult) const override
 			{
-				EResultModificationState ResultState = EResultModificationState::Unmodified;
-
+				METASOUND_TRACE_CPUPROFILER_EVENT_SCOPE(MetaSound::LimitQueryStep::Execute);
 				if (Step.IsValid())
 				{
-					int32 Limit = Step->Limit();
+					const int32 Limit = Step->Limit();
 
-					if (InOutResult.GetSelection().Num() > Limit)
+					for (const FFrontendQueryKey& Key : InOutUpdatedKeys)
 					{
-						ResultState = EResultModificationState::Modified;
-						InOutResult.SetSelection(InOutResult.GetSelection().Slice(0, Limit));
+						if (FFrontendQueryPartition* Partition = InOutResult.Find(Key))
+						{
+							const int32 Num = Partition->Num();
+							if (Num > Limit)
+							{
+								// Create a new partition and fill with desired number
+								// of objects. 
+								FFrontendQueryPartition NewPartition;
+								FFrontendQueryPartition::TConstIterator Iter = Partition->CreateConstIterator();
+								for (int32 i = 0; (i < Limit) && Iter; i++)
+								{
+									NewPartition.Add(*Iter);
+									++Iter;
+								}
+
+								// Replace the old partition with a new one.
+								InOutResult[Key] = NewPartition;
+							}
+						}
 					}
 				}
-
-				return ResultState;
 			}
 
-			virtual bool IsMergeRequiredForIncremental() const 
+			virtual bool IsDependentOnPartitionComposition() const 
 			{
 				// Merge is required because values under limit may be different
 				// after merge.
@@ -715,7 +596,7 @@ namespace Metasound
 	}
 
 	FFrontendQueryKey::FFrontendQueryKey()
-	: Key(TInPlaceType<FFrontendQueryKey::FInvalid>())
+	: Key(TInPlaceType<FFrontendQueryKey::FNull>())
 	, Hash(INDEX_NONE)
 	{}
 
@@ -737,34 +618,34 @@ namespace Metasound
 	{
 	}
 
-	bool FFrontendQueryKey::IsValid() const
+	bool FFrontendQueryKey::IsNull() const
 	{
-		return Key.GetIndex() != FKeyType::IndexOfType<FFrontendQueryKey::FInvalid>();
+		return Key.GetIndex() != FKeyType::IndexOfType<FFrontendQueryKey::FNull>();
 	}
 
 	bool operator==(const FFrontendQueryKey& InLHS, const FFrontendQueryKey& InRHS)
 	{
 		if (InLHS.Hash == InRHS.Hash)
 		{
-			if (InLHS.IsValid() && InRHS.IsValid())
+			if (InLHS.Key.GetIndex() == InRHS.Key.GetIndex())
 			{
-				if (InLHS.Key.GetIndex() == InRHS.Key.GetIndex())
+				switch(InLHS.Key.GetIndex())
 				{
-					switch(InLHS.Key.GetIndex())
-					{
-						case FFrontendQueryKey::FKeyType::IndexOfType<int32>():
-							return InLHS.Key.Get<int32>() == InRHS.Key.Get<int32>();
+					case FFrontendQueryKey::FKeyType::IndexOfType<FFrontendQueryKey::FNull>():
+						return true;
 
-						case FFrontendQueryKey::FKeyType::IndexOfType<FString>():
-							return InLHS.Key.Get<FString>() == InRHS.Key.Get<FString>();
+					case FFrontendQueryKey::FKeyType::IndexOfType<int32>():
+						return InLHS.Key.Get<int32>() == InRHS.Key.Get<int32>();
 
-						case FFrontendQueryKey::FKeyType::IndexOfType<FName>():
-							return InLHS.Key.Get<FName>() == InRHS.Key.Get<FName>();
+					case FFrontendQueryKey::FKeyType::IndexOfType<FString>():
+						return InLHS.Key.Get<FString>() == InRHS.Key.Get<FString>();
 
-						default:
-							// Unhandled case type.
-							checkNoEntry();
-					}
+					case FFrontendQueryKey::FKeyType::IndexOfType<FName>():
+						return InLHS.Key.Get<FName>() == InRHS.Key.Get<FName>();
+
+					default:
+						// Unhandled case type.
+						checkNoEntry();
 				}
 			}
 		}
@@ -791,6 +672,9 @@ namespace Metasound
 
 		switch(InLHS.Key.GetIndex())
 		{
+			case FFrontendQueryKey::FKeyType::IndexOfType<FFrontendQueryKey::FNull>():
+				return false;
+
 			case FFrontendQueryKey::FKeyType::IndexOfType<int32>():
 				return InLHS.Key.Get<int32>() < InRHS.Key.Get<int32>();
 
@@ -813,199 +697,14 @@ namespace Metasound
 		return InKey.Hash;
 	}
 
-	FFrontendQuerySelection::FFrontendQuerySelection(const FFrontendQuerySelection& InOther)
-	{
-		AppendToStorageAndSelection(InOther);
-	}
-
-	FFrontendQuerySelection& FFrontendQuerySelection::operator=(const FFrontendQuerySelection& InOther)
-	{
-		ResetStorageAndSelection();
-		AppendToStorageAndSelection(InOther);
-
-		return *this;
-	}
-
-	FFrontendQuerySelection::FReduceOutputView::FReduceOutputView(TArrayView<FFrontendQueryEntry* const> InEntries)
-	:	InitialEntries(InEntries)
-	{
-	}
-
-	void FFrontendQuerySelection::FReduceOutputView::Add(FFrontendQueryEntry& InResult)
-	{
-		if (!InitialEntries.Contains(&InResult))
-		{
-			SelectedNewEntries.Add(MakeUnique<FFrontendQueryEntry>(InResult));
-		}
-		else
-		{
-			SelectedExistingEntries.Add(&InResult);
-		}
-	}
-
-
-	TArrayView<FFrontendQueryEntry*> FFrontendQuerySelection::GetSelection()
-	{
-		return Selection;
-	}
-
-	TArrayView<const FFrontendQueryEntry * const> FFrontendQuerySelection::GetSelection() const
-	{
-		return MakeArrayView(Selection);
-	}
-
-	void FFrontendQuerySelection::AppendToStorageAndSelection(const FFrontendQuerySelection& InSelection)
-	{
-		TArrayView<const FFrontendQueryEntry* const> InEntries = InSelection.GetSelection();
-		AppendToStorageAndSelection(InSelection.GetSelection());
-	}
-
-	void FFrontendQuerySelection::AppendToStorageAndSelection(FFrontendQuerySelection&& InSelection)
-	{
-		InSelection.ShrinkStorageToSelection();
-		AppendToStorageAndSelection(MoveTemp(InSelection.Storage));
-	}
-
-	void FFrontendQuerySelection::AppendToStorageAndSelection(TArrayView<const FFrontendQueryEntry> InEntries)
-	{
-		const int32 Num = InEntries.Num();
-
-		if (Num >  0)
-		{
-			int32 StorageIndex = Storage.Num();
-			int32 SelectionIndex = Selection.Num();
-
-			for (const FFrontendQueryEntry& Entry : InEntries)
-			{
-				Storage.Add(MakeUnique<FFrontendQueryEntry>(Entry));
-			}
-			Selection.AddZeroed(Num);
-
-			for (int32 i = 0; i < Num; i++)
-			{
-				Selection[SelectionIndex] = Storage[StorageIndex].Get();
-				SelectionIndex++;
-				StorageIndex++;
-			}
-		}
-	}
-
-	void FFrontendQuerySelection::AppendToStorageAndSelection(TArrayView<const FFrontendQueryEntry* const> InEntries)
-	{
-		if (InEntries.Num() >  0)
-		{
-			int32 StorageIndex = Storage.Num();
-			int32 SelectionIndex = Selection.Num();
-
-			int32 NumAdded = 0;
-			for (const FFrontendQueryEntry* Entry : InEntries)
-			{
-				if (nullptr != Entry)
-				{
-					Storage.Add(MakeUnique<FFrontendQueryEntry>(*Entry));
-					NumAdded++;
-				}
-			}
-
-			if (NumAdded > 0)
-			{
-				Selection.AddZeroed(NumAdded);
-
-				for (int32 i = 0; i < NumAdded; i++)
-				{
-					Selection[SelectionIndex] = Storage[StorageIndex].Get();
-					SelectionIndex++;
-					StorageIndex++;
-				}
-			}
-		}
-	}
-
-	void FFrontendQuerySelection::AppendToStorageAndSelection(FReduceOutputView&& InReduceView)
-	{
-		AppendToStorageAndSelection(MoveTemp(InReduceView.SelectedNewEntries));
-		AppendToSelection(InReduceView.SelectedExistingEntries);
-	}
-
-	void FFrontendQuerySelection::AppendToStorageAndSelection(TArray<TUniquePtr<FFrontendQueryEntry>>&& InEntries)
-	{
-		if (InEntries.Num() >  0)
-		{
-			int32 StorageIndex = Storage.Num();
-			int32 SelectionIndex = Selection.Num();
-
-			int32 NumAdded = 0;
-			for (TUniquePtr<FFrontendQueryEntry>& Entry : InEntries)
-			{
-				if (Entry.IsValid())
-				{
-					Storage.Add(MoveTemp(Entry));
-					NumAdded++;
-				}
-			}
-
-			if (NumAdded > 0)
-			{
-				Selection.AddZeroed(NumAdded);
-
-				for (int32 i = 0; i < NumAdded; i++)
-				{
-					Selection[SelectionIndex] = Storage[StorageIndex].Get();
-					SelectionIndex++;
-					StorageIndex++;
-				}
-			}
-		}
-
-	}
-
-	void FFrontendQuerySelection::ResetStorageAndSelection()
-	{
-		Storage.Reset();
-		Selection.Reset();
-	}
-
-	void FFrontendQuerySelection::ResetSelection()
-	{
-		Selection.Reset();
-	}
-
-	void FFrontendQuerySelection::SetSelection(TArrayView<FFrontendQueryEntry*> InEntries)
-	{
-		Selection = InEntries;
-	}
-
-	void FFrontendQuerySelection::AddToSelection(FFrontendQueryEntry* InEntry)
-	{
-		if (nullptr != InEntry)
-		{
-			Selection.Add(InEntry);
-		}
-	}
-
-	void FFrontendQuerySelection::AppendToSelection(TArrayView<FFrontendQueryEntry * const> InEntries)
-	{
-		Selection.Append(InEntries.GetData(), InEntries.Num());
-	}
-
-	void FFrontendQuerySelection::ShrinkStorageToSelection()
-	{
-		TSet<FFrontendQueryEntry*> SelectionSet(Selection);
-		Storage.RemoveAll([&](const TUniquePtr<FFrontendQueryEntry>& InEntry)
-		{
-			return !SelectionSet.Contains(InEntry.Get());
-		});
-	}
-
-
-	FFrontendQuerySelectionView::FFrontendQuerySelectionView(TSharedRef<const FFrontendQuerySelection, ESPMode::ThreadSafe> InResult)
-	:	Result(InResult)
+	uint32 GetTypeHash(const FFrontendQueryEntry& InEntry)
 	{
+		return GetTypeHash(InEntry.ID);
 	}
 
-	TArrayView<const FFrontendQueryEntry* const> FFrontendQuerySelectionView::GetSelection() const
+	bool operator==(const FFrontendQueryEntry& InLHS, const FFrontendQueryEntry& InRHS)
 	{
-		return Result->GetSelection();
+		return InLHS.ID == InRHS.ID;
 	}
 
 	FFrontendQueryStep::FFrontendQueryStep(FStreamFunction&& InFunc)
@@ -1048,11 +747,6 @@ namespace Metasound
 	{
 	}
 
-	FFrontendQueryStep::FFrontendQueryStep(TUniquePtr<IFrontendQuerySource>&& InSource)
-	:	StepExecuter(MakeUnique<FrontendQueryPrivate::FSourceStepExecuter>(MoveTemp(InSource)))
-	{
-	}
-
 	FFrontendQueryStep::FFrontendQueryStep(TUniquePtr<IFrontendQueryStreamStep>&& InStep)
 	:	StepExecuter(MakeUnique<FrontendQueryPrivate::FStreamStepExecuter>(MoveTemp(InStep)))
 	{
@@ -1093,65 +787,57 @@ namespace Metasound
 	{
 	}
 
-	FFrontendQueryStep::EResultModificationState FFrontendQueryStep::Execute(FFrontendQuerySelection& InOutResult)
+	void FFrontendQueryStep::Execute(TSet<FFrontendQueryKey>& InOutUpdatedKeys, FFrontendQuerySelection& InOutResult) const
 	{
 		if (StepExecuter.IsValid())
 		{
-			return StepExecuter->Execute(InOutResult);
+			StepExecuter->Execute(InOutUpdatedKeys, InOutResult);
 		}
-
-		return EResultModificationState::Unmodified;
 	}
 
-	FFrontendQueryStep::EResultModificationState FFrontendQueryStep::Increment(FFrontendQuerySelection& InOutResult)
+	void FFrontendQueryStep::Merge(FIncremental& InOutIncremental, FFrontendQuerySelection& InOutSelection) const
 	{
 		if (StepExecuter.IsValid())
 		{
-			return StepExecuter->Increment(InOutResult);
+			StepExecuter->Merge(InOutIncremental, InOutSelection);
 		}
-
-		return EResultModificationState::Unmodified;
 	}
 
-	FFrontendQueryStep::EResultModificationState FFrontendQueryStep::Merge(const FFrontendQuerySelection& InIncremental, FFrontendQuerySelection& InOutResult)
+	bool FFrontendQueryStep::IsDependentOnPartitionComposition() const
 	{
 		if (StepExecuter.IsValid())
 		{
-			return StepExecuter->Merge(InIncremental, InOutResult);
+			return StepExecuter->IsDependentOnPartitionComposition();
 		}
 
-		return EResultModificationState::Unmodified;
+		return false;
 	}
 
-	bool FFrontendQueryStep::IsMergeRequiredForIncremental() const
+	bool FFrontendQueryStep::CanProcessRemovals() const
 	{
 		if (StepExecuter.IsValid())
 		{
-			return StepExecuter->IsMergeRequiredForIncremental();
+			return StepExecuter->CanProcessRemovals();
 		}
 
-		return false;
-	}
+		return true;
+	} 
 
-	void FFrontendQueryStep::Reset()
+	bool FFrontendQueryStep::CanProduceEntries() const
 	{
 		if (StepExecuter.IsValid())
 		{
-			StepExecuter->Reset();
+			return StepExecuter->CanProduceEntries();
 		}
-	}
+
+		return true;
+	} 
 
 	FFrontendQuery::FFrontendQuery()
 	: Result(MakeShared<FFrontendQuerySelection, ESPMode::ThreadSafe>())
 	{
 	}
 
-	const TArray<TUniquePtr<FFrontendQueryStep>>& FFrontendQuery::GetSteps() const
-	{
-		return Steps;
-	}
-
-
 	FFrontendQuery& FFrontendQuery::AddStreamLambdaStep(FStreamFunction&& InFunc)
 	{
 		return AddFunctionStep(MoveTemp(InFunc));
@@ -1196,120 +882,124 @@ namespace Metasound
 	{
 		if (ensure(InStep.IsValid()))
 		{
-			int32 StepIndex = Steps.Num();
-			if (InStep->IsMergeRequiredForIncremental())
+			FStepInfo StepInfo;
+			StepInfo.Step = MoveTemp(InStep);
+			StepInfo.bProcessRemovals = StepInfo.Step->CanProcessRemovals();
+
+			// Determine if this step is the last step which can produce entries. 
+			if (StepInfo.Step->CanProduceEntries())
+			{
+				FinalEntryProducingStepIndex = Steps.Num();
+			}
+			
+			// If any prior steps have the ability to remove a previously existing 
+			// entry, and this step requires an incremental merge, then we need to 
+			// cache the input to this step. If an entry is deleted from the input 
+			// to this step, all inputs associated with a given key will be reevaluated.
+			//
+			// To cache the input to a step, we can equivalently cache the output 
+			// of prior step.
+			const bool bCacheAncestorStepOutput = Algo::AnyOf(Steps, [](const FStepInfo& Info) { return Info.Step->IsDependentOnPartitionComposition(); });
+			if (bCacheAncestorStepOutput && StepInfo.Step->IsDependentOnPartitionComposition())
+			{
+				Steps.Last().bMergeAndCacheOutput = true;
+			}
+			else if (StepInfo.Step->IsDependentOnPartitionComposition())
 			{
-				StepResultCache.Add(StepIndex);
+				// If the step is dependent upon the partition composition, and
+				// input is not a merged partition, then the incremental must be merged
+				// with the prior result for the step so that downstream steps are
+				// provided the correct input data.
+				StepInfo.bMergeAndCacheOutput = true;
 			}
 
-			Steps.Add(MoveTemp(InStep));
+			Steps.Add(MoveTemp(StepInfo));
 		}
+
 		return *this;
 	}
 
-	FFrontendQuerySelectionView FFrontendQuery::Execute()
+	const FFrontendQuerySelection& FFrontendQuery::Update()
 	{
-		UpdateResult();
-
-		return FFrontendQuerySelectionView(Result);
+		TSet<FFrontendQueryKey> UpdatedKeys;
+		return Update(UpdatedKeys);
 	}
 
-	FFrontendQuerySelectionView FFrontendQuery::Reset()
+	const FFrontendQuerySelection& FFrontendQuery::Update(TSet<FFrontendQueryKey>& OutUpdatedKeys)
 	{
-		Result->ResetStorageAndSelection();
+		UpdateInternal(OutUpdatedKeys);
 
-		for (int32 StepIndex = 0; StepIndex < Steps.Num(); StepIndex++)
-		{
-			if (Steps[StepIndex].IsValid())
-			{
-				Steps[StepIndex]->Reset();
-			}
-		}
-		
-		return FFrontendQuerySelectionView(Result);
+		return *Result;
 	}
 
-	FFrontendQuerySelectionView FFrontendQuery::GetSelection() const
+	const FFrontendQuerySelection& FFrontendQuery::GetSelection() const
 	{
-		return FFrontendQuerySelectionView(Result);
+		return *Result;
 	}
 
-	void FFrontendQuery::UpdateResult()
+	void FFrontendQuery::MergeInternal(FFrontendQueryStep& Step, FFrontendQuery::FIncremental& InOutIncremental, FFrontendQuerySelection& InOutMergedSelection)
 	{
-		FFrontendQuerySelection IncrementalResult;
+		Step.Merge(InOutIncremental, InOutMergedSelection);
+	}
 
-		EResultModificationState IncrementalResultState = EResultModificationState::Unmodified;
+	void FFrontendQuery::UpdateInternal(TSet<FFrontendQueryKey>& OutUpdatedKeys)
+	{
+		FIncremental Incremental;
+
+		const int32 LastStepIndex = Steps.Num() - 1;
 
 		// Perform incremental update sequentially
 		for (int32 StepIndex = 0; StepIndex < Steps.Num(); StepIndex++)
 		{
-			if (ensure(Steps[StepIndex].IsValid()))
-			{
-				FFrontendQueryStep& Step = *Steps[StepIndex];
+			FStepInfo& StepInfo = Steps[StepIndex];
 
-				const bool bMergeIncrementalResults = Step.IsMergeRequiredForIncremental();
-				if (bMergeIncrementalResults)
-				{
-					const bool bIsResultCacheEmpty = StepResultCache[StepIndex].GetSelection().Num() == 0;
-					if (bIsResultCacheEmpty)
-					{
-						// If the prior result for this step is empty, continue using the incremental path.
-						IncrementalResultState = Step.Increment(IncrementalResult);
-						if (EResultModificationState::Modified == IncrementalResultState)
-						{
-							StepResultCache[StepIndex] = IncrementalResult;
-						}
-					}
-					else
-					{
-						Step.Merge(IncrementalResult, StepResultCache[StepIndex]);
+			StepInfo.Step->Execute(Incremental.ActiveKeys, Incremental.ActiveSelection);
 
-						IncrementalResultState = Step.Execute(StepResultCache[StepIndex]);
-						*Result = StepResultCache[StepIndex];
+			if (StepInfo.bProcessRemovals)
+			{
+				StepInfo.Step->Execute(Incremental.ActiveRemovalKeys, Incremental.ActiveRemovalSelection);
+			}
 
-						// After results are merged, downstream steps can no longer 
-						// work on incremental data. Run downstream steps on full
-						// result set.
-						ExecuteSteps(StepIndex + 1);
-						return;
-					}
-				}
-				else
+			const bool bIsLastProducingStep = StepIndex == FinalEntryProducingStepIndex;
+			if (bIsLastProducingStep)
+			{
+				const bool bIsIncrementalEmpty = (Incremental.ActiveKeys.Num() == 0) && (Incremental.ActiveRemovalKeys.Num() == 0);
+				if (bIsIncrementalEmpty)
 				{
-					IncrementalResultState = Step.Increment(IncrementalResult);
+					// Early out if following steps will not produce any new outputs.
+					return;
 				}
 			}
-		}
 
-		if (Steps.Num() > 0)
-		{
-			if (ensure(Steps.Last().IsValid()))
+			// Determine if incremental results need to be merged because it's 
+			// the last step, or because the step requires merging.
+			const bool bIsLastStep = StepIndex == LastStepIndex;
+			if (bIsLastStep)
 			{
-				FFrontendQueryStep& Step = *Steps.Last();
-				if (EResultModificationState::Modified == IncrementalResultState)
-				{
-					// If all incremental steps are performed and result is modified,
-					// merge the incremental result with the output result.
-					Step.Merge(IncrementalResult, *Result);
-				}
+				StepInfo.Step->Merge(Incremental, *Result);
+				OutUpdatedKeys = Incremental.ActiveKeys.Union(Incremental.ActiveRemovalKeys);
+			}
+			else if (StepInfo.bMergeAndCacheOutput)
+			{
+				StepInfo.Step->Merge(Incremental, StepInfo.OutputCache);
+				// If entries were removed during a merge, the subsequent step 
+				// needs to re-evaluate the entire partition. The entire partition 
+				// is added to the active set for each removed key.
+				AppendPartitions(Incremental.ActiveRemovalKeys, StepInfo.OutputCache, Incremental.ActiveKeys, Incremental.ActiveSelection);
 			}
 		}
 	}
 
-	void FFrontendQuery::ExecuteSteps(int32 InStartStepIndex)
+	void FFrontendQuery::AppendPartitions(const TSet<FFrontendQueryKey>& InKeysToAppend, const FFrontendQuerySelection& InSelection, TSet<FFrontendQueryKey>& OutKeysModified, FFrontendQuerySelection& OutSelection) const
 	{
-		for (int32 StepIndex = InStartStepIndex; StepIndex < Steps.Num(); StepIndex++)
+		for (const FFrontendQueryKey& Key : InKeysToAppend)
 		{
-			if (ensure(Steps[StepIndex].IsValid()))
+			if (const FFrontendQueryPartition* Partition = InSelection.Find(Key))
 			{
-				FFrontendQueryStep& Step = *Steps[StepIndex];
-
-				EResultModificationState ExecuteResultState = Step.Execute(*Result);
-
-				const bool bNeedsStepResultCached = (EResultModificationState::Modified == ExecuteResultState) && Step.IsMergeRequiredForIncremental();
-				if (bNeedsStepResultCached)
+				if (Partition->Num() > 0)
 				{
-					StepResultCache[StepIndex] = *Result;
+					OutKeysModified.Add(Key);
+					OutSelection.FindOrAdd(Key) = *Partition;
 				}
 			}
 		}
diff --git a/Engine/Plugins/Runtime/Metasound/Source/MetasoundFrontend/Private/MetasoundFrontendQuerySteps.cpp b/Engine/Plugins/Runtime/Metasound/Source/MetasoundFrontend/Private/MetasoundFrontendQuerySteps.cpp
--- a/Engine/Plugins/Runtime/Metasound/Source/MetasoundFrontend/Private/MetasoundFrontendQuerySteps.cpp
+++ b/Engine/Plugins/Runtime/Metasound/Source/MetasoundFrontend/Private/MetasoundFrontendQuerySteps.cpp
@@ -14,13 +14,13 @@ namespace Metasound
 	{
 	}
 
-	void FNodeClassRegistrationEvents::Stream(TArray<FFrontendQueryEntry>& OutEntries)
+	void FNodeClassRegistrationEvents::Stream(TArray<FFrontendQueryValue>& OutValues)
 	{
 		using namespace Frontend;
 
-		auto AddEntry = [&OutEntries](const FNodeRegistryTransaction& InTransaction)
+		auto AddEntry = [&OutValues](const FNodeRegistryTransaction& InTransaction)
 		{
-			OutEntries.Emplace(FFrontendQueryEntry::FValue(TInPlaceType<FNodeRegistryTransaction>(), InTransaction));
+			OutValues.Emplace(TInPlaceType<FNodeRegistryTransaction>(), InTransaction);
 		};
 		
 		if (FMetasoundFrontendRegistryContainer* Registry = FMetasoundFrontendRegistryContainer::Get())
@@ -29,12 +29,7 @@ namespace Metasound
 		}
 	}
 
-	void FNodeClassRegistrationEvents::Reset()
-	{
-		CurrentTransactionID = Frontend::GetOriginRegistryTransactionID();
-	}
-
-	FFrontendQueryEntry::FKey FMapRegistrationEventsToNodeRegistryKeys::Map(const FFrontendQueryEntry& InEntry) const 
+	FFrontendQueryKey FMapRegistrationEventsToNodeRegistryKeys::Map(const FFrontendQueryEntry& InEntry) const 
 	{
 		using namespace Frontend;
 
@@ -45,27 +40,27 @@ namespace Metasound
 			RegistryKey = InEntry.Value.Get<FNodeRegistryTransaction>().GetNodeRegistryKey();
 		}
 
-		return FFrontendQueryEntry::FKey(RegistryKey);
+		return FFrontendQueryKey(RegistryKey);
 	}
 
-	void FReduceRegistrationEventsToCurrentStatus::Reduce(FFrontendQueryEntry::FKey InKey, TArrayView<FFrontendQueryEntry * const>& InEntries, FReduceOutputView& OutResult) const
+	void FReduceRegistrationEventsToCurrentStatus::Reduce(const FFrontendQueryKey& InKey, FFrontendQueryPartition& InOutEntries) const
 	{
 		using namespace Frontend;
 
 		int32 State = 0;
 		FFrontendQueryEntry* FinalEntry = nullptr;
 
-		for (FFrontendQueryEntry* Entry : InEntries)
+		for (FFrontendQueryEntry& Entry : InOutEntries)
 		{
-			if (ensure(Entry->Value.IsType<FNodeRegistryTransaction>()))
+			if (ensure(Entry.Value.IsType<FNodeRegistryTransaction>()))
 			{
-				const FNodeRegistryTransaction& Transaction = Entry->Value.Get<FNodeRegistryTransaction>();
+				const FNodeRegistryTransaction& Transaction = Entry.Value.Get<FNodeRegistryTransaction>();
 				
 				switch (Transaction.GetTransactionType())
 				{
 					case FNodeRegistryTransaction::ETransactionType::NodeRegistration:
 						State++;
-						FinalEntry = Entry;
+						FinalEntry = &Entry;
 						break;
 
 					case FNodeRegistryTransaction::ETransactionType::NodeUnregistration:
@@ -80,7 +75,13 @@ namespace Metasound
 
 		if ((nullptr != FinalEntry) && (State > 0))
 		{
-			OutResult.Add(*FinalEntry);
+			FFrontendQueryEntry Entry = *FinalEntry;
+			InOutEntries.Reset();
+			InOutEntries.Add(Entry);
+		}
+		else
+		{
+			InOutEntries.Reset();
 		}
 	}
 
@@ -133,14 +134,9 @@ namespace Metasound
 		);
 	}
 
-	FFilterClassesByClassName::FFilterClassesByClassName(const FMetasoundFrontendClassName& InClassName)
-	: ClassName(InClassName)
+	FFrontendQueryKey FMapClassesToClassName::Map(const FFrontendQueryEntry& InEntry) const 
 	{
-	}
-
-	bool FFilterClassesByClassName::Filter(const FFrontendQueryEntry& InEntry) const 
-	{
-		return InEntry.Value.Get<FMetasoundFrontendClass>().Metadata.GetClassName() == ClassName;
+		return FFrontendQueryKey(InEntry.Value.Get<FMetasoundFrontendClass>().Metadata.GetClassName().GetFullName());
 	}
 
 	FFilterClassesByClassID::FFilterClassesByClassID(const FGuid InClassID)
@@ -153,49 +149,33 @@ namespace Metasound
 		return InEntry.Value.Get<FMetasoundFrontendClass>().ID == ClassID;
 	}
 
-	FFrontendQueryEntry::FKey FMapToFullClassName::Map(const FFrontendQueryEntry& InEntry) const
+	FFrontendQueryKey FMapToFullClassName::Map(const FFrontendQueryEntry& InEntry) const
 	{
 		const FMetasoundFrontendClass& FrontendClass = InEntry.Value.Get<FMetasoundFrontendClass>();
-		return FFrontendQueryEntry::FKey(FrontendClass.Metadata.GetClassName().GetFullName());
+		return FFrontendQueryKey(FrontendClass.Metadata.GetClassName().GetFullName());
 	}
 
-	void FReduceClassesToHighestVersion::Reduce(FFrontendQueryEntry::FKey InKey, TArrayView<FFrontendQueryEntry * const>& InEntries, FReduceOutputView& OutResult) const
+	void FReduceClassesToHighestVersion::Reduce(const FFrontendQueryKey& InKey, FFrontendQueryPartition& InOutEntries) const
 	{
 		FFrontendQueryEntry* HighestVersionEntry = nullptr;
-		int32 HighestMajorVersion = -1;
+		FMetasoundFrontendVersionNumber HighestVersion;
 
-		for (FFrontendQueryEntry* Entry : InEntries)
+		for (FFrontendQueryEntry& Entry : InOutEntries)
 		{
-			const int32 EntryMajorVersion = Entry->Value.Get<FMetasoundFrontendClass>().Metadata.GetVersion().Major;
+			const FMetasoundFrontendVersionNumber& Version = Entry.Value.Get<FMetasoundFrontendClass>().Metadata.GetVersion();
 
-			if (!HighestVersionEntry || HighestMajorVersion < EntryMajorVersion)
+			if (!HighestVersionEntry || HighestVersion < Version)
 			{
-				HighestVersionEntry = Entry;
-				HighestMajorVersion = EntryMajorVersion;
+				HighestVersionEntry = &Entry;
+				HighestVersion = Version;
 			}
 		}
 
 		if (HighestVersionEntry)
 		{
-			OutResult.Add(*HighestVersionEntry);
-		}
-	}
-
-	FReduceClassesToMajorVersion::FReduceClassesToMajorVersion(int32 InMajorVersion)
-		: MajorVersion(InMajorVersion)
-	{
-	}
-
-	void FReduceClassesToMajorVersion::Reduce(FFrontendQueryEntry::FKey InKey, TArrayView<FFrontendQueryEntry * const>& InEntries, FReduceOutputView& OutResult) const
-	{
-		for (FFrontendQueryEntry* Entry : InEntries)
-		{
-			const int32 EntryMajorVersion = Entry->Value.Get<FMetasoundFrontendClass>().Metadata.GetVersion().Major;
-			if (MajorVersion == EntryMajorVersion)
-			{
-				OutResult.Add(*Entry);
-				return;
-			}
+			FFrontendQueryEntry Entry = *HighestVersionEntry;
+			InOutEntries.Reset();
+			InOutEntries.Add(Entry);
 		}
 	}
 
diff --git a/Engine/Plugins/Runtime/Metasound/Source/MetasoundFrontend/Private/MetasoundFrontendSearchEngine.cpp b/Engine/Plugins/Runtime/Metasound/Source/MetasoundFrontend/Private/MetasoundFrontendSearchEngine.cpp
--- a/Engine/Plugins/Runtime/Metasound/Source/MetasoundFrontend/Private/MetasoundFrontendSearchEngine.cpp
+++ b/Engine/Plugins/Runtime/Metasound/Source/MetasoundFrontend/Private/MetasoundFrontendSearchEngine.cpp
@@ -26,7 +26,28 @@ namespace Metasound
 	{
 		namespace SearchEngineQuerySteps
 		{
-			class FInterfaceRegistryTransactionSource : public IFrontendQuerySource
+			class FMapClassToFullClassNameAndMajorVersion : public IFrontendQueryMapStep
+			{
+			public:
+				static FFrontendQueryKey GetKey(const FMetasoundFrontendClassName& InClassName, int32 InMajorVersion)
+				{
+					return FFrontendQueryKey(FString::Format(TEXT("{0}_v{1}"), {InClassName.ToString(), InMajorVersion}));
+				}
+
+				virtual FFrontendQueryKey Map(const FFrontendQueryEntry& InEntry) const override
+				{
+					if (ensure(InEntry.Value.IsType<FMetasoundFrontendClass>()))
+					{
+						const FMetasoundFrontendClassMetadata& Metadata = InEntry.Value.Get<FMetasoundFrontendClass>().Metadata;
+						return FMapClassToFullClassNameAndMajorVersion::GetKey(Metadata.GetClassName(), Metadata.GetVersion().Major);
+					}
+
+					return FFrontendQueryKey();
+				}
+			};
+
+
+			class FInterfaceRegistryTransactionSource : public IFrontendQueryStreamStep
 			{
 			public:
 				FInterfaceRegistryTransactionSource()
@@ -34,19 +55,14 @@ namespace Metasound
 				{
 				}
 
-				void Stream(TArray<FFrontendQueryEntry>& OutEntries) override
+				virtual void Stream(TArray<FFrontendQueryValue>& OutEntries) override
 				{
-					auto AddEntry = [&OutEntries](const FInterfaceRegistryTransaction& InTransaction)
+					auto AddValue = [&OutEntries](const FInterfaceRegistryTransaction& InTransaction)
 					{
-						OutEntries.Emplace(FFrontendQueryEntry::FValue(TInPlaceType<FInterfaceRegistryTransaction>(), InTransaction));
+						OutEntries.Emplace(TInPlaceType<FInterfaceRegistryTransaction>(), InTransaction);
 					};
 					
-					IInterfaceRegistry::Get().ForEachRegistryTransactionSince(CurrentTransactionID, &CurrentTransactionID, AddEntry);
-				}
-
-				void Reset() override
-				{
-					CurrentTransactionID = GetOriginRegistryTransactionID();
+					IInterfaceRegistry::Get().ForEachRegistryTransactionSince(CurrentTransactionID, &CurrentTransactionID, AddValue);
 				}
 
 			private:
@@ -56,36 +72,50 @@ namespace Metasound
 			class FMapInterfaceRegistryTransactionsToInterfaceRegistryKeys : public IFrontendQueryMapStep
 			{
 			public:
-				FFrontendQueryEntry::FKey Map(const FFrontendQueryEntry& InEntry) const override
+				virtual FFrontendQueryKey Map(const FFrontendQueryEntry& InEntry) const override
 				{
 					FInterfaceRegistryKey RegistryKey;
 					if (ensure(InEntry.Value.IsType<FInterfaceRegistryTransaction>()))
 					{
 						RegistryKey = InEntry.Value.Get<FInterfaceRegistryTransaction>().GetInterfaceRegistryKey();
 					}
-					return FFrontendQueryEntry::FKey{RegistryKey};
+					return FFrontendQueryKey{RegistryKey};
+				}
+			};
+
+			class FMapInterfaceRegistryTransactionToInterfaceName : public IFrontendQueryMapStep
+			{
+			public:
+				virtual FFrontendQueryKey Map(const FFrontendQueryEntry& InEntry) const override
+				{
+					FInterfaceRegistryKey RegistryKey;
+					if (ensure(InEntry.Value.IsType<FInterfaceRegistryTransaction>()))
+					{
+						return FFrontendQueryKey{InEntry.Value.Get<FInterfaceRegistryTransaction>().GetInterfaceVersion().Name};
+					}
+					return FFrontendQueryKey();
 				}
 			};
 
 			class FReduceInterfaceRegistryTransactionsToCurrentStatus : public IFrontendQueryReduceStep
 			{
 			public:
-				void Reduce(FFrontendQueryEntry::FKey InKey, TArrayView<FFrontendQueryEntry* const>& InEntries, FReduceOutputView& OutResult) const override
+				virtual void Reduce(const FFrontendQueryKey& InKey, FFrontendQueryPartition& InOutEntries) const override
 				{
 					int32 State = 0;
 					FFrontendQueryEntry* FinalEntry = nullptr;
 
-					for (FFrontendQueryEntry* Entry : InEntries)
+					for (FFrontendQueryEntry& Entry : InOutEntries)
 					{
-						if (ensure(Entry->Value.IsType<FInterfaceRegistryTransaction>()))
+						if (ensure(Entry.Value.IsType<FInterfaceRegistryTransaction>()))
 						{
-							const FInterfaceRegistryTransaction& Transaction = Entry->Value.Get<FInterfaceRegistryTransaction>();
+							const FInterfaceRegistryTransaction& Transaction = Entry.Value.Get<FInterfaceRegistryTransaction>();
 							
 							switch (Transaction.GetTransactionType())
 							{
 								case FInterfaceRegistryTransaction::ETransactionType::InterfaceRegistration:
 									State++;
-									FinalEntry = Entry;
+									FinalEntry = &Entry;
 									break;
 
 								case FInterfaceRegistryTransaction::ETransactionType::InterfaceUnregistration:
@@ -100,61 +130,57 @@ namespace Metasound
 
 					if ((nullptr != FinalEntry) && (State > 0))
 					{
-						OutResult.Add(*FinalEntry);
+						FFrontendQueryEntry Entry = *FinalEntry;
+						InOutEntries.Reset();
+						InOutEntries.Add(Entry);
+					}
+					else
+					{
+						InOutEntries.Reset();
 					}
 				}
 			};
 
-			class FReduceInterfaceRegistryTransactionsBySupportedClassName : public IFrontendQueryReduceStep
+			class FFilterInterfaceRegistryTransactionsBySupportedClassName : public IFrontendQueryFilterStep
 			{
 				FName UClassName;
 
 			public:
-				FReduceInterfaceRegistryTransactionsBySupportedClassName(FName InUClassName)
+				FFilterInterfaceRegistryTransactionsBySupportedClassName(FName InUClassName)
 					: UClassName(InUClassName)
 				{
 				}
 
-				void Reduce(FFrontendQueryEntry::FKey InKey, TArrayView<FFrontendQueryEntry* const>& InEntries, FReduceOutputView& OutResult) const override
+				virtual bool Filter(const FFrontendQueryEntry& InEntry) const override
 				{
-					for (FFrontendQueryEntry* Entry : InEntries)
+					if (ensure(InEntry.Value.IsType<FInterfaceRegistryTransaction>()))
 					{
-						if (ensure(Entry->Value.IsType<FInterfaceRegistryTransaction>()))
+						FInterfaceRegistryKey RegistryKey = InEntry.Value.Get<FInterfaceRegistryTransaction>().GetInterfaceRegistryKey();
+						const IInterfaceRegistryEntry* RegEntry = IInterfaceRegistry::Get().FindInterfaceRegistryEntry(RegistryKey);
+						if (ensure(RegEntry))
 						{
-							FInterfaceRegistryKey RegistryKey = Entry->Value.Get<FInterfaceRegistryTransaction>().GetInterfaceRegistryKey();
-							const IInterfaceRegistryEntry* RegEntry = IInterfaceRegistry::Get().FindInterfaceRegistryEntry(RegistryKey);
-							if (ensure(RegEntry))
-							{
-								if (RegEntry->UClassIsSupported(UClassName))
-								{
-									OutResult.Add(*Entry);
-								}
-							}
+							return RegEntry->UClassIsSupported(UClassName);
 						}
 					}
+					return false;
 				}
 			};
 
-			class FReduceInterfaceRegistryTransactionsByDefault : public IFrontendQueryReduceStep
+			class FFilterInterfaceRegistryTransactionsByDefault : public IFrontendQueryFilterStep
 			{
 			public:
-				void Reduce(FFrontendQueryEntry::FKey InKey, TArrayView<FFrontendQueryEntry* const>& InEntries, FReduceOutputView& OutResult) const override
+				virtual bool Filter(const FFrontendQueryEntry& InEntry) const override
 				{
-					for (FFrontendQueryEntry* Entry : InEntries)
+					if (ensure(InEntry.Value.IsType<FInterfaceRegistryTransaction>()))
 					{
-						if (ensure(Entry->Value.IsType<FInterfaceRegistryTransaction>()))
+						FInterfaceRegistryKey RegistryKey = InEntry.Value.Get<FInterfaceRegistryTransaction>().GetInterfaceRegistryKey();
+						const IInterfaceRegistryEntry* RegEntry = IInterfaceRegistry::Get().FindInterfaceRegistryEntry(RegistryKey);
+						if (ensure(RegEntry))
 						{
-							FInterfaceRegistryKey RegistryKey = Entry->Value.Get<FInterfaceRegistryTransaction>().GetInterfaceRegistryKey();
-							const IInterfaceRegistryEntry* RegEntry = IInterfaceRegistry::Get().FindInterfaceRegistryEntry(RegistryKey);
-							if (ensure(RegEntry))
-							{
-								if (RegEntry->IsDefault())
-								{
-									OutResult.Add(*Entry);
-								}
-							}
+							return RegEntry->IsDefault();
 						}
 					}
+					return false;
 				}
 			};
 
@@ -176,35 +202,40 @@ namespace Metasound
 			class FMapInterfaceToInterfaceNameAndMajorVersion : public IFrontendQueryMapStep
 			{
 			public:
-				FFrontendQueryEntry::FKey Map(const FFrontendQueryEntry& InEntry) const override
+				static FFrontendQueryKey GetKey(const FString& InName, int32 InMajorVersion)
+				{
+					return FFrontendQueryKey{FString::Format(TEXT("{0}_{1}"), { InName, InMajorVersion })};
+				}
+
+				virtual FFrontendQueryKey Map(const FFrontendQueryEntry& InEntry) const override
 				{
 					FString Key;
 					if (ensure(InEntry.Value.IsType<FMetasoundFrontendInterface>()))
 					{
 						const FMetasoundFrontendInterface& Interface = InEntry.Value.Get<FMetasoundFrontendInterface>();
-						Key = FString::Format(TEXT("{0}_{1}"), { Interface.Version.Name.ToString(), Interface.Version.Number.Major });
+
+						return FMapInterfaceToInterfaceNameAndMajorVersion::GetKey(Interface.Version.Name.ToString(), Interface.Version.Number.Major);
 					}
-					return FFrontendQueryEntry::FKey{Key};
+					return FFrontendQueryKey{Key};
 				}
 			};
 
 			class FReduceInterfacesToHighestVersion : public IFrontendQueryReduceStep
 			{
 			public:
-				void Reduce(FFrontendQueryEntry::FKey InKey, TArrayView<FFrontendQueryEntry* const>& InEntries, FReduceOutputView& OutResult) const override
+				virtual void Reduce(const FFrontendQueryKey& InKey, FFrontendQueryPartition& InOutEntries) const override
 				{
 					FFrontendQueryEntry* HighestVersionEntry = nullptr;
 					FMetasoundFrontendVersionNumber HighestVersion = FMetasoundFrontendVersionNumber::GetInvalid();
 
-					for (FFrontendQueryEntry* Entry : InEntries)
+					for (FFrontendQueryEntry& Entry : InOutEntries)
 					{
-						if (ensure(Entry->Value.IsType<FMetasoundFrontendInterface>()))
+						if (ensure(Entry.Value.IsType<FMetasoundFrontendInterface>()))
 						{
-							const FMetasoundFrontendVersionNumber& VersionNumber = Entry->Value.Get<FMetasoundFrontendInterface>().Version.Number;
-
+							const FMetasoundFrontendVersionNumber& VersionNumber = Entry.Value.Get<FMetasoundFrontendInterface>().Version.Number;
 							if (VersionNumber > HighestVersion)
 							{
-								HighestVersionEntry = Entry;
+								HighestVersionEntry = &Entry;
 								HighestVersion = VersionNumber;
 							}
 						}
@@ -212,7 +243,13 @@ namespace Metasound
 
 					if (HighestVersionEntry)
 					{
-						OutResult.Add(*HighestVersionEntry);
+						FFrontendQueryEntry Entry = *HighestVersionEntry;
+						InOutEntries.Reset();
+						InOutEntries.Add(Entry);
+					}
+					else
+					{
+						InOutEntries.Reset();
 					}
 				}
 			};
@@ -225,7 +262,7 @@ namespace Metasound
 				{
 				}
 
-				bool Filter(const FFrontendQueryEntry& InEntry) const override
+				virtual bool Filter(const FFrontendQueryEntry& InEntry) const override
 				{
 					if (ensure(InEntry.Value.IsType<FInterfaceRegistryTransaction>()))
 					{
@@ -247,7 +284,7 @@ namespace Metasound
 				{
 				}
 
-				bool Filter(const FFrontendQueryEntry& InEntry) const override
+				virtual bool Filter(const FFrontendQueryEntry& InEntry) const override
 				{
 					if (ensure(InEntry.Value.IsType<FNodeRegistryTransaction>()))
 					{
@@ -268,7 +305,7 @@ namespace Metasound
 				{
 				}
 
-				bool Filter(const FFrontendQueryEntry& InEntry) const override
+				virtual bool Filter(const FFrontendQueryEntry& InEntry) const override
 				{
 					if (ensure(InEntry.Value.IsType<FInterfaceRegistryTransaction>()))
 					{
@@ -287,14 +324,14 @@ namespace Metasound
 			class FMapInterfaceToInterfaceName : public IFrontendQueryMapStep
 			{
 			public:
-				FFrontendQueryEntry::FKey Map(const FFrontendQueryEntry& InEntry) const override
+				virtual FFrontendQueryKey Map(const FFrontendQueryEntry& InEntry) const override
 				{
-					FString Key;
+					FName Key;
 					if (ensure(InEntry.Value.IsType<FMetasoundFrontendInterface>()))
 					{
-						Key = InEntry.Value.Get<FMetasoundFrontendInterface>().Version.Name.ToString();
+						Key = InEntry.Value.Get<FMetasoundFrontendInterface>().Version.Name;
 					}
-					return FFrontendQueryEntry::FKey{Key};
+					return FFrontendQueryKey{Key};
 				}
 			};
 
@@ -325,9 +362,9 @@ namespace Metasound
 				virtual ~FSearchEngine() = default;
 
 				virtual TArray<FMetasoundFrontendClass> FindAllClasses(bool bInIncludeDeprecated) override;
-				virtual TArray<FMetasoundFrontendClass> FindClassesWithName(const FNodeClassName& InName, bool bInSortByVersion) override;
-				virtual bool FindClassWithHighestVersion(const FNodeClassName& InName, FMetasoundFrontendClass& OutClass) override;
-				virtual bool FindClassWithMajorVersion(const FNodeClassName& InName, int32 InMajorVersion, FMetasoundFrontendClass& OutClass) override;
+				virtual TArray<FMetasoundFrontendClass> FindClassesWithName(const FMetasoundFrontendClassName& InName, bool bInSortByVersion) override;
+				virtual bool FindClassWithHighestVersion(const FMetasoundFrontendClassName& InName, FMetasoundFrontendClass& OutClass) override;
+				virtual bool FindClassWithMajorVersion(const FMetasoundFrontendClassName& InName, int32 InMajorVersion, FMetasoundFrontendClass& OutClass) override;
 
 				virtual TArray<FMetasoundFrontendInterface> FindAllInterfaces(bool bInIncludeDeprecated) override;
 				virtual TArray<FMetasoundFrontendInterface> FindUClassDefaultInterfaces(FName InUClassName) override;
@@ -352,6 +389,7 @@ namespace Metasound
 		TArray<FMetasoundFrontendClass> FSearchEngine::FindAllClasses(bool bInIncludeDeprecated)
 		{
 			METASOUND_TRACE_CPUPROFILER_EVENT_SCOPE(metasound::FSearchEngine::FindAllClasses);
+			using namespace SearchEngineQuerySteps;
 
 			FString QueryName = FString(TEXT("FindAllClasses"));
 			if (bInIncludeDeprecated)
@@ -364,9 +402,6 @@ namespace Metasound
 			{
 				TUniquePtr<FFrontendQuery> NewQuery = MakeUnique<FFrontendQuery>();
 
-				// TODO: this query will be slow to update when new nodes are registered. 
-				// Consider reworking call sites, or creating a reduce function
-				// which performs the majority of the query logic.
 				NewQuery->AddStep<FNodeClassRegistrationEvents>()
 					.AddStep<FMapRegistrationEventsToNodeRegistryKeys>()
 					.AddStep<FReduceRegistrationEventsToCurrentStatus>()
@@ -378,6 +413,8 @@ namespace Metasound
 						.AddStep<FReduceClassesToHighestVersion>();
 				}
 
+				NewQuery->AddStep<FMapToNull>();
+
 				Query = AddQuery(QueryName, MoveTemp(NewQuery));
 			}
 
@@ -385,19 +422,22 @@ namespace Metasound
 
 			if (ensure(Query))
 			{
-				FFrontendQuerySelectionView Selection = Query->Execute();
-				TArrayView<const FFrontendQueryEntry* const> Entries = Selection.GetSelection();
-				for (const FFrontendQueryEntry* Entry : Selection.GetSelection())
+				// TODO: cache result and only rebuild if query actually updates. 
+				const FFrontendQuerySelection& Selection = Query->Update();
+				if (const FFrontendQueryPartition* Entries = Selection.Find(FFrontendQueryKey()))
 				{
-					check(Entry->Value.IsType<FMetasoundFrontendClass>());
-					Result.Add(Entry->Value.Get<FMetasoundFrontendClass>());
+					for (const FFrontendQueryEntry& Entry : *Entries)
+					{
+						check(Entry.Value.IsType<FMetasoundFrontendClass>());
+						Result.Add(Entry.Value.Get<FMetasoundFrontendClass>());
+					}
 				}
 			}
 
 			return Result;
 		}
 
-		TArray<FMetasoundFrontendClass> FSearchEngine::FindClassesWithName(const FNodeClassName& InName, bool bInSortByVersion)
+		TArray<FMetasoundFrontendClass> FSearchEngine::FindClassesWithName(const FMetasoundFrontendClassName& InName, bool bInSortByVersion)
 		{
 			METASOUND_TRACE_CPUPROFILER_EVENT_SCOPE(metasound::FSearchEngine::FindClassesWithName);
 			using namespace SearchEngineQuerySteps;
@@ -405,20 +445,18 @@ namespace Metasound
 			FString QueryName = FString(TEXT("FindClassesWithName"));
 			if (bInSortByVersion)
 			{
-				QueryName += TEXT("SortedByVersion");
+				QueryName += TEXT("_SortedByVersion");
 			}
-			QueryName += InName.GetFullName().ToString();
 			FFrontendQuery* Query = FindQuery(QueryName);
 
 			if (!Query)
 			{
 				TUniquePtr<FFrontendQuery> NewQuery = MakeUnique<FFrontendQuery>();
 				NewQuery->AddStep<FNodeClassRegistrationEvents>()
-					.AddStep<FFilterNodeRegistrationEventsByClassName>(InName)
 					.AddStep<FMapRegistrationEventsToNodeRegistryKeys>()
 					.AddStep<FReduceRegistrationEventsToCurrentStatus>()
 					.AddStep<FTransformRegistrationEventsToClasses>()
-					.AddStep<FFilterClassesByClassName>(InName);
+					.AddStep<FMapClassesToClassName>();
 
 				if (bInSortByVersion)
 				{
@@ -432,37 +470,36 @@ namespace Metasound
 
 			if (ensure(nullptr != Query))
 			{
-				FFrontendQuerySelectionView Selection = Query->Execute();
-				TArrayView<const FFrontendQueryEntry* const> Entries = Selection.GetSelection();
-				for (const FFrontendQueryEntry* Entry : Selection.GetSelection())
+				const FFrontendQuerySelection& Selection = Query->Update();
+				const FFrontendQueryKey Key(InName.GetFullName());
+				if (const FFrontendQueryPartition* Entries = Selection.Find(Key))
 				{
-					check(Entry->Value.IsType<FMetasoundFrontendClass>());
-					Result.Add(Entry->Value.Get<FMetasoundFrontendClass>());
+					for (const FFrontendQueryEntry& Entry : *Entries)
+					{
+						check(Entry.Value.IsType<FMetasoundFrontendClass>());
+						Result.Add(Entry.Value.Get<FMetasoundFrontendClass>());
+					}
 				}
 			}
 
 			return Result;
 		}
 
-		bool FSearchEngine::FindClassWithHighestVersion(const FNodeClassName& InName, FMetasoundFrontendClass& OutClass)
+		bool FSearchEngine::FindClassWithHighestVersion(const FMetasoundFrontendClassName& InName, FMetasoundFrontendClass& OutClass)
 		{
 			METASOUND_TRACE_CPUPROFILER_EVENT_SCOPE(metasound::FSearchEngine::FindClassWithHighestVersion);
 			using namespace SearchEngineQuerySteps;
 
-			const FString QueryName = FString(TEXT("FindHighestVersion")) + InName.GetFullName().ToString();
+			const FString QueryName = FString(TEXT("FindHighestVersion"));
 			FFrontendQuery* Query = FindQuery(QueryName);
 			if (!Query)
 			{
 				TUniquePtr<FFrontendQuery> NewQuery = MakeUnique<FFrontendQuery>();
 
-				// TODO: Create index of class names and major version to avoid duplicating
-				// this query for each class name.
 				NewQuery->AddStep<FNodeClassRegistrationEvents>()
-					.AddStep<FFilterNodeRegistrationEventsByClassName>(InName)
 					.AddStep<FMapRegistrationEventsToNodeRegistryKeys>()
 					.AddStep<FReduceRegistrationEventsToCurrentStatus>()
 					.AddStep<FTransformRegistrationEventsToClasses>()
-					.AddStep<FFilterClassesByClassName>(InName)
 					.AddStep<FMapToFullClassName>()
 					.AddStep<FReduceClassesToHighestVersion>();
 
@@ -471,62 +508,57 @@ namespace Metasound
 
 			if (ensure(Query))
 			{
-				FFrontendQuerySelectionView Selection = Query->Execute();
-				TArrayView<const FFrontendQueryEntry* const> Entries = Selection.GetSelection();
-				if (Entries.IsEmpty())
-				{
-					return false;
-				}
+				const FFrontendQuerySelection& Selection = Query->Update();
+				const FFrontendQueryKey Key(InName.GetFullName());
 
-				if (ensure(Entries.Num() == 1))
+				if (const FFrontendQueryPartition* Entries = Selection.Find(Key))
 				{
-					OutClass = Entries[0]->Value.Get<FMetasoundFrontendClass>();
-					return true;
+					if (ensure(Entries->Num() == 1))
+					{
+						OutClass = Entries->CreateConstIterator()->Value.Get<FMetasoundFrontendClass>();
+						return true;
+					}
 				}
 			}
 
 			return false;
 		}
 
-		bool FSearchEngine::FindClassWithMajorVersion(const FNodeClassName& InName, int32 InMajorVersion, FMetasoundFrontendClass& OutClass)
+		bool FSearchEngine::FindClassWithMajorVersion(const FMetasoundFrontendClassName& InName, int32 InMajorVersion, FMetasoundFrontendClass& OutClass)
 		{
 			METASOUND_TRACE_CPUPROFILER_EVENT_SCOPE(metasound::FSearchEngine::FindClassWithMajorVersion);
 			using namespace SearchEngineQuerySteps;
 
-			const FString QueryName = FString(TEXT("FindClassWithMajorVersion")) + InName.GetFullName().ToString() + FString::FromInt(InMajorVersion);
+			const FString QueryName = FString(TEXT("FindClassWithMajorVersion"));
 			FFrontendQuery* Query = FindQuery(QueryName);
 
 			if (!Query)
 			{
 				TUniquePtr<FFrontendQuery> NewQuery = MakeUnique<FFrontendQuery>();
 
-				// TODO: Create index of class names to avoid duplicating
-				// this query for each class name.
 				NewQuery->AddStep<FNodeClassRegistrationEvents>()
-					.AddStep<FFilterNodeRegistrationEventsByClassName>(InName)
 					.AddStep<FMapRegistrationEventsToNodeRegistryKeys>()
 					.AddStep<FReduceRegistrationEventsToCurrentStatus>()
 					.AddStep<FTransformRegistrationEventsToClasses>()
-					.AddStep<FFilterClassesByClassName>(InName)
 					.AddStep<FMapToFullClassName>()
-					.AddStep<FReduceClassesToMajorVersion>(InMajorVersion);
+					.AddStep<FMapClassToFullClassNameAndMajorVersion>();
 
 				Query = AddQuery(QueryName, MoveTemp(NewQuery));
 			}
 
+
 			if (ensure(Query))
 			{
-				FFrontendQuerySelectionView Selection = Query->Execute();
-				TArrayView<const FFrontendQueryEntry* const> Entries = Selection.GetSelection();
-				if (Entries.IsEmpty())
-				{
-					return false;
-				}
+				const FFrontendQueryKey Key = FMapClassToFullClassNameAndMajorVersion::GetKey(InName, InMajorVersion);
+				const FFrontendQuerySelection& Selection = Query->Update();
 
-				if (ensure(Entries.Num() == 1))
+				if (const FFrontendQueryPartition* Entries = Selection.Find(Key))
 				{
-					OutClass = Entries[0]->Value.Get<FMetasoundFrontendClass>();
-					return true;
+					if (ensure(Entries->Num() == 1))
+					{
+						OutClass = Entries->CreateConstIterator()->Value.Get<FMetasoundFrontendClass>();
+						return true;
+					}
 				}
 			}
 
@@ -554,6 +586,7 @@ namespace Metasound
 					NewQuery->AddStep<FMapInterfaceToInterfaceNameAndMajorVersion>()
 					.AddStep<FReduceInterfacesToHighestVersion>();
 				}
+				NewQuery->AddStep<FMapToNull>();
 
 				Query = AddQuery(QueryName, MoveTemp(NewQuery));
 			}
@@ -562,14 +595,15 @@ namespace Metasound
 
 			if (ensure(Query))
 			{
-				FFrontendQuerySelectionView Selection = Query->Execute();
-				TArrayView<const FFrontendQueryEntry* const> Entries = Selection.GetSelection();
-
-				for (const FFrontendQueryEntry* Entry : Selection.GetSelection())
+				const FFrontendQuerySelection& Selection = Query->Update();
+				if (const FFrontendQueryPartition* Entries = Selection.Find(FFrontendQueryKey()))
 				{
-					if (ensure(Entry->Value.IsType<FMetasoundFrontendInterface>()))
+					for (const FFrontendQueryEntry& Entry : *Entries)
 					{
-						Result.Add(Entry->Value.Get<FMetasoundFrontendInterface>());
+						if (ensure(Entry.Value.IsType<FMetasoundFrontendInterface>()))
+						{
+							Result.Add(Entry.Value.Get<FMetasoundFrontendInterface>());
+						}
 					}
 				}
 			}
@@ -579,44 +613,24 @@ namespace Metasound
 
 		TArray<FMetasoundFrontendInterface> FSearchEngine::FindUClassDefaultInterfaces(FName InUClassName)
 		{
-			using namespace SearchEngineQuerySteps;
-
-			const FString QueryName = FString::Format(TEXT("FindUClassDefaultInterfaces_Class:{0}"), { InUClassName.ToString() });
-
-			FFrontendQuery* Query = FindQuery(QueryName);
-			if (!Query)
-			{
-				TUniquePtr<FFrontendQuery> NewQuery = MakeUnique<FFrontendQuery>();
-
-				NewQuery->AddStep<FInterfaceRegistryTransactionSource>()
-					.AddStep<FMapInterfaceRegistryTransactionsToInterfaceRegistryKeys>()
-					.AddStep<FReduceInterfaceRegistryTransactionsToCurrentStatus>()
-					.AddStep<FReduceInterfaceRegistryTransactionsByDefault>()
-					.AddStep<FReduceInterfaceRegistryTransactionsBySupportedClassName>(InUClassName)
-					.AddStep<FTransformInterfaceRegistryTransactionToInterface>()
-					.AddStep<FMapInterfaceToInterfaceNameAndMajorVersion>()
-					.AddStep<FReduceInterfacesToHighestVersion>();
-
-				Query = AddQuery(QueryName, MoveTemp(NewQuery));
-			}
-
-			TArray<FMetasoundFrontendInterface> Result;
+			constexpr bool bIncludeDeprecated = false;
+			TArray<FMetasoundFrontendInterface> AllInterfaces = FindAllInterfaces(bIncludeDeprecated);
 
-			if (ensure(Query))
+			auto DoesNotSupportUClass = [&InUClassName](const FMetasoundFrontendInterface& InInterface)
 			{
-				FFrontendQuerySelectionView Selection = Query->Execute();
-				TArrayView<const FFrontendQueryEntry* const> Entries = Selection.GetSelection();
-
-				for (const FFrontendQueryEntry* Entry : Selection.GetSelection())
+				FInterfaceRegistryKey RegistryKey = GetInterfaceRegistryKey(InInterface);
+				const IInterfaceRegistryEntry* RegEntry = IInterfaceRegistry::Get().FindInterfaceRegistryEntry(RegistryKey);
+				if (ensure(RegEntry))
 				{
-					if (ensure(Entry->Value.IsType<FMetasoundFrontendInterface>()))
-					{
-						Result.Add(Entry->Value.Get<FMetasoundFrontendInterface>());
-					}
+					return RegEntry->UClassIsSupported(InUClassName);
 				}
-			}
 
-			return Result;
+				return true;
+			};
+
+			AllInterfaces.RemoveAllSwap(DoesNotSupportUClass);
+
+			return AllInterfaces;
 		}
 
 		TArray<FMetasoundFrontendVersion> FSearchEngine::FindAllRegisteredInterfacesWithName(FName InInterfaceName)
@@ -631,9 +645,9 @@ namespace Metasound
 				TUniquePtr<FFrontendQuery> NewQuery = MakeUnique<FFrontendQuery>();
 
 				NewQuery->AddStep<FInterfaceRegistryTransactionSource>()
-					.AddStep<FFilterInterfaceRegistryTransactionsByName>(InInterfaceName)
 					.AddStep<FMapInterfaceRegistryTransactionsToInterfaceRegistryKeys>()
 					.AddStep<FReduceInterfaceRegistryTransactionsToCurrentStatus>()
+					.AddStep<FMapInterfaceRegistryTransactionToInterfaceName>()
 					.AddStep<FTransformInterfaceRegistryTransactionToInterfaceVersion>();
 
 				Query = AddQuery(QueryName, MoveTemp(NewQuery));
@@ -643,14 +657,16 @@ namespace Metasound
 
 			if (ensure(Query))
 			{
-				FFrontendQuerySelectionView Selection = Query->Execute();
-				TArrayView<const FFrontendQueryEntry* const> Entries = Selection.GetSelection();
-
-				for (const FFrontendQueryEntry* Entry : Selection.GetSelection())
+				const FFrontendQueryKey Key{InInterfaceName};
+				const FFrontendQuerySelection& Selection = Query->Update();
+				if (const FFrontendQueryPartition* Entries = Selection.Find(Key))
 				{
-					if (ensure(Entry->Value.IsType<FMetasoundFrontendVersion>()))
+					for (const FFrontendQueryEntry& Entry : *Entries)
 					{
-						Result.Add(Entry->Value.Get<FMetasoundFrontendVersion>());
+						if (ensure(Entry.Value.IsType<FMetasoundFrontendVersion>()))
+						{
+							Result.Add(Entry.Value.Get<FMetasoundFrontendVersion>());
+						}
 					}
 				}
 			}
@@ -661,7 +677,7 @@ namespace Metasound
 		bool FSearchEngine::FindInterfaceWithHighestVersion(FName InInterfaceName, FMetasoundFrontendInterface& OutInterface)
 		{
 			using namespace SearchEngineQuerySteps;
-			const FString QueryName = FString::Format(TEXT("FindInterfaceWithHighestVersion_Name:{0}"), { InInterfaceName.ToString() });
+			const FString QueryName = TEXT("FindInterfaceWithHighestVersion");
 
 			FFrontendQuery* Query = FindQuery(QueryName);
 			if (!Query)
@@ -669,7 +685,6 @@ namespace Metasound
 				TUniquePtr<FFrontendQuery> NewQuery = MakeUnique<FFrontendQuery>();
 
 				NewQuery->AddStep<FInterfaceRegistryTransactionSource>()
-					.AddStep<FFilterInterfaceRegistryTransactionsByName>(InInterfaceName)
 					.AddStep<FMapInterfaceRegistryTransactionsToInterfaceRegistryKeys>()
 					.AddStep<FReduceInterfaceRegistryTransactionsToCurrentStatus>()
 					.AddStep<FTransformInterfaceRegistryTransactionToInterface>()
@@ -681,17 +696,15 @@ namespace Metasound
 
 			if (ensure(Query))
 			{
-				FFrontendQuerySelectionView Selection = Query->Execute();
-				TArrayView<const FFrontendQueryEntry* const> Entries = Selection.GetSelection();
-				if (Entries.IsEmpty())
+				const FFrontendQueryKey Key{InInterfaceName};
+				const FFrontendQuerySelection& Selection = Query->Update();
+				if (const FFrontendQueryPartition* Entries = Selection.Find(Key))
 				{
-					return false;
-				}
-
-				if (ensure(Entries.Num() == 1))
-				{
-					OutInterface = Entries[0]->Value.Get<FMetasoundFrontendInterface>();
-					return true;
+					if (ensure(Entries->Num() == 1))
+					{
+						OutInterface = Entries->CreateConstIterator()->Value.Get<FMetasoundFrontendInterface>();
+						return true;
+					}
 				}
 			}
 
@@ -701,7 +714,7 @@ namespace Metasound
 		bool FSearchEngine::FindInterfaceWithMajorVersion(FName InInterfaceName, int32 InMajorVersion, FMetasoundFrontendInterface& OutInterface)
 		{
 			using namespace SearchEngineQuerySteps;
-			const FString QueryName = FString::Format(TEXT("FindInterfaceWithMajorVersion_Name:{0}_MajorVersion:{1}"), { InInterfaceName.ToString(), InMajorVersion });
+			const FString QueryName = TEXT("FindInterfaceWithMajorVersion");
 
 			FFrontendQuery* Query = FindQuery(QueryName);
 			if (!Query)
@@ -709,11 +722,10 @@ namespace Metasound
 				TUniquePtr<FFrontendQuery> NewQuery = MakeUnique<FFrontendQuery>();
 
 				NewQuery->AddStep<FInterfaceRegistryTransactionSource>()
-					.AddStep<FFilterInterfaceRegistryTransactionsByNameAndMajorVersion>(InInterfaceName, InMajorVersion)
 					.AddStep<FMapInterfaceRegistryTransactionsToInterfaceRegistryKeys>()
 					.AddStep<FReduceInterfaceRegistryTransactionsToCurrentStatus>()
 					.AddStep<FTransformInterfaceRegistryTransactionToInterface>()
-					.AddStep<FMapInterfaceToInterfaceName>()
+					.AddStep<FMapInterfaceToInterfaceNameAndMajorVersion>()
 					.AddStep<FReduceInterfacesToHighestVersion>();
 
 				Query = AddQuery(QueryName, MoveTemp(NewQuery));
@@ -721,17 +733,15 @@ namespace Metasound
 
 			if (ensure(Query))
 			{
-				FFrontendQuerySelectionView Selection = Query->Execute();
-				TArrayView<const FFrontendQueryEntry* const> Entries = Selection.GetSelection();
-				if (Entries.IsEmpty())
+				const FFrontendQueryKey Key = FMapInterfaceToInterfaceNameAndMajorVersion::GetKey(InInterfaceName.ToString(), InMajorVersion);
+				const FFrontendQuerySelection& Selection = Query->Update();
+				if (const FFrontendQueryPartition* Entries = Selection.Find(Key))
 				{
-					return false;
-				}
-
-				if (ensure(Entries.Num() == 1))
-				{
-					OutInterface = Entries[0]->Value.Get<FMetasoundFrontendInterface>();
-					return true;
+					if (ensure(Entries->Num() == 1))
+					{
+						OutInterface = Entries->CreateConstIterator()->Value.Get<FMetasoundFrontendInterface>();
+						return true;
+					}
 				}
 			}
 
diff --git a/Engine/Plugins/Runtime/Metasound/Source/MetasoundFrontend/Public/MetasoundFrontendQuery.h b/Engine/Plugins/Runtime/Metasound/Source/MetasoundFrontend/Public/MetasoundFrontendQuery.h
--- a/Engine/Plugins/Runtime/Metasound/Source/MetasoundFrontend/Public/MetasoundFrontendQuery.h
+++ b/Engine/Plugins/Runtime/Metasound/Source/MetasoundFrontend/Public/MetasoundFrontendQuery.h
@@ -9,6 +9,41 @@
 #include "MetasoundFrontendArchetypeRegistry.h"
 #include "Misc/TVariant.h"
 
+/** MetaSound Frontend Query
+ *
+ * MetaSound Frontend Query provides a way to systematically organize and update
+ * streaming data associated with the MetaSound Frontend. It is a streaming MapReduce 
+ * framework for querying streams of data (https://en.wikipedia.org/wiki/MapReduce) 
+ *
+ * While it does not support the computational parallelism commonly found in MapReduce 
+ * frameworks, it does offer:
+ * 	- An encapsulated and reusable set of methods for manipulating streamed data.
+ *  - Support for incremental updates (a.k.a. streamed data).
+ * 	- An indexed output for efficient lookup. 
+ *
+ * Data within MetaSound Frontend Query is organized similarly to a NoSQL database
+ * (https://en.wikipedia.org/wiki/NoSQL). Each object (FFrontendQueryEntry) is 
+ * assigned a unique ID. Keys (FFrontendQueryKey) are associated with sets of entries
+ * (FFrontendQueryPartition) and allow partitions to be retrieved efficiently. 
+ * Each partition holds a set of entries which is determined by the steps in the 
+ * query (FFrontendQuery). A FFrontendQueryKey or FFrontendQueryValue represent
+ * one of multiple types by using a TVariant<>.
+ *
+ * A query contains a sequence of steps that get executed on streaming data. The 
+ * various types of steps reflect common operations performed in MapReduce and 
+ * NoSQL database queries. 
+ *
+ * Step Types
+ * 	Stream: 	Produce a stream of FFrontendQueryValues.
+ * 	Map:		Map a FFrontendQueryEntry to a partition associated with a FFrontendQueryKey.
+ * 	Reduce: 	Apply an incremental summarization of a FFrontendQueryPartition.
+ * 	Transform:	Alter a FFrontendQueryValue.
+ * 	Filter:		Remove FFrontendQueryValues with a test function.
+ * 	Score:		Calculate a score for a FFrontendQueryValue.
+ * 	Sort:		Sort a FFrontendQueryPartition.
+ * 	Limit:		Limit the size of a FFrontendQueryPartition.
+ */
+
 namespace Metasound
 {
 	/** FFrontendQueryKey allows entries to be partitioned by their key. A key
@@ -31,171 +66,39 @@ namespace Metasound
 		friend bool operator<(const FFrontendQueryKey& InLHS, const FFrontendQueryKey& InRHS);
 		friend uint32 GetTypeHash(const FFrontendQueryKey& InKey);
 		
-		bool IsValid() const;
+		bool IsNull() const;
 
 	private:
-		struct FInvalid {};
+		struct FNull{};
 
-		using FKeyType = TVariant<FInvalid, int32, FString, FName>;
+		using FKeyType = TVariant<FNull, int32, FString, FName>;
 		FKeyType Key;
 		uint32 Hash;
 	};
 
-	/** FFrontendQueryEntry represents one value in the query. It contains a key,
-	 * value and score. 
-	 */
+	/** A FFrontendQueryValue contains data of interest. */
+	using FFrontendQueryValue = TVariant<FMetasoundFrontendVersion, Frontend::FNodeRegistryTransaction, FMetasoundFrontendClass, Frontend::FInterfaceRegistryTransaction, FMetasoundFrontendInterface>;
+
+	/** FFrontendQueryEntry represents one value in the query. It contains an ID,
+	 * value and score.  */
 	struct METASOUNDFRONTEND_API FFrontendQueryEntry
 	{
-		using FValue = TVariant<Frontend::FNodeRegistryTransaction, FMetasoundFrontendVersion, FMetasoundFrontendClass, Frontend::FInterfaceRegistryTransaction, FMetasoundFrontendInterface>;
-		using FKey = FFrontendQueryKey;
+		using FValue = FFrontendQueryValue;
 
-		FKey Key;
+		FGuid ID;
 		FValue Value;
 		float Score = 0.f;
 
-		FFrontendQueryEntry(const FFrontendQueryEntry&) = default;
-		FFrontendQueryEntry& operator=(const FFrontendQueryEntry&) = default;
-
-		FFrontendQueryEntry(FValue&& InValue)
-		:	Value(MoveTemp(InValue))
-		{
-		}
-
-		FFrontendQueryEntry(const FValue& InValue)
-		:	Value(InValue)
-		{
-		}
-	};
-
-	/** FFrontendQuerySelection represents a selection utilized during a query.
-	 *
-	 * Note: FFrontendQuerySelection owns the memory for the related entries. 
-	 */
-	class METASOUNDFRONTEND_API FFrontendQuerySelection
-	{
-	public:
-		/** FReduceOutputView provides an interface for storing results during
-		 * a "Reduce" query step.  It is primarily used to only expose operations
-		 * which are acceptable to call on a FFrontendQuerySelection during a 
-		 * call to "Reduce".
-		 */
-		class FReduceOutputView
-		{
-			// deleted constructors
-			FReduceOutputView(const FReduceOutputView&) = delete;
-			FReduceOutputView(FReduceOutputView&&) = delete;
-
-			// deleted operators
-			FReduceOutputView& operator=(const FReduceOutputView&) = delete;
-			FReduceOutputView& operator=(FReduceOutputView&&) = delete;
-
-		public:
-			// Set initial entries which are stored elsewhere
-			FReduceOutputView(TArrayView<FFrontendQueryEntry* const> InEntries);
-			~FReduceOutputView() = default;
-
-			/** Add a result to the selection */
-			void Add(FFrontendQueryEntry& InResult);
-
-
-
-		private:
-			friend class FFrontendQuerySelection;
-
-			TArray<FFrontendQueryEntry*> InitialEntries;
-
-			TArray<FFrontendQueryEntry*> SelectedExistingEntries;
-			TArray<TUniquePtr<FFrontendQueryEntry>> SelectedNewEntries;
-		};
-
-		FFrontendQuerySelection() = default;
-		FFrontendQuerySelection(const FFrontendQuerySelection&);
-		FFrontendQuerySelection& operator=(const FFrontendQuerySelection&);
-
-		/** Get an array of all entries selected. */
-		TArrayView<FFrontendQueryEntry*> GetSelection();
-
-		/** Get an array of all entries selected. */
-		TArrayView<const FFrontendQueryEntry * const> GetSelection() const;
-
-		/** Add entries to both storage and current selection. */
-		void AppendToStorageAndSelection(const FFrontendQuerySelection& InSelection);
-
-		/** Add entries to both storage and current selection. */
-		void AppendToStorageAndSelection(FFrontendQuerySelection&& InSelection);
-
-		/** Add entries to both storage and current selection. */
-		void AppendToStorageAndSelection(TArrayView<const FFrontendQueryEntry> InEntries);
-
-		/** Add entries to both storage and current selection. */
-		void AppendToStorageAndSelection(TArrayView<const FFrontendQueryEntry* const> InEntries);
-
-		/** Add entries to both storage and current selection. */
-		void AppendToStorageAndSelection(FReduceOutputView&& InReduceView);
-
-
-		/** Reset both storage and selection. */
-		void ResetStorageAndSelection();
-
-		/** Reset selection, but keep storage. */
-		void ResetSelection();
-
-		/** Filters the selection by calling Func on each entry. If Func returns
-		 * false, the entry is removed.
-		 */
-		template<typename FilterFuncType>
-		int32 FilterSelection(FilterFuncType Func)
-		{
-			auto InvertedFilter = [&](const FFrontendQueryEntry* InEntry) -> bool
-			{
-				return !Func(InEntry);
-			};
-			return Selection.RemoveAll(InvertedFilter);
-		}
-
-		/** Set selection. All entries must already be in storage. */
-		void SetSelection(TArrayView<FFrontendQueryEntry*> InEntries);
-
-		/** Add to selection. The entry must already be in storage. */
-		void AddToSelection(FFrontendQueryEntry* InEntry);
-
-		/** Append entries to selection. Each selection must already be in storage. */
-		void AppendToSelection(TArrayView<FFrontendQueryEntry * const> InEntries);
-
-	private:
-		void AppendToStorageAndSelection(TArray<TUniquePtr<FFrontendQueryEntry>>&& InEntries);
-		void ShrinkStorageToSelection();
-
-		TArray<FFrontendQueryEntry*> Selection;
-		TArray<TUniquePtr<FFrontendQueryEntry>> Storage;
-	};
-
-	/** FFrontendQuerySelectionView holds the result of query. */
-	class METASOUNDFRONTEND_API FFrontendQuerySelectionView
-	{
-		FFrontendQuerySelectionView() = delete;
-	public:
-
-		/** The FFrontendQuerySelectionView is constructed with an existing
-		 * FFrontendQuerySelection.
-		 */
-		FFrontendQuerySelectionView(TSharedRef<const FFrontendQuerySelection, ESPMode::ThreadSafe> InResult);
-
-		/** Get the selected results. */
-		TArrayView<const FFrontendQueryEntry* const> GetSelection() const;
-
-	private:
-		TSharedRef<const FFrontendQuerySelection, ESPMode::ThreadSafe> Result;
+		friend uint32 GetTypeHash(const FFrontendQueryEntry& InEntry);
+		friend bool operator==(const FFrontendQueryEntry& InLHS, const FFrontendQueryEntry& InRHS);
 	};
 
-	class IFrontendQuerySource
-	{
-		public:
-			virtual ~IFrontendQuerySource() = default;
+	/** A FFrontendQueryPartition represents a set of entries associated with a 
+	 * single FFrontendQueryKey. */
+	using FFrontendQueryPartition = TSet<FFrontendQueryEntry>;
 
-			virtual void Stream(TArray<FFrontendQueryEntry>& OutEntries) = 0;
-			virtual void Reset() = 0;
-	};
+	/** A FFrontendQuerySelection holds a map of keys to partitions. */
+	using FFrontendQuerySelection = TMap<FFrontendQueryKey, FFrontendQueryPartition>;
 
 	/** Interface for an individual step in a query */
 	class IFrontendQueryStep
@@ -209,7 +112,7 @@ namespace Metasound
 	{
 		public:
 			virtual ~IFrontendQueryStreamStep() = default;
-			virtual void Stream(TArray<FFrontendQueryEntry>& OutEntries) = 0;
+			virtual void Stream(TArray<FFrontendQueryValue>& OutEntries) = 0;
 	};
 
 	/** Interface for a query step which transforms an entry's value. */
@@ -225,17 +128,15 @@ namespace Metasound
 	{
 		public:
 			virtual ~IFrontendQueryMapStep() = default;
-			virtual FFrontendQueryEntry::FKey Map(const FFrontendQueryEntry& InEntry) const = 0;
+			virtual FFrontendQueryKey Map(const FFrontendQueryEntry& InEntry) const = 0;
 	};
 
 	/** Interface for a query step which reduces entries with the same key. */
 	class IFrontendQueryReduceStep : public IFrontendQueryStep
 	{
 		public:
-			using FReduceOutputView = FFrontendQuerySelection::FReduceOutputView;
-
 			virtual ~IFrontendQueryReduceStep() = default;
-			virtual void Reduce(FFrontendQueryEntry::FKey InKey, TArrayView<FFrontendQueryEntry * const>& InEntries, FReduceOutputView& OutResult) const = 0;
+			virtual void Reduce(const FFrontendQueryKey& InKey, FFrontendQueryPartition& InOutEntries) const = 0;
 	};
 
 	/** Interface for a query step which filters entries. */
@@ -277,44 +178,51 @@ namespace Metasound
 		FFrontendQueryStep() = delete;
 
 	public:
-		enum class EResultModificationState : uint8
+		// Represents an incremental update to the existing data.
+		struct FIncremental
 		{
-			Modified,
-			Unmodified,
+			// Keys that are affected by this incremental update.
+			TSet<FFrontendQueryKey> ActiveKeys;
+			// The selection being manipulated in the incremental update.
+			FFrontendQuerySelection ActiveSelection;
+
+			// Keys that contain active removals.
+			TSet<FFrontendQueryKey> ActiveRemovalKeys;
+			// Selection containing entries to remove during a merge.
+			FFrontendQuerySelection ActiveRemovalSelection;
 		};
 
 		/* Interface for executing a step in the query. */
 		struct IStepExecuter
 		{
-			using EResultModificationState = FFrontendQueryStep::EResultModificationState;
-
 			virtual ~IStepExecuter() = default;
 
-			// Perform an incremental step. Assume a previous result already exists. 
-			virtual EResultModificationState Increment(FFrontendQuerySelection& InOutResult) = 0;
+			// Merge new result with the existing result from this step.
+			virtual void Merge(FIncremental& InOutIncremental, FFrontendQuerySelection& InOutSelection) const = 0;
 
-			// Merge an incremental result with the prior result from this step.
-			virtual EResultModificationState Merge(const FFrontendQuerySelection& InIncremental, FFrontendQuerySelection& InOutResult) = 0;
+			// Execute step. 
+			virtual void Execute(TSet<FFrontendQueryKey>& InOutUpdatedKeys, FFrontendQuerySelection& InOutResult) const = 0;
 
-			// Execute step. Assume not other prior results exist.
-			virtual EResultModificationState Execute(FFrontendQuerySelection& InOutResult) = 0;
+			// Returns true if a steps result is conditioned on the composition of a partition.
+			//
+			// Most steps are only dependent upon individual entries, but some
+			// (Reduce, Limit, Sort) are specifically dependent upon the composition
+			// of the Partition. They require special handling during incremental
+			// updates.
+			virtual bool IsDependentOnPartitionComposition() const = 0;
 
-			// Reset internal state.
-			virtual void Reset() = 0;
+			// Return true if the step can be used to process downstream removals.
+			virtual bool CanProcessRemovals() const = 0;
 
-			// Returns true if a merge is required for a modified incremental update.
-			// 
-			// If this is true, Merge() will be called after Incremental() if Incremental
-			// returns EResultModificationState::Modified. 
-			virtual bool IsMergeRequiredForIncremental() const = 0;
+			// Return true if the step can produce new entries. This information
+			// is used to early-out on queries with no new entries. 
+			virtual bool CanProduceEntries() const = 0;
 		};
 
-		using FReduceOutputView = FFrontendQuerySelection::FReduceOutputView;
-
-		using FStreamFunction = TUniqueFunction<void (TArray<FFrontendQueryEntry>&)>;
+		using FStreamFunction = TUniqueFunction<void (TArray<FFrontendQueryValue>&)>;
 		using FTransformFunction = TFunction<void (FFrontendQueryEntry::FValue&)>;
-		using FMapFunction = TFunction<FFrontendQueryEntry::FKey (const FFrontendQueryEntry&)>;
-		using FReduceFunction = TFunction<void (FFrontendQueryEntry::FKey, TArrayView<FFrontendQueryEntry * const>, FReduceOutputView& )>;
+		using FMapFunction = TFunction<FFrontendQueryKey (const FFrontendQueryEntry&)>;
+		using FReduceFunction = TFunction<void (const FFrontendQueryKey& InKey, FFrontendQueryPartition& InOutEntries)>;
 		using FFilterFunction = TFunction<bool (const FFrontendQueryEntry&)>;
 		using FScoreFunction = TFunction<float (const FFrontendQueryEntry&)>;
 		using FSortFunction = TFunction<bool (const FFrontendQueryEntry& InEntryLHS, const FFrontendQueryEntry& InEntryRHS)>;
@@ -330,8 +238,7 @@ namespace Metasound
 		FFrontendQueryStep(FSortFunction&& InSort);
 		FFrontendQueryStep(FLimitFunction&& InLimit);
 
-		/** Create a query step using a IFrotnedQueryStep */
-		FFrontendQueryStep(TUniquePtr<IFrontendQuerySource>&& InSource);
+		/** Create a query step using a IFrontedQueryStep */
 		FFrontendQueryStep(TUniquePtr<IFrontendQueryStreamStep>&& InStep);
 		FFrontendQueryStep(TUniquePtr<IFrontendQueryTransformStep>&& InStep);
 		FFrontendQueryStep(TUniquePtr<IFrontendQueryMapStep>&& InStep);
@@ -342,37 +249,36 @@ namespace Metasound
 		FFrontendQueryStep(TUniquePtr<IFrontendQueryLimitStep>&& InStep);
 
 
-		// Perform an incremental step. Assume a previous result already exists. 
-		EResultModificationState Increment(FFrontendQuerySelection& InOutResult);
-
 		// Merge an incremental result with the prior result from this step.
-		EResultModificationState Merge(const FFrontendQuerySelection& InIncremental, FFrontendQuerySelection& InOutResult);
+		void Merge(FIncremental& InIncremental, FFrontendQuerySelection& InOutSelection) const;
 
 		// Execute step. Assume not other prior results exist.
-		EResultModificationState Execute(FFrontendQuerySelection& InOutResult);
+		void Execute(TSet<FFrontendQueryKey>& InOutUpdatedKeys, FFrontendQuerySelection& InOutResult) const;
 
-		// Returns true if a merge is required for a modified incremental update.
-		// 
-		// If this is true, Merge() will be called after Incremental() if Incremental
-		// returns EResultModificationState::Modified. 
-		bool IsMergeRequiredForIncremental() const;
+		// Returns true if a steps result is conditioned on the composition of a partition.
+		//
+		// Most steps are only dependent upon individual entries, but some
+		// (Reduce, Limit, Sort) are specifically dependent upon the composition
+		// of the Partition. They require special handling during incremental
+		// updates.
+		bool IsDependentOnPartitionComposition() const;
 
-		// Reset internal state.
-		void Reset();
+		// Return true if the step can be used to process downstream removals.
+		bool CanProcessRemovals() const;
 
-	private:
+		// Return true if the step can produce new entries. This information
+		// is used to early-out on queries with no new entries. 
+		bool CanProduceEntries() const;
 
+	private:
 		TUniquePtr<IStepExecuter> StepExecuter;
 	};
 
-	/** FFrontendQuery contains a set of query steps which can be executed 
-	 * to produce a FFrontendQuerySelectionView
-	 */
+	/** FFrontendQuery contains a set of query steps which produce a FFrontendQuerySelectionView */
 	class METASOUNDFRONTEND_API FFrontendQuery
 	{
 	public:
 
-		using EResultModificationState = FFrontendQueryStep::EResultModificationState;
 		using FStreamFunction = FFrontendQueryStep::FStreamFunction;
 		using FTransformFunction = FFrontendQueryStep::FTransformFunction;
 		using FMapFunction = FFrontendQueryStep::FMapFunction;
@@ -383,12 +289,12 @@ namespace Metasound
 		using FLimitFunction = FFrontendQueryStep::FLimitFunction;
 
 		FFrontendQuery();
+		FFrontendQuery(FFrontendQuery&&) = default;
+		FFrontendQuery& operator=(FFrontendQuery&&) = default;
 
 		FFrontendQuery(const FFrontendQuery&) = delete;
 		FFrontendQuery& operator=(const FFrontendQuery&) = delete;
 
-		/** Return all steps in a query. */
-		const TArray<TUniquePtr<FFrontendQueryStep>>& GetSteps() const;
 
 		/** Add a step to the query. */
 		template<typename StepType, typename... ArgTypes>
@@ -416,22 +322,31 @@ namespace Metasound
 		FFrontendQuery& AddStep(TUniquePtr<FFrontendQueryStep>&& InStep);
 
 		/** Calls all steps in the query and returns the selection. */
-		FFrontendQuerySelectionView Execute();
-
-		/** Resets the query result by removing all entries. */
-		FFrontendQuerySelectionView Reset();
+		const FFrontendQuerySelection& Update(TSet<FFrontendQueryKey>& OutUpdatedKeys);
+		const FFrontendQuerySelection& Update();
 
 		/** Returns the current result. */
-		FFrontendQuerySelectionView GetSelection() const;
+		const FFrontendQuerySelection& GetSelection() const;
 
 	private:
+		using FIncremental = FFrontendQueryStep::FIncremental;
 
-		void UpdateResult();
-		void ExecuteSteps(int32 InStartStepIndex);
+		void UpdateInternal(TSet<FFrontendQueryKey>& OutUpdatedKeys);
+		void MergeInternal(FFrontendQueryStep& Step, FIncremental& InOutIncremental, FFrontendQuerySelection& InOutMergedSelection);
 
 		TSharedRef<FFrontendQuerySelection, ESPMode::ThreadSafe> Result;
 
-		TArray<TUniquePtr<FFrontendQueryStep>> Steps;
-		TMap<int32, FFrontendQuerySelection> StepResultCache;
+		struct FStepInfo
+		{
+			TUniquePtr<FFrontendQueryStep> Step;
+			FFrontendQuerySelection OutputCache;
+			bool bMergeAndCacheOutput = false;
+			bool bProcessRemovals = false;
+		};
+
+		void AppendPartitions(const TSet<FFrontendQueryKey>& InKeysToAppend, const FFrontendQuerySelection& InSelection, TSet<FFrontendQueryKey>& OutKeysModified, FFrontendQuerySelection& OutSelection) const;
+
+		TArray<FStepInfo> Steps;
+		int32 FinalEntryProducingStepIndex = INDEX_NONE;
 	};
 }
diff --git a/Engine/Plugins/Runtime/Metasound/Source/MetasoundFrontend/Public/MetasoundFrontendQuerySteps.h b/Engine/Plugins/Runtime/Metasound/Source/MetasoundFrontend/Public/MetasoundFrontendQuerySteps.h
--- a/Engine/Plugins/Runtime/Metasound/Source/MetasoundFrontend/Public/MetasoundFrontendQuerySteps.h
+++ b/Engine/Plugins/Runtime/Metasound/Source/MetasoundFrontend/Public/MetasoundFrontendQuerySteps.h
@@ -9,14 +9,23 @@
 
 namespace Metasound
 {
+	class FMapToNull : public IFrontendQueryMapStep
+	{
+	public:
+		virtual FFrontendQueryKey Map(const FFrontendQueryEntry& InEntry) const override
+		{
+			return FFrontendQueryKey();
+		}
+	};
+
 	/** Streams node classes that have been newly registered or unregistered since last call to Stream()
 	 */
-	class METASOUNDFRONTEND_API FNodeClassRegistrationEvents : public IFrontendQuerySource
+	class METASOUNDFRONTEND_API FNodeClassRegistrationEvents : public IFrontendQueryStreamStep
 	{
 	public:
 		FNodeClassRegistrationEvents();
-		void Stream(TArray<FFrontendQueryEntry>& OutEntries) override;
-		void Reset() override;
+		virtual void Stream(TArray<FFrontendQueryValue>& OutValues) override;
+		//void Reset() override;
 
 	private:
 		Frontend::FRegistryTransactionID CurrentTransactionID;
@@ -26,7 +35,7 @@ namespace Metasound
 	class METASOUNDFRONTEND_API FMapRegistrationEventsToNodeRegistryKeys : public IFrontendQueryMapStep
 	{
 	public:
-		FFrontendQueryEntry::FKey Map(const FFrontendQueryEntry& InEntry) const override;
+		virtual FFrontendQueryKey Map(const FFrontendQueryEntry& InEntry) const override;
 	};
 
 	/** Reduces registration events mapped to the same key by inspecting their add/remove state in
@@ -35,7 +44,7 @@ namespace Metasound
 	class METASOUNDFRONTEND_API FReduceRegistrationEventsToCurrentStatus: public IFrontendQueryReduceStep
 	{
 	public:
-		void Reduce(FFrontendQueryEntry::FKey InKey, TArrayView<FFrontendQueryEntry * const>& InEntries, FReduceOutputView& OutResult) const override;
+		virtual void Reduce(const FFrontendQueryKey& InKey, FFrontendQueryPartition& InOutEntries) const override;
 	};
 
 	/** Transforms a registration event into a FMetasoundFrontendClass. */
@@ -56,7 +65,7 @@ namespace Metasound
 
 		FFilterClassesByInputVertexDataType(const FName& InTypeName);
 
-		bool Filter(const FFrontendQueryEntry& InEntry) const override;
+		virtual bool Filter(const FFrontendQueryEntry& InEntry) const override;
 
 	private:
 		FName InputVertexTypeName;
@@ -73,21 +82,17 @@ namespace Metasound
 
 		FFilterClassesByOutputVertexDataType(const FName& InTypeName);
 
-		bool Filter(const FFrontendQueryEntry& InEntry) const override;
+		virtual bool Filter(const FFrontendQueryEntry& InEntry) const override;
 
 	private:
 		FName OutputVertexTypeName;
 	};
 
-	class METASOUNDFRONTEND_API FFilterClassesByClassName : public IFrontendQueryFilterStep
+	class METASOUNDFRONTEND_API FMapClassesToClassName : public IFrontendQueryMapStep
 	{
 	public: 
-		FFilterClassesByClassName(const FMetasoundFrontendClassName& InClassName);
+		virtual FFrontendQueryKey Map(const FFrontendQueryEntry& InEntry) const override;
 
-		bool Filter(const FFrontendQueryEntry& InEntry) const override;
-
-	private:
-		FMetasoundFrontendClassName ClassName;
 	};
 
 	class METASOUNDFRONTEND_API FFilterClassesByClassID : public IFrontendQueryFilterStep
@@ -95,7 +100,7 @@ namespace Metasound
 	public:
 		FFilterClassesByClassID(const FGuid InClassID);
 
-		bool Filter(const FFrontendQueryEntry& InEntry) const override;
+		virtual bool Filter(const FFrontendQueryEntry& InEntry) const override;
 
 	private:
 		FGuid ClassID;
@@ -104,29 +109,18 @@ namespace Metasound
 	class METASOUNDFRONTEND_API FMapToFullClassName : public IFrontendQueryMapStep
 	{
 	public:
-		FFrontendQueryEntry::FKey Map(const FFrontendQueryEntry& InEntry) const override;
+		virtual FFrontendQueryKey Map(const FFrontendQueryEntry& InEntry) const override;
 	};
 
 	class METASOUNDFRONTEND_API FReduceClassesToHighestVersion : public IFrontendQueryReduceStep
 	{
 	public:
-		void Reduce(FFrontendQueryEntry::FKey InKey, TArrayView<FFrontendQueryEntry * const>& InEntries, FReduceOutputView& OutResult) const override;
-	};
-
-	class METASOUNDFRONTEND_API FReduceClassesToMajorVersion : public IFrontendQueryReduceStep
-	{
-	public:
-		FReduceClassesToMajorVersion(int32 InMajorVersion);
-
-		void Reduce(FFrontendQueryEntry::FKey InKey, TArrayView<FFrontendQueryEntry * const>& InEntries, FReduceOutputView& OutResult) const override;
-
-	private:
-		int32 MajorVersion = -1;
+		virtual void Reduce(const FFrontendQueryKey& InKey, FFrontendQueryPartition& InOutEntries) const override;
 	};
 
 	class METASOUNDFRONTEND_API FSortClassesByVersion : public IFrontendQuerySortStep
 	{
 	public:
-		bool Sort(const FFrontendQueryEntry& InEntryLHS, const FFrontendQueryEntry& InEntryRHS) const override;
+		virtual bool Sort(const FFrontendQueryEntry& InEntryLHS, const FFrontendQueryEntry& InEntryRHS) const override;
 	};
 }
diff --git a/Engine/Plugins/Runtime/Metasound/Source/MetasoundFrontend/Public/MetasoundFrontendSearchEngine.h b/Engine/Plugins/Runtime/Metasound/Source/MetasoundFrontend/Public/MetasoundFrontendSearchEngine.h
--- a/Engine/Plugins/Runtime/Metasound/Source/MetasoundFrontend/Public/MetasoundFrontendSearchEngine.h
+++ b/Engine/Plugins/Runtime/Metasound/Source/MetasoundFrontend/Public/MetasoundFrontendSearchEngine.h
@@ -28,13 +28,13 @@ namespace Metasound
 				/** Find all classes with the given ClassName.
 				  * (Optional) Sort matches based on version.
 				  */
-				virtual TArray<FMetasoundFrontendClass> FindClassesWithName(const FNodeClassName& InName, bool bInSortByVersion) = 0;
+				virtual TArray<FMetasoundFrontendClass> FindClassesWithName(const FMetasoundFrontendClassName& InName, bool bInSortByVersion) = 0;
 
 				/** Find the highest version of a class with the given ClassName. Returns false if not found, true if found. */
-				virtual bool FindClassWithHighestVersion(const FNodeClassName& InName, FMetasoundFrontendClass& OutClass) = 0;
+				virtual bool FindClassWithHighestVersion(const FMetasoundFrontendClassName& InName, FMetasoundFrontendClass& OutClass) = 0;
 
 				/** Find the class with the given ClassName & Major Version. Returns false if not found, true if found. */
-				virtual bool FindClassWithMajorVersion(const FNodeClassName& InName, int32 InMajorVersion, FMetasoundFrontendClass& OutClass) = 0;
+				virtual bool FindClassWithMajorVersion(const FMetasoundFrontendClassName& InName, int32 InMajorVersion, FMetasoundFrontendClass& OutClass) = 0;
 
 				virtual TArray<FMetasoundFrontendInterface> FindAllInterfaces(bool bInIncludeDeprecated = false) = 0;
 
