Keep metasound search engine warm
- Adds a "Prime" call to the MetasoundSearchEngine so we can precache a bunch of queries.
- Adds more metasound traces
- Encapsulates queries into structs to avoid a bunch of boilerplate code everywhere.

diff --git a/Engine/Plugins/Runtime/Metasound/Source/MetasoundFrontend/Private/MetasoundFrontendRegistries.cpp b/Engine/Plugins/Runtime/Metasound/Source/MetasoundFrontend/Private/MetasoundFrontendRegistries.cpp
--- a/Engine/Plugins/Runtime/Metasound/Source/MetasoundFrontend/Private/MetasoundFrontendRegistries.cpp
+++ b/Engine/Plugins/Runtime/Metasound/Source/MetasoundFrontend/Private/MetasoundFrontendRegistries.cpp
@@ -6,16 +6,16 @@
 #include "CoreMinimal.h"
 #include "MetasoundFrontendDataTypeRegistry.h"
 #include "MetasoundFrontendRegistryTransaction.h"
+#include "MetasoundFrontendSearchEngine.h"
 #include "MetasoundLog.h"
 #include "MetasoundRouter.h"
+#include "MetasoundTrace.h"
 #include "Misc/ScopeLock.h"
-#include "HAL/PlatformTime.h"
 
 #ifndef WITH_METASOUND_FRONTEND
 #define WITH_METASOUND_FRONTEND 0
 #endif
 
-
 namespace Metasound
 {
 	namespace Frontend
@@ -70,8 +70,6 @@ namespace Metasound
 				}
 			}
 
-
-
 			// Registry container private implementation.
 			class FRegistryContainerImpl : public FMetasoundFrontendRegistryContainer
 			{
@@ -160,24 +158,25 @@ namespace Metasound
 
 			void FRegistryContainerImpl::RegisterPendingNodes()
 			{
-				FScopeLock ScopeLock(&LazyInitCommandCritSection);
+				METASOUND_TRACE_CPUPROFILER_EVENT_SCOPE(metasound::FRegistryContainerImpl::RegisterPendingNodes);
+				{
+					FScopeLock ScopeLock(&LazyInitCommandCritSection);
 
-				UE_LOG(LogMetaSound, Display, TEXT("Processing %i Metasounds Frontend Registration Requests."), LazyInitCommands.Num());
-				uint64 CurrentTime = FPlatformTime::Cycles64();
+					for (TUniqueFunction<void()>& Command : LazyInitCommands)
+					{
+						Command();
+					}
 
-				for (TUniqueFunction<void()>& Command : LazyInitCommands)
-				{
-					Command();
+					LazyInitCommands.Empty();
 				}
 
-				LazyInitCommands.Empty();
-
-				uint64 CyclesUsed = FPlatformTime::Cycles64() - CurrentTime;
-				UE_LOG(LogMetaSound, Display, TEXT("Initializing Metasounds Frontend took %f seconds."), FPlatformTime::ToSeconds64(CyclesUsed));
+				// Prime search engine after bulk registration.
+				ISearchEngine::Get().Prime();
 			}
 
 			bool FRegistryContainerImpl::EnqueueInitCommand(TUniqueFunction<void()>&& InFunc)
 			{
+
 				FScopeLock ScopeLock(&LazyInitCommandCritSection);
 				if (LazyInitCommands.Num() >= MaxNumNodesAndDatatypesToInitialize)
 				{
diff --git a/Engine/Plugins/Runtime/Metasound/Source/MetasoundFrontend/Private/MetasoundFrontendSearchEngine.cpp b/Engine/Plugins/Runtime/Metasound/Source/MetasoundFrontend/Private/MetasoundFrontendSearchEngine.cpp
--- a/Engine/Plugins/Runtime/Metasound/Source/MetasoundFrontend/Private/MetasoundFrontendSearchEngine.cpp
+++ b/Engine/Plugins/Runtime/Metasound/Source/MetasoundFrontend/Private/MetasoundFrontendSearchEngine.cpp
@@ -41,12 +41,10 @@ namespace Metasound
 						const FMetasoundFrontendClassMetadata& Metadata = InEntry.Value.Get<FMetasoundFrontendClass>().Metadata;
 						return FMapClassToFullClassNameAndMajorVersion::GetKey(Metadata.GetClassName(), Metadata.GetVersion().Major);
 					}
-
 					return FFrontendQueryKey();
 				}
 			};
 
-
 			class FInterfaceRegistryTransactionSource : public IFrontendQueryStreamStep
 			{
 			public:
@@ -141,48 +139,6 @@ namespace Metasound
 				}
 			};
 
-			class FFilterInterfaceRegistryTransactionsBySupportedClassName : public IFrontendQueryFilterStep
-			{
-				FName UClassName;
-
-			public:
-				FFilterInterfaceRegistryTransactionsBySupportedClassName(FName InUClassName)
-					: UClassName(InUClassName)
-				{
-				}
-
-				virtual bool Filter(const FFrontendQueryEntry& InEntry) const override
-				{
-					if (ensure(InEntry.Value.IsType<FInterfaceRegistryTransaction>()))
-					{
-						FInterfaceRegistryKey RegistryKey = InEntry.Value.Get<FInterfaceRegistryTransaction>().GetInterfaceRegistryKey();
-						const IInterfaceRegistryEntry* RegEntry = IInterfaceRegistry::Get().FindInterfaceRegistryEntry(RegistryKey);
-						if (ensure(RegEntry))
-						{
-							return RegEntry->UClassIsSupported(UClassName);
-						}
-					}
-					return false;
-				}
-			};
-
-			class FFilterInterfaceRegistryTransactionsByDefault : public IFrontendQueryFilterStep
-			{
-			public:
-				virtual bool Filter(const FFrontendQueryEntry& InEntry) const override
-				{
-					if (ensure(InEntry.Value.IsType<FInterfaceRegistryTransaction>()))
-					{
-						FInterfaceRegistryKey RegistryKey = InEntry.Value.Get<FInterfaceRegistryTransaction>().GetInterfaceRegistryKey();
-						const IInterfaceRegistryEntry* RegEntry = IInterfaceRegistry::Get().FindInterfaceRegistryEntry(RegistryKey);
-						if (ensure(RegEntry))
-						{
-							return RegEntry->IsDefault();
-						}
-					}
-					return false;
-				}
-			};
 
 			class FTransformInterfaceRegistryTransactionToInterface : public IFrontendQueryTransformStep
 			{
@@ -254,73 +210,6 @@ namespace Metasound
 				}
 			};
 
-			class FFilterInterfaceRegistryTransactionsByName : public IFrontendQueryFilterStep
-			{
-			public:
-				FFilterInterfaceRegistryTransactionsByName(FName InName)
-				: Name(InName)
-				{
-				}
-
-				virtual bool Filter(const FFrontendQueryEntry& InEntry) const override
-				{
-					if (ensure(InEntry.Value.IsType<FInterfaceRegistryTransaction>()))
-					{
-						return Name == InEntry.Value.Get<FInterfaceRegistryTransaction>().GetInterfaceVersion().Name;
-					}
-					return false;
-				}
-
-			private:
-				FName Name;
-			};
-
-
-			class FFilterNodeRegistrationEventsByClassName : public IFrontendQueryFilterStep
-			{
-			public:
-				FFilterNodeRegistrationEventsByClassName(const FMetasoundFrontendClassName& InName)
-				: Name(InName)
-				{
-				}
-
-				virtual bool Filter(const FFrontendQueryEntry& InEntry) const override
-				{
-					if (ensure(InEntry.Value.IsType<FNodeRegistryTransaction>()))
-					{
-						return Name == InEntry.Value.Get<FNodeRegistryTransaction>().GetNodeClassInfo().ClassName;
-					}
-					return false;
-				}
-			private:
-				FMetasoundFrontendClassName Name;
-			};
-
-			class FFilterInterfaceRegistryTransactionsByNameAndMajorVersion : public IFrontendQueryFilterStep
-			{
-			public:
-				FFilterInterfaceRegistryTransactionsByNameAndMajorVersion(const FName& InName, int32 InMajorVersion)
-				: Name(InName)
-				, MajorVersion(InMajorVersion)
-				{
-				}
-
-				virtual bool Filter(const FFrontendQueryEntry& InEntry) const override
-				{
-					if (ensure(InEntry.Value.IsType<FInterfaceRegistryTransaction>()))
-					{
-						const FMetasoundFrontendVersion& Version = InEntry.Value.Get<FInterfaceRegistryTransaction>().GetInterfaceVersion();
-
-						return (Name == Version.Name) && (MajorVersion == Version.Number.Major);
-					}
-					return false;
-				}
-
-			private:
-				FName Name;
-				int32 MajorVersion;
-			};
-
 			class FMapInterfaceToInterfaceName : public IFrontendQueryMapStep
 			{
 			public:
@@ -350,423 +239,585 @@ namespace Metasound
 			};
 		}
 
-		class FSearchEngine : public ISearchEngine
+		// Base interface for a search engine query. This allows the query to be
+		// primed to improve execution performance in critical call-stacks.
+		struct ISearchEngineQuery
 		{
-				FSearchEngine(const FSearchEngine&) = delete;
-				FSearchEngine(FSearchEngine&&) = delete;
-				FSearchEngine& operator=(const FSearchEngine&) = delete;
-				FSearchEngine& operator=(FSearchEngine&&) = delete;
-
-			public:
-				FSearchEngine() = default;
-				virtual ~FSearchEngine() = default;
-
-				virtual TArray<FMetasoundFrontendClass> FindAllClasses(bool bInIncludeDeprecated) override;
-				virtual TArray<FMetasoundFrontendClass> FindClassesWithName(const FMetasoundFrontendClassName& InName, bool bInSortByVersion) override;
-				virtual bool FindClassWithHighestVersion(const FMetasoundFrontendClassName& InName, FMetasoundFrontendClass& OutClass) override;
-				virtual bool FindClassWithMajorVersion(const FMetasoundFrontendClassName& InName, int32 InMajorVersion, FMetasoundFrontendClass& OutClass) override;
-
-				virtual TArray<FMetasoundFrontendInterface> FindAllInterfaces(bool bInIncludeDeprecated) override;
-				virtual TArray<FMetasoundFrontendInterface> FindUClassDefaultInterfaces(FName InUClassName) override;
-				virtual TArray<FMetasoundFrontendVersion> FindAllRegisteredInterfacesWithName(FName InInterfaceName) override;
-				virtual bool FindInterfaceWithHighestVersion(FName InInterfaceName, FMetasoundFrontendInterface& OutInterface) override;
-				virtual bool FindInterfaceWithMajorVersion(FName InInterfaceName, int32 InMajorVersion, FMetasoundFrontendInterface& OutInterface) override;
-
-			private:
-				FFrontendQuery* FindQuery(const FString& InQueryName);
-				FFrontendQuery* AddQuery(const FString& InQueryName, TUniquePtr<FFrontendQuery>&& InQuery);
-
-				// Store of queries indexed by unique strings.
-				TMap<FString, TUniquePtr<FFrontendQuery>> QueryCache;
+			// Prime the query to improve execution performance during subsequent 
+			// queries. 
+			virtual void Prime() = 0;
+			virtual ~ISearchEngineQuery() = default;
 		};
 
-		ISearchEngine& ISearchEngine::Get()
-		{
-			static FSearchEngine SearchEngine;
-			return SearchEngine;
-		}
-
-		TArray<FMetasoundFrontendClass> FSearchEngine::FindAllClasses(bool bInIncludeDeprecated)
+		// TSearchEngineQuery provides a base for all queries used in the search
+		// engine. It requires a `QueryPolicy` as a template parameter. 
+		//
+		// A `QueryPolicy` is a struct containing a static function to create the 
+		// query, a static function to build a result from the query selection, 
+		// and query result type.
+		//
+		// Example:
+		//
+		// struct FQueryPolicyExample
+		// {
+		//		using ResultType = FMetasoundFrontendClass;
+		//
+		//		static FFrontendQuery CreateQuery()
+		//		{
+		//			using namespace SearchEngineQuerySteps;
+		//			FFrontendQuery Query;
+		//			Query.AddStep<FNodeClassRegistrationEvents>()
+		//				.AddStep<FMapRegistrationEventsToNodeRegistryKeys>()
+		//				.AddStep<FReduceRegistrationEventsToCurrentStatus>()
+		//				.AddStep<FTransformRegistrationEventsToClasses>();
+		//			return Query;
+		//		}
+		//
+		//		static ResultType BuildResult(const FFrontendQueryPartition& InPartition)
+		//		{
+		//			FMetasoundFrontendClass MetaSoundClass;
+		//			if (ensureMsgf(InPartition.Num() > 0, TEXT("Cannot retrieve class from empty partition")))
+		//			{
+		//				MetaSoundClass = InPartition.CreateConstIterator()->Value.Get<FMetasoundFrontendClass>();
+		//			}
+		//			return MetaSoundClass;
+		//		}
+		// };
+		template<typename QueryPolicy>
+		struct TSearchEngineQuery : public ISearchEngineQuery
 		{
-			METASOUND_TRACE_CPUPROFILER_EVENT_SCOPE(metasound::FSearchEngine::FindAllClasses);
-			using namespace SearchEngineQuerySteps;
+			using ResultType = typename QueryPolicy::ResultType;
 
-			FString QueryName = FString(TEXT("FindAllClasses"));
-			if (bInIncludeDeprecated)
+			TSearchEngineQuery()
+			: Query(QueryPolicy::CreateQuery())
 			{
-				QueryName += TEXT("IncludingDeprecated");
 			}
 
-			FFrontendQuery* Query = FindQuery(QueryName);
-			if (!Query)
+			virtual void Prime() override
 			{
-				TUniquePtr<FFrontendQuery> NewQuery = MakeUnique<FFrontendQuery>();
+				FScopeLock Lock(&QueryCriticalSection);
+				Update();
+			}
 
-				NewQuery->AddStep<FNodeClassRegistrationEvents>()
-					.AddStep<FMapRegistrationEventsToNodeRegistryKeys>()
-					.AddStep<FReduceRegistrationEventsToCurrentStatus>()
-					.AddStep<FTransformRegistrationEventsToClasses>();
+			// Updates the query to the most recent value and returns the
+			// result.
+			ResultType UpdateAndFindResult(const FFrontendQueryKey& InKey)
+			{
+				FScopeLock Lock(&QueryCriticalSection);
 
-				if (!bInIncludeDeprecated)
+				Update();
+				if (const ResultType* Result = ResultCache.Find(InKey))
 				{
-					NewQuery->AddStep<FMapToFullClassName>()
-						.AddStep<FReduceClassesToHighestVersion>();
+					return *Result;
 				}
+				return ResultType();
+			}
 
-				NewQuery->AddStep<FMapToNull>();
+			// Updates the query to the most recent value and assigns the value 
+			// to OutResult. Returns true on success, false on failure. 
+			bool UpdateAndFindResult(const FFrontendQueryKey& InKey, ResultType& OutResult)
+			{
+				FScopeLock Lock(&QueryCriticalSection);
 
-				Query = AddQuery(QueryName, MoveTemp(NewQuery));
+				Update();
+				if (const ResultType* Result = ResultCache.Find(InKey))
+				{
+					OutResult = *Result;
+					return true;
+				}
+				return false;
 			}
 
-			TArray<FMetasoundFrontendClass> Result;
-
-			if (ensure(Query))
+		private:
+			void Update()
 			{
-				// TODO: cache result and only rebuild if query actually updates. 
-				const FFrontendQuerySelection& Selection = Query->Update();
-				if (const FFrontendQueryPartition* Entries = Selection.Find(FFrontendQueryKey()))
+				METASOUND_TRACE_CPUPROFILER_EVENT_SCOPE(MetaSound::SearchEngineQuery::Update);
+
+				TSet<FFrontendQueryKey> UpdatedKeys;
+				const FFrontendQuerySelection& Selection = Query.Update(UpdatedKeys);
+
+				if (UpdatedKeys.Num() > 0)
 				{
-					for (const FFrontendQueryEntry& Entry : *Entries)
+					METASOUND_TRACE_CPUPROFILER_EVENT_SCOPE(MetaSound::SearchEngineQuery::BuildResult);
+
+					for (const FFrontendQueryKey& Key : UpdatedKeys)
 					{
-						check(Entry.Value.IsType<FMetasoundFrontendClass>());
-						Result.Add(Entry.Value.Get<FMetasoundFrontendClass>());
+						const FFrontendQueryPartition* Partition = Query.GetSelection().Find(Key);
+						if (nullptr != Partition)
+						{
+							ResultCache.Add(Key, QueryPolicy::BuildResult(*Partition));
+						}
+						else
+						{
+							ResultCache.Remove(Key);
+						}
 					}
 				}
 			}
 
+			FCriticalSection QueryCriticalSection;
+			FFrontendQuery Query;
+			TMap<FFrontendQueryKey, ResultType> ResultCache;
+		};
+
+		TArray<FMetasoundFrontendClass> BuildArrayOfClassesFromPartition(const FFrontendQueryPartition& InPartition)
+		{
+			TArray<FMetasoundFrontendClass> Result;
+			for (const FFrontendQueryEntry& Entry : InPartition)
+			{
+				check(Entry.Value.IsType<FMetasoundFrontendClass>());
+				Result.Add(Entry.Value.Get<FMetasoundFrontendClass>());
+			}
 			return Result;
 		}
 
-		TArray<FMetasoundFrontendClass> FSearchEngine::FindClassesWithName(const FMetasoundFrontendClassName& InName, bool bInSortByVersion)
+		FMetasoundFrontendClass BuildSingleClassFromPartition(const FFrontendQueryPartition& InPartition)
 		{
-			METASOUND_TRACE_CPUPROFILER_EVENT_SCOPE(metasound::FSearchEngine::FindClassesWithName);
-			using namespace SearchEngineQuerySteps;
+			FMetasoundFrontendClass MetaSoundClass;
+			if (ensureMsgf(InPartition.Num() > 0, TEXT("Cannot retrieve class from empty partition")))
+			{
+				MetaSoundClass = InPartition.CreateConstIterator()->Value.Get<FMetasoundFrontendClass>();
+			}
+			return MetaSoundClass;
+		}
 
-			FString QueryName = FString(TEXT("FindClassesWithName"));
-			if (bInSortByVersion)
+		TArray<FMetasoundFrontendInterface> BuildArrayOfInterfacesFromPartition(const FFrontendQueryPartition& InPartition)
+		{
+			TArray<FMetasoundFrontendInterface> Result;
+			for (const FFrontendQueryEntry& Entry : InPartition)
+			{
+				check(Entry.Value.IsType<FMetasoundFrontendInterface>());
+				Result.Add(Entry.Value.Get<FMetasoundFrontendInterface>());
+			}
+			return Result;
+		}
+
+		FMetasoundFrontendInterface BuildSingleInterfaceFromPartition(const FFrontendQueryPartition& InPartition)
+		{
+			FMetasoundFrontendInterface MetaSoundInterface;
+			if (ensureMsgf(InPartition.Num() > 0, TEXT("Cannot retrieve class from empty partition")))
 			{
-				QueryName += TEXT("_SortedByVersion");
+				MetaSoundInterface = InPartition.CreateConstIterator()->Value.Get<FMetasoundFrontendInterface>();
 			}
-			FFrontendQuery* Query = FindQuery(QueryName);
+			return MetaSoundInterface;
+		}
 
-			if (!Query)
+		TArray<FMetasoundFrontendVersion> BuildArrayOfVersionsFromPartition(const FFrontendQueryPartition& InPartition)
+		{
+			TArray<FMetasoundFrontendVersion> Result;
+			for (const FFrontendQueryEntry& Entry : InPartition)
 			{
-				TUniquePtr<FFrontendQuery> NewQuery = MakeUnique<FFrontendQuery>();
-				NewQuery->AddStep<FNodeClassRegistrationEvents>()
+				check(Entry.Value.IsType<FMetasoundFrontendVersion>());
+				Result.Add(Entry.Value.Get<FMetasoundFrontendVersion>());
+			}
+			return Result;
+		}
+
+		// Policy for finding all registered metasound classes. 
+		struct FFindAllClassesQueryPolicy
+		{
+			using ResultType = TArray<FMetasoundFrontendClass>;
+
+			static FFrontendQuery CreateQuery()
+			{
+				using namespace SearchEngineQuerySteps;
+				FFrontendQuery Query;
+				Query.AddStep<FNodeClassRegistrationEvents>()
 					.AddStep<FMapRegistrationEventsToNodeRegistryKeys>()
 					.AddStep<FReduceRegistrationEventsToCurrentStatus>()
 					.AddStep<FTransformRegistrationEventsToClasses>()
-					.AddStep<FMapClassesToClassName>();
-
-				if (bInSortByVersion)
-				{
-					NewQuery->AddStep<FSortClassesByVersion>();
-				}
+					.AddStep<FMapToNull>();
+				return Query;
+			}
 
-				Query = AddQuery(QueryName, MoveTemp(NewQuery));
+			static ResultType BuildResult(const FFrontendQueryPartition& InPartition)
+			{
+				return BuildArrayOfClassesFromPartition(InPartition);
 			}
+		};
 
-			TArray<FMetasoundFrontendClass> Result;
+		// Policy for finding all registered metasound classes, including deprecated classes. 
+		struct FFindAllClassesIncludingDeprecatedQueryPolicy
+		{
+			using ResultType = TArray<FMetasoundFrontendClass>;
 
-			if (ensure(nullptr != Query))
+			static FFrontendQuery CreateQuery()
 			{
-				const FFrontendQuerySelection& Selection = Query->Update();
-				const FFrontendQueryKey Key(InName.GetFullName());
-				if (const FFrontendQueryPartition* Entries = Selection.Find(Key))
-				{
-					for (const FFrontendQueryEntry& Entry : *Entries)
-					{
-						check(Entry.Value.IsType<FMetasoundFrontendClass>());
-						Result.Add(Entry.Value.Get<FMetasoundFrontendClass>());
-					}
-				}
+				using namespace SearchEngineQuerySteps;
+				FFrontendQuery Query;
+				Query.AddStep<FNodeClassRegistrationEvents>()
+					.AddStep<FMapRegistrationEventsToNodeRegistryKeys>()
+					.AddStep<FReduceRegistrationEventsToCurrentStatus>()
+					.AddStep<FTransformRegistrationEventsToClasses>()
+					.AddStep<FMapToNull>();
+				return Query;
 			}
 
-			return Result;
-		}
+			static ResultType BuildResult(const FFrontendQueryPartition& InPartition)
+			{
+				return BuildArrayOfClassesFromPartition(InPartition);
+			}
+		};
 
-		bool FSearchEngine::FindClassWithHighestVersion(const FMetasoundFrontendClassName& InName, FMetasoundFrontendClass& OutClass)
+		// Policy for finding all registered metasound classes sorted by version 
+		// and indexed by name.
+		struct FFindClassesWithNameSortedQueryPolicy
 		{
-			METASOUND_TRACE_CPUPROFILER_EVENT_SCOPE(metasound::FSearchEngine::FindClassWithHighestVersion);
-			using namespace SearchEngineQuerySteps;
+			using ResultType = TArray<FMetasoundFrontendClass>;
 
-			const FString QueryName = FString(TEXT("FindHighestVersion"));
-			FFrontendQuery* Query = FindQuery(QueryName);
-			if (!Query)
+			static FFrontendQuery CreateQuery()
 			{
-				TUniquePtr<FFrontendQuery> NewQuery = MakeUnique<FFrontendQuery>();
-
-				NewQuery->AddStep<FNodeClassRegistrationEvents>()
+				using namespace SearchEngineQuerySteps;
+				FFrontendQuery Query;
+				Query.AddStep<FNodeClassRegistrationEvents>()
 					.AddStep<FMapRegistrationEventsToNodeRegistryKeys>()
 					.AddStep<FReduceRegistrationEventsToCurrentStatus>()
 					.AddStep<FTransformRegistrationEventsToClasses>()
-					.AddStep<FMapToFullClassName>()
-					.AddStep<FReduceClassesToHighestVersion>();
-
-				Query = AddQuery(QueryName, MoveTemp(NewQuery));
+					.AddStep<FMapClassesToClassName>()
+					.AddStep<FSortClassesByVersion>();
+				return Query;
 			}
 
-			if (ensure(Query))
+			static ResultType BuildResult(const FFrontendQueryPartition& InPartition)
 			{
-				const FFrontendQuerySelection& Selection = Query->Update();
-				const FFrontendQueryKey Key(InName.GetFullName());
-
-				if (const FFrontendQueryPartition* Entries = Selection.Find(Key))
-				{
-					if (ensure(Entries->Num() == 1))
-					{
-						OutClass = Entries->CreateConstIterator()->Value.Get<FMetasoundFrontendClass>();
-						return true;
-					}
-				}
+				return BuildArrayOfClassesFromPartition(InPartition);
 			}
+		};
 
-			return false;
-		}
-
-		bool FSearchEngine::FindClassWithMajorVersion(const FMetasoundFrontendClassName& InName, int32 InMajorVersion, FMetasoundFrontendClass& OutClass)
+		// Policy for finding all registered metasound classes indexed by name.
+		struct FFindClassesWithNameUnsortedQueryPolicy
 		{
-			METASOUND_TRACE_CPUPROFILER_EVENT_SCOPE(metasound::FSearchEngine::FindClassWithMajorVersion);
-			using namespace SearchEngineQuerySteps;
+			using ResultType = TArray<FMetasoundFrontendClass>;
 
-			const FString QueryName = FString(TEXT("FindClassWithMajorVersion"));
-			FFrontendQuery* Query = FindQuery(QueryName);
+			static FFrontendQuery CreateQuery()
+			{
+				using namespace SearchEngineQuerySteps;
+				FFrontendQuery Query;
+				Query.AddStep<FNodeClassRegistrationEvents>()
+					.AddStep<FMapRegistrationEventsToNodeRegistryKeys>()
+					.AddStep<FReduceRegistrationEventsToCurrentStatus>()
+					.AddStep<FTransformRegistrationEventsToClasses>()
+					.AddStep<FMapClassesToClassName>();
+				return Query;
+			}
 
-			if (!Query)
+			static ResultType BuildResult(const FFrontendQueryPartition& InPartition)
 			{
-				TUniquePtr<FFrontendQuery> NewQuery = MakeUnique<FFrontendQuery>();
+				return BuildArrayOfClassesFromPartition(InPartition);
+			}
+		};
+
+		// Policy for finding highest versioned metasound class by name.
+		struct FFindClassWithHighestVersionQueryPolicy
+		{
+			using ResultType = FMetasoundFrontendClass;
 
-				NewQuery->AddStep<FNodeClassRegistrationEvents>()
+			static FFrontendQuery CreateQuery()
+			{
+				using namespace SearchEngineQuerySteps;
+				FFrontendQuery Query;
+				Query.AddStep<FNodeClassRegistrationEvents>()
 					.AddStep<FMapRegistrationEventsToNodeRegistryKeys>()
 					.AddStep<FReduceRegistrationEventsToCurrentStatus>()
 					.AddStep<FTransformRegistrationEventsToClasses>()
 					.AddStep<FMapToFullClassName>()
-					.AddStep<FMapClassToFullClassNameAndMajorVersion>();
+					.AddStep<FReduceClassesToHighestVersion>();
+				return Query;
+			}
 
-				Query = AddQuery(QueryName, MoveTemp(NewQuery));
+			static ResultType BuildResult(const FFrontendQueryPartition& InPartition)
+			{
+				return BuildSingleClassFromPartition(InPartition);
 			}
+		};
 
+		// Policy for finding highest versioned metasound class by name and major version.
+		struct FFindClassWithMajorVersionQueryPolicy
+		{
+			using ResultType = FMetasoundFrontendClass;
 
-			if (ensure(Query))
+			static FFrontendQuery CreateQuery()
 			{
-				const FFrontendQueryKey Key = FMapClassToFullClassNameAndMajorVersion::GetKey(InName, InMajorVersion);
-				const FFrontendQuerySelection& Selection = Query->Update();
-
-				if (const FFrontendQueryPartition* Entries = Selection.Find(Key))
-				{
-					if (ensure(Entries->Num() == 1))
-					{
-						OutClass = Entries->CreateConstIterator()->Value.Get<FMetasoundFrontendClass>();
-						return true;
-					}
-				}
+				using namespace SearchEngineQuerySteps;
+				FFrontendQuery Query;
+				Query.AddStep<FNodeClassRegistrationEvents>()
+					.AddStep<FMapRegistrationEventsToNodeRegistryKeys>()
+					.AddStep<FReduceRegistrationEventsToCurrentStatus>()
+					.AddStep<FTransformRegistrationEventsToClasses>()
+					.AddStep<FMapClassToFullClassNameAndMajorVersion>();
+				return Query;
 			}
 
-			return false;
-		}
+			static ResultType BuildResult(const FFrontendQueryPartition& InPartition)
+			{
+				return BuildSingleClassFromPartition(InPartition);
+			}
+		};
 
-		TArray<FMetasoundFrontendInterface> FSearchEngine::FindAllInterfaces(bool bInIncludeDeprecated)
+		// Policy for finding all registered interfaces. 
+		struct FFindAllInterfacesQueryPolicy
 		{
-			using namespace SearchEngineQuerySteps;
-
-			const FString QueryName = FString::Format(TEXT("FindAllInterfaces_IncludeDeprecated:{0}"), { bInIncludeDeprecated });
+			using ResultType = TArray<FMetasoundFrontendInterface>;
 
-			FFrontendQuery* Query = FindQuery(QueryName);
-			if (!Query)
+			static FFrontendQuery CreateQuery()
 			{
-				TUniquePtr<FFrontendQuery> NewQuery = MakeUnique<FFrontendQuery>();
-
-				NewQuery->AddStep<FInterfaceRegistryTransactionSource>()
+				using namespace SearchEngineQuerySteps;
+				FFrontendQuery Query;
+				Query.AddStep<FInterfaceRegistryTransactionSource>()
 					.AddStep<FMapInterfaceRegistryTransactionsToInterfaceRegistryKeys>()
 					.AddStep<FReduceInterfaceRegistryTransactionsToCurrentStatus>()
-					.AddStep<FTransformInterfaceRegistryTransactionToInterface>();
-
-				if (!bInIncludeDeprecated)
-				{
-					NewQuery->AddStep<FMapInterfaceToInterfaceNameAndMajorVersion>()
-					.AddStep<FReduceInterfacesToHighestVersion>();
-				}
-				NewQuery->AddStep<FMapToNull>();
-
-				Query = AddQuery(QueryName, MoveTemp(NewQuery));
+					.AddStep<FTransformInterfaceRegistryTransactionToInterface>()
+					.AddStep<FMapInterfaceToInterfaceNameAndMajorVersion>()
+					.AddStep<FReduceInterfacesToHighestVersion>()
+					.AddStep<FMapToNull>();
+				return Query;
 			}
 
-			TArray<FMetasoundFrontendInterface> Result;
-
-			if (ensure(Query))
+			static ResultType BuildResult(const FFrontendQueryPartition& InPartition)
 			{
-				const FFrontendQuerySelection& Selection = Query->Update();
-				if (const FFrontendQueryPartition* Entries = Selection.Find(FFrontendQueryKey()))
-				{
-					for (const FFrontendQueryEntry& Entry : *Entries)
-					{
-						if (ensure(Entry.Value.IsType<FMetasoundFrontendInterface>()))
-						{
-							Result.Add(Entry.Value.Get<FMetasoundFrontendInterface>());
-						}
-					}
-				}
+				return BuildArrayOfInterfacesFromPartition(InPartition);
 			}
+		};
 
-			return Result;
-		}
-
-		TArray<FMetasoundFrontendInterface> FSearchEngine::FindUClassDefaultInterfaces(FName InUClassName)
+		// Policy for finding all registered interfaces (including deprecated). 
+		struct FFindAllInterfacesIncludingDeprecatedQueryPolicy
 		{
-			constexpr bool bIncludeDeprecated = false;
-			TArray<FMetasoundFrontendInterface> AllInterfaces = FindAllInterfaces(bIncludeDeprecated);
+			using ResultType = TArray<FMetasoundFrontendInterface>;
 
-			auto DoesNotSupportUClass = [&InUClassName](const FMetasoundFrontendInterface& InInterface)
+			static FFrontendQuery CreateQuery()
 			{
-				FInterfaceRegistryKey RegistryKey = GetInterfaceRegistryKey(InInterface);
-				const IInterfaceRegistryEntry* RegEntry = IInterfaceRegistry::Get().FindInterfaceRegistryEntry(RegistryKey);
-				if (ensure(RegEntry))
-				{
-					return RegEntry->UClassIsSupported(InUClassName);
-				}
-
-				return true;
-			};
-
-			AllInterfaces.RemoveAllSwap(DoesNotSupportUClass);
+				using namespace SearchEngineQuerySteps;
+				FFrontendQuery Query;
+				Query.AddStep<FInterfaceRegistryTransactionSource>()
+					.AddStep<FMapInterfaceRegistryTransactionsToInterfaceRegistryKeys>()
+					.AddStep<FReduceInterfaceRegistryTransactionsToCurrentStatus>()
+					.AddStep<FTransformInterfaceRegistryTransactionToInterface>()
+					.AddStep<FMapToNull>();
+				return Query;
+			}
 
-			return AllInterfaces;
-		}
+			static TArray<FMetasoundFrontendInterface> BuildResult(const FFrontendQueryPartition& InPartition)
+			{
+				return BuildArrayOfInterfacesFromPartition(InPartition);
+			}
+		};
 
-		TArray<FMetasoundFrontendVersion> FSearchEngine::FindAllRegisteredInterfacesWithName(FName InInterfaceName)
+		// Policy for finding all registered interface versions (name & version number)
+		// by name.
+		struct FFindAllRegisteredInterfacesWithNameQueryPolicy
 		{
-			using namespace SearchEngineQuerySteps;
-
-			const FString QueryName = FString::Format(TEXT("FindAllRegisteredInterfacesWithName_Name:{0}"), { InInterfaceName.ToString() });
+			using ResultType = TArray<FMetasoundFrontendVersion>;
 
-			FFrontendQuery* Query = FindQuery(QueryName);
-			if (!Query)
+			static FFrontendQuery CreateQuery()
 			{
-				TUniquePtr<FFrontendQuery> NewQuery = MakeUnique<FFrontendQuery>();
-
-				NewQuery->AddStep<FInterfaceRegistryTransactionSource>()
+				using namespace SearchEngineQuerySteps;
+				FFrontendQuery Query;
+				Query.AddStep<FInterfaceRegistryTransactionSource>()
 					.AddStep<FMapInterfaceRegistryTransactionsToInterfaceRegistryKeys>()
 					.AddStep<FReduceInterfaceRegistryTransactionsToCurrentStatus>()
 					.AddStep<FMapInterfaceRegistryTransactionToInterfaceName>()
 					.AddStep<FTransformInterfaceRegistryTransactionToInterfaceVersion>();
-
-				Query = AddQuery(QueryName, MoveTemp(NewQuery));
+				return Query;
 			}
 
-			TArray<FMetasoundFrontendVersion> Result;
-
-			if (ensure(Query))
+			static TArray<FMetasoundFrontendVersion> BuildResult(const FFrontendQueryPartition& InPartition)
 			{
-				const FFrontendQueryKey Key{InInterfaceName};
-				const FFrontendQuerySelection& Selection = Query->Update();
-				if (const FFrontendQueryPartition* Entries = Selection.Find(Key))
-				{
-					for (const FFrontendQueryEntry& Entry : *Entries)
-					{
-						if (ensure(Entry.Value.IsType<FMetasoundFrontendVersion>()))
-						{
-							Result.Add(Entry.Value.Get<FMetasoundFrontendVersion>());
-						}
-					}
-				}
+				return BuildArrayOfVersionsFromPartition(InPartition);
 			}
+		};
 
-			return Result;
-		}
-
-		bool FSearchEngine::FindInterfaceWithHighestVersion(FName InInterfaceName, FMetasoundFrontendInterface& OutInterface)
+		// Policy for finding the highest version registered interfaces by name.
+		struct FFindInterfaceWithHighestVersionQueryPolicy
 		{
-			using namespace SearchEngineQuerySteps;
-			const FString QueryName = TEXT("FindInterfaceWithHighestVersion");
+			using ResultType = FMetasoundFrontendInterface;
 
-			FFrontendQuery* Query = FindQuery(QueryName);
-			if (!Query)
+			static FFrontendQuery CreateQuery()
 			{
-				TUniquePtr<FFrontendQuery> NewQuery = MakeUnique<FFrontendQuery>();
-
-				NewQuery->AddStep<FInterfaceRegistryTransactionSource>()
+				using namespace SearchEngineQuerySteps;
+				FFrontendQuery Query;
+				Query.AddStep<FInterfaceRegistryTransactionSource>()
 					.AddStep<FMapInterfaceRegistryTransactionsToInterfaceRegistryKeys>()
 					.AddStep<FReduceInterfaceRegistryTransactionsToCurrentStatus>()
 					.AddStep<FTransformInterfaceRegistryTransactionToInterface>()
 					.AddStep<FMapInterfaceToInterfaceName>()
 					.AddStep<FReduceInterfacesToHighestVersion>();
-
-				Query = AddQuery(QueryName, MoveTemp(NewQuery));
+				return Query;
 			}
 
-			if (ensure(Query))
+			static FMetasoundFrontendInterface BuildResult(const FFrontendQueryPartition& InPartition)
 			{
-				const FFrontendQueryKey Key{InInterfaceName};
-				const FFrontendQuerySelection& Selection = Query->Update();
-				if (const FFrontendQueryPartition* Entries = Selection.Find(Key))
-				{
-					if (ensure(Entries->Num() == 1))
-					{
-						OutInterface = Entries->CreateConstIterator()->Value.Get<FMetasoundFrontendInterface>();
-						return true;
-					}
-				}
+				return BuildSingleInterfaceFromPartition(InPartition);
 			}
+		};
+
+		class FSearchEngine : public ISearchEngine
+		{
+			FSearchEngine(const FSearchEngine&) = delete;
+			FSearchEngine(FSearchEngine&&) = delete;
+			FSearchEngine& operator=(const FSearchEngine&) = delete;
+			FSearchEngine& operator=(FSearchEngine&&) = delete;
+
+		public:
+			FSearchEngine() = default;
+			virtual ~FSearchEngine() = default;
+
+			virtual void Prime() override;
+			virtual TArray<FMetasoundFrontendClass> FindAllClasses(bool bInIncludeDeprecated) override;
+			virtual TArray<FMetasoundFrontendClass> FindClassesWithName(const FMetasoundFrontendClassName& InName, bool bInSortByVersion) override;
+			virtual bool FindClassWithHighestVersion(const FMetasoundFrontendClassName& InName, FMetasoundFrontendClass& OutClass) override;
+			virtual bool FindClassWithMajorVersion(const FMetasoundFrontendClassName& InName, int32 InMajorVersion, FMetasoundFrontendClass& OutClass) override;
+
+			virtual TArray<FMetasoundFrontendInterface> FindAllInterfaces(bool bInIncludeDeprecated) override;
+			virtual TArray<FMetasoundFrontendInterface> FindUClassDefaultInterfaces(FName InUClassName) override;
+			virtual TArray<FMetasoundFrontendVersion> FindAllRegisteredInterfacesWithName(FName InInterfaceName) override;
+			virtual bool FindInterfaceWithHighestVersion(FName InInterfaceName, FMetasoundFrontendInterface& OutInterface) override;
+
+		private:
+
+			TSearchEngineQuery<FFindAllClassesQueryPolicy> FindAllClassesQuery;
+			TSearchEngineQuery<FFindAllClassesIncludingDeprecatedQueryPolicy> FindAllClassesIncludingDeprecatedQuery;
+			TSearchEngineQuery<FFindClassesWithNameUnsortedQueryPolicy> FindClassesWithNameUnsortedQuery;
+			TSearchEngineQuery<FFindClassesWithNameSortedQueryPolicy> FindClassesWithNameSortedQuery;
+			TSearchEngineQuery<FFindClassWithHighestVersionQueryPolicy> FindClassWithHighestVersionQuery;
+			TSearchEngineQuery<FFindClassWithMajorVersionQueryPolicy> FindClassWithMajorVersionQuery;
+			TSearchEngineQuery<FFindAllInterfacesQueryPolicy> FindAllInterfacesQuery;
+			TSearchEngineQuery<FFindAllInterfacesIncludingDeprecatedQueryPolicy> FindAllInterfacesIncludingDeprecatedQuery;
+			TSearchEngineQuery<FFindAllRegisteredInterfacesWithNameQueryPolicy> FindAllRegisteredInterfacesWithNameQuery;
+			TSearchEngineQuery<FFindInterfaceWithHighestVersionQueryPolicy> FindInterfaceWithHighestVersionQuery;
+		};
 
-			return false;
+		ISearchEngine& ISearchEngine::Get()
+		{
+			static FSearchEngine SearchEngine;
+			return SearchEngine;
 		}
 
-		bool FSearchEngine::FindInterfaceWithMajorVersion(FName InInterfaceName, int32 InMajorVersion, FMetasoundFrontendInterface& OutInterface)
+		void FSearchEngine::Prime()
 		{
-			using namespace SearchEngineQuerySteps;
-			const FString QueryName = TEXT("FindInterfaceWithMajorVersion");
+			METASOUND_TRACE_CPUPROFILER_EVENT_SCOPE(metasound::FSearchEngine::Prime);
+
+			FindAllClassesQuery.Prime();
+			FindAllClassesIncludingDeprecatedQuery.Prime();
+			FindClassesWithNameUnsortedQuery.Prime();
+			FindClassesWithNameSortedQuery.Prime();
+			FindClassWithHighestVersionQuery.Prime();
+			FindClassWithMajorVersionQuery.Prime();
+			FindAllInterfacesIncludingDeprecatedQuery.Prime();
+			FindAllInterfacesQuery.Prime();
+			FindAllRegisteredInterfacesWithNameQuery.Prime();
+			FindInterfaceWithHighestVersionQuery.Prime();
+		}
+
+		TArray<FMetasoundFrontendClass> FSearchEngine::FindAllClasses(bool bInIncludeDeprecated)
+		{
+			METASOUND_TRACE_CPUPROFILER_EVENT_SCOPE(metasound::FSearchEngine::FindAllClasses);
 
-			FFrontendQuery* Query = FindQuery(QueryName);
-			if (!Query)
+			FFrontendQueryKey NullKey;
+			if (bInIncludeDeprecated)
 			{
-				TUniquePtr<FFrontendQuery> NewQuery = MakeUnique<FFrontendQuery>();
+				return FindAllClassesIncludingDeprecatedQuery.UpdateAndFindResult(NullKey);
+			}
+			else
+			{
+				return FindAllClassesQuery.UpdateAndFindResult(NullKey);
+			}
+		}
 
-				NewQuery->AddStep<FInterfaceRegistryTransactionSource>()
-					.AddStep<FMapInterfaceRegistryTransactionsToInterfaceRegistryKeys>()
-					.AddStep<FReduceInterfaceRegistryTransactionsToCurrentStatus>()
-					.AddStep<FTransformInterfaceRegistryTransactionToInterface>()
-					.AddStep<FMapInterfaceToInterfaceNameAndMajorVersion>()
-					.AddStep<FReduceInterfacesToHighestVersion>();
+		TArray<FMetasoundFrontendClass> FSearchEngine::FindClassesWithName(const FMetasoundFrontendClassName& InName, bool bInSortByVersion)
+		{
+			METASOUND_TRACE_CPUPROFILER_EVENT_SCOPE(metasound::FSearchEngine::FindClassesWithName);
 
-				Query = AddQuery(QueryName, MoveTemp(NewQuery));
+			const FFrontendQueryKey Key(InName.GetFullName());
+			if (bInSortByVersion)
+			{
+				return FindClassesWithNameSortedQuery.UpdateAndFindResult(Key);
 			}
-
-			if (ensure(Query))
+			else 
 			{
-				const FFrontendQueryKey Key = FMapInterfaceToInterfaceNameAndMajorVersion::GetKey(InInterfaceName.ToString(), InMajorVersion);
-				const FFrontendQuerySelection& Selection = Query->Update();
-				if (const FFrontendQueryPartition* Entries = Selection.Find(Key))
-				{
-					if (ensure(Entries->Num() == 1))
-					{
-						OutInterface = Entries->CreateConstIterator()->Value.Get<FMetasoundFrontendInterface>();
-						return true;
-					}
-				}
+				return FindClassesWithNameUnsortedQuery.UpdateAndFindResult(Key);
 			}
+		}
+
+		bool FSearchEngine::FindClassWithHighestVersion(const FMetasoundFrontendClassName& InName, FMetasoundFrontendClass& OutClass)
+		{
+			METASOUND_TRACE_CPUPROFILER_EVENT_SCOPE(metasound::FSearchEngine::FindClassWithHighestVersion);
+
+			const FFrontendQueryKey Key{InName.GetFullName()};
+			return FindClassWithHighestVersionQuery.UpdateAndFindResult(Key, OutClass);
+		}
+
+		bool FSearchEngine::FindClassWithMajorVersion(const FMetasoundFrontendClassName& InName, int32 InMajorVersion, FMetasoundFrontendClass& OutClass)
+		{
+			METASOUND_TRACE_CPUPROFILER_EVENT_SCOPE(metasound::FSearchEngine::FindClassWithMajorVersion);
+			using namespace SearchEngineQuerySteps;
 
-			return false;
+			const FFrontendQueryKey Key = FMapClassToFullClassNameAndMajorVersion::GetKey(InName, InMajorVersion);
+			return FindClassWithMajorVersionQuery.UpdateAndFindResult(Key, OutClass);
 		}
 
-		FFrontendQuery* FSearchEngine::FindQuery(const FString& InQueryName)
+		TArray<FMetasoundFrontendInterface> FSearchEngine::FindAllInterfaces(bool bInIncludeDeprecated)
 		{
-			if (ClearMetaSoundFrontendSearchEngineCacheCVar)
+			METASOUND_TRACE_CPUPROFILER_EVENT_SCOPE(metasound::FSearchEngine::FindAllInterfaces);
+
+			const FFrontendQueryKey NullKey;
+			if (bInIncludeDeprecated)
 			{
-				ClearMetaSoundFrontendSearchEngineCacheCVar = 0;
-				QueryCache.Reset();
+				return FindAllInterfacesIncludingDeprecatedQuery.UpdateAndFindResult(NullKey);
 			}
-
-			if (TUniquePtr<FFrontendQuery>* Query = QueryCache.Find(InQueryName))
+			else
 			{
-				return Query->Get();
+				return FindAllInterfacesQuery.UpdateAndFindResult(NullKey);
 			}
+		}
+
+		TArray<FMetasoundFrontendInterface> FSearchEngine::FindUClassDefaultInterfaces(FName InUClassName)
+		{
+			METASOUND_TRACE_CPUPROFILER_EVENT_SCOPE(metasound::FSearchEngine::FindUClassDefaultInterfaces);
+			constexpr bool bIncludeDeprecated = false;
+			TArray<FMetasoundFrontendInterface> AllInterfaces = FindAllInterfaces(bIncludeDeprecated);
 
-			return nullptr;
+
+			// All interfaces are filtered here because a UClass name cannot effectively 
+			// be used as a "Key" for this query.
+			//
+			// We cannot know all the UClasses associated with each default interface given only a UClassName
+			// because we need to inspect UClass inheritance to handle derived UClasses
+			// which should be compatible with their parent's default interfaces. 
+			//
+			// This module does not have access to the UClass inheritance structure.
+			// Instead, this information is accessed in the IInterfaceRegistryEntry
+			// concrete implementation. 
+			auto DoesNotSupportUClass = [&InUClassName](const FMetasoundFrontendInterface& InInterface)
+			{
+				FInterfaceRegistryKey RegistryKey = GetInterfaceRegistryKey(InInterface);
+				const IInterfaceRegistryEntry* RegEntry = IInterfaceRegistry::Get().FindInterfaceRegistryEntry(RegistryKey);
+				if (ensure(RegEntry))
+				{
+					return !RegEntry->UClassIsSupported(InUClassName);
+				}
+
+				return true;
+			};
+
+			AllInterfaces.RemoveAllSwap(DoesNotSupportUClass);
+
+			return AllInterfaces;
 		}
 
-		FFrontendQuery* FSearchEngine::AddQuery(const FString& InQueryName, TUniquePtr<FFrontendQuery>&& InQuery)
+		TArray<FMetasoundFrontendVersion> FSearchEngine::FindAllRegisteredInterfacesWithName(FName InInterfaceName)
 		{
-			return QueryCache.Add(InQueryName, MoveTemp(InQuery)).Get();
+			METASOUND_TRACE_CPUPROFILER_EVENT_SCOPE(metasound::FSearchEngine::FindAllRegisteredInterfacesWithName);
+
+			const FFrontendQueryKey Key{InInterfaceName};
+			return FindAllRegisteredInterfacesWithNameQuery.UpdateAndFindResult(Key);
+		}
+
+		bool FSearchEngine::FindInterfaceWithHighestVersion(FName InInterfaceName, FMetasoundFrontendInterface& OutInterface)
+		{
+			METASOUND_TRACE_CPUPROFILER_EVENT_SCOPE(metasound::FSearchEngine::FindInterfaceWithHighestVersion);
+
+			const FFrontendQueryKey Key{InInterfaceName};
+			return FindInterfaceWithHighestVersionQuery.UpdateAndFindResult(Key, OutInterface);
 		}
 	}
 }
diff --git a/Engine/Plugins/Runtime/Metasound/Source/MetasoundFrontend/Public/MetasoundFrontendSearchEngine.h b/Engine/Plugins/Runtime/Metasound/Source/MetasoundFrontend/Public/MetasoundFrontendSearchEngine.h
--- a/Engine/Plugins/Runtime/Metasound/Source/MetasoundFrontend/Public/MetasoundFrontendSearchEngine.h
+++ b/Engine/Plugins/Runtime/Metasound/Source/MetasoundFrontend/Public/MetasoundFrontendSearchEngine.h
@@ -20,6 +20,9 @@ namespace Metasound
 				/** Return an instance of a search engine. */
 				static ISearchEngine& Get();
 
+				/** Updates internal state to speed up queries. */
+				virtual void Prime() = 0;
+
 				/** Find all FMetasoundFrontendClasses.
 				  * (Optional) Include deprecated classes (i.e. versions of classes that are not the highest major version).
 				  */
@@ -44,8 +47,6 @@ namespace Metasound
 
 				virtual bool FindInterfaceWithHighestVersion(FName InInterfaceName, FMetasoundFrontendInterface& OutInterface) = 0;
 
-				virtual bool FindInterfaceWithMajorVersion(FName InInterfaceName, int32 InMajorVersion, FMetasoundFrontendInterface& OutInterface) = 0;
-
 				virtual ~ISearchEngine() = default;
 
 			protected:
