Fix for new metasounds getting all supported interfaces instead of only default supported interfaces.
#rnx

diff --git a/Engine/Plugins/Runtime/Metasound/Source/MetasoundFrontend/Private/MetasoundFrontendSearchEngine.cpp b/Engine/Plugins/Runtime/Metasound/Source/MetasoundFrontend/Private/MetasoundFrontendSearchEngine.cpp
--- a/Engine/Plugins/Runtime/Metasound/Source/MetasoundFrontend/Private/MetasoundFrontendSearchEngine.cpp
+++ b/Engine/Plugins/Runtime/Metasound/Source/MetasoundFrontend/Private/MetasoundFrontendSearchEngine.cpp
@@ -210,6 +210,24 @@ namespace Metasound
 				}
 			};
 
+			class FRemoveInterfacesWhichAreNotDefault : public IFrontendQueryFilterStep
+			{
+				virtual bool Filter(const FFrontendQueryEntry& InEntry) const override
+				{
+					if (InEntry.Value.IsType<FMetasoundFrontendInterface>())
+					{
+						FInterfaceRegistryKey RegistryKey = GetInterfaceRegistryKey(InEntry.Value.Get<FMetasoundFrontendInterface>());
+						const IInterfaceRegistryEntry* RegEntry = IInterfaceRegistry::Get().FindInterfaceRegistryEntry(RegistryKey);
+
+						if (nullptr != RegEntry)
+						{
+							return RegEntry->IsDefault();
+						}
+					}
+					return false;
+				}
+			};
+
 			class FMapInterfaceToInterfaceName : public IFrontendQueryMapStep
 			{
 			public:
@@ -579,6 +597,32 @@ namespace Metasound
 			}
 		};
 
+		// Policy for finding all registered default interfaces. 
+		struct FFindAllDefaultInterfacesQueryPolicy
+		{
+			using ResultType = TArray<FMetasoundFrontendInterface>;
+
+			static FFrontendQuery CreateQuery()
+			{
+				using namespace SearchEngineQuerySteps;
+				FFrontendQuery Query;
+				Query.AddStep<FInterfaceRegistryTransactionSource>()
+					.AddStep<FMapInterfaceRegistryTransactionsToInterfaceRegistryKeys>()
+					.AddStep<FReduceInterfaceRegistryTransactionsToCurrentStatus>()
+					.AddStep<FTransformInterfaceRegistryTransactionToInterface>()
+					.AddStep<FMapInterfaceToInterfaceNameAndMajorVersion>()
+					.AddStep<FReduceInterfacesToHighestVersion>()
+					.AddStep<FRemoveInterfacesWhichAreNotDefault>()
+					.AddStep<FMapToNull>();
+				return Query;
+			}
+
+			static ResultType BuildResult(const FFrontendQueryPartition& InPartition)
+			{
+				return BuildArrayOfInterfacesFromPartition(InPartition);
+			}
+		};
+
 		// Policy for finding all registered interfaces (including deprecated). 
 		struct FFindAllInterfacesIncludingDeprecatedQueryPolicy
 		{
@@ -681,6 +725,7 @@ namespace Metasound
 			TSearchEngineQuery<FFindClassWithHighestVersionQueryPolicy> FindClassWithHighestVersionQuery;
 			TSearchEngineQuery<FFindClassWithMajorVersionQueryPolicy> FindClassWithMajorVersionQuery;
 			TSearchEngineQuery<FFindAllInterfacesQueryPolicy> FindAllInterfacesQuery;
+			TSearchEngineQuery<FFindAllDefaultInterfacesQueryPolicy> FFindAllDefaultInterfacesQuery;
 			TSearchEngineQuery<FFindAllInterfacesIncludingDeprecatedQueryPolicy> FindAllInterfacesIncludingDeprecatedQuery;
 			TSearchEngineQuery<FFindAllRegisteredInterfacesWithNameQueryPolicy> FindAllRegisteredInterfacesWithNameQuery;
 			TSearchEngineQuery<FFindInterfaceWithHighestVersionQueryPolicy> FindInterfaceWithHighestVersionQuery;
@@ -773,11 +818,11 @@ namespace Metasound
 		TArray<FMetasoundFrontendInterface> FSearchEngine::FindUClassDefaultInterfaces(FName InUClassName)
 		{
 			METASOUND_TRACE_CPUPROFILER_EVENT_SCOPE(metasound::FSearchEngine::FindUClassDefaultInterfaces);
-			constexpr bool bIncludeDeprecated = false;
-			TArray<FMetasoundFrontendInterface> AllInterfaces = FindAllInterfaces(bIncludeDeprecated);
 
+			const FFrontendQueryKey NullKey;
+			TArray<FMetasoundFrontendInterface> DefaultInterfaces = FFindAllDefaultInterfacesQuery.UpdateAndFindResult(NullKey);
 
-			// All interfaces are filtered here because a UClass name cannot effectively 
+			// All default interfaces are filtered here because a UClass name cannot effectively 
 			// be used as a "Key" for this query.
 			//
 			// We cannot know all the UClasses associated with each default interface given only a UClassName
@@ -799,9 +844,9 @@ namespace Metasound
 				return true;
 			};
 
-			AllInterfaces.RemoveAllSwap(DoesNotSupportUClass);
+			DefaultInterfaces.RemoveAllSwap(DoesNotSupportUClass);
 
-			return AllInterfaces;
+			return DefaultInterfaces;
 		}
 
 		TArray<FMetasoundFrontendVersion> FSearchEngine::FindAllRegisteredInterfacesWithName(FName InInterfaceName)
