Fix for creation of files when WaveWriter is unhooked

diff --git a/Engine/Plugins/Runtime/Metasound/Source/MetasoundStandardNodes/Private/MetasoundWaveWriterNode.cpp b/Engine/Plugins/Runtime/Metasound/Source/MetasoundStandardNodes/Private/MetasoundWaveWriterNode.cpp
--- a/Engine/Plugins/Runtime/Metasound/Source/MetasoundStandardNodes/Private/MetasoundWaveWriterNode.cpp
+++ b/Engine/Plugins/Runtime/Metasound/Source/MetasoundStandardNodes/Private/MetasoundWaveWriterNode.cpp
@@ -456,19 +456,29 @@ namespace Metasound
 
 		FStringReadRef FilenamePrefix = InputCol.GetDataReadReferenceOrConstructWithVertexDefault<FString>(InputInterface, GetFilenamePrefixPinName(), Settings);
 
+		int32 NumConnectedAudioPins = 0;
 		TArray<FAudioBufferReadRef> InputBuffers;
 		for (int32 i = 0; i < NumInputChannels; ++i)
 		{
-			InputBuffers.Add(InputCol.GetDataReadReferenceOrConstruct<FAudioBuffer>(GetAudioInputName(i), InParams.OperatorSettings));
+			const FVertexName PinName = GetAudioInputName(i);
+			NumConnectedAudioPins += (int32)InputCol.ContainsDataReadReference<FAudioBuffer>(PinName);
+			InputBuffers.Add(InputCol.GetDataReadReferenceOrConstruct<FAudioBuffer>(PinName, InParams.OperatorSettings));
 		}
-
-		return MakeUnique<TWaveWriterOperator>(
-			Settings,
-			MoveTemp(InputBuffers),
-			InputCol.GetDataReadReferenceOrConstructWithVertexDefault<bool>(InputInterface, GetEnabledPinName(), Settings),
-			GetNameCache(),
-			*FilenamePrefix
+		
+		// Only create a real operator if there's some connected pins.
+		if (NumConnectedAudioPins > 0)
+		{
+			return MakeUnique<TWaveWriterOperator>(
+				Settings,
+				MoveTemp(InputBuffers),
+				InputCol.GetDataReadReferenceOrConstructWithVertexDefault<bool>(InputInterface, GetEnabledPinName(), Settings),
+				GetNameCache(),
+				*FilenamePrefix
 			);
+		}
+
+		// Create a no-op operator.
+		return MakeUnique<FNoOpOperator>();
 	}
 
 	template<int32 NumInputChannels>
