MetaSound oscillator generator fixes
- change triangle oscillator implementation
- fix DC offset issue for square wave
- apply fade in after unipolar transformation to prevent popping

diff --git a/Engine/Plugins/Runtime/Metasound/Source/MetasoundStandardNodes/Private/MetasoundOscillatorNodes.cpp b/Engine/Plugins/Runtime/Metasound/Source/MetasoundStandardNodes/Private/MetasoundOscillatorNodes.cpp
--- a/Engine/Plugins/Runtime/Metasound/Source/MetasoundStandardNodes/Private/MetasoundOscillatorNodes.cpp
+++ b/Engine/Plugins/Runtime/Metasound/Source/MetasoundStandardNodes/Private/MetasoundOscillatorNodes.cpp
@@ -36,6 +36,7 @@ namespace Metasound
 			float FrequencyHz = 0.f;
 			float GlideEaseFactor = 0.0f;
 			float PulseWidth = 0.f;
+			bool BiPolar = true;
 			TArrayView<float> AlignedBuffer;
 			TArrayView<const float> FM;
 		};
@@ -102,46 +103,97 @@ namespace Metasound
 				float* Out = InArgs.AlignedBuffer.GetData();
 				const float OneOverSampleRate = 1.f / InArgs.SampleRate;
 				
-				// Constant Freq between this an last update? 
-				if (FMath::IsNearlyEqual(InArgs.FrequencyHz, CurrentFreq) || CurrentFreq < 0.f || FMath::IsNearlyEqual(InArgs.GlideEaseFactor, 1.0f))
+				// TODO: break this into separate calls because there is a lot of code duplication to prevent per sample branching 
+				if (InArgs.BiPolar)
 				{
-					CurrentFreq = InArgs.FrequencyHz;
-					const float DeltaPhase = InArgs.FrequencyHz * OneOverSampleRate;
-					while (RemainingSamplesInBlock > 0)
+					// Constant Freq between this an last update? 
+					if (FMath::IsNearlyEqual(InArgs.FrequencyHz, CurrentFreq) || CurrentFreq < 0.f || FMath::IsNearlyEqual(InArgs.GlideEaseFactor, 1.0f))
 					{
-						Wrap(Phase);
+						CurrentFreq = InArgs.FrequencyHz;
+						const float DeltaPhase = InArgs.FrequencyHz * OneOverSampleRate;
+						while (RemainingSamplesInBlock > 0)
+						{
+							Wrap(Phase);
 
-						float Output = Osc(Phase, DeltaPhase, InArgs);
-						Output *= CurrentFade;
-						*Out++ = Output;
+							float Output = Osc(Phase, DeltaPhase, InArgs);
+							Output *= CurrentFade;
+							*Out++ = Output;
 
-						CurrentFade = CurrentFade + FadeSmooth * (1.0f - CurrentFade);
+							CurrentFade = CurrentFade + FadeSmooth * (1.0f - CurrentFade);
 
-						Phase += DeltaPhase;
+							Phase += DeltaPhase;
 
-						RemainingSamplesInBlock--;
+							RemainingSamplesInBlock--;
+						}
+					}
+					else
+					{
+						while (RemainingSamplesInBlock > 0)
+						{
+							Wrap(Phase);
+
+							const float DeltaPhase = CurrentFreq * OneOverSampleRate;
+
+							float Output = Osc(Phase, DeltaPhase, InArgs);
+							Output *= CurrentFade;
+							*Out++ = Output;
+
+							CurrentFade = CurrentFade + FadeSmooth * (1.0f - CurrentFade);
+
+							Phase += DeltaPhase;
+
+							// lerp frequency based on the glide factor
+							CurrentFreq = CurrentFreq + InArgs.GlideEaseFactor * (InArgs.FrequencyHz - CurrentFreq);
+
+							RemainingSamplesInBlock--;
+						}
 					}
 				}
- 				else
- 				{			
-					while (RemainingSamplesInBlock > 0)
+				else // unipolar
+				{
+					// Constant Freq between this an last update? 
+					if (FMath::IsNearlyEqual(InArgs.FrequencyHz, CurrentFreq) || CurrentFreq < 0.f || FMath::IsNearlyEqual(InArgs.GlideEaseFactor, 1.0f))
+					{
+						CurrentFreq = InArgs.FrequencyHz;
+						const float DeltaPhase = InArgs.FrequencyHz * OneOverSampleRate;
+						while (RemainingSamplesInBlock > 0)
+						{
+							Wrap(Phase);
+
+							float Output = Osc(Phase, DeltaPhase, InArgs);
+							Output = Audio::GetUnipolar(Output);
+							Output *= CurrentFade;
+							*Out++ = Output;
+
+							CurrentFade = CurrentFade + FadeSmooth * (1.0f - CurrentFade);
+
+							Phase += DeltaPhase;
+
+							RemainingSamplesInBlock--;
+						}
+					}
+					else
 					{
-						Wrap(Phase);
+						while (RemainingSamplesInBlock > 0)
+						{
+							Wrap(Phase);
 
-						const float DeltaPhase = CurrentFreq * OneOverSampleRate;
-						
-						float Output = Osc(Phase, DeltaPhase, InArgs);
-						Output *= CurrentFade;
-						*Out++ = Output;
+							const float DeltaPhase = CurrentFreq * OneOverSampleRate;
 
-						CurrentFade = CurrentFade + FadeSmooth * (1.0f - CurrentFade);
+							float Output = Osc(Phase, DeltaPhase, InArgs);
+							Output = Audio::GetUnipolar(Output);
+							Output *= CurrentFade;
+							*Out++ = Output;
 
-						Phase += DeltaPhase;
+							CurrentFade = CurrentFade + FadeSmooth * (1.0f - CurrentFade);
 
-						// lerp frequency based on the glide factor
-						CurrentFreq = CurrentFreq + InArgs.GlideEaseFactor * (InArgs.FrequencyHz - CurrentFreq);
+							Phase += DeltaPhase;
 
-						RemainingSamplesInBlock--;
+							// lerp frequency based on the glide factor
+							CurrentFreq = CurrentFreq + InArgs.GlideEaseFactor * (InArgs.FrequencyHz - CurrentFreq);
+
+							RemainingSamplesInBlock--;
+						}
 					}
 				}
 			}
@@ -170,50 +222,105 @@ namespace Metasound
 				const float* FM = InArgs.FM.GetData();
 				const float OneOverSampleRate = 1.f / InArgs.SampleRate;
 
-				// Constant Base Freq between this and last update? 
-				if (FMath::IsNearlyEqual(InArgs.FrequencyHz, CurrentFreq) || CurrentFreq < 0.f || FMath::IsNearlyEqual(InArgs.GlideEaseFactor, 1.0f))
+				// TODO: break this into separate calls because there is a lot of code duplication to prevent per sample branching 
+				if (InArgs.BiPolar)
 				{
-					while (RemainingSamplesInBlock > 0)
+					// Constant Base Freq between this and last update? 
+					if (FMath::IsNearlyEqual(InArgs.FrequencyHz, CurrentFreq) || CurrentFreq < 0.f || FMath::IsNearlyEqual(InArgs.GlideEaseFactor, 1.0f))
 					{
-						const float PerSampleFreq = InArgs.FrequencyHz + *FM++;
-						const float DeltaPhase = PerSampleFreq * OneOverSampleRate;
+						while (RemainingSamplesInBlock > 0)
+						{
+							const float PerSampleFreq = InArgs.FrequencyHz + *FM++;
+							const float DeltaPhase = PerSampleFreq * OneOverSampleRate;
+
+							Wrap(Phase);
 
-						Wrap(Phase);
+							float Output = Osc(Phase, DeltaPhase, InArgs);
+							Output *= CurrentFade;
+							*Out++ = Output;
 
-						float Output = Osc(Phase, DeltaPhase, InArgs);
-						Output *= CurrentFade;
-						*Out++ = Output;
-												
-						CurrentFade = CurrentFade + FadeSmooth * (1.0f - CurrentFade);
+							CurrentFade = CurrentFade + FadeSmooth * (1.0f - CurrentFade);
 
-						Phase += DeltaPhase;
+							Phase += DeltaPhase;
 
-						RemainingSamplesInBlock--;
+							RemainingSamplesInBlock--;
+						}
+
+						CurrentFreq = InArgs.FrequencyHz;
 					}
+					else
+					{
+						while (RemainingSamplesInBlock > 0)
+						{
+							const float ModulatedFreqSum = CurrentFreq + *FM++;
+							const float DeltaPhase = ModulatedFreqSum * OneOverSampleRate;
+
+							Wrap(Phase);
 
-					CurrentFreq = InArgs.FrequencyHz;
+							float Output = Osc(Phase, DeltaPhase, InArgs);
+							Output *= CurrentFade;
+							*Out++ = Output;
+
+							CurrentFade = CurrentFade + FadeSmooth * (1.0f - CurrentFade);
+
+							Phase += DeltaPhase;
+
+							// lerp frequency based on the glide factor
+							CurrentFreq = CurrentFreq + InArgs.GlideEaseFactor * (InArgs.FrequencyHz - CurrentFreq);
+
+							RemainingSamplesInBlock--;
+						}
+					}
 				}
-				else
-				{	
-					while (RemainingSamplesInBlock > 0)
+				else // unipolar
+				{
+					// Constant Base Freq between this and last update? 
+					if (FMath::IsNearlyEqual(InArgs.FrequencyHz, CurrentFreq) || CurrentFreq < 0.f || FMath::IsNearlyEqual(InArgs.GlideEaseFactor, 1.0f))
 					{
-						const float ModulatedFreqSum = CurrentFreq + *FM++;
-						const float DeltaPhase = ModulatedFreqSum * OneOverSampleRate;
+						while (RemainingSamplesInBlock > 0)
+						{
+							const float PerSampleFreq = InArgs.FrequencyHz + *FM++;
+							const float DeltaPhase = PerSampleFreq * OneOverSampleRate;
+
+							Wrap(Phase);
 
-						Wrap(Phase);
+							float Output = Osc(Phase, DeltaPhase, InArgs);
+							Output = Audio::GetUnipolar(Output);
+							Output *= CurrentFade;
+							*Out++ = Output;
 
-						float Output = Osc(Phase, DeltaPhase, InArgs);
-						Output *= CurrentFade;
-						*Out++ = Output;
-						
-						CurrentFade = CurrentFade + FadeSmooth * (1.0f - CurrentFade);
+							CurrentFade = CurrentFade + FadeSmooth * (1.0f - CurrentFade);
 
-						Phase += DeltaPhase;
+							Phase += DeltaPhase;
 
-						// lerp frequency based on the glide factor
-						CurrentFreq = CurrentFreq + InArgs.GlideEaseFactor * (InArgs.FrequencyHz - CurrentFreq);
+							RemainingSamplesInBlock--;
+						}
 
-						RemainingSamplesInBlock--;
+						CurrentFreq = InArgs.FrequencyHz;
+					}
+					else
+					{
+						while (RemainingSamplesInBlock > 0)
+						{
+							const float ModulatedFreqSum = CurrentFreq + *FM++;
+							const float DeltaPhase = ModulatedFreqSum * OneOverSampleRate;
+
+							Wrap(Phase);
+
+							float Output = Osc(Phase, DeltaPhase, InArgs);
+							Output = Audio::GetUnipolar(Output);
+							Output *= CurrentFade;
+							*Out++ = Output;
+
+							CurrentFade = CurrentFade + FadeSmooth * (1.0f - CurrentFade);
+
+							Phase += DeltaPhase;
+
+							// lerp frequency based on the glide factor
+							CurrentFreq = CurrentFreq + InArgs.GlideEaseFactor * (InArgs.FrequencyHz - CurrentFreq);
+
+							RemainingSamplesInBlock--;
+						}
 					}
 				}
 			}
@@ -280,6 +387,10 @@ namespace Metasound
 			void operator()(const FGeneratorArgs& Args)
 			{
 				Rotator.GenerateBuffer(Args.SampleRate, Args.FrequencyHz, Args.AlignedBuffer.GetData(), Args.AlignedBuffer.Num());
+				if (!Args.BiPolar)
+				{
+					Audio::ConvertBipolarBufferToUnipolar(Args.AlignedBuffer.GetData(), Args.AlignedBuffer.Num());
+				}
 			}
 		};		
 		
@@ -361,15 +472,11 @@ namespace Metasound
 				float SquareSaw2 = Audio::GetBipolar(NewPhase);
 				SquareSaw2 += PolySmoothSaw(NewPhase, InPhaseDelta);
 
-				// Subtracting 2 saws creates a square wave!
-				float Output = 0.5f * SquareSaw1 - 0.5f * SquareSaw2;
-
-				// Apply DC correction
-				const float Correction = (InArgs.PulseWidth < 0.5f) ?
-					(1.0f / (1.0f - InArgs.PulseWidth)) : (1.0f / InArgs.PulseWidth);
-
-				Output *= Correction;
-				return Output;
+				// Subtract 2 saws, then apply DC correction 
+				// Simplified version of 
+				// float Output = 0.5f * SquareSaw1 - 0.5f * SquareSaw2;
+				// Output = 2.0f * (Output + InArgs.PulseWidth) - 1.0f;
+				return SquareSaw1 - SquareSaw2 + 2.0f * (InArgs.PulseWidth - 0.5f);
 			}
 		};
 
@@ -388,22 +495,10 @@ namespace Metasound
 
 			float operator()(float InPhase, float InPhaseDelta, const FGeneratorArgs& InArgs)
 			{
-				// if the oscillator phase wrapped, our previous phase will be higher than current phase
-				if (PrevPhase > InPhase)
-				{
-					// Flip the sign if we wrapped phase
-					TriangleSign *= -1.0f;
-				}
-				PrevPhase = InPhase;
-
-
-				// Get saw trivial output
-				float SawToothOutput = Audio::GetBipolar(InPhase);
-				const float SawSquaredInvMod = (1.0f - SawToothOutput * SawToothOutput) * TriangleSign;
-
-				const float Differentiated = SawSquaredInvMod - DPW_z1;
-				DPW_z1 = SawSquaredInvMod;
-				return Differentiated * InArgs.SampleRate / (4.0f * InArgs.FrequencyHz * (1.0f - InPhaseDelta));
+				constexpr float FASTASIN_HALF_PI{ 1.5707963050f };
+				float PhaseRadians = (InPhase * PI * 2.f);
+				float Output = FMath::FastAsin(Audio::FastSin3(PhaseRadians - PI)) / FASTASIN_HALF_PI;
+				return Output;
 			}
 		};
 
@@ -519,28 +614,22 @@ namespace Metasound
 			// Clamp frequencies into Nyquist range.
 			const float ClampedFreq = FMath::Clamp(*BaseFrequency, -Nyquist, Nyquist);
 			const float ClampedGlideEase = Audio::GetLogFrequencyClamped(*GlideFactor, { 0.0f, 1.0f }, { 1.0f, 0.0001f });
-
+			const bool OutputBiPolar = *BiPolar;
 			AudioBuffer->Zero();
 
 			Derived* Self = static_cast<Derived*>(this);
 			PhaseReset->ExecuteBlock
 			(
-				[Self, ClampedFreq, ClampedGlideEase](int32 InFrameStart, int32 InFrameEnd)
+				[Self, ClampedFreq, ClampedGlideEase, OutputBiPolar](int32 InFrameStart, int32 InFrameEnd)
 				{
-					Self->Generate(InFrameStart, InFrameEnd, ClampedFreq, ClampedGlideEase);
+					Self->Generate(InFrameStart, InFrameEnd, ClampedFreq, ClampedGlideEase, OutputBiPolar);
 				},
-				[Self, ClampedFreq, ClampedGlideEase](int32 InFrameStart, int32 InFrameEnd)
+				[Self, ClampedFreq, ClampedGlideEase, OutputBiPolar](int32 InFrameStart, int32 InFrameEnd)
 				{
 					Self->ResetPhase(*Self->PhaseOffset);
-					Self->Generate(InFrameStart, InFrameEnd, ClampedFreq, ClampedGlideEase);
+					Self->Generate(InFrameStart, InFrameEnd, ClampedFreq, ClampedGlideEase, OutputBiPolar);
 				}
 			);
-
-			// Convert to Unipolar
-			if (*Enabled && !*BiPolar)
-			{
-				Audio::ConvertBipolarBufferToUnipolar(AudioBuffer->GetData(), AudioBuffer->Num());
-			}
 		}
 
 	protected:
@@ -568,7 +657,7 @@ namespace Metasound
 			: Super(InConstructParams)
 		{}
 
-		void Generate(int32 InStartFrame, int32 InEndFrame, float InClampedFreq, float InClampedGlideEase)
+		void Generate(int32 InStartFrame, int32 InEndFrame, float InClampedFreq, float InClampedGlideEase, bool BiPolar)
 		{		
 			int32 NumFrames = InEndFrame - InStartFrame;
 
@@ -580,6 +669,7 @@ namespace Metasound
 					InClampedFreq,
 					InClampedGlideEase,
 					0.f,
+					BiPolar,
 					MakeArrayView(this->AudioBuffer->GetData() + InStartFrame, NumFrames), // Not aligned.
 				});
 			}			
@@ -603,7 +693,7 @@ namespace Metasound
 			return Inputs;
 		}
 
-		void Generate(int32 InStartFrame, int32 InEndFrame, float InClampedFreq, float InClampedGlideEase)
+		void Generate(int32 InStartFrame, int32 InEndFrame, float InClampedFreq, float InClampedGlideEase, bool BiPolar)
 		{
 			int32 NumFrames = InEndFrame - InStartFrame;
 			if (*this->Enabled && NumFrames > 0)
@@ -614,6 +704,7 @@ namespace Metasound
 					InClampedFreq,
 					InClampedGlideEase,
 					0.f,
+					BiPolar, 
 					MakeArrayView(this->AudioBuffer->GetData() + InStartFrame, NumFrames), // Not aligned.
 					MakeArrayView(this->Fm->GetData() + InStartFrame, NumFrames) // Not aligned.
 				});
@@ -718,7 +809,7 @@ namespace Metasound
 			
 			// Check to see if we have an FM input connected.
 			bool bHasFM = InputCol.ContainsDataReadReference<FAudioBuffer>(FrequencyModPinName);			
-			FEnumSineGenerationTypeReadRef Type = InputCol.GetDataReadReferenceOrConstruct<FEnumSineGenerationType>(SineTypePin, ESineGenerationType::Rotation);
+			FEnumSineGenerationTypeReadRef Type = InputCol.GetDataReadReferenceOrConstruct<FEnumSineGenerationType>(SineTypePin, ESineGenerationType::Wavetable);
 			if (bHasFM)
 			{
 
@@ -893,7 +984,7 @@ namespace Metasound
 
 	DEFINE_METASOUND_ENUM_BEGIN(ESquareGenerationType, FEnumSquareGenerationType, "SquareGenerationType")
 		DEFINE_METASOUND_ENUM_ENTRY(ESquareGenerationType::PolySmooth, LOCTEXT("SquarePolySmoothDescription", "Poly Smooth"), LOCTEXT("PolySmoothDescriptionTT", "PolySmooth (i.e. BLEP)")),
-		DEFINE_METASOUND_ENUM_ENTRY(ESquareGenerationType::Trivial, LOCTEXT("SquareTrivialDescription", "Trivial"), LOCTEXT("TrivialDescriptionTT", "The most basic raw implementation")),
+		DEFINE_METASOUND_ENUM_ENTRY(ESquareGenerationType::Trivial, LOCTEXT("SquareTrivialDescription", "Trivial"), LOCTEXT("TrivialDescriptionTT", "The most basic raw implementation. Does not obey pulse width.")),
 		//DEFINE_METASOUND_ENUM_ENTRY(ESquareGenerationType::Wavetable, LOCTEXT("SquareWavetableDescription", "Wavetable"), LOCTEXT("SquareWavetableDescriptionTT", "Use a Wavetable interpolation to generate the Waveform"))
 	DEFINE_METASOUND_ENUM_END()
 	
@@ -906,7 +997,7 @@ namespace Metasound
 			: Super(InConstructParams)
 			, PulseWidth(InPulseWidth)
 		{}
-		void Generate(int32 InStartFrame, int32 InEndFrame, float InClampedFreq, float InClampedGlideEase)
+		void Generate(int32 InStartFrame, int32 InEndFrame, float InClampedFreq, float InClampedGlideEase, bool BiPolar)
 		{
 			int32 NumFrames = InEndFrame - InStartFrame;
 			float ClampedPulseWidth = FMath::Clamp(*PulseWidth, 0.0f, 1.0f);
@@ -918,6 +1009,7 @@ namespace Metasound
 					InClampedFreq,
 					InClampedGlideEase,
 					ClampedPulseWidth,
+					BiPolar, 
 					MakeArrayView(this->AudioBuffer->GetData() + InStartFrame, NumFrames), // Not aligned.
 				});
 			}
@@ -940,7 +1032,7 @@ namespace Metasound
 			check(InFm->GetData());
 			check(InConstructParams.Settings.GetNumFramesPerBlock() == InFm->Num());
 		}
-		void Generate(int32 InStartFrame, int32 InEndFrame, float InClampedFreq, float InClampedGlideEase)
+		void Generate(int32 InStartFrame, int32 InEndFrame, float InClampedFreq, float InClampedGlideEase, bool BiPolar)
 		{
 			int32 NumFrames = InEndFrame - InStartFrame;
 			float ClampedPulseWidth = FMath::Clamp(*PulseWidth, 0.0f, 1.0f);
@@ -952,6 +1044,7 @@ namespace Metasound
 					InClampedFreq,
 					InClampedGlideEase,
 					ClampedPulseWidth,
+					BiPolar, 
 					MakeArrayView(this->AudioBuffer->GetData() + InStartFrame, NumFrames), // Not aligned.
 					MakeArrayView(this->FM->GetData() + InStartFrame, NumFrames), // Not aligned.
 				});
diff --git a/Engine/Source/Runtime/SignalProcessing/Private/Osc.cpp b/Engine/Source/Runtime/SignalProcessing/Private/Osc.cpp
--- a/Engine/Source/Runtime/SignalProcessing/Private/Osc.cpp
+++ b/Engine/Source/Runtime/SignalProcessing/Private/Osc.cpp
@@ -303,13 +303,11 @@ namespace Audio
 				float SquareSaw2 = GetBipolar(NewPhase);
 				SquareSaw2 += PolySmooth(NewPhase, PhaseInc);
 
-				// Subtracting 2 saws creates a square wave!
-				Output = 0.5f * SquareSaw1 - 0.5f * SquareSaw2;
-
-				// Apply DC correction
-				const float Correction = (CurrentPulseWidth < 0.5f) ? (1.0f / (1.0f - CurrentPulseWidth)) : (1.0f / CurrentPulseWidth);
-
-				Output *= Correction;
+				// Subtract 2 saws, then apply DC correction 
+				// Simplified version of 
+				// float Output = 0.5f * SquareSaw1 - 0.5f * SquareSaw2;
+				// Output = 2.0f * (Output + CurrentPulseWidth) - 1.0f;
+				return SquareSaw1 - SquareSaw2 + 2.0f * (CurrentPulseWidth - 0.5f);
 			}
 			break;
 
