Metasound trigger compare bool

diff --git a/Engine/Plugins/Runtime/Metasound/Source/MetasoundStandardNodes/Private/MetasoundTriggerCompareNode.cpp b/Engine/Plugins/Runtime/Metasound/Source/MetasoundStandardNodes/Private/MetasoundTriggerCompareNode.cpp
--- a/Engine/Plugins/Runtime/Metasound/Source/MetasoundStandardNodes/Private/MetasoundTriggerCompareNode.cpp
+++ b/Engine/Plugins/Runtime/Metasound/Source/MetasoundStandardNodes/Private/MetasoundTriggerCompareNode.cpp
@@ -34,6 +34,9 @@ namespace Metasound
 
 	using FTriggerCompareNodeFloat = TTriggerCompareNode<float>;
 	METASOUND_REGISTER_NODE(FTriggerCompareNodeFloat)
+
+	using FTriggerCompareNodeBool = TTriggerCompareNode<bool>;
+	METASOUND_REGISTER_NODE(FTriggerCompareNodeBool)
 }
 
 #undef LOCTEXT_NAMESPACE // MetasoundTriggerDelayNode
