Fixing tooltip on trigger accumulator

diff --git a/Engine/Plugins/Runtime/Metasound/Source/MetasoundStandardNodes/Public/MetasoundTriggerAccumulatorNode.h b/Engine/Plugins/Runtime/Metasound/Source/MetasoundStandardNodes/Public/MetasoundTriggerAccumulatorNode.h
--- a/Engine/Plugins/Runtime/Metasound/Source/MetasoundStandardNodes/Public/MetasoundTriggerAccumulatorNode.h
+++ b/Engine/Plugins/Runtime/Metasound/Source/MetasoundStandardNodes/Public/MetasoundTriggerAccumulatorNode.h
@@ -21,7 +21,7 @@ namespace Metasound
 
 	namespace TriggerAccumulatorVertexNames
 	{
-		METASOUND_PARAM(InputAutoReset, "Auto Reset", "Input trigger which results in a delayed trigger.");	
+		METASOUND_PARAM(InputAutoReset, "Auto Reset", "Input trigger resets the trigger accumulation count.");	
 		METASOUND_PARAM(InputTrigger, "In {0}", "Trigger {0} input. All trigger inputs must be triggered before the output trigger is hit.");
 		METASOUND_PARAM(AccumulateOutputOnTrigger, "Out", "Triggered when all input triggers have been triggered. Call Reset to reset the state or use \"Auto Reset\".");
 	}
