Tooltip fix for ADSR envelope

diff --git a/Engine/Plugins/Runtime/Metasound/Source/MetasoundStandardNodes/Private/MetasoundADSREnvelopeNode.cpp b/Engine/Plugins/Runtime/Metasound/Source/MetasoundStandardNodes/Private/MetasoundADSREnvelopeNode.cpp
--- a/Engine/Plugins/Runtime/Metasound/Source/MetasoundStandardNodes/Private/MetasoundADSREnvelopeNode.cpp
+++ b/Engine/Plugins/Runtime/Metasound/Source/MetasoundStandardNodes/Private/MetasoundADSREnvelopeNode.cpp
@@ -28,7 +28,7 @@ namespace Metasound
 		METASOUND_PARAM(InputAttackTime, "Attack Time", "Attack time of the envelope.");
 		METASOUND_PARAM(InputDecayTime, "Decay Time", "Decay time of the envelope.");
 		METASOUND_PARAM(InputSustainLevel, "Sustain Level", "The sustain level.");
-		METASOUND_PARAM(InputReleaseTime, "Release Time", "Decay time of the envelope.");
+		METASOUND_PARAM(InputReleaseTime, "Release Time", "Release time of the envelope.");
 
 		METASOUND_PARAM(InputAttackCurve, "Attack Curve", "The exponential curve factor of the attack. 1.0 = linear growth, < 1.0 logorithmic growth, > 1.0 exponential growth.");
 		METASOUND_PARAM(InputDecayCurve, "Decay Curve", "The exponential curve factor of the decay. 1.0 = linear decay, < 1.0 exponential decay, > 1.0 logorithmic decay.");
