Fixing waveplayer tooltip

diff --git a/Engine/Plugins/Runtime/Metasound/Source/MetasoundEngine/Private/MetasoundWavePlayerNode.cpp b/Engine/Plugins/Runtime/Metasound/Source/MetasoundEngine/Private/MetasoundWavePlayerNode.cpp
--- a/Engine/Plugins/Runtime/Metasound/Source/MetasoundEngine/Private/MetasoundWavePlayerNode.cpp
+++ b/Engine/Plugins/Runtime/Metasound/Source/MetasoundEngine/Private/MetasoundWavePlayerNode.cpp
@@ -743,7 +743,7 @@ namespace Metasound
 				TOutputDataVertexModel<FTrigger>(OutputTriggerOnCuePointName, OutputTriggerOnCuePointTT),
 				TOutputDataVertexModel<int32>(OutputCuePointIDName, OutputCuePointIDTT),
 				TOutputDataVertexModel<FString>(OutputCuePointLabelName, OutputCuePointLabelTT),
-				TOutputDataVertexModel<float>(OutputLoopPercentName, OutputAudioRightNameTT),
+				TOutputDataVertexModel<float>(OutputLoopPercentName, OutputLoopPercentTT),
 				TOutputDataVertexModel<float>(OutputPlaybackLocationName, OutputPlaybackLocationTT),
 				TOutputDataVertexModel<FAudioBuffer>(OutputAudioLeftName, OutputAudioLeftNameTT),
 				TOutputDataVertexModel<FAudioBuffer>(OutputAudioRightName, OutputAudioRightNameTT)
