MetasoundDiffuserNode
- essentially takes the APF series from Flexiverb and splits it out as its own node in order to "smear" audio without the long delay lines of a full reverb.

diff --git a/Engine/Plugins/Runtime/Metasound/Source/MetasoundStandardNodes/Private/MetasoundDiffuserNode.cpp b/Engine/Plugins/Runtime/Metasound/Source/MetasoundStandardNodes/Private/MetasoundDiffuserNode.cpp
new file mode 100644
--- /dev/null
+++ b/Engine/Plugins/Runtime/Metasound/Source/MetasoundStandardNodes/Private/MetasoundDiffuserNode.cpp
@@ -0,0 +1,219 @@
+// Copyright Epic Games, Inc. All Rights Reserved.
+
+#include "Internationalization/Text.h"
+#include "MetasoundExecutableOperator.h"
+#include "MetasoundNodeRegistrationMacro.h"
+#include "MetasoundDataTypeRegistrationMacro.h"
+#include "MetasoundPrimitives.h"
+#include "MetasoundParamHelper.h"
+#include "MetasoundStandardNodesNames.h"
+#include "MetasoundStandardNodesCategories.h"
+#include "MetasoundAudioBuffer.h"
+#include "DSP/LongDelayAPF.h"
+
+#define LOCTEXT_NAMESPACE "MetasoundStandardNodes_DiffuserNode"
+namespace Metasound
+{
+	namespace DiffuserNode
+	{
+		METASOUND_PARAM(InputAudio, "Input Audio", "Incoming Audio");
+		METASOUND_PARAM(InputDiffusionDepth, "Depth", "1 - 5: The number of filters to use to diffuse the audio. Will not update while running.");
+		METASOUND_PARAM(InputFeedbackGain, "Feedback", "0 - 1: the amount of feedback on each diffuser.");
+		METASOUND_PARAM(OutputAudio, "Output Audio", "Diffuse Audio");
+
+		// APF Lengths, as a ratio of samplerate. 
+		//These were influenced by the APF lengths used in Flexiverb.cpp, except in ms instead of samples
+		static const float APFLengthsMs[] = { 5.3125f, 1.1458f, 11.5833f, 9.1875f, 1.21f, 7.104167f };
+		constexpr int32 MaxNumAPFs = 6;
+
+		// max number of samples the APFs will allocate for internal work buffers
+		constexpr int32 MaxBufferLengthSamples = 240;
+	}
+
+	class METASOUNDSTANDARDNODES_API FDiffuserOperator : public TExecutableOperator<FDiffuserOperator>
+	{
+	public:
+
+		static const FNodeClassMetadata& GetNodeInfo()
+		{
+			auto InitNodeInfo = []() -> FNodeClassMetadata
+			{
+				FNodeClassMetadata Info;
+				Info.ClassName = { StandardNodes::Namespace, TEXT("Diffuser"), StandardNodes::AudioVariant };
+				Info.MajorVersion = 1;
+				Info.MinorVersion = 0;
+				Info.DisplayName = LOCTEXT("Metasound_DiffuserDisplayName", "Diffuser");
+				Info.Description = LOCTEXT("Metasound_DiffuserNodeDescription", "Applies diffusion to incoming audio.");
+				Info.Author = PluginAuthor;
+				Info.PromptIfMissing = PluginNodeMissingPrompt;
+				Info.DefaultInterface = GetVertexInterface();
+				Info.CategoryHierarchy = { NodeCategories::Delays };
+
+				return Info;
+			};
+
+			static const FNodeClassMetadata Info = InitNodeInfo();
+
+			return Info;
+		}
+
+		static const FVertexInterface& GetVertexInterface()
+		{
+			using namespace DiffuserNode;
+
+			static const FVertexInterface Interface(
+				FInputVertexInterface(
+					TInputDataVertexModel<FAudioBuffer>(METASOUND_GET_PARAM_NAME_AND_TT(InputAudio)),
+					TInputDataVertexModel<int32>(METASOUND_GET_PARAM_NAME_AND_TT(InputDiffusionDepth), 4),
+					TInputDataVertexModel<float>(METASOUND_GET_PARAM_NAME_AND_TT(InputFeedbackGain), 0.707f)
+				),
+				FOutputVertexInterface(
+					TOutputDataVertexModel<FAudioBuffer>(METASOUND_GET_PARAM_NAME_AND_TT(OutputAudio))
+				)
+			);
+
+			return Interface;
+		}
+
+		static TUniquePtr<IOperator> CreateOperator(const FCreateOperatorParams& InParams, FBuildErrorArray& OutErrors)
+		{
+			using namespace DiffuserNode;
+
+			const FDataReferenceCollection& InputCollection = InParams.InputDataReferences;
+			const FInputVertexInterface& InputInterface = GetVertexInterface().GetInputInterface();
+
+			FAudioBufferReadRef AudioIn = InputCollection.GetDataReadReferenceOrConstruct<FAudioBuffer>(METASOUND_GET_PARAM_NAME(InputAudio), InParams.OperatorSettings);
+			FInt32ReadRef DepthIn = InputCollection.GetDataReadReferenceOrConstructWithVertexDefault<int32>(InputInterface, METASOUND_GET_PARAM_NAME(InputDiffusionDepth), InParams.OperatorSettings);
+			FFloatReadRef FeedbackIn = InputCollection.GetDataReadReferenceOrConstructWithVertexDefault<float>(InputInterface, METASOUND_GET_PARAM_NAME(InputFeedbackGain), InParams.OperatorSettings);
+
+			return MakeUnique<FDiffuserOperator>(InParams.OperatorSettings, AudioIn, DepthIn, FeedbackIn);
+		}
+
+		FDiffuserOperator(const FOperatorSettings& InSettings,
+			const FAudioBufferReadRef& InAudioInput,
+			const FInt32ReadRef& InDepth,
+			const FFloatReadRef& InFeedback)
+			: AudioInput(InAudioInput)
+			, Depth(InDepth)
+			, Feedback(InFeedback)
+			, AudioOutput(FAudioBufferWriteRef::CreateNew(InSettings))
+		{
+			using namespace DiffuserNode;
+
+			ClampedDepth = FMath::Clamp(*InDepth, 1, MaxNumAPFs);
+
+			const float FeedbackGain = *InFeedback;
+			const float SampleRatio =  InSettings.GetSampleRate() / 1000.f;
+
+			Delays.Reset(ClampedDepth);
+
+			for (int32 FilterIdx = 0; FilterIdx < ClampedDepth; ++FilterIdx)
+			{
+				// This is all because FLongDelayAPF has a TUniquePtr<FAlignedBlockBuffer> and can't be copied instantiated directly into an array
+				TUniquePtr<Audio::FLongDelayAPF>& NewDelay = GetDelayFromIndex(FilterIdx);
+				NewDelay = MakeUnique<Audio::FLongDelayAPF>(FeedbackGain, SampleRatio * APFLengthsMs[FilterIdx], MaxBufferLengthSamples);
+
+				Delays.Add(NewDelay.Get());
+			}
+		}
+
+		virtual FDataReferenceCollection GetInputs() const override
+		{
+			using namespace DiffuserNode;
+
+			FDataReferenceCollection InputDataReferences;
+			InputDataReferences.AddDataReadReference(METASOUND_GET_PARAM_NAME(InputAudio), AudioInput);
+			InputDataReferences.AddDataReadReference(METASOUND_GET_PARAM_NAME(InputDiffusionDepth), Depth);
+			InputDataReferences.AddDataReadReference(METASOUND_GET_PARAM_NAME(InputFeedbackGain), Feedback);
+
+			return InputDataReferences;
+		}
+
+		virtual FDataReferenceCollection GetOutputs() const override
+		{
+			using namespace DiffuserNode;
+
+			FDataReferenceCollection OutputDataReferences;
+			OutputDataReferences.AddDataReadReference(METASOUND_GET_PARAM_NAME(OutputAudio), AudioOutput);
+
+			return OutputDataReferences;
+		}
+
+		void Execute()
+		{
+			const int32 NumSamples = AudioInput->Num();
+
+			FMemory::Memcpy(AudioOutput->GetData(), AudioInput->GetData(), NumSamples * sizeof(float));
+
+			for (Audio::FLongDelayAPF* Delay : Delays)
+			{
+				Delay->SetG(*Feedback);
+				Delay->ProcessAudio(*AudioOutput);
+			}
+		}
+
+	private:
+		// The input audio buffer
+		FAudioBufferReadRef AudioInput;
+
+		// The amount that the wave is shaped
+		FInt32ReadRef Depth;
+
+		int32 ClampedDepth = 1;
+
+		// DC offset to apply before gain
+		FFloatReadRef Feedback;
+
+		// The audio output
+		FAudioBufferWriteRef AudioOutput;
+
+		TArray<Audio::FLongDelayAPF*> Delays;
+
+		// Can't hold the unique ptrs of this directly in an array because of hidden copy constructors
+		TUniquePtr<Audio::FLongDelayAPF> Delay_0;
+		TUniquePtr<Audio::FLongDelayAPF> Delay_1;
+		TUniquePtr<Audio::FLongDelayAPF> Delay_2;
+		TUniquePtr<Audio::FLongDelayAPF> Delay_3;
+		TUniquePtr<Audio::FLongDelayAPF> Delay_4;
+		TUniquePtr<Audio::FLongDelayAPF> Delay_5;
+
+		FORCEINLINE TUniquePtr<Audio::FLongDelayAPF>& GetDelayFromIndex(int32 Idx)
+		{
+			switch (Idx)
+			{
+			case 0:
+				return Delay_0;
+			case 1:
+				return Delay_1;
+			case 2:
+				return Delay_2;
+			case 3:
+				return Delay_3;
+			case 4:
+				return Delay_4;
+			case 5:
+				return Delay_5;
+			default:
+				ensure(false);
+			}
+
+			return Delay_0;
+		}
+	};
+
+	class FDiffuserNode : public FNodeFacade
+	{
+	public:
+		/**
+		 * Constructor used by the Metasound Frontend.
+		 */
+		FDiffuserNode(const FNodeInitData& InitData)
+			: FNodeFacade(InitData.InstanceName, InitData.InstanceID, TFacadeOperatorClass<FDiffuserOperator>())
+		{
+		}
+	};
+
+	METASOUND_REGISTER_NODE(FDiffuserNode)
+}
+
+#undef LOCTEXT_NAMESPACE
\ No newline at end of file
diff --git a/Engine/Source/Runtime/SignalProcessing/Private/LongDelayAPF.cpp b/Engine/Source/Runtime/SignalProcessing/Private/LongDelayAPF.cpp
--- a/Engine/Source/Runtime/SignalProcessing/Private/LongDelayAPF.cpp
+++ b/Engine/Source/Runtime/SignalProcessing/Private/LongDelayAPF.cpp
@@ -34,6 +34,33 @@ FLongDelayAPF::FLongDelayAPF(float InG, int32 InNumDelaySamples, int32 InMaxNumI
 FLongDelayAPF::~FLongDelayAPF()
 {}
 
+void FLongDelayAPF::ProcessAudio(FAlignedFloatBuffer& Samples)
+{
+	const int32 InNum = Samples.Num();
+
+	float* InData = Samples.GetData();
+	float* OutDelayData = WorkBuffer.GetData();
+
+	// Process audio one block at a time.
+	int32 LeftOver = InNum;
+	int32 BufferIndex = 0;
+	while (LeftOver != 0)
+	{
+		// Determine block size for this loop.
+		int32 NumToProcess = FMath::Min<int32>(NumInternalBufferSamples, LeftOver);
+		const float* InDelayData = DelayLine->InspectSamples(NumToProcess);
+
+		ProcessAudioBlock(&InData[BufferIndex], InDelayData, NumToProcess, &InData[BufferIndex], OutDelayData);
+
+		// Update delay line with new data.
+		DelayLine->RemoveSamples(NumToProcess);
+		DelayLine->AddSamples(OutDelayData, NumToProcess);
+
+		LeftOver -= NumToProcess;
+		BufferIndex += NumToProcess;
+	}
+}
+
 void FLongDelayAPF::ProcessAudio(const FAlignedFloatBuffer& InSamples, FAlignedFloatBuffer& OutSamples)
 {
 	const float* InData = InSamples.GetData();
@@ -102,9 +129,6 @@ void FLongDelayAPF::ProcessAudio(const FAlignedFloatBuffer& InSamples, FAlignedF
 void FLongDelayAPF::ProcessAudioBlock(const float* InSamples, const float* InDelaySamples, const int32 InNum, float* OutSamples, float* OutDelaySamples)
 {
 	// Calculate new delay line samples. "w[n] = x[n] + gw[n - d]"
-	//WeightedSumBuffer(InDelaySamples, G, InSamples, OutDelaySamples, InNum);
-	
-	
 	int32 NumToSIMD = InNum - (InNum % AUDIO_NUM_FLOATS_PER_VECTOR_REGISTER);
 
 	VectorRegister VG = MakeVectorRegister(G, G, G, G);
@@ -127,8 +151,6 @@ void FLongDelayAPF::ProcessAudioBlock(const float* InSamples, const float* InDel
 		// y[n] = -G * w[n] + w[n - D]
 		VectorRegister VOut = VectorMultiplyAdd(VOutDelay, VNG, VInDelay);
 		VectorStoreAligned(VOut, &OutSamples[i]);
-
-
 	}
 
 	// Calculate allpass for remaining samples that we couldn't SIMD
diff --git a/Engine/Source/Runtime/SignalProcessing/Public/DSP/LongDelayAPF.h b/Engine/Source/Runtime/SignalProcessing/Public/DSP/LongDelayAPF.h
--- a/Engine/Source/Runtime/SignalProcessing/Public/DSP/LongDelayAPF.h
+++ b/Engine/Source/Runtime/SignalProcessing/Public/DSP/LongDelayAPF.h
@@ -27,6 +27,9 @@ namespace Audio
 		// Process InSamples and place filtered data in OutSamples
 		void ProcessAudio(const FAlignedFloatBuffer& InSamples, FAlignedFloatBuffer& OutSamples);
 
+		// Process Samples in place
+		void ProcessAudio(FAlignedFloatBuffer& Samples);
+
 		// Process InSamples and place filtered data in OutSamples and delay line samples in OutDelaySamples.
 		// The OutDelaySamples correspond to "w[n] = InSamples[n] + InG * w[n - InDelay]"
 		void ProcessAudio(const FAlignedFloatBuffer& InSamples, FAlignedFloatBuffer& OutSamples, FAlignedFloatBuffer& OutDelaySamples);
