Compile fix (variable shadowing) for release build

diff --git a/Engine/Plugins/Runtime/Metasound/Source/MetasoundStandardNodes/Private/MetasoundOscillatorNodes.cpp b/Engine/Plugins/Runtime/Metasound/Source/MetasoundStandardNodes/Private/MetasoundOscillatorNodes.cpp
--- a/Engine/Plugins/Runtime/Metasound/Source/MetasoundStandardNodes/Private/MetasoundOscillatorNodes.cpp
+++ b/Engine/Plugins/Runtime/Metasound/Source/MetasoundStandardNodes/Private/MetasoundOscillatorNodes.cpp
@@ -614,20 +614,19 @@ namespace Metasound
 			// Clamp frequencies into Nyquist range.
 			const float ClampedFreq = FMath::Clamp(*BaseFrequency, -Nyquist, Nyquist);
 			const float ClampedGlideEase = Audio::GetLogFrequencyClamped(*GlideFactor, { 0.0f, 1.0f }, { 1.0f, 0.0001f });
-			const bool OutputBiPolar = *BiPolar;
 			AudioBuffer->Zero();
 
 			Derived* Self = static_cast<Derived*>(this);
 			PhaseReset->ExecuteBlock
 			(
-				[Self, ClampedFreq, ClampedGlideEase, OutputBiPolar](int32 InFrameStart, int32 InFrameEnd)
+				[Self, ClampedFreq, ClampedGlideEase](int32 InFrameStart, int32 InFrameEnd)
 				{
-					Self->Generate(InFrameStart, InFrameEnd, ClampedFreq, ClampedGlideEase, OutputBiPolar);
+					Self->Generate(InFrameStart, InFrameEnd, ClampedFreq, ClampedGlideEase);
 				},
-				[Self, ClampedFreq, ClampedGlideEase, OutputBiPolar](int32 InFrameStart, int32 InFrameEnd)
+				[Self, ClampedFreq, ClampedGlideEase](int32 InFrameStart, int32 InFrameEnd)
 				{
 					Self->ResetPhase(*Self->PhaseOffset);
-					Self->Generate(InFrameStart, InFrameEnd, ClampedFreq, ClampedGlideEase, OutputBiPolar);
+					Self->Generate(InFrameStart, InFrameEnd, ClampedFreq, ClampedGlideEase);
 				}
 			);
 		}
@@ -657,7 +656,7 @@ namespace Metasound
 			: Super(InConstructParams)
 		{}
 
-		void Generate(int32 InStartFrame, int32 InEndFrame, float InClampedFreq, float InClampedGlideEase, bool BiPolar)
+		void Generate(int32 InStartFrame, int32 InEndFrame, float InClampedFreq, float InClampedGlideEase)
 		{		
 			int32 NumFrames = InEndFrame - InStartFrame;
 
@@ -669,7 +668,7 @@ namespace Metasound
 					InClampedFreq,
 					InClampedGlideEase,
 					0.f,
-					BiPolar,
+					*this->BiPolar,
 					MakeArrayView(this->AudioBuffer->GetData() + InStartFrame, NumFrames), // Not aligned.
 				});
 			}			
@@ -693,7 +692,7 @@ namespace Metasound
 			return Inputs;
 		}
 
-		void Generate(int32 InStartFrame, int32 InEndFrame, float InClampedFreq, float InClampedGlideEase, bool BiPolar)
+		void Generate(int32 InStartFrame, int32 InEndFrame, float InClampedFreq, float InClampedGlideEase)
 		{
 			int32 NumFrames = InEndFrame - InStartFrame;
 			if (*this->Enabled && NumFrames > 0)
@@ -704,7 +703,7 @@ namespace Metasound
 					InClampedFreq,
 					InClampedGlideEase,
 					0.f,
-					BiPolar, 
+					*this->BiPolar,
 					MakeArrayView(this->AudioBuffer->GetData() + InStartFrame, NumFrames), // Not aligned.
 					MakeArrayView(this->Fm->GetData() + InStartFrame, NumFrames) // Not aligned.
 				});
@@ -997,7 +996,7 @@ namespace Metasound
 			: Super(InConstructParams)
 			, PulseWidth(InPulseWidth)
 		{}
-		void Generate(int32 InStartFrame, int32 InEndFrame, float InClampedFreq, float InClampedGlideEase, bool BiPolar)
+		void Generate(int32 InStartFrame, int32 InEndFrame, float InClampedFreq, float InClampedGlideEase)
 		{
 			int32 NumFrames = InEndFrame - InStartFrame;
 			float ClampedPulseWidth = FMath::Clamp(*PulseWidth, 0.0f, 1.0f);
@@ -1009,7 +1008,7 @@ namespace Metasound
 					InClampedFreq,
 					InClampedGlideEase,
 					ClampedPulseWidth,
-					BiPolar, 
+					*this->BiPolar, 
 					MakeArrayView(this->AudioBuffer->GetData() + InStartFrame, NumFrames), // Not aligned.
 				});
 			}
@@ -1032,7 +1031,7 @@ namespace Metasound
 			check(InFm->GetData());
 			check(InConstructParams.Settings.GetNumFramesPerBlock() == InFm->Num());
 		}
-		void Generate(int32 InStartFrame, int32 InEndFrame, float InClampedFreq, float InClampedGlideEase, bool BiPolar)
+		void Generate(int32 InStartFrame, int32 InEndFrame, float InClampedFreq, float InClampedGlideEase)
 		{
 			int32 NumFrames = InEndFrame - InStartFrame;
 			float ClampedPulseWidth = FMath::Clamp(*PulseWidth, 0.0f, 1.0f);
@@ -1044,7 +1043,7 @@ namespace Metasound
 					InClampedFreq,
 					InClampedGlideEase,
 					ClampedPulseWidth,
-					BiPolar, 
+					*this->BiPolar,
 					MakeArrayView(this->AudioBuffer->GetData() + InStartFrame, NumFrames), // Not aligned.
 					MakeArrayView(this->FM->GetData() + InStartFrame, NumFrames), // Not aligned.
 				});
