[Backout] - CL18345241
Original CL Desc
-----------------------------------------------------------------
EngineTest fix-ups for interface changes
#rnx

diff --git a/Engine/Plugins/Runtime/Metasound/Source/MetasoundEngineTest/MetasoundEngineTest.Build.cs b/Engine/Plugins/Runtime/Metasound/Source/MetasoundEngineTest/MetasoundEngineTest.Build.cs
--- a/Engine/Plugins/Runtime/Metasound/Source/MetasoundEngineTest/MetasoundEngineTest.Build.cs
+++ b/Engine/Plugins/Runtime/Metasound/Source/MetasoundEngineTest/MetasoundEngineTest.Build.cs
@@ -16,7 +16,6 @@ namespace UnrealBuildTool.Rules
 					"Core",
 					"CoreUObject",
 					"Engine",
-					"AudioExtensions",
 					"MetasoundGraphCore",
 					"MetasoundGenerator",
 					"MetasoundFrontend",
diff --git a/Engine/Plugins/Runtime/Metasound/Source/MetasoundEngineTest/Private/EngineTestMetasoundSource.cpp b/Engine/Plugins/Runtime/Metasound/Source/MetasoundEngineTest/Private/EngineTestMetasoundSource.cpp
--- a/Engine/Plugins/Runtime/Metasound/Source/MetasoundEngineTest/Private/EngineTestMetasoundSource.cpp
+++ b/Engine/Plugins/Runtime/Metasound/Source/MetasoundEngineTest/Private/EngineTestMetasoundSource.cpp
@@ -2,13 +2,10 @@
 
 #include "AudioDevice.h"
 #include "Components/AudioComponent.h"
-#include "IAudioParameterInterfaceRegistry.h"
 #include "Interfaces/IPluginManager.h"
 #include "MetasoundFrontendController.h"
 #include "MetasoundFrontendSearchEngine.h"
-#include "MetasoundOutputFormatInterfaces.h"
 #include "MetasoundSource.h"
-#include "MetasoundSourceInterface.h"
 #include "Misc/AutomationTest.h"
 #include "Tests/AutomationCommon.h"
 
@@ -63,46 +60,58 @@ namespace EngineTestMetasoundSourcePrivate
 
 	FMetasoundFrontendDocument CreateMetaSoundMonoSourceDocument()
 	{
-		using namespace Audio;
 		using namespace Metasound;
 		using namespace Metasound::Frontend;
 
 		FMetasoundFrontendDocument Document;
+		Document.InterfaceVersions.Add(FMetasoundFrontendVersion{"MonoSource", {1, 1}});
 
 		Document.RootGraph.Metadata.SetClassName(FMetasoundFrontendClassName { "Namespace", "Unit Test Node", *LexToString(FGuid::NewGuid()) });
 		Document.RootGraph.Metadata.SetType(EMetasoundFrontendClassType::Graph);
 
 		FDocumentHandle DocumentHandle = IDocumentController::CreateDocumentHandle(Document);
-		FGraphHandle RootGraph = DocumentHandle->GetRootGraph();
+		Metasound::Frontend::FGraphHandle RootGraph = DocumentHandle->GetRootGraph();
 		check(RootGraph->IsValid());
 
-		// Add default source & mono interface members (OnPlay, OnFinished & Mono Out)
-		FModifyRootGraphInterfaces InterfaceTransform({ SourceInterface::GetVersion(), OutputFormatMonoInterface::GetVersion() }, { });
-		InterfaceTransform.Transform(DocumentHandle);
+		// Input on play
+		FMetasoundFrontendClassInput OnPlayInput;
+		OnPlayInput.Name = TEXT("On Play");
+		OnPlayInput.TypeName = GetMetasoundDataTypeName<FTrigger>();
+		OnPlayInput.VertexID = FGuid::NewGuid();
 
-		// Input on Play
-		FNodeHandle OnPlayOutputNode = RootGraph->GetOutputNodeWithName(SourceInterface::Inputs::OnPlay);
-		check(OnPlayOutputNode->IsValid());
+		FNodeHandle OnPlayInputNode = RootGraph->AddInputVertex(OnPlayInput);
+		check(OnPlayInputNode->IsValid());
 
 		// Input Frequency
 		FMetasoundFrontendClassInput FrequencyInput;
-		FrequencyInput.Name = "Frequency";
+		FrequencyInput.Name = TEXT("Frequency");
 		FrequencyInput.TypeName = GetMetasoundDataTypeName<float>();
 		FrequencyInput.VertexID = FGuid::NewGuid();
 		FrequencyInput.DefaultLiteral.Set(100.f);
+
 		FNodeHandle FrequencyInputNode = RootGraph->AddInputVertex(FrequencyInput);
 		check(FrequencyInputNode->IsValid());
 
-		// Output On Finished
-		FNodeHandle OnFinishedOutputNode = RootGraph->GetOutputNodeWithName(SourceInterface::Outputs::OnFinished);
-		check(OnFinishedOutputNode->IsValid());
+
+		// Output Is Finished
+		FMetasoundFrontendClassOutput IsFinishedOutput;
+		IsFinishedOutput.Name = TEXT("On Finished");
+		IsFinishedOutput.TypeName = Metasound::GetMetasoundDataTypeName<FTrigger>();
+		IsFinishedOutput.VertexID = FGuid::NewGuid();
+		FNodeHandle IsFinishedOutputNode = RootGraph->AddOutputVertex(IsFinishedOutput);
+		check(IsFinishedOutputNode->IsValid());
 
 		// Output Audio
-		FNodeHandle AudioOutputNode = RootGraph->GetOutputNodeWithName(OutputFormatMonoInterface::Outputs::MonoOut);
-		check(AudioOutputNode->IsValid());
+		FMetasoundFrontendClassOutput AudioOutput;
+		AudioOutput.Name = TEXT("Audio:0");
+		AudioOutput.TypeName = Metasound::GetMetasoundDataTypeName<FAudioBuffer>();
+		AudioOutput.VertexID = FGuid::NewGuid();
+		
+		FNodeHandle OutputAudioNode = RootGraph->AddOutputVertex(AudioOutput);
+		check(OutputAudioNode->IsValid());
 
-		// osc node
-		FNodeHandle OscNode = AddNode(*RootGraph, { "UE", "Sine", "Audio" });
+		// osc node 
+		FNodeHandle OscNode = AddNode(*RootGraph, {TEXT("UE"), TEXT("Sine"), TEXT("Audio")});
 
 		// Make connections:
 
@@ -113,7 +122,7 @@ namespace EngineTestMetasoundSourcePrivate
 
 		// oscillator to output
 		OutputToConnect = OscNode->GetOutputWithVertexName("Audio");
-		InputToConnect = AudioOutputNode->GetInputWithVertexName(OutputFormatMonoInterface::Outputs::MonoOut);
+		InputToConnect = OutputAudioNode->GetInputWithVertexName("Audio:0");
 		ensure(InputToConnect->Connect(*OutputToConnect));
 
 		return Document;
