Fix for InterpTo Nodes and Conversion Nodes not propogating their initial output values at the point of construction

diff --git a/Engine/Plugins/Runtime/Metasound/Source/MetasoundStandardNodes/Private/MetasoundConversionNodes.cpp b/Engine/Plugins/Runtime/Metasound/Source/MetasoundStandardNodes/Private/MetasoundConversionNodes.cpp
--- a/Engine/Plugins/Runtime/Metasound/Source/MetasoundStandardNodes/Private/MetasoundConversionNodes.cpp
+++ b/Engine/Plugins/Runtime/Metasound/Source/MetasoundStandardNodes/Private/MetasoundConversionNodes.cpp
@@ -217,6 +217,8 @@ namespace Metasound
 			: InputValue(InInputValue)
 			, OutputValue(TDataWriteReferenceFactory<ToType>::CreateAny(InSettings))
 		{
+			using namespace MetasoundConversionNodePrivate;
+			TConversionNodeSpecialization<FromType, ToType>::GetConvertedValue(*InputValue, *OutputValue);
 		}
 
 		virtual ~TConversionOperator() = default;
diff --git a/Engine/Plugins/Runtime/Metasound/Source/MetasoundStandardNodes/Private/MetasoundDelayNode.cpp b/Engine/Plugins/Runtime/Metasound/Source/MetasoundStandardNodes/Private/MetasoundDelayNode.cpp
--- a/Engine/Plugins/Runtime/Metasound/Source/MetasoundStandardNodes/Private/MetasoundDelayNode.cpp
+++ b/Engine/Plugins/Runtime/Metasound/Source/MetasoundStandardNodes/Private/MetasoundDelayNode.cpp
@@ -57,10 +57,10 @@ namespace Metasound
 		// The amount of delay time
 		FTimeReadRef DelayTime;
 
-		// The the dry level
+		// The dry level
 		FFloatReadRef DryLevel;
 
-		// The the wet level
+		// The wet level
 		FFloatReadRef WetLevel;
 
 		// The feedback amount
@@ -72,18 +72,18 @@ namespace Metasound
 		// The internal delay buffer
 		Audio::FDelay DelayBuffer;
 
-		// The current delay time
+		// The previous delay time
 		float PrevDelayTimeMsec;
 
 		// Feedback sample
 		float FeedbackSample;
 	};
 
-	FDelayOperator::FDelayOperator(const FOperatorSettings& InSettings, 
-		const FAudioBufferReadRef& InAudioInput, 
+	FDelayOperator::FDelayOperator(const FOperatorSettings& InSettings,
+		const FAudioBufferReadRef& InAudioInput,
 		const FTimeReadRef& InDelayTime,
-		const FFloatReadRef& InDryLevel, 
-		const FFloatReadRef& InWetLevel, 
+		const FFloatReadRef& InDryLevel,
+		const FFloatReadRef& InWetLevel,
 		const FFloatReadRef& InFeedback)
 
 		: AudioInput(InAudioInput)
@@ -171,7 +171,7 @@ namespace Metasound
 		static const FVertexInterface Interface(
 			FInputVertexInterface(
 				TInputDataVertexModel<FAudioBuffer>(Delay::InParamNameAudioInput, LOCTEXT("AudioInputTooltip", "Audio input.")),
-				TInputDataVertexModel<FTime>(Delay::InParamNameDelayTime, LOCTEXT("DelayTimeTooltip", "The amount of time to delay the audio."), 1.0f),
+				TInputDataVertexModel<FTime>(Delay::InParamNameDelayTime, LOCTEXT("DelayTimeTooltip", "The amount of time to delay the audio, in seconds."), 1.0f),
 				TInputDataVertexModel<float>(Delay::InParamNameDryLevel, LOCTEXT("DryLevelTooltip", "The dry level of the delay."), 0.0f),
 				TInputDataVertexModel<float>(Delay::InParamNameWetLevel, LOCTEXT("WetlevelTooltip", "The wet level of the delay."), 1.0f),
 				TInputDataVertexModel<float>(Delay::InParamNameFeedbackAmount, LOCTEXT("FeedbackTooltip", "Feedback amount."), 0.0f)
diff --git a/Engine/Plugins/Runtime/Metasound/Source/MetasoundStandardNodes/Private/MetasoundInterpToNode.cpp b/Engine/Plugins/Runtime/Metasound/Source/MetasoundStandardNodes/Private/MetasoundInterpToNode.cpp
--- a/Engine/Plugins/Runtime/Metasound/Source/MetasoundStandardNodes/Private/MetasoundInterpToNode.cpp
+++ b/Engine/Plugins/Runtime/Metasound/Source/MetasoundStandardNodes/Private/MetasoundInterpToNode.cpp
@@ -81,6 +81,8 @@ namespace Metasound
 		BlockTimeDelta = 1.0f / BlockRate;
 
 		PreviousTargetValue = *TargetValue;
+
+		*ValueOutput = *TargetValue;
 	}
 
 	FDataReferenceCollection FInterpToOperator::GetInputs() const
