Removing blacklist terminology from MetaSound settings

diff --git a/Engine/Plugins/Runtime/Metasound/Config/BaseMetasound.ini b/Engine/Plugins/Runtime/Metasound/Config/BaseMetasound.ini
new file mode 100644
--- /dev/null
+++ b/Engine/Plugins/Runtime/Metasound/Config/BaseMetasound.ini
@@ -0,0 +1,3 @@
+[CoreRedirects]
++PropertyRedirects=(OldName="MetaSoundSettings.AutoUpdateBlacklist",NewName="AutoUpdateDenylist")
++PropertyRedirects=(OldName="MetaSoundSettings.AutoUpdateAssetBlacklist",NewName="AutoUpdateDenylist")
diff --git a/Engine/Plugins/Runtime/Metasound/Source/MetasoundEngine/Private/MetasoundAssetSubsystem.cpp b/Engine/Plugins/Runtime/Metasound/Source/MetasoundEngine/Private/MetasoundAssetSubsystem.cpp
--- a/Engine/Plugins/Runtime/Metasound/Source/MetasoundEngine/Private/MetasoundAssetSubsystem.cpp
+++ b/Engine/Plugins/Runtime/Metasound/Source/MetasoundEngine/Private/MetasoundAssetSubsystem.cpp
@@ -252,12 +252,12 @@ void UMetaSoundAssetSubsystem::RebuildDenyListCache(const UAssetManager& InAsset
 
 	AutoUpdateDenyListCache.Reset();
 
-	for (const FMetasoundFrontendClassName& ClassName : Settings->AutoUpdateBlacklist)
+	for (const FMetasoundFrontendClassName& ClassName : Settings->AutoUpdateDenylist)
 	{
 		AutoUpdateDenyListCache.Add(ClassName.GetFullName());
 	}
 
-	for (const FDefaultMetaSoundAssetAutoUpdateSettings& UpdateSettings : Settings->AutoUpdateAssetBlacklist)
+	for (const FDefaultMetaSoundAssetAutoUpdateSettings& UpdateSettings : Settings->AutoUpdateAssetDenylist)
 	{
 		if (UAssetManager* AssetManager = UAssetManager::GetIfValid())
 		{
diff --git a/Engine/Plugins/Runtime/Metasound/Source/MetasoundEngine/Public/MetasoundSettings.h b/Engine/Plugins/Runtime/Metasound/Source/MetasoundEngine/Public/MetasoundSettings.h
--- a/Engine/Plugins/Runtime/Metasound/Source/MetasoundEngine/Public/MetasoundSettings.h
+++ b/Engine/Plugins/Runtime/Metasound/Source/MetasoundEngine/Public/MetasoundSettings.h
@@ -46,11 +46,11 @@ public:
 
 	/** List of native MetaSound classes whose node references should not be AutoUpdated. */
 	UPROPERTY(EditAnywhere, config, Category = AutoUpdate, meta = (DisplayName = "DenyList", EditCondition = "bAutoUpdateEnabled"))
-	TArray<FMetasoundFrontendClassName> AutoUpdateBlacklist;
+	TArray<FMetasoundFrontendClassName> AutoUpdateDenylist;
 
 	/** List of MetaSound assets whose node references should not be AutoUpdated. */
 	UPROPERTY(EditAnywhere, config, Category = AutoUpdate, meta = (DisplayName = "Asset DenyList", EditCondition = "bAutoUpdateEnabled"))
-	TArray<FDefaultMetaSoundAssetAutoUpdateSettings> AutoUpdateAssetBlacklist;
+	TArray<FDefaultMetaSoundAssetAutoUpdateSettings> AutoUpdateAssetDenylist;
 
 	/** Directories to scan & automatically register MetaSound post initial asset scan on engine start-up.
 	  * May speed up subsequent calls to playback MetaSounds post asset scan but increases application load time.
