Fixing ADSR toottip

diff --git a/Engine/Plugins/Runtime/Metasound/Source/MetasoundStandardNodes/Private/MetasoundADSREnvelopeNode.cpp b/Engine/Plugins/Runtime/Metasound/Source/MetasoundStandardNodes/Private/MetasoundADSREnvelopeNode.cpp
--- a/Engine/Plugins/Runtime/Metasound/Source/MetasoundStandardNodes/Private/MetasoundADSREnvelopeNode.cpp
+++ b/Engine/Plugins/Runtime/Metasound/Source/MetasoundStandardNodes/Private/MetasoundADSREnvelopeNode.cpp
@@ -37,7 +37,7 @@ namespace Metasound
 		METASOUND_PARAM(OutputOnAttackTrigger, "On Attack Triggered", "Triggers when the envelope attack is triggered.");
 		METASOUND_PARAM(OutputOnDecayTrigger, "On Decay Triggered", "Triggers when the envelope decay begins and attack is finished.");
 		METASOUND_PARAM(OutputOnSustainTrigger, "On Sustain Triggered", "Triggers when the envelope sustain begins and attack is finished.");
-		METASOUND_PARAM(OutputOnReleaseTrigger, "On Release Triggered", "Triggers when the envelope attack is triggered.");
+		METASOUND_PARAM(OutputOnReleaseTrigger, "On Release Triggered", "Triggers when the envelope release is triggered.");
 		METASOUND_PARAM(OutputOnDone, "On Done", "Triggers when the envelope finishes.");
 		METASOUND_PARAM(OutputEnvelopeValue, "Out Envelope", "The output value of the envelope.");
 	}
