- Rename & redirect MetaSoundEditorGraph<Input>Literals to MetaSoundEditorGraph<DefaultMember>Literals
- First pass at consolidating editor code for API members/vertices where possible
- Make Variables & Output now use default literal for input value (so user doesn't have to hunt down literal box in graph)

Various bug fixes:
- Fixes crash in post undo after updating a variable default
- Fixes ':Array' flavor showing up in new SearchableComboBox when selecting `Is Array` for literal type
- Synchronize variable setter default literal with variable interface default literal
- Provide default literal field for outputs & variables for consistency
    - Added win is widget support for single input/output types (i.e. floats for now)
- Don't show trigger simulate option on non-transmissible inputs (ex. OnPlay)
- Set simulate button visibility to only show when parameter transmission is supported by input
- Make simulate buttons in graph only visible when sound plays
- Start sounds automatically when clicking simulate in inspector
- Fixed initial condition "holding" trigger state for trigger routing now that Transmission System is only used inter-MetaSound (not intra-MetaSound) and we are deprecating Send/Receives

#rnx

diff --git a/Engine/Plugins/Runtime/AudioWidgets/Source/AudioWidgets/Private/SAudioSlider.cpp b/Engine/Plugins/Runtime/AudioWidgets/Source/AudioWidgets/Private/SAudioSlider.cpp
--- a/Engine/Plugins/Runtime/AudioWidgets/Source/AudioWidgets/Private/SAudioSlider.cpp
+++ b/Engine/Plugins/Runtime/AudioWidgets/Source/AudioWidgets/Private/SAudioSlider.cpp
@@ -22,6 +22,7 @@ void SAudioSliderBase::Construct(const SAudioSliderBase::FArguments& InArgs)
 {
 	Style = InArgs._Style;
 	OnValueChanged = InArgs._OnValueChanged;
+	OnValueCommitted = InArgs._OnValueCommitted;
 	ValueAttribute = InArgs._Value;
 	SliderBackgroundColor = InArgs._SliderBackgroundColor;
 	SliderBarColor = InArgs._SliderBarColor;
@@ -59,19 +60,24 @@ void SAudioSliderBase::Construct(const SAudioSliderBase::FArguments& InArgs)
 			OnValueChanged.ExecuteIfBound(Value);
 			const float OutputValue = GetOutputValue(Value);
 			Label->SetValueText(OutputValue);
+		})
+		.OnMouseCaptureEnd_Lambda([this]()
+		{
+			OnValueCommitted.ExecuteIfBound(ValueAttribute.Get());
 		});
 
 	// Text label
 	SAssignNew(Label, SAudioTextBox)
 		.Style(&Style->TextBoxStyle)
 		.OnValueTextCommitted_Lambda([this](const FText& Text, ETextCommit::Type CommitType)
-			{
-				const float OutputValue = FCString::Atof(*Text.ToString());
-				const float LinValue = GetLinValue(OutputValue);
-				ValueAttribute.Set(LinValue);
-				Slider->SetValue(LinValue);
-				OnValueChanged.ExecuteIfBound(LinValue);
-			});
+		{
+			const float OutputValue = FCString::Atof(*Text.ToString());
+			const float LinValue = GetLinValue(OutputValue);
+			ValueAttribute.Set(LinValue);
+			Slider->SetValue(LinValue);
+			OnValueChanged.ExecuteIfBound(LinValue);
+			OnValueCommitted.ExecuteIfBound(LinValue);
+		});
 
 	ChildSlot
 	[
diff --git a/Engine/Plugins/Runtime/AudioWidgets/Source/AudioWidgets/Public/SAudioSlider.h b/Engine/Plugins/Runtime/AudioWidgets/Source/AudioWidgets/Public/SAudioSlider.h
--- a/Engine/Plugins/Runtime/AudioWidgets/Source/AudioWidgets/Public/SAudioSlider.h
+++ b/Engine/Plugins/Runtime/AudioWidgets/Source/AudioWidgets/Public/SAudioSlider.h
@@ -71,6 +71,9 @@ public:
 		/** Called when the value is changed by slider or typing */
 		SLATE_EVENT(FOnFloatValueChanged, OnValueChanged)
 
+		/** Called when the value is committed by mouse capture ending */
+		SLATE_EVENT(FOnFloatValueChanged, OnValueCommitted)
+
 	SLATE_END_ARGS()
 
 	SAudioSliderBase();
@@ -79,6 +82,9 @@ public:
 	// Holds a delegate that is executed when the slider's value changed.
 	FOnFloatValueChanged OnValueChanged;
 
+	// Holds a delegate that is executed when the slider's value is committed (mouse capture ends).
+	FOnFloatValueChanged OnValueCommitted;
+
 	/**
 	 * Construct the widget.
 	 *
diff --git a/Engine/Plugins/Runtime/Metasound/Config/BaseMetasound.ini b/Engine/Plugins/Runtime/Metasound/Config/BaseMetasound.ini
--- a/Engine/Plugins/Runtime/Metasound/Config/BaseMetasound.ini
+++ b/Engine/Plugins/Runtime/Metasound/Config/BaseMetasound.ini
@@ -1,3 +1,20 @@
 [CoreRedirects]
 +PropertyRedirects=(OldName="MetaSoundSettings.AutoUpdateBlacklist",NewName="AutoUpdateDenylist")
 +PropertyRedirects=(OldName="MetaSoundSettings.AutoUpdateAssetBlacklist",NewName="AutoUpdateDenylist")
+
++StructRedirects=(OldName="MetasoundEditorGraphInputBoolRef",NewName="MetasoundEditorGraphMemberDefaultBoolRef")
++StructRedirects=(OldName="MetasoundEditorGraphInputIntRef",NewName="MetasoundEditorGraphMemberDefaultIntRef")
++StructRedirects=(OldName="MetasoundEditorGraphInputObjectRef",NewName="MetasoundEditorGraphMemberDefaultObjectRef")
+ 
++ClassRedirects=(OldName="MetasoundEditorGraphInputLiteral",NewName="MetasoundEditorGraphMemberDefaultLiteral")
+ 
++ClassRedirects=(OldName="MetasoundEditorGraphInputBool",NewName="MetasoundEditorGraphMemberDefaultBool")
++ClassRedirects=(OldName="MetasoundEditorGraphInputBoolArray",NewName="MetasoundEditorGraphMemberDefaultBoolArray")
++ClassRedirects=(OldName="MetasoundEditorGraphInputInt",NewName="MetasoundEditorGraphMemberDefaultInt")
++ClassRedirects=(OldName="MetasoundEditorGraphInputIntArray",NewName="MetasoundEditorGraphMemberDefaultIntArray")
++ClassRedirects=(OldName="MetasoundEditorGraphInputFloat",NewName="MetasoundEditorGraphMemberDefaultFloat")
++ClassRedirects=(OldName="MetasoundEditorGraphInputFloatArray",NewName="MetasoundEditorGraphMemberDefaultFloatArray")
++ClassRedirects=(OldName="MetasoundEditorGraphInputString",NewName="MetasoundEditorGraphMemberDefaultString")
++ClassRedirects=(OldName="MetasoundEditorGraphInputStringArray",NewName="MetasoundEditorGraphMemberDefaultStringArray")
++ClassRedirects=(OldName="MetasoundEditorGraphInputObject",NewName="MetasoundEditorGraphMemberDefaultObject")
++ClassRedirects=(OldName="MetasoundEditorGraphInputObjectArray",NewName="MetasoundEditorGraphMemberDefaultObjectArray")
\ No newline at end of file
diff --git a/Engine/Plugins/Runtime/Metasound/Source/MetasoundEditor/Private/MetasoundDetailCustomization.cpp b/Engine/Plugins/Runtime/Metasound/Source/MetasoundEditor/Private/MetasoundDetailCustomization.cpp
--- a/Engine/Plugins/Runtime/Metasound/Source/MetasoundEditor/Private/MetasoundDetailCustomization.cpp
+++ b/Engine/Plugins/Runtime/Metasound/Source/MetasoundEditor/Private/MetasoundDetailCustomization.cpp
@@ -111,7 +111,7 @@ namespace Metasound
 							FModifyRootGraphInterfaces ModifyTransform({ }, { InterfaceToAdd });
 							ModifyTransform.SetDefaultNodeLocations(false); // Don't automatically add nodes to ed graph
 							ModifyTransform.Transform(DocumentHandle);
-							MetaSoundAsset->SetInterfacesUpdatedOnSynchronization();
+							MetaSoundAsset->SetUpdateDetailsOnSynchronization();
 						}
 
 						UpdateInterfaceNames();
@@ -170,7 +170,7 @@ namespace Metasound
 
 					UpdateInterfaceNames();
 					InterfaceComboBox->RefreshOptions();
-					MetaSoundAsset->SetInterfacesUpdatedOnSynchronization();
+					MetaSoundAsset->SetUpdateDetailsOnSynchronization();
 					MetaSoundAsset->SetSynchronizationRequired();
 				}), LOCTEXT("RemoveInterfaceTooltip", "Removes all interfaces from the given MetaSound."))
 			];
@@ -215,7 +215,7 @@ namespace Metasound
 
 					UpdateInterfaceNames();
 					InterfaceComboBox->RefreshOptions();
-					MetaSoundAsset->SetInterfacesUpdatedOnSynchronization();
+					MetaSoundAsset->SetUpdateDetailsOnSynchronization();
 					MetaSoundAsset->SetSynchronizationRequired();
 				}), LOCTEXT("RemoveInterfaceTooltip", "Removes the associated interface from the MetaSound."));
 
diff --git a/Engine/Plugins/Runtime/Metasound/Source/MetasoundEditor/Private/MetasoundEditor.cpp b/Engine/Plugins/Runtime/Metasound/Source/MetasoundEditor/Private/MetasoundEditor.cpp
--- a/Engine/Plugins/Runtime/Metasound/Source/MetasoundEditor/Private/MetasoundEditor.cpp
+++ b/Engine/Plugins/Runtime/Metasound/Source/MetasoundEditor/Private/MetasoundEditor.cpp
@@ -27,6 +27,7 @@
 #include "MetasoundEditorCommands.h"
 #include "MetasoundEditorGraph.h"
 #include "MetasoundEditorGraphBuilder.h"
+#include "MetasoundEditorGraphInputNode.h"
 #include "MetasoundEditorGraphNode.h"
 #include "MetasoundEditorGraphSchema.h"
 #include "MetasoundEditorGraphValidation.h"
@@ -934,6 +935,7 @@ namespace Metasound
 							SetPreviewID(INDEX_NONE);
 							PlayTime = 0.0;
 							PlayTimeWidget->SetText(FText::GetEmpty());
+
 							return EActiveTimerReturnType::Stop;
 						}
 					})
@@ -1346,6 +1348,10 @@ namespace Metasound
 				{
 					Selection.Add(OutputNode->Output);
 				}
+				else if (UMetasoundEditorGraphVariableNode* VariableNode = Cast<UMetasoundEditorGraphVariableNode>(NodeObject))
+				{
+					Selection.Add(VariableNode->Variable);
+				}
 				else
 				{
 					Selection.Add(NodeObject);
@@ -2173,28 +2179,30 @@ namespace Metasound
 			FMetasoundAssetBase* MetasoundAsset = IMetasoundUObjectRegistry::Get().GetObjectAsAssetBase(Metasound);
 			check(MetasoundAsset);
 
-			if (MetasoundAsset->GetSynchronizationPending())
+			if (MetasoundAsset->GetSynchronizationRequired())
 			{
 				// Capture before synchronizing as the flag is cleared therein.
-				const bool bInterfacesUpdated = MetasoundAsset->GetSynchronizationInterfacesUpdated();
-
+				const bool bShouldRefreshDetails = MetasoundAsset->GetSynchronizationUpdateDetails();
 				FGraphBuilder::SynchronizeGraph(*Metasound);
 
-				if (MetasoundDetails.IsValid())
+				if (bShouldRefreshDetails)
+				{
+					RefreshDetails();
+				}
+				else
 				{
-					// Only refresh the details panel if viewing the graph details of the object is no longer valid
-					auto ShouldRefreshDetails = [bInterfacesUpdated](TWeakObjectPtr<UObject> Obj)
+					// Also refresh if the object in the panel has gone invalid
+					auto ShouldRefreshDetails = [](TWeakObjectPtr<UObject> Obj)
 					{
 						if (!Obj.IsValid() || !IsValid(Obj.Get()))
 						{
 							return true;
 						}
 
-						return bInterfacesUpdated && IMetasoundUObjectRegistry::Get().IsRegisteredClass(Obj.Get());
+						return IMetasoundUObjectRegistry::Get().IsRegisteredClass(Obj.Get());
 					};
 
-					const bool bShouldRefreshDetails = Algo::AnyOf(MetasoundDetails->GetSelectedObjects(), ShouldRefreshDetails);
-					if (bShouldRefreshDetails)
+					if (Algo::AnyOf(MetasoundDetails->GetSelectedObjects(), ShouldRefreshDetails))
 					{
 						RefreshDetails();
 					}
diff --git a/Engine/Plugins/Runtime/Metasound/Source/MetasoundEditor/Private/MetasoundEditor.h b/Engine/Plugins/Runtime/Metasound/Source/MetasoundEditor/Private/MetasoundEditor.h
--- a/Engine/Plugins/Runtime/Metasound/Source/MetasoundEditor/Private/MetasoundEditor.h
+++ b/Engine/Plugins/Runtime/Metasound/Source/MetasoundEditor/Private/MetasoundEditor.h
@@ -114,6 +114,9 @@ namespace Metasound
 			virtual TStatId GetStatId() const override;
 			virtual ETickableTickType GetTickableTickType() const override { return ETickableTickType::Always; }
 
+			virtual void Play() override;
+			virtual void Stop() override;
+
 			/** Whether pasting the currently selected nodes is permissible */
 			bool CanPasteNodes();
 
@@ -241,8 +244,6 @@ namespace Metasound
 
 			/** Toolbar command methods */
 			void ExecuteNode();
-			void Play();
-			void Stop();
 
 			/** Whether we can play the current selection of nodes */
 			bool CanExecuteNode() const;
diff --git a/Engine/Plugins/Runtime/Metasound/Source/MetasoundEditor/Private/MetasoundEditorGraph.cpp b/Engine/Plugins/Runtime/Metasound/Source/MetasoundEditor/Private/MetasoundEditorGraph.cpp
--- a/Engine/Plugins/Runtime/Metasound/Source/MetasoundEditor/Private/MetasoundEditorGraph.cpp
+++ b/Engine/Plugins/Runtime/Metasound/Source/MetasoundEditor/Private/MetasoundEditorGraph.cpp
@@ -8,16 +8,49 @@
 #include "Interfaces/ITargetPlatform.h"
 #include "MetasoundAssetBase.h"
 #include "MetasoundEditorGraphBuilder.h"
-#include "MetasoundEditorGraphInputNodes.h"
+#include "MetasoundEditorGraphInputNode.h"
+#include "MetasoundEditorGraphMemberDefaults.h"
 #include "MetasoundEditorGraphNode.h"
 #include "MetasoundEditorGraphValidation.h"
 #include "MetasoundEditorModule.h"
 #include "MetasoundUObjectRegistry.h"
+#include "MetasoundVariableNodes.h"
 #include "MetasoundVertex.h"
 #include "ScopedTransaction.h"
+#include "UObject/UnrealType.h"
 
 #define LOCTEXT_NAMESPACE "MetaSoundEditor"
 
+
+namespace Metasound
+{
+	namespace Editor
+	{
+		namespace GraphPrivate
+		{
+			void OnLiteralChanged(UMetasoundEditorGraphMember& InMember, EPropertyChangeType::Type InChangeType)
+			{
+				constexpr bool bPostTransaction = false;
+				InMember.UpdateFrontendDefaultLiteral(bPostTransaction);
+
+				const bool bCommitChange = InChangeType != EPropertyChangeType::Interactive;
+				if (bCommitChange)
+				{
+					InMember.MarkNodesForRefresh();
+					if (UObject* MetaSound = InMember.GetOutermostObject())
+					{
+						FGraphBuilder::RegisterGraphWithFrontend(*MetaSound);
+						if (FMetasoundAssetBase* MetaSoundAsset = IMetasoundUObjectRegistry::Get().GetObjectAsAssetBase(MetaSound))
+						{
+							MetaSoundAsset->SetUpdateDetailsOnSynchronization();
+						}
+					}
+				}
+			}
+		} // namespace GraphPrivate
+	} // namespace Editor
+} // namespace Metasound
+
 UMetasoundEditorGraph* UMetasoundEditorGraphMember::GetOwningGraph()
 {
 	return Cast<UMetasoundEditorGraph>(GetOuter());
@@ -47,8 +80,79 @@ void UMetasoundEditorGraphMember::MarkNodesForRefresh()
 	}
 }
 
-void UMetasoundEditorGraphVertex::OnDataTypeChanged()
+void UMetasoundEditorGraphMember::ConformLiteralDataType()
+{
+	using namespace Metasound::Editor;
+	using namespace Metasound::Frontend;
+
+	IMetasoundEditorModule& EditorModule = FModuleManager::GetModuleChecked<IMetasoundEditorModule>("MetaSoundEditor");
+	const FEditorDataType& EditorDataType = EditorModule.FindDataTypeChecked(TypeName);
+	const EMetasoundFrontendLiteralType LiteralType = static_cast<EMetasoundFrontendLiteralType>(EditorDataType.RegistryInfo.PreferredLiteralType);
+
+	TSubclassOf<UMetasoundEditorGraphMemberDefaultLiteral> LiteralClass = EditorModule.FindDefaultLiteralClass(LiteralType);
+	if (!LiteralClass)
+	{
+		LiteralClass = UMetasoundEditorGraphMemberDefaultLiteral::StaticClass();
+	}
+
+	const UMetasoundEditorGraphMemberDefaultLiteral* Literal = GetLiteral();
+	if (!Literal || Literal->GetClass() != LiteralClass)
+	{
+		UMetasoundEditorGraphMemberDefaultLiteral* NewLiteral = NewObject<UMetasoundEditorGraphMemberDefaultLiteral>(this, LiteralClass, FName(), RF_Transactional);
+		SetLiteral(NewLiteral);
+	}
+}
+
+FName UMetasoundEditorGraphMember::GetDataType() const
 {
+	return TypeName;
+}
+
+#if WITH_EDITOR
+void UMetasoundEditorGraphMember::PostEditUndo()
+{
+	using namespace Metasound;
+
+	Super::PostEditUndo();
+
+	if (IsValid(this))
+	{
+		SetDataType(TypeName);
+
+		constexpr bool bPostTransaction = false;
+		UpdateFrontendDefaultLiteral(bPostTransaction);
+
+		if (UMetasoundEditorGraph* Graph = GetOwningGraph())
+		{
+			if (IsValid(Graph))
+			{
+				UObject* MetaSound = Graph->GetMetasound();
+				if (IsValid(MetaSound))
+				{
+					FMetasoundAssetBase* MetasoundAsset = IMetasoundUObjectRegistry::Get().GetObjectAsAssetBase(MetaSound);
+					check(MetasoundAsset);
+					MetasoundAsset->SetUpdateDetailsOnSynchronization();
+					MetasoundAsset->SetSynchronizationRequired();
+				}
+			}
+		}
+	}
+}
+#endif // WITH_EDITOR
+
+void UMetasoundEditorGraphVertex::InitMember(FName InDataType, const FMetasoundFrontendLiteral& InDefaultLiteral, FGuid InNodeID, FMetasoundFrontendClassName&& InClassName)
+{
+	TypeName = InDataType;
+	NodeID = InNodeID;
+	ClassName = MoveTemp(InClassName);
+
+	ConformLiteralDataType();
+
+	UMetasoundEditorGraphMemberDefaultLiteral* Literal = GetLiteral();
+	if (ensure(Literal))
+	{
+		Literal->SetFromLiteral(InDefaultLiteral);
+	}
 }
 
 TArray<UMetasoundEditorGraphNode*> UMetasoundEditorGraphVertex::GetNodes() const
@@ -141,9 +245,9 @@ void UMetasoundEditorGraphVertex::SetMemberName(const FName& InNewName, bool bPo
 			{
 				Pin->PinName = InNewName;
 			}
-
-			Node->bRefreshNode = true;
 		}
+
+		MarkNodesForRefresh();
 	}
 
 	NameChanged.Broadcast(NodeID);
@@ -182,19 +286,24 @@ void UMetasoundEditorGraphVertex::SetDisplayName(const FText& InNewName, bool bP
 			{
 				Pin->PinFriendlyName = InNewName;
 			}
-
-			Node->bRefreshNode = true;
 		}
+
+		MarkNodesForRefresh();
 	}
 
 	NameChanged.Broadcast(NodeID);
 }
 
-void UMetasoundEditorGraphVertex::SetDataType(FName InNewType, bool bPostTransaction, bool bRegisterParentGraph)
+void UMetasoundEditorGraphVertex::SetDataType(FName InNewType, bool bPostTransaction)
 {
 	using namespace Metasound::Editor;
 	using namespace Metasound::Frontend;
 
+	if (InNewType == GetDataType())
+	{
+		return;
+	}
+
 	UMetasoundEditorGraph* Graph = GetOwningGraph();
 	if (!ensure(Graph))
 	{
@@ -239,12 +348,9 @@ void UMetasoundEditorGraphVertex::SetDataType(FName InNewType, bool bPostTransac
 
 	ClassName = NewNodeHandle->GetClassMetadata().GetClassName();
 	NodeID = NewNodeHandle->GetID();
-	TypeName = InNewType;
 
-	// 5. Report data type changed immediately after assignment to child
-	// class(es) so underlying data can be fixed-up prior to recreating
-	// referencing nodes.
-	OnDataTypeChanged();
+	TypeName = InNewType;
+	ConformLiteralDataType();
 
 	// 6. Create new node references in the same locations as the old locations
 	for (FVector2D Location : NodeLocations)
@@ -255,14 +361,9 @@ void UMetasoundEditorGraphVertex::SetDataType(FName InNewType, bool bPostTransac
 	// Notify now that the node has a new ID (doing so before creating & syncing Frontend Node &
 	// EdGraph variable can result in refreshing editors while in a desync'ed state)
 	NameChanged.Broadcast(NodeID);
-
-	if (bRegisterParentGraph)
-	{
-		FGraphBuilder::RegisterGraphWithFrontend(Metasound);
-	}
 }
 
-Metasound::Frontend::FNodeHandle UMetasoundEditorGraphVertex::GetNodeHandle() 
+Metasound::Frontend::FNodeHandle UMetasoundEditorGraphVertex::GetNodeHandle()
 {
 	using namespace Metasound;
 
@@ -341,62 +442,43 @@ bool UMetasoundEditorGraphVertex::CanRename(const FText& InNewName, FText& OutEr
 	return bIsNameValid;
 }
 
-void UMetasoundEditorGraphInputLiteral::PostEditUndo()
+#if WITH_EDITOR
+void UMetasoundEditorGraphMemberDefaultLiteral::PostEditChangeProperty(FPropertyChangedEvent& InPropertyChangedEvent)
 {
-	Super::PostEditUndo();
-
-	if (!IsValid(this))
+	if (UMetasoundEditorGraphMember* Member = Cast<UMetasoundEditorGraphMember>(GetOuter()))
 	{
-		return;
+		Metasound::Editor::GraphPrivate::OnLiteralChanged(*Member, InPropertyChangedEvent.ChangeType);
 	}
+}
 
-	UpdateDocumentInputLiteral(false /* bPostTransaction */);
+void UMetasoundEditorGraphMemberDefaultLiteral::PostEditChangeChainProperty(FPropertyChangedChainEvent& InPropertyChangedEvent)
+{
+	if (UMetasoundEditorGraphMember* Member = Cast<UMetasoundEditorGraphMember>(GetOuter()))
+	{
+		Metasound::Editor::GraphPrivate::OnLiteralChanged(*Member, InPropertyChangedEvent.ChangeType);
+	}
 }
 
-void UMetasoundEditorGraphInputLiteral::UpdateDocumentInputLiteral(bool bPostTransaction)
+void UMetasoundEditorGraphMemberDefaultLiteral::PostEditUndo()
 {
-	using namespace Metasound;
-	using namespace Metasound::Frontend;
+	Super::PostEditUndo();
 
-	UMetasoundEditorGraphInput* Input = Cast<UMetasoundEditorGraphInput>(GetOuter());
-	UMetasoundEditorGraph* MetasoundGraph = CastChecked<UMetasoundEditorGraph>(Input->GetOuter());
-	UObject* Metasound = MetasoundGraph->GetMetasound();
-	if (!ensure(Metasound))
+	if (!IsValid(this))
 	{
 		return;
 	}
 
-	const FScopedTransaction Transaction(LOCTEXT("Set Input Default", "Set MetaSound Input Default"), bPostTransaction);
-	Metasound->Modify();
+	constexpr bool bPostTransaction = false;
+	UMetasoundEditorGraphMember* Member = CastChecked<UMetasoundEditorGraphMember>(GetOuter());
+	Member->UpdateFrontendDefaultLiteral(bPostTransaction);
 
-	FMetasoundAssetBase* MetasoundAsset = IMetasoundUObjectRegistry::Get().GetObjectAsAssetBase(Metasound);
-	check(MetasoundAsset);
-
-	FGraphHandle GraphHandle = MetasoundAsset->GetRootGraphHandle();
-	FNodeHandle NodeHandle = GraphHandle->GetNodeWithID(Input->NodeID);
-
-	const Metasound::FVertexName& NodeName = NodeHandle->GetNodeName();
-	const FGuid VertexID = GraphHandle->GetVertexIDForInputVertex(NodeName);
-	GraphHandle->SetDefaultInput(VertexID, GetDefault());
-
-	// Disabled as internal call to validation to all other open graphs
-	// is expensive and can be spammed by dragging values 
-// 	Metasound::Editor::FGraphBuilder::RegisterGraphWithFrontend(*Metasound);
-
-	const bool bIsPreviewing = MetasoundGraph->IsPreviewing();
-	if (bIsPreviewing)
+	Member->MarkNodesForRefresh();
+	if (UMetasoundEditorGraph* Graph = Member->GetOwningGraph())
 	{
-		UAudioComponent* PreviewComponent = GEditor->GetPreviewAudioComponent();
-		check(PreviewComponent);
-
-		if (TScriptInterface<IAudioParameterControllerInterface> ParamInterface = PreviewComponent)
-		{
-			Metasound::Frontend::FConstNodeHandle ConstNodeHandle = Input->GetConstNodeHandle();
-			Metasound::FVertexName VertexKey = NodeHandle->GetNodeName();
-			UpdatePreviewInstance(VertexKey, ParamInterface);
-		}
+		Graph->SetSynchronizationRequired();
 	}
 }
+#endif // WITH_EDITOR
 
 Metasound::Editor::ENodeSection UMetasoundEditorGraphInput::GetSectionID() const 
 {
@@ -425,62 +507,100 @@ const FText& UMetasoundEditorGraphInput::GetGraphMemberLabel() const
 	return Label;
 }
 
-void UMetasoundEditorGraphInput::PostEditUndo()
+void UMetasoundEditorGraphInput::ResetToClassDefault()
 {
-	Super::PostEditUndo();
+	using namespace Metasound::Editor;
+	using namespace Metasound::Frontend;
 
-	if (!IsValid(this))
-	{
-		return;
-	}
+	UMetasoundEditorGraph* MetaSoundGraph = GetOwningGraph();
+	FGraphHandle GraphHandle = MetaSoundGraph->GetGraphHandle();
+	FNodeHandle NodeHandle = GraphHandle->GetNodeWithID(NodeID);
 
-	if (!Literal)
-	{
-		if (UMetasoundEditorGraph* MetasoundGraph = Cast<UMetasoundEditorGraph>(GetOuter()))
-		{
-			using namespace Metasound::Editor;
-			MetasoundGraph->RemoveMember(*this);
+	FMetasoundFrontendLiteral DefaultLiteral;
+	DefaultLiteral.SetFromLiteral(IDataTypeRegistry::Get().CreateDefaultLiteral(GetDataType()));
 
-			if (UObject* Object = MetasoundGraph->GetMetasound())
-			{
-				Metasound::Editor::FGraphBuilder::RegisterGraphWithFrontend(*Object);
+	Literal->Modify();
+	Literal->SetFromLiteral(DefaultLiteral);
 
-				FMetasoundAssetBase* MetasoundAsset = Metasound::IMetasoundUObjectRegistry::Get().GetObjectAsAssetBase(Object);
-				check(MetasoundAsset);
-				MetasoundAsset->SetSynchronizationRequired();
-			}
+	const Metasound::FVertexName& NodeName = NodeHandle->GetNodeName();
+	const FGuid VertexID = GraphHandle->GetVertexIDForInputVertex(NodeName);
+	GraphHandle->SetDefaultInput(VertexID, DefaultLiteral);
+
+	if (GEditor && MetaSoundGraph->IsPreviewing())
+	{
+		UAudioComponent* PreviewComponent = GEditor->GetPreviewAudioComponent();
+		check(PreviewComponent);
+
+		if (TScriptInterface<IAudioParameterControllerInterface> ParamInterface = PreviewComponent)
+		{
+			Metasound::Frontend::FConstNodeHandle ConstNodeHandle = GetConstNodeHandle();
+			Metasound::FVertexName VertexKey = NodeHandle->GetNodeName();
+			Literal->UpdatePreviewInstance(VertexKey, ParamInterface);
 		}
-		return;
 	}
 
-	Literal->UpdateDocumentInputLiteral(false /* bPostTransaction */);
-	UpdateEditorLiteralType();
+	if (UMetasoundEditorGraph* Graph = GetOwningGraph())
+	{
+		Graph->SetSynchronizationRequired();
+	}
 }
 
-void UMetasoundEditorGraphInput::OnDataTypeChanged()
+void UMetasoundEditorGraphInput::SetMemberName(const FName& InNewName, bool bPostTransaction)
 {
-	UpdateEditorLiteralType();
+	// Renaming vertex members must stop the preview component to avoid confusion afterward
+	// with newly named input not passing updated values to active previewed instance.
+	if (UMetasoundEditorGraph* Graph = GetOwningGraph())
+	{
+		if (GEditor && Graph->IsPreviewing())
+		{
+			GEditor->ResetPreviewAudioComponent();
+		}
+	}
+
+	Super::SetMemberName(InNewName, bPostTransaction);
 }
 
-void UMetasoundEditorGraphInput::UpdateEditorLiteralType()
+void UMetasoundEditorGraphInput::UpdateFrontendDefaultLiteral(bool bPostTransaction)
 {
 	using namespace Metasound;
-	using namespace Metasound::Editor;
 	using namespace Metasound::Frontend;
 
-	IMetasoundEditorModule& EditorModule = FModuleManager::GetModuleChecked<IMetasoundEditorModule>("MetaSoundEditor");
-	const FEditorDataType& EditorDataType = EditorModule.FindDataTypeChecked(TypeName);
-	const EMetasoundFrontendLiteralType LiteralType = static_cast<EMetasoundFrontendLiteralType>(EditorDataType.RegistryInfo.PreferredLiteralType);
+	UMetasoundEditorGraph* MetaSoundGraph = GetOwningGraph();
+	UObject* Metasound = MetaSoundGraph->GetMetasound();
+	if (!ensure(Metasound))
+	{
+		return;
+	}
 
-	TSubclassOf<UMetasoundEditorGraphInputLiteral> InputLiteralClass = EditorModule.FindInputLiteralClass(LiteralType);
-	if (!InputLiteralClass)
+	if (!ensure(Literal))
 	{
-		InputLiteralClass = UMetasoundEditorGraphInputLiteral::StaticClass();
+		return;
 	}
 
-	if ((nullptr == Literal) || (Literal->GetClass() != InputLiteralClass))
+	const FScopedTransaction Transaction(LOCTEXT("Set Input Default", "Set MetaSound Input Default"), bPostTransaction);
+	Metasound->Modify();
+
+	FMetasoundAssetBase* MetasoundAsset = IMetasoundUObjectRegistry::Get().GetObjectAsAssetBase(Metasound);
+	check(MetasoundAsset);
+
+	FGraphHandle GraphHandle = MetasoundAsset->GetRootGraphHandle();
+	FNodeHandle NodeHandle = GraphHandle->GetNodeWithID(NodeID);
+
+	const Metasound::FVertexName& NodeName = NodeHandle->GetNodeName();
+	const FGuid VertexID = GraphHandle->GetVertexIDForInputVertex(NodeName);
+	GraphHandle->SetDefaultInput(VertexID, Literal->GetDefault());
+
+	if (GEditor && MetaSoundGraph->IsPreviewing())
 	{
-		Literal = NewObject<UMetasoundEditorGraphInputLiteral>(this, InputLiteralClass, FName(), RF_Transactional);
+		UAudioComponent* PreviewComponent = GEditor->GetPreviewAudioComponent();
+		check(PreviewComponent);
+
+		if (TScriptInterface<IAudioParameterControllerInterface> ParamInterface = PreviewComponent)
+		{
+			Metasound::Frontend::FConstNodeHandle ConstNodeHandle = GetConstNodeHandle();
+			Metasound::FVertexName VertexKey = NodeHandle->GetNodeName();
+			Literal->UpdatePreviewInstance(VertexKey, ParamInterface);
+		}
 	}
 }
 
@@ -506,11 +626,90 @@ const FText& UMetasoundEditorGraphOutput::GetGraphMemberLabel() const
 	return Label;
 }
 
+void UMetasoundEditorGraphOutput::ResetToClassDefault()
+{
+	using namespace Metasound;
+	using namespace Metasound::Frontend;
+
+	TArray<UMetasoundEditorGraphNode*> Nodes = GetNodes();
+	for (UMetasoundEditorGraphNode* Node : Nodes)
+	{
+		TArray<FInputHandle> Inputs = Node->GetNodeHandle()->GetInputs();
+		if (ensure(!Inputs.IsEmpty()))
+		{
+			Inputs.Last()->ClearLiteral();
+		}
+	}
+
+	FMetasoundFrontendLiteral DefaultLiteral;
+	DefaultLiteral.SetFromLiteral(IDataTypeRegistry::Get().CreateDefaultLiteral(GetDataType()));
+
+	Literal->Modify();
+	Literal->SetFromLiteral(DefaultLiteral);
+
+	if (UMetasoundEditorGraph* Graph = GetOwningGraph())
+	{
+		Graph->SetSynchronizationRequired();
+	}
+}
+
+void UMetasoundEditorGraphOutput::UpdateFrontendDefaultLiteral(bool bPostTransaction)
+{
+	using namespace Metasound;
+	using namespace Metasound::Frontend;
+
+	UObject* Metasound = nullptr;
+	UMetasoundEditorGraph* MetaSoundGraph = GetOwningGraph();
+	if (ensure(MetaSoundGraph))
+	{
+		Metasound = MetaSoundGraph->GetMetasound();
+	}
+
+	if (!ensure(Metasound))
+	{
+		return;
+	}
+
+	if (!ensure(Literal))
+	{
+		return;
+	}
+
+	const FScopedTransaction Transaction(LOCTEXT("Set Output Default", "Set MetaSound Output Default"), bPostTransaction);
+	Metasound->Modify();
+
+	const FMetasoundFrontendLiteral DefaultLiteral = Literal->GetDefault();
+
+	TArray<UMetasoundEditorGraphNode*> Nodes = GetNodes();
+	for (UMetasoundEditorGraphNode* Node : Nodes)
+	{
+		TArray<FInputHandle> Inputs = Node->GetNodeHandle()->GetInputs();
+		if (ensure(!Inputs.IsEmpty()))
+		{
+			Inputs.Last()->SetLiteral(DefaultLiteral);
+		}
+	}
+}
+
 Metasound::Editor::ENodeSection UMetasoundEditorGraphOutput::GetSectionID() const 
 {
 	return Metasound::Editor::ENodeSection::Outputs;
 }
 
+void UMetasoundEditorGraphVariable::InitMember(FName InDataType, const FMetasoundFrontendLiteral& InDefaultLiteral, FGuid InVariableID)
+{
+	TypeName = InDataType;
+	VariableID = InVariableID;
+
+	ConformLiteralDataType();
+
+	UMetasoundEditorGraphMemberDefaultLiteral* MemberDefault = GetLiteral();
+	if (ensure(MemberDefault))
+	{
+		MemberDefault->SetFromLiteral(InDefaultLiteral);
+	}
+}
+
 const FText& UMetasoundEditorGraphVariable::GetGraphMemberLabel() const
 {
 	static const FText Label = LOCTEXT("GraphMemberLabel_Variable", "Variable");
@@ -676,11 +875,16 @@ void UMetasoundEditorGraphVariable::SetDisplayName(const FText& InNewName, bool
 	NameChanged.Broadcast(VariableID);
 }
 
-void UMetasoundEditorGraphVariable::SetDataType(FName InNewType, bool bPostTransaction, bool bRegisterParentGraph)
+void UMetasoundEditorGraphVariable::SetDataType(FName InNewType, bool bPostTransaction)
 {
 	using namespace Metasound::Editor;
 	using namespace Metasound::Frontend;
 
+	if (InNewType == GetDataType())
+	{
+		return;
+	}
+
 	UMetasoundEditorGraph* Graph = GetOwningGraph();
 	if (!ensure(Graph))
 	{
@@ -714,7 +918,7 @@ void UMetasoundEditorGraphVariable::SetDataType(FName InNewType, bool bPostTrans
 		CachedData.Locations = GetVariableNodeLocations();
 
 
-		// Remove the current variable 
+		// Remove the current variable
 		Graph->RemoveMemberNodes(*this);
 		FGraphHandle FrontendGraph = Graph->GetGraphHandle();
 		FrontendGraph->RemoveVariable(VariableID);
@@ -731,28 +935,23 @@ void UMetasoundEditorGraphVariable::SetDataType(FName InNewType, bool bPostTrans
 
 		// Setup this object with new variable data
 		VariableID = FrontendVariable->GetID();
-		TypeName = InNewType;
 
 		constexpr bool bPostSubTransaction = false;
 		SetMemberName(CachedData.MemberName, bPostSubTransaction);
 		SetDisplayName(CachedData.DisplayName, bPostSubTransaction);
 		SetDescription(CachedData.Description, bPostSubTransaction);
 
+		TypeName = InNewType;
+		ConformLiteralDataType();
+
 		// Add the nodes with the same identifier data but new datatype.
 		UObject& Metasound = Graph->GetMetasoundChecked();
 		AddVariableNodes(Metasound, FrontendGraph, CachedData.Locations);
-
-		UpdateEditorLiteralType();
 	}
 
 	// Notify now that the variable has a new ID (doing so before creating & syncing Frontend Node &
 	// EdGraph variable can result in refreshing editors while in a desync'ed state)
 	NameChanged.Broadcast(VariableID);
-
-	if (bRegisterParentGraph)
-	{
-		FGraphBuilder::RegisterGraphWithFrontend(Graph->GetMetasoundChecked());
-	}
 }
 
 UMetasoundEditorGraphVariable::FVariableEditorNodes UMetasoundEditorGraphVariable::GetVariableNodes() const
@@ -864,57 +1063,51 @@ void UMetasoundEditorGraphVariable::AddVariableNodes(UObject& InMetasound, Metas
 	}
 }
 
-void UMetasoundEditorGraphVariable::UpdateEditorLiteralType()
+const FGuid& UMetasoundEditorGraphVariable::GetVariableID() const
 {
-	using namespace Metasound::Editor;
+	return VariableID;
+}
 
-	IMetasoundEditorModule& EditorModule = FModuleManager::GetModuleChecked<IMetasoundEditorModule>("MetaSoundEditor");
-	FMetasoundFrontendLiteral DefaultFrontendLiteral = GetConstVariableHandle()->GetLiteral();
-	EMetasoundFrontendLiteralType FrontendLiteralType = DefaultFrontendLiteral.GetType();
+void UMetasoundEditorGraphVariable::ResetToClassDefault()
+{
+	using namespace Metasound;
+	using namespace Metasound::Frontend;
 
-	TSubclassOf<UMetasoundEditorGraphInputLiteral> InputLiteralClass = EditorModule.FindInputLiteralClass(FrontendLiteralType);
-	if (!InputLiteralClass)
-	{
-		InputLiteralClass = UMetasoundEditorGraphInputLiteral::StaticClass();
-	}
+	FMetasoundFrontendLiteral DefaultLiteral;
+	DefaultLiteral.SetFromLiteral(IDataTypeRegistry::Get().CreateDefaultLiteral(GetDataType()));
 
-	if ((nullptr == Literal) || (Literal->GetClass() != InputLiteralClass))
-	{
-		Literal = NewObject<UMetasoundEditorGraphInputLiteral>(this, InputLiteralClass, FName(), RF_Transactional);
-		Literal->SetFromLiteral(DefaultFrontendLiteral);
-	}
-}
+	Literal->Modify();
+	Literal->SetFromLiteral(DefaultLiteral);
 
-void UMetasoundEditorGraphVariable::SetFrontendVariable(const Metasound::Frontend::FConstVariableHandle& InVariable)
-{
-	if (ensure(InVariable->IsValid()))
+	FVariableHandle VariableHandle = GetVariableHandle();
+	VariableHandle->SetLiteral(DefaultLiteral);
+
+	FNodeHandle MutatorNode = VariableHandle->FindMutatorNode();
+	if (MutatorNode->IsValid())
 	{
-		VariableID = InVariable->GetID();
-		TypeName = InVariable->GetDataType();
-		UpdateEditorLiteralType();
+		FInputHandle InputHandle = MutatorNode->GetInputWithVertexName(VariableNames::GetInputDataName());
+		if (ensure(InputHandle->IsValid()))
+		{
+			InputHandle->ClearLiteral();
+		}
 	}
-	else
+
+	if (UMetasoundEditorGraph* Graph = GetOwningGraph())
 	{
-		TypeName = FName();
-		VariableID = FGuid();
-		Literal = nullptr;
+		Graph->SetSynchronizationRequired();
 	}
 }
 
-const FGuid& UMetasoundEditorGraphVariable::GetVariableID() const
-{
-	return VariableID;
-}
-
-void UMetasoundEditorGraphVariable::UpdateDocumentVariable(bool bPostTransaction)
+void UMetasoundEditorGraphVariable::UpdateFrontendDefaultLiteral(bool bPostTransaction)
 {
 	using namespace Metasound;
 	using namespace Metasound::Frontend;
 
 	UObject* Metasound = nullptr;
-	if (UMetasoundEditorGraph* MetasoundGraph = GetOwningGraph())
+	UMetasoundEditorGraph* MetaSoundGraph = GetOwningGraph();
+	if (ensure(MetaSoundGraph))
 	{
-		Metasound = MetasoundGraph->GetMetasound();
+		Metasound = MetaSoundGraph->GetMetasound();
 	}
 
 	if (!ensure(Metasound))
@@ -930,47 +1123,21 @@ void UMetasoundEditorGraphVariable::UpdateDocumentVariable(bool bPostTransaction
 	const FScopedTransaction Transaction(LOCTEXT("Set Variable Default", "Set MetaSound Variable Default"), bPostTransaction);
 	Metasound->Modify();
 
-	GetVariableHandle()->SetLiteral(Literal->GetDefault());
-}
-
-
-#if WITH_EDITOR
-void UMetasoundEditorGraphVariable::PostEditUndo()
-{
-	Super::PostEditUndo();
-
-	if (!IsValid(this))
-	{
-		return;
-	}
+	const FMetasoundFrontendLiteral DefaultLiteral = Literal->GetDefault();
+	FVariableHandle VariableHandle = GetVariableHandle();
+	VariableHandle->SetLiteral(DefaultLiteral);
 
-	if (!Literal)
+	FNodeHandle MutatorNode = VariableHandle->FindMutatorNode();
+	if (MutatorNode->IsValid())
 	{
-		if (UMetasoundEditorGraph* MetasoundGraph = GetOwningGraph())
+		FInputHandle InputHandle = MutatorNode->GetInputWithVertexName(VariableNames::GetInputDataName());
+		if (ensure(InputHandle->IsValid()))
 		{
-			MetasoundGraph->RemoveMember(*this);
-
-			if (UObject* Object = MetasoundGraph->GetMetasound())
-			{
-				using namespace Metasound::Editor;
-
-				FGraphBuilder::RegisterGraphWithFrontend(*Object);
-				if (TSharedPtr<FEditor> MetaSoundEditor = FGraphBuilder::GetEditorForMetasound(*Object))
-				{
-					// Refresh details panel in case this variable was selected when it was deleted.
-					MetaSoundEditor->RefreshDetails();
-				}
-			}
+			InputHandle->SetLiteral(DefaultLiteral);
 		}
-		return;
 	}
-
-	UpdateDocumentVariable(false /* bPostTransaction */);
-	UpdateEditorLiteralType();
 }
 
-#endif // WITH_EDITOR
-
 UMetasoundEditorGraphInputNode* UMetasoundEditorGraph::CreateInputNode(Metasound::Frontend::FNodeHandle InNodeHandle, bool bInSelectNewNode)
 {
 	using namespace Metasound::Editor;
@@ -1007,7 +1174,6 @@ UMetasoundEditorGraphInputNode* UMetasoundEditorGraph::CreateInputNode(Metasound
 	return nullptr;	
 }
 
-
 Metasound::Frontend::FDocumentHandle UMetasoundEditorGraph::GetDocumentHandle()
 {
 	return GetGraphHandle()->GetOwningDocument();
@@ -1077,11 +1243,14 @@ void UMetasoundEditorGraph::SetSynchronizationRequired(bool bInClearUpdateNotes)
 		FMetasoundAssetBase* MetasoundAsset = IMetasoundUObjectRegistry::Get().GetObjectAsAssetBase(ParentMetasound);
 		check(MetasoundAsset);
 
-		MetasoundAsset->SetSynchronizationRequired();
 		if (bInClearUpdateNotes)
 		{
 			MetasoundAsset->SetClearNodeNotesOnSynchronization();
 		}
+		else
+		{
+			MetasoundAsset->SetSynchronizationRequired();
+		}
 	}
 }
 
@@ -1104,7 +1273,6 @@ UMetasoundEditorGraphInput* UMetasoundEditorGraph::FindInput(FName InName) const
 	return Input ? Input->Get() : nullptr;
 }
 
-
 UMetasoundEditorGraphInput* UMetasoundEditorGraph::FindOrAddInput(Metasound::Frontend::FNodeHandle InNodeHandle)
 {
 	using namespace Metasound::Editor;
@@ -1129,25 +1297,17 @@ UMetasoundEditorGraphInput* UMetasoundEditorGraph::FindOrAddInput(Metasound::Fro
 		return Input;
 	}
 
-	if (UMetasoundEditorGraphInput* NewInput = NewObject<UMetasoundEditorGraphInput>(this, FName(), RF_Transactional))
+	UMetasoundEditorGraphInput* NewInput = NewObject<UMetasoundEditorGraphInput>(this, FName(), RF_Transactional);
+	if (ensure(NewInput))
 	{
-		NewInput->NodeID = NodeID;
-		NewInput->ClassName = InNodeHandle->GetClassMetadata().GetClassName();
-		NewInput->TypeName = TypeName;
-
 		FMetasoundFrontendLiteral DefaultLiteral = Graph->GetDefaultInput(VertexID);
-		EMetasoundFrontendLiteralType LiteralType = DefaultLiteral.GetType();
-		IMetasoundEditorModule& EditorModule = FModuleManager::GetModuleChecked<IMetasoundEditorModule>("MetaSoundEditor");
-		TSubclassOf<UMetasoundEditorGraphInputLiteral> InputLiteralClass = EditorModule.FindInputLiteralClass(LiteralType);
-
-		NewInput->Literal = NewObject<UMetasoundEditorGraphInputLiteral>(NewInput, InputLiteralClass, FName(), RF_Transactional);
-		NewInput->Literal->SetFromLiteral(DefaultLiteral);
-
+		FMetasoundFrontendClassName ClassName = InNodeHandle->GetClassMetadata().GetClassName();
+		NewInput->InitMember(TypeName, DefaultLiteral, NodeID, MoveTemp(ClassName));
 		Inputs.Add(NewInput);
+
 		return NewInput;
 	}
 
-	checkNoEntry();
 	return nullptr;
 }
 
@@ -1193,17 +1353,19 @@ UMetasoundEditorGraphOutput* UMetasoundEditorGraph::FindOrAddOutput(Metasound::F
 		return Output;
 	}
 
-	if (UMetasoundEditorGraphOutput* NewOutput = NewObject<UMetasoundEditorGraphOutput>(this, FName(), RF_Transactional))
+	UMetasoundEditorGraphOutput* NewOutput = NewObject<UMetasoundEditorGraphOutput>(this, FName(), RF_Transactional);
+	if (ensure(NewOutput))
 	{
-		NewOutput->NodeID = NodeID;
-		NewOutput->ClassName = InNodeHandle->GetClassMetadata().GetClassName();
-		NewOutput->TypeName = TypeName;
+		FMetasoundFrontendLiteral DefaultLiteral;
+		DefaultLiteral.SetFromLiteral(IDataTypeRegistry::Get().CreateDefaultLiteral(TypeName));
+
+		FMetasoundFrontendClassName ClassName = InNodeHandle->GetClassMetadata().GetClassName();
+		NewOutput->InitMember(TypeName, DefaultLiteral, NodeID, MoveTemp(ClassName));
 		Outputs.Add(NewOutput);
 
 		return NewOutput;
 	}
 
-	checkNoEntry();
 	return nullptr;
 }
 
@@ -1211,7 +1373,12 @@ UMetasoundEditorGraphVariable* UMetasoundEditorGraph::FindVariable(const FGuid&
 {
 	const TObjectPtr<UMetasoundEditorGraphVariable>* Variable = Variables.FindByPredicate([&InVariableID](const TObjectPtr<UMetasoundEditorGraphVariable>& InVariable)
 	{
-		return InVariable->GetVariableID() == InVariableID;
+		if (InVariable)
+		{
+			return InVariable->GetVariableID() == InVariableID;
+		}
+
+		return false;
 	});
 	return Variable ? Variable->Get() : nullptr;
 }
@@ -1221,26 +1388,25 @@ UMetasoundEditorGraphVariable* UMetasoundEditorGraph::FindOrAddVariable(const Me
 	using namespace Metasound::Editor;
 	using namespace Metasound::Frontend;
 
-
 	FName TypeName = InVariableHandle->GetDataType();
 	const FGuid VariableID = InVariableHandle->GetID();
-
 	const FGuid NodeID = InVariableHandle->GetID();
+
 	if (TObjectPtr<UMetasoundEditorGraphVariable> EditorVariable = FindVariable(VariableID))
 	{
 		ensure(EditorVariable->TypeName == TypeName);
 		return EditorVariable;
 	}
 
-	if (UMetasoundEditorGraphVariable* NewVariable = NewObject<UMetasoundEditorGraphVariable>(this, FName(), RF_Transactional))
+	UMetasoundEditorGraphVariable* NewVariable = NewObject<UMetasoundEditorGraphVariable>(this, FName(), RF_Transactional);
+	if (ensure(NewVariable))
 	{
-		NewVariable->SetFrontendVariable(InVariableHandle);
+		const FMetasoundFrontendLiteral DefaultLiteral = InVariableHandle->GetLiteral();
+		NewVariable->InitMember(InVariableHandle->GetDataType(), DefaultLiteral, InVariableHandle->GetID());
 		Variables.Add(NewVariable);
-
 		return NewVariable;
 	}
 
-	checkNoEntry();
 	return nullptr;
 }
 
@@ -1341,18 +1507,23 @@ void UMetasoundEditorGraph::SetPreviewID(uint32 InPreviewID)
 
 bool UMetasoundEditorGraph::IsPreviewing() const
 {
-	UAudioComponent* PreviewComponent = GEditor->GetPreviewAudioComponent();
-	if (!PreviewComponent)
+	if (GEditor)
 	{
-		return false;
-	}
+		UAudioComponent* PreviewComponent = GEditor->GetPreviewAudioComponent();
+		if (!PreviewComponent)
+		{
+			return false;
+		}
 
-	if (!PreviewComponent->IsPlaying())
-	{
-		return false;
+		if (!PreviewComponent->IsPlaying())
+		{
+			return false;
+		}
+
+		return PreviewComponent->GetUniqueID() == PreviewID;
 	}
 
-	return PreviewComponent->GetUniqueID() == PreviewID;
+	return false;
 }
 
 bool UMetasoundEditorGraph::IsEditable() const
diff --git a/Engine/Plugins/Runtime/Metasound/Source/MetasoundEditor/Private/MetasoundEditorGraphBuilder.cpp b/Engine/Plugins/Runtime/Metasound/Source/MetasoundEditor/Private/MetasoundEditorGraphBuilder.cpp
--- a/Engine/Plugins/Runtime/Metasound/Source/MetasoundEditor/Private/MetasoundEditorGraphBuilder.cpp
+++ b/Engine/Plugins/Runtime/Metasound/Source/MetasoundEditor/Private/MetasoundEditorGraphBuilder.cpp
@@ -15,9 +15,9 @@
 #include "MetasoundEditor.h"
 #include "MetasoundEditorGraph.h"
 #include "MetasoundEditorGraphBuilder.h"
+#include "MetasoundEditorGraphInputNode.h"
+#include "MetasoundEditorGraphMemberDefaults.h"
 #include "MetasoundEditorGraphNode.h"
-#include "MetasoundEditorGraphInputNodes.h"
-#include "MetasoundEditorGraphSchema.h"
 #include "MetasoundEditorGraphValidation.h"
 #include "MetasoundEditorModule.h"
 #include "MetasoundEditorSettings.h"
@@ -37,7 +37,6 @@
 #include "Toolkits/ToolkitManager.h"
 #include "Widgets/Notifications/SNotificationList.h"
 
-
 #define LOCTEXT_NAMESPACE "MetaSoundEditor"
 
 
@@ -150,24 +149,23 @@ namespace Metasound
 				case EMetasoundFrontendClassType::VariableAccessor:
 				case EMetasoundFrontendClassType::VariableDeferredAccessor:
 				case EMetasoundFrontendClassType::VariableMutator:
-					{
-						// All variables nodes use the same pin name for inputs and outputs
-						// and do not reflect the name of the variable in their pin names. 
-						// The MetaSound editor needs a pin name which reflects that variable's
-						// name to have a consistent look and behavior to input and output nodes.
-						Frontend::FConstVariableHandle Variable = OwningNode->GetOwningGraph()->FindVariableContainingNode(OwningNode->GetID());
-						return Variable->GetName();
-					}
+				{
+					// All variables nodes use the same pin name for user-modifiable node
+					// inputs and outputs and the editor does not display the pin's name. The
+					// editor instead displays the variable's name in place of the pin name to
+					// maintain a consistent look and behavior to input and output nodes.
+					return VariableNames::GetOutputDataName();
+				}
 				case EMetasoundFrontendClassType::Input:
 				case EMetasoundFrontendClassType::Output:
-					{
-						return OwningNode->GetNodeName();
-					}
+				{
+					return OwningNode->GetNodeName();
+				}
 
 				default:
-					{
-						return InFrontendOutput.GetName();
-					}
+				{
+					return InFrontendOutput.GetName();
+				}
 			}
 		}
 
@@ -182,25 +180,24 @@ namespace Metasound
 				case EMetasoundFrontendClassType::VariableAccessor:
 				case EMetasoundFrontendClassType::VariableDeferredAccessor:
 				case EMetasoundFrontendClassType::VariableMutator:
-					{
-						// All variables nodes use the same pin name for inputs and outputs
-						// and do not reflect the name of the variable in their pin names. 
-						// The MetaSound editor needs a pin name which reflects that variable's
-						// name to have a consistent look and behavior to input and output nodes.
-						Frontend::FConstVariableHandle Variable = OwningNode->GetOwningGraph()->FindVariableContainingNode(OwningNode->GetID());
-						return Variable->GetName();
-					}
+				{
+					// All variables nodes use the same pin name for user-modifiable node
+					// inputs and outputs and the editor does not display the pin's name. The
+					// editor instead displays the variable's name in place of the pin name to
+					// maintain a consistent look and behavior to input and output nodes.
+					return VariableNames::GetInputDataName();
+				}
 
 				case EMetasoundFrontendClassType::Input:
 				case EMetasoundFrontendClassType::Output:
-					{
-						return OwningNode->GetNodeName();
-					}
+				{
+					return OwningNode->GetNodeName();
+				}
 
 				default:
-					{
-						return InFrontendInput.GetName();
-					}
+				{
+					return InFrontendInput.GetName();
+				}
 			}
 		}
 
@@ -312,28 +309,6 @@ namespace Metasound
 			return NewGraphNode;
 		}
 
-		bool FGraphBuilder::InitGraph(UObject& InMetaSound)
-		{
-			FMetasoundAssetBase* MetaSoundAsset = IMetasoundUObjectRegistry::Get().GetObjectAsAssetBase(&InMetaSound);
-			check(MetaSoundAsset);
-
-			// Initial graph generation is not something to be managed by the transaction stack, so don't track dirty state until
-			// after initial setup if necessary.
-			UMetasoundEditorGraph* Graph = Cast<UMetasoundEditorGraph>(MetaSoundAsset->GetGraph());
-			if (!Graph)
-			{
-				Graph = NewObject<UMetasoundEditorGraph>(&InMetaSound, FName(), RF_Transactional);
-				Graph->Schema = UMetasoundEditorGraphSchema::StaticClass();
-				MetaSoundAsset->SetGraph(Graph);
-
-				// Has to be done inline to have valid graph initially when opening editor for the first time
-				SynchronizeGraph(InMetaSound);
-				return true;
-			}
-
-			return false;
-		}
-
 		void FGraphBuilder::InitGraphNode(Frontend::FNodeHandle& InNodeHandle, UMetasoundEditorGraphNode* NewGraphNode, UObject& InMetaSound)
 		{
 			NewGraphNode->CreateNewGuid();
@@ -767,7 +742,7 @@ namespace Metasound
 					if (ensure(Input))
 					{
 						const FName PinName = Pin->GetFName();
-						NodeHandle = AddInputNodeHandle(InMetaSound, Input->TypeName, InGraphNode.GetTooltipText(), nullptr, &PinName);
+						NodeHandle = AddInputNodeHandle(InMetaSound, Input->GetDataType(), InGraphNode.GetTooltipText(), nullptr, &PinName);
 					}
 				}
 			}
@@ -782,7 +757,7 @@ namespace Metasound
 					if (ensure(Output))
 					{
 						const FName PinName = Pin->GetFName();
-						NodeHandle = FGraphBuilder::AddOutputNodeHandle(InMetaSound, Output->TypeName, InGraphNode.GetTooltipText(), &PinName);
+						NodeHandle = FGraphBuilder::AddOutputNodeHandle(InMetaSound, Output->GetDataType(), InGraphNode.GetTooltipText(), &PinName);
 					}
 				}
 			}
@@ -1010,10 +985,14 @@ namespace Metasound
 			TArray<FInputHandle> InputHandles;
 			TArray<UEdGraphPin*> InputPins;
 
+			UMetasoundEditorGraphNode* Node = CastChecked<UMetasoundEditorGraphNode>(InPin.GetOwningNode());
+
 			if (InPin.Direction == EGPD_Input)
 			{
-				FNodeHandle NodeHandle = CastChecked<UMetasoundEditorGraphNode>(InPin.GetOwningNode())->GetNodeHandle();
-				FInputHandle InputHandle = NodeHandle->GetInputWithVertexName(InPin.GetFName());
+				const FName PinName = InPin.GetFName();
+
+				FNodeHandle NodeHandle = Node->GetNodeHandle();
+				FInputHandle InputHandle = NodeHandle->GetInputWithVertexName(PinName);
 
 				// Input can be invalid if renaming a vertex member
 				if (InputHandle->IsValid())
@@ -1044,7 +1023,6 @@ namespace Metasound
 			{
 				FInputHandle InputHandle = InputHandles[i];
 				FConstOutputHandle OutputHandle = InputHandle->GetConnectedOutput();
-				const FMetasoundFrontendNodeStyle& Style = OutputHandle->GetOwningNode()->GetNodeStyle();
 
 				InputHandle->Disconnect();
 
@@ -1054,6 +1032,10 @@ namespace Metasound
 					SynchronizePinLiteral(*InputPins[i]);
 				}
 			}
+
+			UObject& MetaSound = Node->GetMetasoundChecked();
+			FMetasoundAssetBase* MetaSoundAsset = IMetasoundUObjectRegistry::Get().GetObjectAsAssetBase(&MetaSound);
+			MetaSoundAsset->SetSynchronizationRequired();
 		}
 
 		void FGraphBuilder::InitMetaSound(UObject& InMetaSound, const FString& InAuthor)
@@ -2033,80 +2015,69 @@ namespace Metasound
 				}
 			});
 
-			// Remove stale inputs and outputs. 
+			// Remove stale inputs and outputs.
 			bIsEditorGraphDirty |= !ToRemove.IsEmpty();
 			for (UMetasoundEditorGraphMember* GraphMember: ToRemove)
 			{
 				Graph->RemoveMember(*GraphMember);
 			}
 
-			// Synchronize data types of input nodes. 
-			GraphHandle->IterateNodes([&](FNodeHandle NodeHandle)
+			UMetasoundEditorGraphMember* Member = nullptr;
+
+			auto SynchronizeMember = [](UMetasoundEditorGraphVertex& InVertex)
 			{
-				if (UMetasoundEditorGraphInput* Input = Graph->FindInput(NodeHandle->GetID()))
+				FConstNodeHandle NodeHandle = InVertex.GetConstNodeHandle();
+				TArray<FConstInputHandle> InputHandles = NodeHandle->GetConstInputs();
+				if (ensure(InputHandles.Num() == 1))
 				{
-					TArray<FConstInputHandle> InputHandles = NodeHandle->GetConstInputs();
-					if (ensure(InputHandles.Num() == 1))
+					FConstInputHandle InputHandle = InputHandles.Last();
+					const FName NewDataType = InputHandle->GetDataType();
+					if (InVertex.GetDataType() != NewDataType)
 					{
-						FConstInputHandle InputHandle = InputHandles[0];
-						const FName NewDataType = InputHandle->GetDataType();
-						if (Input->TypeName != NewDataType)
-						{
-							constexpr bool bIncludeNamespace = true;
-							FText NodeDisplayName = FGraphBuilder::GetDisplayName(*NodeHandle, bIncludeNamespace);
-							UE_LOG(LogMetasoundEditor, Verbose, TEXT("Synchronizing Input '%s': Updating DataType to '%s'."), *NodeDisplayName.ToString(), *NewDataType.ToString());
+						constexpr bool bIncludeNamespace = true;
+						FText NodeDisplayName = FGraphBuilder::GetDisplayName(*NodeHandle, bIncludeNamespace);
+						UE_LOG(LogMetasoundEditor, Verbose, TEXT("Synchronizing Member '%s': Updating DataType to '%s'."), *NodeDisplayName.ToString(), *NewDataType.ToString());
 
-							FMetasoundFrontendLiteral DefaultLiteral;
-							DefaultLiteral.SetFromLiteral(IDataTypeRegistry::Get().CreateDefaultLiteral(NewDataType));
-							if (const FMetasoundFrontendLiteral* InputLiteral = InputHandle->GetLiteral())
-							{
-								DefaultLiteral = *InputLiteral;
-							}
+						FMetasoundFrontendLiteral DefaultLiteral;
+						DefaultLiteral.SetFromLiteral(IDataTypeRegistry::Get().CreateDefaultLiteral(NewDataType));
+						if (const FMetasoundFrontendLiteral* InputLiteral = InputHandle->GetLiteral())
+						{
+							DefaultLiteral = *InputLiteral;
+						}
 
-							Input->ClassName = NodeHandle->GetClassMetadata().GetClassName();
-							Input->TypeName = NewDataType;
+						InVertex.ClassName = NodeHandle->GetClassMetadata().GetClassName();
 
-							// Report data type changed immediately after assignment to child
-							// class(es) so underlying data can be fixed-up prior to recreating
-							// referencing nodes.
-							Input->OnDataTypeChanged();
+						constexpr bool bPostTransaction = false;
+						InVertex.SetDataType(NewDataType, bPostTransaction);
 
-							if (DefaultLiteral.IsValid())
-							{
-								Input->Literal->SetFromLiteral(DefaultLiteral);
-							}
+						if (DefaultLiteral.IsValid())
+						{
+							InVertex.GetLiteral()->SetFromLiteral(DefaultLiteral);
 						}
 					}
 				}
+			};
+
+			// Synchronize data types of input nodes.
+			GraphHandle->IterateNodes([&](FNodeHandle NodeHandle)
+			{
+				if (UMetasoundEditorGraphInput* Input = Graph->FindInput(NodeHandle->GetID()))
+				{
+					SynchronizeMember(*CastChecked<UMetasoundEditorGraphVertex>(Input));
+				}
 			}, EMetasoundFrontendClassType::Input);
 
-			// Synchronize data types of output nodes. 
+			// Synchronize data types of output nodes.
 			GraphHandle->IterateNodes([&](FNodeHandle NodeHandle)
 			{
 				if (UMetasoundEditorGraphOutput* Output = Graph->FindOutput(NodeHandle->GetID()))
 				{
-					TArray<FConstOutputHandle> OutputHandles = NodeHandle->GetConstOutputs();
-					if (ensure(OutputHandles.Num() == 1))
-					{
-						const FName NewDataType = OutputHandles[0]->GetDataType();
-						if (Output->TypeName != NewDataType)
-						{
-							constexpr bool bIncludeNamespace = true;
-							FText NodeDisplayName = FGraphBuilder::GetDisplayName(*NodeHandle, bIncludeNamespace);
-							UE_LOG(LogMetasoundEditor, Verbose, TEXT("Synchronizing Output '%s': Updating DataType to '%s'."), *NodeDisplayName.ToString(), *NewDataType.ToString());
-
-							Output->ClassName = NodeHandle->GetClassMetadata().GetClassName();
-							Output->TypeName = NewDataType;
-
-							// Report data type changed immediately after assignment to child
-							// class(es) so underlying data can be fixed-up prior to recreating
-							// referencing nodes.
-							Output->OnDataTypeChanged();
-						}
-					}
+					SynchronizeMember(*CastChecked<UMetasoundEditorGraphVertex>(Output));
 				}
 			}, EMetasoundFrontendClassType::Output);
 
+
+
 			return bIsEditorGraphDirty;
 		}
 	} // namespace Editor
diff --git a/Engine/Plugins/Runtime/Metasound/Source/MetasoundEditor/Private/MetasoundEditorGraphBuilder.h b/Engine/Plugins/Runtime/Metasound/Source/MetasoundEditor/Private/MetasoundEditorGraphBuilder.h
--- a/Engine/Plugins/Runtime/Metasound/Source/MetasoundEditor/Private/MetasoundEditorGraphBuilder.h
+++ b/Engine/Plugins/Runtime/Metasound/Source/MetasoundEditor/Private/MetasoundEditorGraphBuilder.h
@@ -1,14 +1,13 @@
 // Copyright Epic Games, Inc. All Rights Reserved.
 #pragma once
 
-#include "CoreMinimal.h"
-#include "EdGraph/EdGraph.h"
-#include "EdGraph/EdGraphNode.h"
-#include "EdGraph/EdGraphPin.h"
-#include "MetasoundEditorGraphInputNodes.h"
 #include "MetasoundFrontend.h"
 #include "MetasoundFrontendController.h"
 #include "MetasoundFrontendDocument.h"
+#include "Internationalization/Text.h"
+#include "Templates/SharedPointer.h"
+#include "UObject/NameTypes.h"
+#include "UObject/NoExportTypes.h"
 
 
 // Forward Declarations
@@ -16,10 +15,15 @@ class SGraphEditor;
 class UEdGraphNode;
 class UEdGraphPin;
 class UMetaSound;
+class UMetasoundEditorGraphExternalNode;
 class UMetasoundEditorGraphNode;
+class UMetasoundEditorGraphInputNode;
+class UMetasoundEditorGraphOutputNode;
+class UMetasoundEditorGraphVariableNode;
 
 struct FEdGraphPinType;
 struct FMetasoundFrontendNodeStyle;
+struct FMetasoundFrontendLiteral;
 
 
 namespace Metasound
@@ -64,9 +68,6 @@ namespace Metasound
 			static TSharedPtr<FEditor> GetEditorForMetasound(const UObject& MetaSound);
 			static TSharedPtr<FEditor> GetEditorForGraph(const UEdGraph& EdGraph);
 
-			// Initializes editor graph for given MetaSound
-			static bool InitGraph(UObject& InMetaSound);
-
 			// Wraps RegisterGraphWithFrontend logic in Frontend with any additional logic required to refresh editor & respective editor object state.
 			static void RegisterGraphWithFrontend(UObject& InMetaSound);
 
diff --git a/Engine/Plugins/Runtime/Metasound/Source/MetasoundEditor/Private/MetasoundEditorGraphInputNode.cpp b/Engine/Plugins/Runtime/Metasound/Source/MetasoundEditor/Private/MetasoundEditorGraphInputNode.cpp
new file mode 100644
--- /dev/null
+++ b/Engine/Plugins/Runtime/Metasound/Source/MetasoundEditor/Private/MetasoundEditorGraphInputNode.cpp
@@ -0,0 +1,79 @@
+// Copyright Epic Games, Inc. All Rights Reserved.
+#include "MetasoundEditorGraphInputNode.h"
+
+#include "Algo/Count.h"
+#include "Algo/Transform.h"
+#include "AudioParameterControllerInterface.h"
+#include "EdGraph/EdGraphNode.h"
+#include "GraphEditorSettings.h"
+#include "MetasoundDataReference.h"
+#include "MetasoundEditorGraph.h"
+#include "MetasoundEditorGraphNode.h"
+#include "MetasoundEditorSettings.h"
+#include "MetasoundFrontend.h"
+#include "MetasoundFrontendController.h"
+#include "MetasoundFrontendDocument.h"
+#include "MetasoundFrontendLiteral.h"
+#include "MetasoundFrontendRegistries.h"
+#include "MetasoundPrimitives.h"
+#include "Misc/Guid.h"
+#include "Sound/SoundWave.h"
+#include "UObject/ObjectMacros.h"
+#include "UObject/SoftObjectPath.h"
+
+UMetasoundEditorGraphMember* UMetasoundEditorGraphInputNode::GetMember() const
+{
+	return Input;
+}
+
+FMetasoundFrontendClassName UMetasoundEditorGraphInputNode::GetClassName() const
+{
+	if (ensure(Input))
+	{
+		return Input->ClassName;
+	}
+
+	return Super::GetClassName();
+}
+
+void UMetasoundEditorGraphInputNode::UpdatePreviewInstance(const Metasound::FVertexName& InParameterName, TScriptInterface<IAudioParameterControllerInterface>& InParameterInterface) const
+{
+	if (Input)
+	{
+		if (UMetasoundEditorGraphMemberDefaultLiteral* DefaultLiteral = Input->GetLiteral())
+		{
+			DefaultLiteral->UpdatePreviewInstance(InParameterName, InParameterInterface);
+		}
+	}
+}
+
+FGuid UMetasoundEditorGraphInputNode::GetNodeID() const
+{
+	if (Input)
+	{
+		return Input->NodeID;
+	}
+
+	return Super::GetNodeID();
+}
+
+FLinearColor UMetasoundEditorGraphInputNode::GetNodeTitleColor() const
+{
+	if (const UMetasoundEditorSettings* EditorSettings = GetDefault<UMetasoundEditorSettings>())
+	{
+		return EditorSettings->InputNodeTitleColor;
+	}
+
+	return Super::GetNodeTitleColor();
+}
+
+FSlateIcon UMetasoundEditorGraphInputNode::GetNodeTitleIcon() const
+{
+	static const FName NativeIconName = "MetasoundEditor.Graph.Node.Class.Input";
+	return FSlateIcon("MetaSoundStyle", NativeIconName);
+}
+
+void UMetasoundEditorGraphInputNode::SetNodeID(FGuid InNodeID)
+{
+	Input->NodeID = InNodeID;
+}
diff --git a/Engine/Plugins/Runtime/Metasound/Source/MetasoundEditor/Private/MetasoundEditorGraphInputNodes.cpp b/Engine/Plugins/Runtime/Metasound/Source/MetasoundEditor/Private/MetasoundEditorGraphInputNodes.cpp
deleted file mode 100644
diff --git a/Engine/Plugins/Runtime/Metasound/Source/MetasoundEditor/Private/MetasoundEditorGraphMemberDefaults.cpp b/Engine/Plugins/Runtime/Metasound/Source/MetasoundEditor/Private/MetasoundEditorGraphMemberDefaults.cpp
new file mode 100644
--- /dev/null
+++ b/Engine/Plugins/Runtime/Metasound/Source/MetasoundEditor/Private/MetasoundEditorGraphMemberDefaults.cpp
@@ -0,0 +1,354 @@
+// Copyright Epic Games, Inc. All Rights Reserved.
+#include "MetasoundEditorGraphMemberDefaults.h"
+
+#include "Algo/Count.h"
+#include "Algo/Transform.h"
+#include "AudioParameterControllerInterface.h"
+#include "EdGraph/EdGraphNode.h"
+#include "GraphEditorSettings.h"
+#include "MetasoundDataReference.h"
+#include "MetasoundEditorGraph.h"
+#include "MetasoundEditorGraphNode.h"
+#include "MetasoundEditorSettings.h"
+#include "MetasoundFrontend.h"
+#include "MetasoundFrontendController.h"
+#include "MetasoundFrontendDocument.h"
+#include "MetasoundFrontendLiteral.h"
+#include "MetasoundFrontendRegistries.h"
+#include "MetasoundPrimitives.h"
+#include "Misc/Guid.h"
+#include "Sound/SoundWave.h"
+#include "UObject/ObjectMacros.h"
+#include "UObject/SoftObjectPath.h"
+
+namespace Metasound
+{
+	namespace Editor
+	{
+		namespace MemberDefaultsPrivate
+		{
+			template <typename TPODType>
+			void ConvertLiteral(const FMetasoundFrontendLiteral& InLiteral, TPODType& OutValue)
+			{
+				const FName& TypeName = GetMetasoundDataTypeName<TPODType>();
+				OutValue = InLiteral.ToLiteral(TypeName).Value.Get<TPODType>();
+			}
+
+			template <typename TPODType, typename TLiteralType = TPODType>
+			void ConvertLiteralToArray(const FMetasoundFrontendLiteral& InLiteral, TArray<TLiteralType>& OutArray)
+			{
+				const FName TypeName = CreateArrayTypeNameFromElementTypeName(GetMetasoundDataTypeName<TPODType>());
+				TArray<TPODType> NewValue = InLiteral.ToLiteral(TypeName).Value.Get<TArray<TPODType>>();
+				Algo::Transform(NewValue, OutArray, [](const TPODType& InValue) { return TLiteralType { InValue }; });
+			}
+		}
+	}
+}
+
+FMetasoundFrontendLiteral UMetasoundEditorGraphMemberDefaultBool::GetDefault() const 
+{
+	FMetasoundFrontendLiteral Literal;
+	Literal.Set(Default.Value);
+	return Literal;
+}
+
+EMetasoundFrontendLiteralType UMetasoundEditorGraphMemberDefaultBool::GetLiteralType() const 
+{
+	return EMetasoundFrontendLiteralType::Boolean;
+}
+
+void UMetasoundEditorGraphMemberDefaultBool::SetFromLiteral(const FMetasoundFrontendLiteral& InLiteral)
+{
+	Metasound::Editor::MemberDefaultsPrivate::ConvertLiteral<bool>(InLiteral, Default.Value);
+}
+
+void UMetasoundEditorGraphMemberDefaultBool::UpdatePreviewInstance(const Metasound::FVertexName& InParameterName, TScriptInterface<IAudioParameterControllerInterface>& InParameterInterface) const
+{
+	InParameterInterface->SetBoolParameter(InParameterName, Default.Value);
+}
+
+FMetasoundFrontendLiteral UMetasoundEditorGraphMemberDefaultBoolArray::GetDefault() const
+{
+	TArray<bool> BoolArray;
+	Algo::Transform(Default, BoolArray, [](const FMetasoundEditorGraphMemberDefaultBoolRef& InValue) { return InValue.Value; });
+
+	FMetasoundFrontendLiteral Literal;
+	Literal.Set(BoolArray);
+	return Literal;
+}
+
+EMetasoundFrontendLiteralType UMetasoundEditorGraphMemberDefaultBoolArray::GetLiteralType() const
+{
+	return EMetasoundFrontendLiteralType::BooleanArray;
+}
+
+void UMetasoundEditorGraphMemberDefaultBoolArray::SetFromLiteral(const FMetasoundFrontendLiteral& InLiteral)
+{
+	Metasound::Editor::MemberDefaultsPrivate::ConvertLiteralToArray<bool, FMetasoundEditorGraphMemberDefaultBoolRef>(InLiteral, Default);
+}
+
+void UMetasoundEditorGraphMemberDefaultBoolArray::UpdatePreviewInstance(const Metasound::FVertexName& InParameterName, TScriptInterface<IAudioParameterControllerInterface>& InParameterInterface) const
+{
+	TArray<bool> BoolArray;
+	Algo::Transform(Default, BoolArray, [](const FMetasoundEditorGraphMemberDefaultBoolRef& InValue) { return InValue.Value; });
+	InParameterInterface->SetBoolArrayParameter(InParameterName, BoolArray);
+}
+
+FMetasoundFrontendLiteral UMetasoundEditorGraphMemberDefaultInt::GetDefault() const
+{
+	FMetasoundFrontendLiteral Literal;
+	Literal.Set(Default.Value);
+	return Literal;
+}
+
+EMetasoundFrontendLiteralType UMetasoundEditorGraphMemberDefaultInt::GetLiteralType() const
+{
+	return EMetasoundFrontendLiteralType::Integer;
+}
+
+void UMetasoundEditorGraphMemberDefaultInt::SetFromLiteral(const FMetasoundFrontendLiteral& InLiteral)
+{
+	Metasound::Editor::MemberDefaultsPrivate::ConvertLiteral<int32>(InLiteral, Default.Value);
+}
+
+void UMetasoundEditorGraphMemberDefaultInt::UpdatePreviewInstance(const Metasound::FVertexName& InParameterName, TScriptInterface<IAudioParameterControllerInterface>& InParameterInterface) const
+{
+	InParameterInterface->SetIntParameter(InParameterName, Default.Value);
+}
+
+FMetasoundFrontendLiteral UMetasoundEditorGraphMemberDefaultIntArray::GetDefault() const
+{
+	TArray<int32> IntArray;
+	Algo::Transform(Default, IntArray, [](const FMetasoundEditorGraphMemberDefaultIntRef& InValue) { return InValue.Value; });
+
+	FMetasoundFrontendLiteral Literal;
+	Literal.Set(IntArray);
+
+	return Literal;
+}
+
+EMetasoundFrontendLiteralType UMetasoundEditorGraphMemberDefaultIntArray::GetLiteralType() const
+{
+	return EMetasoundFrontendLiteralType::IntegerArray;
+}
+
+void UMetasoundEditorGraphMemberDefaultIntArray::SetFromLiteral(const FMetasoundFrontendLiteral& InLiteral)
+{
+	Metasound::Editor::MemberDefaultsPrivate::ConvertLiteralToArray<int32, FMetasoundEditorGraphMemberDefaultIntRef>(InLiteral, Default);
+}
+
+void UMetasoundEditorGraphMemberDefaultIntArray::UpdatePreviewInstance(const Metasound::FVertexName& InParameterName, TScriptInterface<IAudioParameterControllerInterface>& InParameterInterface) const
+{
+	TArray<int32> IntArray;
+	Algo::Transform(Default, IntArray, [](const FMetasoundEditorGraphMemberDefaultIntRef& InValue) { return InValue.Value; });
+	InParameterInterface->SetIntArrayParameter(InParameterName, IntArray);
+}
+
+FMetasoundFrontendLiteral UMetasoundEditorGraphMemberDefaultFloat::GetDefault() const
+{
+	FMetasoundFrontendLiteral Literal;
+	Literal.Set(Default);
+	return Literal;
+}
+
+EMetasoundFrontendLiteralType UMetasoundEditorGraphMemberDefaultFloat::GetLiteralType() const
+{
+	return EMetasoundFrontendLiteralType::Float;
+}
+
+void UMetasoundEditorGraphMemberDefaultFloat::SetFromLiteral(const FMetasoundFrontendLiteral& InLiteral)
+{
+	Metasound::Editor::MemberDefaultsPrivate::ConvertLiteral<float>(InLiteral, Default);
+}
+
+void UMetasoundEditorGraphMemberDefaultFloat::UpdatePreviewInstance(const Metasound::FVertexName& InParameterName, TScriptInterface<IAudioParameterControllerInterface>& InParameterInterface) const
+{
+	InParameterInterface->SetFloatParameter(InParameterName, Default);
+}
+
+void UMetasoundEditorGraphMemberDefaultFloat::PostEditChangeProperty(FPropertyChangedEvent& PropertyChangedEvent)
+{
+	if (PropertyChangedEvent.GetPropertyName().IsEqual(GET_MEMBER_NAME_CHECKED(UMetasoundEditorGraphMemberDefaultFloat, Default)))
+	{
+		OnDefaultValueChanged.Broadcast(Default);
+	}
+	else if (PropertyChangedEvent.GetPropertyName().IsEqual(GET_MEMBER_NAME_CHECKED(UMetasoundEditorGraphMemberDefaultFloat, WidgetType)) || PropertyChangedEvent.GetPropertyName().IsEqual(GET_MEMBER_NAME_CHECKED(UMetasoundEditorGraphMemberDefaultFloat, WidgetValueType)))
+	{
+		if (WidgetValueType == EMetasoundMemberDefaultWidgetValueType::Linear)
+		{
+			Range = FVector2D(0.0f, 1.0f);
+		}
+		else if (WidgetValueType == EMetasoundMemberDefaultWidgetValueType::Frequency)
+		{
+			Range = FVector2D(20.0f, 20000.0f);
+		}
+		else if (WidgetValueType == EMetasoundMemberDefaultWidgetValueType::Volume)
+		{
+			Range = FVector2D(-100.0f, 0.0f);
+		}
+
+		OnRangeChanged.Broadcast(Range);
+		// If the widget type is changed to none, we need to refresh clamping the value or not, since if the widget was a slider before, the value was clamped
+		OnClampChanged.Broadcast(ClampDefault);
+	}
+	else if (PropertyChangedEvent.MemberProperty->GetFName().IsEqual(GET_MEMBER_NAME_CHECKED(UMetasoundEditorGraphMemberDefaultFloat, Range)))
+	{
+		if (PropertyChangedEvent.ChangeType != EPropertyChangeType::Interactive)
+		{
+			float Min = Range.X;
+			float Max = Range.Y;
+			if (Min < Max)
+			{
+				OnRangeChanged.Broadcast(Range);
+				SetDefault(FMath::Clamp(Default, Min, Max));
+			}
+		}
+	}
+	else if (PropertyChangedEvent.GetPropertyName().IsEqual(GET_MEMBER_NAME_CHECKED(UMetasoundEditorGraphMemberDefaultFloat, ClampDefault)))
+	{
+		// set range to reasonable limit given current value
+		Range = FVector2D(FMath::Min(0.0f, Default), FMath::Max(0.0f, Default));
+
+		OnClampChanged.Broadcast(ClampDefault);
+	}
+
+	UMetasoundEditorGraphMember* Member = GetParentMember();
+	if (ensure(Member))
+	{
+		Member->MarkNodesForRefresh();
+	}
+}
+
+void UMetasoundEditorGraphMemberDefaultFloat::SetDefault(const float InDefault)
+{
+	Default = InDefault;
+	OnDefaultValueChanged.Broadcast(InDefault);
+}
+
+float UMetasoundEditorGraphMemberDefaultFloat::GetDefault()
+{
+	return Default;
+}
+
+FVector2D UMetasoundEditorGraphMemberDefaultFloat::GetRange()
+{
+	return Range;
+}
+
+FMetasoundFrontendLiteral UMetasoundEditorGraphMemberDefaultFloatArray::GetDefault() const
+{
+	FMetasoundFrontendLiteral Literal;
+	Literal.Set(Default);
+	return Literal;
+}
+
+EMetasoundFrontendLiteralType UMetasoundEditorGraphMemberDefaultFloatArray::GetLiteralType() const
+{
+	return EMetasoundFrontendLiteralType::FloatArray;
+}
+
+void UMetasoundEditorGraphMemberDefaultFloatArray::SetFromLiteral(const FMetasoundFrontendLiteral& InLiteral)
+{
+	Metasound::Editor::MemberDefaultsPrivate::ConvertLiteralToArray<float>(InLiteral, Default);
+}
+
+void UMetasoundEditorGraphMemberDefaultFloatArray::UpdatePreviewInstance(const Metasound::FVertexName& InParameterName, TScriptInterface<IAudioParameterControllerInterface>& InParameterInterface) const
+{
+	InParameterInterface->SetFloatArrayParameter(InParameterName, Default);
+}
+
+FMetasoundFrontendLiteral UMetasoundEditorGraphMemberDefaultString::GetDefault() const
+{
+	FMetasoundFrontendLiteral Literal;
+	Literal.Set(Default);
+	return Literal;
+}
+
+EMetasoundFrontendLiteralType UMetasoundEditorGraphMemberDefaultString::GetLiteralType() const
+{
+	return EMetasoundFrontendLiteralType::String;
+}
+
+void UMetasoundEditorGraphMemberDefaultString::SetFromLiteral(const FMetasoundFrontendLiteral& InLiteral)
+{
+	Metasound::Editor::MemberDefaultsPrivate::ConvertLiteral<FString>(InLiteral, Default);
+}
+
+void UMetasoundEditorGraphMemberDefaultString::UpdatePreviewInstance(const Metasound::FVertexName& InParameterName, TScriptInterface<IAudioParameterControllerInterface>& InParameterInterface) const
+{
+	InParameterInterface->SetStringParameter(InParameterName, Default);
+}
+
+FMetasoundFrontendLiteral UMetasoundEditorGraphMemberDefaultStringArray::GetDefault() const
+{
+	FMetasoundFrontendLiteral Literal;
+	Literal.Set(Default);
+	return Literal;
+}
+
+EMetasoundFrontendLiteralType UMetasoundEditorGraphMemberDefaultStringArray::GetLiteralType() const
+{
+	return EMetasoundFrontendLiteralType::StringArray;
+}
+
+void UMetasoundEditorGraphMemberDefaultStringArray::SetFromLiteral(const FMetasoundFrontendLiteral& InLiteral)
+{
+	Metasound::Editor::MemberDefaultsPrivate::ConvertLiteralToArray<FString>(InLiteral, Default);
+}
+
+void UMetasoundEditorGraphMemberDefaultStringArray::UpdatePreviewInstance(const Metasound::FVertexName& InParameterName, TScriptInterface<IAudioParameterControllerInterface>& InParameterInterface) const
+{
+	InParameterInterface->SetStringArrayParameter(InParameterName, Default);
+}
+
+FMetasoundFrontendLiteral UMetasoundEditorGraphMemberDefaultObject::GetDefault() const
+{
+	FMetasoundFrontendLiteral Literal;
+	Literal.Set(Default.Object);
+	return Literal;
+}
+
+EMetasoundFrontendLiteralType UMetasoundEditorGraphMemberDefaultObject::GetLiteralType() const
+{
+	return EMetasoundFrontendLiteralType::UObject;
+}
+
+void UMetasoundEditorGraphMemberDefaultObject::SetFromLiteral(const FMetasoundFrontendLiteral& InLiteral)
+{
+	ensure(InLiteral.TryGet(Default.Object));
+}
+
+void UMetasoundEditorGraphMemberDefaultObject::UpdatePreviewInstance(const Metasound::FVertexName& InParameterName, TScriptInterface<IAudioParameterControllerInterface>& InParameterInterface) const
+{
+	// TODO. We need proxy object here safely.
+}
+
+FMetasoundFrontendLiteral UMetasoundEditorGraphMemberDefaultObjectArray::GetDefault() const
+{
+	TArray<UObject*> ObjectArray;
+	Algo::Transform(Default, ObjectArray, [](const FMetasoundEditorGraphMemberDefaultObjectRef& InValue) { return InValue.Object; });
+
+	FMetasoundFrontendLiteral Literal;
+	Literal.Set(ObjectArray);
+	return Literal;
+}
+
+EMetasoundFrontendLiteralType UMetasoundEditorGraphMemberDefaultObjectArray::GetLiteralType() const
+{
+	return EMetasoundFrontendLiteralType::UObjectArray;
+}
+
+void UMetasoundEditorGraphMemberDefaultObjectArray::SetFromLiteral(const FMetasoundFrontendLiteral& InLiteral)
+{
+	TArray<UObject*> ObjectArray;
+	ensure(InLiteral.TryGet(ObjectArray));
+
+	Algo::Transform(ObjectArray, Default, [](UObject* InValue) { return FMetasoundEditorGraphMemberDefaultObjectRef { InValue }; });
+}
+
+void UMetasoundEditorGraphMemberDefaultObjectArray::UpdatePreviewInstance(const Metasound::FVertexName& InParameterName, TScriptInterface<IAudioParameterControllerInterface>& InParameterInterface) const
+{
+	TArray<UObject*> ObjectArray;
+	Algo::Transform(Default, ObjectArray, [](const FMetasoundEditorGraphMemberDefaultObjectRef& InValue) { return InValue.Object; });
+	// TODO. We need proxy object here safely.
+}
diff --git a/Engine/Plugins/Runtime/Metasound/Source/MetasoundEditor/Private/MetasoundEditorGraphNode.cpp b/Engine/Plugins/Runtime/Metasound/Source/MetasoundEditor/Private/MetasoundEditorGraphNode.cpp
--- a/Engine/Plugins/Runtime/Metasound/Source/MetasoundEditor/Private/MetasoundEditorGraphNode.cpp
+++ b/Engine/Plugins/Runtime/Metasound/Source/MetasoundEditor/Private/MetasoundEditorGraphNode.cpp
@@ -472,6 +472,54 @@ FString UMetasoundEditorGraphNode::GetDocumentationExcerptName() const
 	return FString::Printf(TEXT("%s%s"), UMetaSound::StaticClass()->GetPrefixCPP(), *UMetaSound::StaticClass()->GetName());
 }
 
+void UMetasoundEditorGraphOutputNode::PinDefaultValueChanged(UEdGraphPin* InPin)
+{
+	using namespace Metasound::Editor;
+	using namespace Metasound::Frontend;
+
+	if (InPin && InPin->Direction == EGPD_Input)
+	{
+		GetMetasoundChecked().Modify();
+
+		FInputHandle InputHandle = FGraphBuilder::GetInputHandleFromPin(InPin);
+		if (InputHandle->IsValid())
+		{
+			FMetasoundFrontendLiteral LiteralValue;
+			if (FGraphBuilder::GetPinLiteral(*InPin, LiteralValue))
+			{
+				InputHandle->SetLiteral(LiteralValue);
+			}
+
+			if (Output)
+			{
+				UMetasoundEditorGraphMemberDefaultLiteral* Literal = Output->GetLiteral();
+				if (ensure(Literal))
+				{
+					Literal->SetFromLiteral(LiteralValue);
+				}
+
+				constexpr bool bPostTransaction = false;
+				Output->UpdateFrontendDefaultLiteral(bPostTransaction);
+			}
+		}
+	}
+}
+
+bool UMetasoundEditorGraphOutputNode::EnableInteractWidgets() const
+{
+	using namespace Metasound::Frontend;
+
+	bool bEnabled = true;
+	GetConstNodeHandle()->IterateConstInputs([bIsEnabled = &bEnabled](FConstInputHandle InputHandle)
+	{
+		if (InputHandle->IsConnectionUserModifiable())
+		{
+			*bIsEnabled &= !InputHandle->IsConnected();
+		}
+	});
+	return bEnabled;
+}
+
 FMetasoundFrontendClassName UMetasoundEditorGraphOutputNode::GetClassName() const
 {
 	if (ensure(Output))
@@ -513,10 +561,14 @@ FLinearColor UMetasoundEditorGraphOutputNode::GetNodeTitleColor() const
 	return Super::GetNodeTitleColor();
 }
 
+UMetasoundEditorGraphMember* UMetasoundEditorGraphOutputNode::GetMember() const
+{
+	return Output;
+}
+
 FSlateIcon UMetasoundEditorGraphOutputNode::GetNodeTitleIcon() const
 {
-	static const FName NativeIconName = "MetasoundEditor.Graph.Node.Class.Output";
-	return FSlateIcon("MetaSoundStyle", NativeIconName);
+	return FSlateIcon("MetaSoundStyle", "MetasoundEditor.Graph.Node.Class.Output");
 }
 
 bool UMetasoundEditorGraphExternalNode::RefreshPinMetadata()
@@ -896,16 +948,48 @@ FSlateIcon UMetasoundEditorGraphExternalNode::GetNodeTitleIcon() const
 {
 	if (bIsClassNative)
 	{
-		static const FName NativeIconName = "MetasoundEditor.Graph.Node.Class.Native";
-		return FSlateIcon("MetaSoundStyle", NativeIconName);
+		return FSlateIcon("MetaSoundStyle", "MetasoundEditor.Graph.Node.Class.Native");
 	}
 	else
 	{
-		static const FName GraphRefIconName = "MetasoundEditor.Graph.Node.Class.Graph";
-		return FSlateIcon("MetasoundStyle", GraphRefIconName);
+		return FSlateIcon("MetasoundStyle", "MetasoundEditor.Graph.Node.Class.Graph");
 	}
 }
 
+UMetasoundEditorGraphMember* UMetasoundEditorGraphVariableNode::GetMember() const
+{
+	return Variable;
+}
+
+bool UMetasoundEditorGraphVariableNode::EnableInteractWidgets() const
+{
+	using namespace Metasound::Frontend;
+
+	bool bEnabled = false;
+
+	if (Variable)
+	{
+		FConstVariableHandle VariableHandle = Variable->GetConstVariableHandle();
+		FConstNodeHandle MutatorNode = VariableHandle->FindMutatorNode();
+		if (MutatorNode->IsValid())
+		{
+			if (MutatorNode->GetID() == NodeID)
+			{
+				bEnabled = true;
+				MutatorNode->IterateConstInputs([bIsEnabled = &bEnabled](FConstInputHandle InputHandle)
+				{
+					if (InputHandle->IsConnectionUserModifiable())
+					{
+						// Don't enable if variable input is connected
+						*bIsEnabled &= !InputHandle->IsConnected();
+					}
+				});
+			}
+		}
+	}
+
+	return bEnabled;
+}
 
 FMetasoundFrontendClassName UMetasoundEditorGraphVariableNode::GetClassName() const
 {
@@ -917,6 +1001,44 @@ FGuid UMetasoundEditorGraphVariableNode::GetNodeID() const
 	return NodeID;
 }
 
+void UMetasoundEditorGraphVariableNode::PinDefaultValueChanged(UEdGraphPin* Pin)
+{
+	using namespace Metasound::Editor;
+	using namespace Metasound::Frontend;
+
+	if (Pin && Pin->Direction == EGPD_Input)
+	{
+		GetMetasoundChecked().Modify();
+
+		FInputHandle InputHandle = FGraphBuilder::GetInputHandleFromPin(Pin);
+		if (InputHandle->IsValid())
+		{
+			FMetasoundFrontendLiteral LiteralValue;
+			if (FGraphBuilder::GetPinLiteral(*Pin, LiteralValue))
+			{
+				InputHandle->SetLiteral(LiteralValue);
+			}
+
+			// If this is the mutator node, synchronize the variable default literal with this default.
+			FNodeHandle MutatorNode = Variable->GetVariableHandle()->FindMutatorNode();
+			if (MutatorNode->IsValid())
+			{
+				if (MutatorNode->GetID() == NodeID)
+				{
+					UMetasoundEditorGraphMemberDefaultLiteral* Literal = Variable->GetLiteral();
+					if (ensure(Literal))
+					{
+						Literal->SetFromLiteral(LiteralValue);
+					}
+
+					constexpr bool bPostTransaction = false;
+					Variable->UpdateFrontendDefaultLiteral(bPostTransaction);
+				}
+			}
+		}
+	}
+}
+
 FLinearColor UMetasoundEditorGraphVariableNode::GetNodeTitleColor() const
 {
 	if (const UMetasoundEditorSettings* EditorSettings = GetDefault<UMetasoundEditorSettings>())
@@ -929,13 +1051,11 @@ FLinearColor UMetasoundEditorGraphVariableNode::GetNodeTitleColor() const
 
 FSlateIcon UMetasoundEditorGraphVariableNode::GetNodeTitleIcon() const
 {
-	static const FName NativeIconName = "MetasoundEditor.Graph.Node.Class.Variable";
-	return FSlateIcon("MetaSoundStyle", NativeIconName);
+	return FSlateIcon("MetaSoundStyle", "MetasoundEditor.Graph.Node.Class.Variable");
 }
 
 void UMetasoundEditorGraphVariableNode::SetNodeID(FGuid InNodeID)
 {
 	NodeID = InNodeID;
 }
-
 #undef LOCTEXT_NAMESPACE
diff --git a/Engine/Plugins/Runtime/Metasound/Source/MetasoundEditor/Private/MetasoundEditorGraphNodeFactory.h b/Engine/Plugins/Runtime/Metasound/Source/MetasoundEditor/Private/MetasoundEditorGraphNodeFactory.h
--- a/Engine/Plugins/Runtime/Metasound/Source/MetasoundEditor/Private/MetasoundEditorGraphNodeFactory.h
+++ b/Engine/Plugins/Runtime/Metasound/Source/MetasoundEditor/Private/MetasoundEditorGraphNodeFactory.h
@@ -1,11 +1,12 @@
 // Copyright Epic Games, Inc. All Rights Reserved.
 #pragma once
 
-#include "CoreMinimal.h"
 #include "EdGraphUtilities.h"
+#include "EdGraph/EdGraphNode.h"
 #include "MetasoundEditorGraphNode.h"
 #include "SGraphNode.h"
 #include "SMetasoundGraphNode.h"
+#include "Templates/SharedPointer.h"
 #include "Widgets/DeclarativeSyntaxSupport.h"
 
 
@@ -13,9 +14,11 @@ class FMetasoundGraphNodeFactory : public FGraphPanelNodeFactory
 {
 	virtual TSharedPtr<SGraphNode> CreateNode(UEdGraphNode* InNode) const override
 	{
+		using namespace Metasound::Editor;
+
 		if (InNode->IsA<UMetasoundEditorGraphNode>())
 		{
-			return SNew(SMetasoundGraphNode, InNode);
+			return SNew(SMetaSoundGraphNode, InNode);
 		}
 
 		return nullptr;
diff --git a/Engine/Plugins/Runtime/Metasound/Source/MetasoundEditor/Private/MetasoundEditorGraphSchema.cpp b/Engine/Plugins/Runtime/Metasound/Source/MetasoundEditor/Private/MetasoundEditorGraphSchema.cpp
--- a/Engine/Plugins/Runtime/Metasound/Source/MetasoundEditor/Private/MetasoundEditorGraphSchema.cpp
+++ b/Engine/Plugins/Runtime/Metasound/Source/MetasoundEditor/Private/MetasoundEditorGraphSchema.cpp
@@ -21,6 +21,7 @@
 #include "MetasoundEditorCommands.h"
 #include "MetasoundEditorGraph.h"
 #include "MetasoundEditorGraphBuilder.h"
+#include "MetasoundEditorGraphInputNode.h"
 #include "MetasoundEditorGraphNode.h"
 #include "MetasoundEditorModule.h"
 #include "MetasoundEditorSettings.h"
@@ -29,7 +30,9 @@
 #include "MetasoundFrontendRegistries.h"
 #include "MetasoundFrontendSearchEngine.h"
 #include "MetasoundLiteral.h"
+#include "MetasoundNodeInterface.h"
 #include "MetasoundStandardNodesCategories.h"
+#include "MetasoundVariableNodes.h"
 #include "ScopedTransaction.h"
 #include "Styling/SlateStyleRegistry.h"
 #include "Toolkits/ToolkitManager.h"
@@ -103,7 +106,8 @@ namespace Metasound
 			static const FText VariableMutatorDisplayNameFormat = LOCTEXT("DisplayNameAddVariableMutatorFormat", "Set {0}");
 			static const FText VariableMutatorTooltipFormat = LOCTEXT("TooltipAddVariableMutatorFormat", "Adds a setter for the variable '{0}' to the graph.");
 
-			bool TryConnectNewNodeToPin(UEdGraphNode& NewGraphNode, UEdGraphPin* FromPin)
+			// Connects to first pin with the same DataType
+			bool TryConnectNewNodeToMatchingDataTypePin(UEdGraphNode& NewGraphNode, UEdGraphPin* FromPin)
 			{
 				using namespace Metasound::Frontend;
 
@@ -185,6 +189,78 @@ namespace Metasound
 					}
 				}
 			}
+
+			UEdGraphNode* PromoteToVariable(const FName BaseName, UEdGraphPin& FromPin, const FName DataType, const FMetasoundFrontendClass& InVariableClass, const FVector2D& InLocation)
+			{
+				using namespace Frontend;
+
+				UEdGraphNode* ConnectedNode = Cast<UEdGraphNode>(FromPin.GetOwningNode());
+				if (!ensure(ConnectedNode))
+				{
+					return nullptr;
+				}
+
+				const FMetasoundFrontendClassName& ClassName = InVariableClass.Metadata.GetClassName();
+
+				UMetasoundEditorGraph* MetaSoundGraph = CastChecked<UMetasoundEditorGraph>(ConnectedNode->GetGraph());
+				const FName NodeName = FGraphBuilder::GenerateUniqueVariableName(MetaSoundGraph->GetGraphHandle(), BaseName.ToString());
+
+				const FScopedTransaction Transaction(FText::Format(
+					LOCTEXT("PromoteNodeVertexToGraphVariableFormat", "Promote MetaSound Node {0} to {1}"),
+					FText::FromName(NodeName),
+					FText::FromName(ClassName.Namespace)));
+
+				UObject& ParentMetaSound = MetaSoundGraph->GetMetasoundChecked();
+				ParentMetaSound.Modify();
+				MetaSoundGraph->Modify();
+
+				// Cache the default literal from the pin if connecting to an input
+				FMetasoundFrontendLiteral DefaultValue;
+				if (FromPin.Direction == EGPD_Input)
+				{
+					FGraphBuilder::GetPinLiteral(FromPin, DefaultValue);
+				}
+
+				FVariableHandle VariableHandle = FGraphBuilder::AddVariableHandle(ParentMetaSound, DataType);
+				if (ensure(VariableHandle->IsValid()))
+				{
+					UMetasoundEditorGraphVariable* Variable = MetaSoundGraph->FindOrAddVariable(VariableHandle);
+					if (ensure(Variable))
+					{
+						constexpr bool bPostSubTransaction = false;
+						Variable->SetMemberName(NodeName, bPostSubTransaction);
+
+						FNodeHandle NodeHandle = FGraphBuilder::AddVariableNodeHandle(ParentMetaSound, Variable->GetVariableID(), ClassName.ToNodeClassName());
+
+						if (UMetasoundEditorGraphVariableNode* NewGraphNode = FGraphBuilder::AddVariableNode(ParentMetaSound, NodeHandle, InLocation))
+						{
+							// Set the literal using the new value if connecting to an input
+							if (FromPin.Direction == EGPD_Input)
+							{
+								UMetasoundEditorGraphMemberDefaultLiteral* DefaultLiteral = Variable->GetLiteral();
+								if (ensure(DefaultLiteral))
+								{
+									DefaultLiteral->SetFromLiteral(DefaultValue);
+								}
+
+								// Ensures the setter node value is synced with the editor literal value
+								constexpr bool bPostTransaction = false;
+								Variable->UpdateFrontendDefaultLiteral(bPostTransaction);
+							}
+
+							UEdGraphNode* EdGraphNode = CastChecked<UEdGraphNode>(NewGraphNode);
+							if (ensure(SchemaPrivate::TryConnectNewNodeToMatchingDataTypePin(*EdGraphNode, &FromPin)))
+							{
+								FGraphBuilder::RegisterGraphWithFrontend(ParentMetaSound);
+								MetaSoundGraph->SetSynchronizationRequired();
+								return EdGraphNode;
+							}
+						}
+					}
+				}
+
+				return nullptr;
+			}
 		}
 	} // namespace Editor
 } // namespace Metasound
@@ -241,7 +317,7 @@ UEdGraphNode* FMetasoundGraphSchemaAction_NewNode::PerformAction(UEdGraph* Paren
 	if (UMetasoundEditorGraphExternalNode* NewGraphNode = FGraphBuilder::AddExternalNode(ParentMetasound, ClassMetadata, Location, bSelectNewNode))
 	{
 		NewGraphNode->Modify();
-		SchemaPrivate::TryConnectNewNodeToPin(*NewGraphNode, FromPin);
+		SchemaPrivate::TryConnectNewNodeToMatchingDataTypePin(*NewGraphNode, FromPin);
 		return NewGraphNode;
 	}
 
@@ -303,7 +379,7 @@ UEdGraphNode* FMetasoundGraphSchemaAction_NewInput::PerformAction(UEdGraph* Pare
 	{
 		NewGraphNode->Modify();
 		UEdGraphNode* EdGraphNode = CastChecked<UEdGraphNode>(NewGraphNode);
-		SchemaPrivate::TryConnectNewNodeToPin(*EdGraphNode, FromPin);
+		SchemaPrivate::TryConnectNewNodeToMatchingDataTypePin(*EdGraphNode, FromPin);
 
 		FGraphBuilder::RegisterGraphWithFrontend(ParentMetasound);
 		return EdGraphNode;
@@ -313,10 +389,9 @@ UEdGraphNode* FMetasoundGraphSchemaAction_NewInput::PerformAction(UEdGraph* Pare
 }
 
 FMetasoundGraphSchemaAction_PromoteToInput::FMetasoundGraphSchemaAction_PromoteToInput()
-	: FMetasoundGraphSchemaAction_NewInput(
+	: FMetasoundGraphSchemaAction(
 		FText(),
 		LOCTEXT("PromoteToInputName", "Promote To Graph Input"),
-		FGuid(),
 		LOCTEXT("PromoteToInputTooltip", "Promotes node input to graph input"),
 		Metasound::Editor::EPrimaryContextGroup::Common)
 {
@@ -352,7 +427,7 @@ UEdGraphNode* FMetasoundGraphSchemaAction_PromoteToInput::PerformAction(UEdGraph
 			{
 				UEdGraphNode* EdGraphNode = CastChecked<UEdGraphNode>(NewGraphNode);
 
-				if (ensure(SchemaPrivate::TryConnectNewNodeToPin(*EdGraphNode, FromPin)))
+				if (ensure(SchemaPrivate::TryConnectNewNodeToMatchingDataTypePin(*EdGraphNode, FromPin)))
 				{
 					TSharedPtr<FEditor> MetasoundEditor = FGraphBuilder::GetEditorForGraph(*ParentGraph);
 					if (MetasoundEditor.IsValid())
@@ -370,6 +445,102 @@ UEdGraphNode* FMetasoundGraphSchemaAction_PromoteToInput::PerformAction(UEdGraph
 	return nullptr;
 }
 
+FMetasoundGraphSchemaAction_PromoteToVariable_AccessorNode::FMetasoundGraphSchemaAction_PromoteToVariable_AccessorNode()
+	: FMetasoundGraphSchemaAction(
+		FText(),
+		LOCTEXT("PromoteToVariableGetterName", "Promote To Graph Variable"),
+		LOCTEXT("PromoteToInputTooltip", "Promotes node input to graph variable and creates a setter node"),
+		Metasound::Editor::EPrimaryContextGroup::Common)
+{
+}
+
+UEdGraphNode* FMetasoundGraphSchemaAction_PromoteToVariable_AccessorNode::PerformAction(UEdGraph* ParentGraph, UEdGraphPin* FromPin, const FVector2D InLocation, bool bSelectNewNode /* = true */)
+{
+	using namespace Metasound;
+	using namespace Metasound::Editor;
+	using namespace Metasound::Frontend;
+
+	FInputHandle InputHandle = FGraphBuilder::GetInputHandleFromPin(FromPin);
+	if (!ensure(InputHandle->IsValid()))
+	{
+		return nullptr;
+	}
+
+	const FName NodeName = InputHandle->GetName();
+	const FName DataType = InputHandle->GetDataType();
+	FMetasoundFrontendClass VariableClass;
+	if (ensure(IDataTypeRegistry::Get().GetFrontendVariableAccessorClass(DataType, VariableClass)))
+	{
+		return SchemaPrivate::PromoteToVariable(NodeName, *FromPin, DataType, VariableClass, InLocation);
+	}
+
+	return nullptr;
+}
+
+FMetasoundGraphSchemaAction_PromoteToVariable_DeferredAccessorNode::FMetasoundGraphSchemaAction_PromoteToVariable_DeferredAccessorNode()
+	: FMetasoundGraphSchemaAction(
+		FText(),
+		LOCTEXT("PromoteToVariableDeferredGetterName", "Promote To Graph Variable (Deferred)"),
+		LOCTEXT("PromoteToInputTooltip", "Promotes node input to graph variable and creates a deferred getter node"),
+		Metasound::Editor::EPrimaryContextGroup::Common)
+{
+}
+
+UEdGraphNode* FMetasoundGraphSchemaAction_PromoteToVariable_DeferredAccessorNode::PerformAction(UEdGraph* ParentGraph, UEdGraphPin* FromPin, const FVector2D InLocation, bool bSelectNewNode /* = true */)
+{
+	using namespace Metasound;
+	using namespace Metasound::Editor;
+	using namespace Metasound::Frontend;
+
+	FInputHandle InputHandle = FGraphBuilder::GetInputHandleFromPin(FromPin);
+	if (!ensure(InputHandle->IsValid()))
+	{
+		return nullptr;
+	}
+
+	const FName NodeName = InputHandle->GetName();
+	const FName DataType = InputHandle->GetDataType();
+	FMetasoundFrontendClass VariableClass;
+	if (ensure(IDataTypeRegistry::Get().GetFrontendVariableDeferredAccessorClass(DataType, VariableClass)))
+	{
+		return SchemaPrivate::PromoteToVariable(NodeName, *FromPin, DataType, VariableClass, InLocation);
+	}
+
+	return nullptr;
+}
+
+FMetasoundGraphSchemaAction_PromoteToVariable_MutatorNode::FMetasoundGraphSchemaAction_PromoteToVariable_MutatorNode()
+	: FMetasoundGraphSchemaAction(
+		FText(),
+		LOCTEXT("PromoteToVariableSetterName", "Promote To Graph Variable"),
+		LOCTEXT("PromoteToVariableSetterTooltip", "Promotes node input to graph variable and creates a setter node"),
+		Metasound::Editor::EPrimaryContextGroup::Common)
+{
+}
+
+UEdGraphNode* FMetasoundGraphSchemaAction_PromoteToVariable_MutatorNode::PerformAction(UEdGraph* ParentGraph, UEdGraphPin* FromPin, const FVector2D InLocation, bool bSelectNewNode /* = true */)
+{
+	using namespace Metasound;
+	using namespace Metasound::Editor;
+	using namespace Metasound::Frontend;
+
+	FOutputHandle OutputHandle = FGraphBuilder::GetOutputHandleFromPin(FromPin);
+	if (!ensure(OutputHandle->IsValid()))
+	{
+		return nullptr;
+	}
+
+	const FName NodeName = OutputHandle->GetName();
+	const FName DataType = OutputHandle->GetDataType();
+	FMetasoundFrontendClass VariableClass;
+	if (ensure(IDataTypeRegistry::Get().GetFrontendVariableMutatorClass(DataType, VariableClass)))
+	{
+		return SchemaPrivate::PromoteToVariable(NodeName, *FromPin, DataType, VariableClass, InLocation);
+	}
+
+	return nullptr;
+}
+
 FMetasoundGraphSchemaAction_NewOutput::FMetasoundGraphSchemaAction_NewOutput(FText InNodeCategory, FText InDisplayName, FGuid InOutputNodeID, FText InToolTip, Metasound::Editor::EPrimaryContextGroup InGroup)
 	: FMetasoundGraphSchemaAction(MoveTemp(InNodeCategory), MoveTemp(InDisplayName), MoveTemp(InToolTip), InGroup)
 	, NodeID(InOutputNodeID)
@@ -422,7 +593,7 @@ UEdGraphNode* FMetasoundGraphSchemaAction_NewOutput::PerformAction(UEdGraph* Par
 
 	if (UMetasoundEditorGraphOutputNode* NewGraphNode = FGraphBuilder::AddOutputNode(ParentMetasound, NodeHandle, Location, bSelectNewNode))
 	{
-		SchemaPrivate::TryConnectNewNodeToPin(*NewGraphNode, FromPin);
+		SchemaPrivate::TryConnectNewNodeToMatchingDataTypePin(*NewGraphNode, FromPin);
 		FGraphBuilder::RegisterGraphWithFrontend(ParentMetasound);
 		return NewGraphNode;
 	}
@@ -431,10 +602,9 @@ UEdGraphNode* FMetasoundGraphSchemaAction_NewOutput::PerformAction(UEdGraph* Par
 }
 
 FMetasoundGraphSchemaAction_PromoteToOutput::FMetasoundGraphSchemaAction_PromoteToOutput()
-	: FMetasoundGraphSchemaAction_NewOutput(
+	: FMetasoundGraphSchemaAction(
 		FText(),
 		LOCTEXT("PromoteToOutputName", "Promote To Graph Output"),
-		FGuid(),
 		LOCTEXT("PromoteToOutputTooltip", "Promotes node output to graph output"),
 		Metasound::Editor::EPrimaryContextGroup::Common)
 {
@@ -469,7 +639,7 @@ UEdGraphNode* FMetasoundGraphSchemaAction_PromoteToOutput::PerformAction(UEdGrap
 			{
 				UEdGraphNode* EdGraphNode = CastChecked<UEdGraphNode>(NewGraphNode);
 
-				if (ensure(SchemaPrivate::TryConnectNewNodeToPin(*EdGraphNode, FromPin)))
+				if (ensure(SchemaPrivate::TryConnectNewNodeToMatchingDataTypePin(*EdGraphNode, FromPin)))
 				{
 					TSharedPtr<FEditor> MetasoundEditor = FGraphBuilder::GetEditorForGraph(*ParentGraph);
 					if (MetasoundEditor.IsValid())
@@ -487,8 +657,6 @@ UEdGraphNode* FMetasoundGraphSchemaAction_PromoteToOutput::PerformAction(UEdGrap
 	return nullptr;
 }
 
-
-
 FMetasoundGraphSchemaAction_NewVariableNode::FMetasoundGraphSchemaAction_NewVariableNode(FText InNodeCategory, FText InDisplayName, FGuid InVariableID, FText InToolTip)
 	: FMetasoundGraphSchemaAction(MoveTemp(InNodeCategory), MoveTemp(InDisplayName), MoveTemp(InToolTip), Metasound::Editor::EPrimaryContextGroup::Variables)
 	, VariableID(InVariableID)
@@ -506,7 +674,6 @@ UEdGraphNode* FMetasoundGraphSchemaAction_NewVariableNode::PerformAction(UEdGrap
 		{
 			if (UMetasoundEditorGraphVariable* Variable = MetasoundGraph->FindVariable(VariableID))
 			{
-
 				const FScopedTransaction Transaction(LOCTEXT("AddNewVariableAccessorNode", "Add New MetaSound Variable Accessor Node"));
 				ParentMetasound->Modify();
 				MetasoundGraph->Modify();
@@ -518,7 +685,15 @@ UEdGraphNode* FMetasoundGraphSchemaAction_NewVariableNode::PerformAction(UEdGrap
 					if (UMetasoundEditorGraphVariableNode* NewGraphNode = FGraphBuilder::AddVariableNode(*ParentMetasound, FrontendNode, Location, bSelectNewNode))
 					{
 						NewGraphNode->Modify();
-						SchemaPrivate::TryConnectNewNodeToPin(*NewGraphNode, FromPin);
+
+						// Ensures the setter node value is synced with the editor literal value
+						constexpr bool bPostTransaction = false;
+						Variable->UpdateFrontendDefaultLiteral(bPostTransaction);
+
+						SchemaPrivate::TryConnectNewNodeToMatchingDataTypePin(*NewGraphNode, FromPin);
+
+						MetasoundGraph->SetSynchronizationRequired();
+
 						return NewGraphNode;
 					}
 				}
@@ -529,11 +704,11 @@ UEdGraphNode* FMetasoundGraphSchemaAction_NewVariableNode::PerformAction(UEdGrap
 	return nullptr;
 }
 
-const FSlateBrush* FMetasoundGraphSchemaAction_NewVariableNode::GetIconBrush() const 
+const FSlateBrush* FMetasoundGraphSchemaAction_NewVariableNode::GetIconBrush() const
 {
 	if (const ISlateStyle* MetasoundStyle = FSlateStyleRegistry::FindSlateStyle("MetaSoundStyle"))
 	{
-		return MetasoundStyle->GetBrush("MetasoundEditor.Graph.Node.Class.Default");
+		return MetasoundStyle->GetBrush("MetasoundEditor.Graph.Node.Class.Variable");
 	}
 	return Super::GetIconBrush();
 }
@@ -553,7 +728,7 @@ FMetasoundGraphSchemaAction_NewVariableAccessorNode::FMetasoundGraphSchemaAction
 {
 }
 
-Metasound::Frontend::FNodeHandle FMetasoundGraphSchemaAction_NewVariableAccessorNode::CreateFrontendVariableNode(const Metasound::Frontend::FGraphHandle& InFrontendGraph, const FGuid& InVariableID) const 
+Metasound::Frontend::FNodeHandle FMetasoundGraphSchemaAction_NewVariableAccessorNode::CreateFrontendVariableNode(const Metasound::Frontend::FGraphHandle& InFrontendGraph, const FGuid& InVariableID) const
 {
 	return InFrontendGraph->AddVariableAccessorNode(InVariableID);
 }
@@ -563,7 +738,7 @@ FMetasoundGraphSchemaAction_NewVariableDeferredAccessorNode::FMetasoundGraphSche
 {
 }
 
-Metasound::Frontend::FNodeHandle FMetasoundGraphSchemaAction_NewVariableDeferredAccessorNode::CreateFrontendVariableNode(const Metasound::Frontend::FGraphHandle& InFrontendGraph, const FGuid& InVariableID) const 
+Metasound::Frontend::FNodeHandle FMetasoundGraphSchemaAction_NewVariableDeferredAccessorNode::CreateFrontendVariableNode(const Metasound::Frontend::FGraphHandle& InFrontendGraph, const FGuid& InVariableID) const
 {
 	return InFrontendGraph->AddVariableDeferredAccessorNode(InVariableID);
 }
@@ -573,11 +748,12 @@ FMetasoundGraphSchemaAction_NewVariableMutatorNode::FMetasoundGraphSchemaAction_
 {
 }
 
-Metasound::Frontend::FNodeHandle FMetasoundGraphSchemaAction_NewVariableMutatorNode::CreateFrontendVariableNode(const Metasound::Frontend::FGraphHandle& InFrontendGraph, const FGuid& InVariableID) const 
+Metasound::Frontend::FNodeHandle FMetasoundGraphSchemaAction_NewVariableMutatorNode::CreateFrontendVariableNode(const Metasound::Frontend::FGraphHandle& InFrontendGraph, const FGuid& InVariableID) const
 {
 	using namespace Metasound::Frontend;
-	// Only one mutator node should exist per variable. Check to make sure that
-	// there is a mutator node does not already exist for this variable.
+
+	// Only one mutator node should exist per variable. Check to
+	// make sure that one does not already exist for this variable.
 	FConstVariableHandle Variable = InFrontendGraph->FindVariable(InVariableID);
 	FConstNodeHandle MutatorNode = Variable->FindMutatorNode();
 	if (ensure(!MutatorNode->IsValid()))
@@ -727,7 +903,10 @@ void UMetasoundEditorGraphSchema::GetGraphContextActions(FGraphContextMenuBuilde
 				return bHasOutputOfType;
 			});
 
-			static_cast<FGraphActionMenuBuilder&>(ContextMenuBuilder).AddAction(MakeShared<FMetasoundGraphSchemaAction_PromoteToInput>());
+			FGraphActionMenuBuilder& ActionMenuBuilder = static_cast<FGraphActionMenuBuilder&>(ContextMenuBuilder);
+			ActionMenuBuilder.AddAction(MakeShared<FMetasoundGraphSchemaAction_PromoteToInput>());
+			ActionMenuBuilder.AddAction(MakeShared<FMetasoundGraphSchemaAction_PromoteToVariable_AccessorNode>());
+			ActionMenuBuilder.AddAction(MakeShared<FMetasoundGraphSchemaAction_PromoteToVariable_DeferredAccessorNode>());
 		}
 
 		if (FromPin->Direction == EGPD_Output)
@@ -750,7 +929,9 @@ void UMetasoundEditorGraphSchema::GetGraphContextActions(FGraphContextMenuBuilde
 				return bHasInputOfType;
 			});
 
-			static_cast<FGraphActionMenuBuilder&>(ContextMenuBuilder).AddAction(MakeShared<FMetasoundGraphSchemaAction_PromoteToOutput>());
+			FGraphActionMenuBuilder& ActionMenuBuilder = static_cast<FGraphActionMenuBuilder&>(ContextMenuBuilder);
+			ActionMenuBuilder.AddAction(MakeShared<FMetasoundGraphSchemaAction_PromoteToOutput>());
+			ActionMenuBuilder.AddAction(MakeShared<FMetasoundGraphSchemaAction_PromoteToVariable_MutatorNode>());
 		}
 	}
 	else
@@ -1020,23 +1201,31 @@ FText UMetasoundEditorGraphSchema::GetPinDisplayName(const UEdGraphPin* Pin) con
 	check(Pin);
 
 	UMetasoundEditorGraphNode* Node = CastChecked<UMetasoundEditorGraphNode>(Pin->GetOwningNode());
-	FNodeHandle NodeHandle = Node->GetNodeHandle();
+	FConstNodeHandle NodeHandle = Node->GetConstNodeHandle();
 	const EMetasoundFrontendClassType ClassType = NodeHandle->GetClassMetadata().GetType();
 
 	switch (ClassType)
 	{
 		case EMetasoundFrontendClassType::Input:
 		case EMetasoundFrontendClassType::Output:
-		{
-			constexpr bool bIncludeNamespace = true;
-			return FGraphBuilder::GetDisplayName(*NodeHandle, bIncludeNamespace);
-		}
-
-		case EMetasoundFrontendClassType::Literal:
 		case EMetasoundFrontendClassType::Variable:
 		case EMetasoundFrontendClassType::VariableAccessor:
 		case EMetasoundFrontendClassType::VariableDeferredAccessor:
 		case EMetasoundFrontendClassType::VariableMutator:
+		{
+			UMetasoundEditorGraphMemberNode* MemberNode = Cast<UMetasoundEditorGraphMemberNode>(Node);
+			if (ensure(MemberNode))
+			{
+				UMetasoundEditorGraphMember* Member = MemberNode->GetMember();
+				if (ensure(Member))
+				{
+					return MemberNode->GetMember()->GetDisplayName();
+				}
+			}
+			return Super::GetPinDisplayName(Pin);
+		}
+
+		case EMetasoundFrontendClassType::Literal:
 		case EMetasoundFrontendClassType::External:
 		{
 			if (Pin->Direction == EGPD_Input)
@@ -1060,8 +1249,6 @@ FText UMetasoundEditorGraphSchema::GetPinDisplayName(const UEdGraphPin* Pin) con
 		}
 
 		case EMetasoundFrontendClassType::Graph:
-			// TODO: Implement above
-
 		case EMetasoundFrontendClassType::Invalid:
 		default:
 		{
diff --git a/Engine/Plugins/Runtime/Metasound/Source/MetasoundEditor/Private/MetasoundEditorModule.cpp b/Engine/Plugins/Runtime/Metasound/Source/MetasoundEditor/Private/MetasoundEditorModule.cpp
--- a/Engine/Plugins/Runtime/Metasound/Source/MetasoundEditor/Private/MetasoundEditorModule.cpp
+++ b/Engine/Plugins/Runtime/Metasound/Source/MetasoundEditor/Private/MetasoundEditorModule.cpp
@@ -1,7 +1,6 @@
 // Copyright Epic Games, Inc. All Rights Reserved.
 #include "MetasoundEditorModule.h"
 
-
 #include "AssetRegistry/AssetData.h"
 #include "AssetRegistry/AssetRegistryModule.h"
 #include "AssetRegistryModule.h"
@@ -22,9 +21,8 @@
 #include "MetasoundEditorGraph.h"
 #include "MetasoundEditorGraphBuilder.h"
 #include "MetasoundEditorGraphConnectionDrawingPolicy.h"
-#include "MetasoundEditorGraphInputNodes.h"
+#include "MetasoundEditorGraphMemberDefaults.h"
 #include "MetasoundEditorGraphNodeFactory.h"
-#include "MetasoundEditorGraphSchema.h"
 #include "MetasoundEditorSettings.h"
 #include "MetasoundFrontendDataTypeRegistry.h"
 #include "MetasoundFrontendDocument.h"
@@ -300,9 +298,9 @@ namespace Metasound
 				}
 			}
 
-			void RegisterNodeInputClasses()
+			void RegisterInputDefaultClasses()
 			{
-				TSubclassOf<UMetasoundEditorGraphInputLiteral> NodeClass;
+				TSubclassOf<UMetasoundEditorGraphMemberDefaultLiteral> NodeClass;
 				for (TObjectIterator<UClass> ClassIt; ClassIt; ++ClassIt)
 				{
 					UClass* Class = *ClassIt;
@@ -316,14 +314,14 @@ namespace Metasound
 						continue;
 					}
 
-					if (!ClassIt->IsChildOf(UMetasoundEditorGraphInputLiteral::StaticClass()))
+					if (!ClassIt->IsChildOf(UMetasoundEditorGraphMemberDefaultLiteral::StaticClass()))
 					{
 						continue;
 					}
 
-					if (const UMetasoundEditorGraphInputLiteral* InputCDO = Class->GetDefaultObject<UMetasoundEditorGraphInputLiteral>())
+					if (const UMetasoundEditorGraphMemberDefaultLiteral* DefaultLiteralCDO = Class->GetDefaultObject<UMetasoundEditorGraphMemberDefaultLiteral>())
 					{
-						NodeInputClassRegistry.Add(InputCDO->GetLiteralType(), InputCDO->GetClass());
+						InputDefaultLiteralClassRegistry.Add(DefaultLiteralCDO->GetLiteralType(), DefaultLiteralCDO->GetClass());
 					}
 				}
 			}
@@ -469,9 +467,9 @@ namespace Metasound
 				return ExplicitProxyClasses.Contains(&InClass);
 			}
 
-			virtual TUniquePtr<IMetaSoundInputLiteralCustomization> CreateInputLiteralCustomization(UClass& InClass, IDetailCategoryBuilder& InDefaultCategoryBuilder) const override
+			virtual TUniquePtr<IMemberDefaultLiteralCustomization> CreateMemberDefaultLiteralCustomization(UClass& InClass, IDetailCategoryBuilder& InDefaultCategoryBuilder) const override
 			{
-				const TUniquePtr<IMetaSoundInputLiteralCustomizationFactory>* CustomizationFactory = LiteralCustomizationFactories.Find(&InClass);
+				const TUniquePtr<IMemberDefaultLiteralCustomizationFactory>* CustomizationFactory = LiteralCustomizationFactories.Find(&InClass);
 				if (CustomizationFactory && CustomizationFactory->IsValid())
 				{
 					return (*CustomizationFactory)->CreateLiteralCustomization(InDefaultCategoryBuilder);
@@ -480,9 +478,9 @@ namespace Metasound
 				return nullptr;
 			}
 
-			virtual const TSubclassOf<UMetasoundEditorGraphInputLiteral> FindInputLiteralClass(EMetasoundFrontendLiteralType InLiteralType) const override
+			virtual const TSubclassOf<UMetasoundEditorGraphMemberDefaultLiteral> FindDefaultLiteralClass(EMetasoundFrontendLiteralType InLiteralType) const override
 			{
-				return NodeInputClassRegistry.FindRef(InLiteralType);
+				return InputDefaultLiteralClassRegistry.FindRef(InLiteralType);
 			}
 
 			virtual const FEditorDataType* FindDataType(FName InDataTypeName) const override
@@ -545,23 +543,23 @@ namespace Metasound
 					FOnGetDetailCustomizationInstance::CreateLambda([]() { return MakeShared<FMetasoundVariableDetailCustomization>(); }));
 
 				PropertyModule.RegisterCustomPropertyTypeLayout(
-					"MetasoundEditorGraphInputBoolRef",
-					FOnGetPropertyTypeCustomizationInstance::CreateLambda([]() { return MakeShared<FMetasoundInputBoolDetailCustomization>(); }));
+					"MetaSoundEditorGraphMemberDefaultBoolRef",
+					FOnGetPropertyTypeCustomizationInstance::CreateLambda([]() { return MakeShared<FMetasoundMemberDefaultBoolDetailCustomization>(); }));
 
 				PropertyModule.RegisterCustomPropertyTypeLayout(
-					"MetasoundEditorGraphInputIntRef",
-					FOnGetPropertyTypeCustomizationInstance::CreateLambda([]() { return MakeShared<FMetasoundInputIntDetailCustomization>(); }));
+					"MetaSoundEditorGraphMemberDefaultIntRef",
+					FOnGetPropertyTypeCustomizationInstance::CreateLambda([]() { return MakeShared<FMetasoundMemberDefaultIntDetailCustomization>(); }));
 
 				PropertyModule.RegisterCustomPropertyTypeLayout(
-					"MetasoundEditorGraphInputObjectRef",
-					FOnGetPropertyTypeCustomizationInstance::CreateLambda([]() { return MakeShared<FMetasoundInputObjectDetailCustomization>(); }));
+					"MetaSoundEditorGraphMemberDefaultObjectRef",
+					FOnGetPropertyTypeCustomizationInstance::CreateLambda([]() { return MakeShared<FMetasoundMemberDefaultObjectDetailCustomization>(); }));
 
-				LiteralCustomizationFactories.Add(UMetasoundEditorGraphInputFloat::StaticClass(), MakeUnique<FMetasoundFloatLiteralCustomizationFactory>());
+				LiteralCustomizationFactories.Add(UMetasoundEditorGraphMemberDefaultFloat::StaticClass(), MakeUnique<FMetasoundFloatLiteralCustomizationFactory>());
 
 				StyleSet = MakeShared<FSlateStyle>();
 
 				RegisterCoreDataTypes();
-				RegisterNodeInputClasses();
+				RegisterInputDefaultClasses();
 
 				GraphConnectionFactory = MakeShared<FGraphConnectionDrawingPolicyFactory>();
 				FEdGraphUtilities::RegisterVisualPinConnectionFactory(GraphConnectionFactory);
@@ -636,9 +634,9 @@ namespace Metasound
 
 			TArray<TSharedPtr<FAssetTypeActions_Base>> AssetActions;
 			TMap<FName, FEditorDataType> DataTypeInfo;
-			TMap<EMetasoundFrontendLiteralType, const TSubclassOf<UMetasoundEditorGraphInputLiteral>> NodeInputClassRegistry;
+			TMap<EMetasoundFrontendLiteralType, const TSubclassOf<UMetasoundEditorGraphMemberDefaultLiteral>> InputDefaultLiteralClassRegistry;
 
-			TMap<UClass*, TUniquePtr<IMetaSoundInputLiteralCustomizationFactory>> LiteralCustomizationFactories;
+			TMap<UClass*, TUniquePtr<IMemberDefaultLiteralCustomizationFactory>> LiteralCustomizationFactories;
 
 			TSharedPtr<FMetasoundGraphNodeFactory> GraphNodeFactory;
 			TSharedPtr<FGraphPanelPinConnectionFactory> GraphConnectionFactory;
diff --git a/Engine/Plugins/Runtime/Metasound/Source/MetasoundEditor/Private/MetasoundFactory.cpp b/Engine/Plugins/Runtime/Metasound/Source/MetasoundEditor/Private/MetasoundFactory.cpp
--- a/Engine/Plugins/Runtime/Metasound/Source/MetasoundEditor/Private/MetasoundFactory.cpp
+++ b/Engine/Plugins/Runtime/Metasound/Source/MetasoundEditor/Private/MetasoundFactory.cpp
@@ -3,7 +3,9 @@
 
 #include "Kismet/KismetSystemLibrary.h"
 #include "Metasound.h"
+#include "MetasoundEditorGraph.h"
 #include "MetasoundEditorGraphBuilder.h"
+#include "MetasoundEditorGraphSchema.h"
 #include "MetasoundFrontendTransform.h"
 #include "MetasoundSource.h"
 
@@ -16,6 +18,7 @@ namespace Metasound
 		T* CreateNewMetaSoundObject(UObject* InParent, FName InName, EObjectFlags InFlags, UObject* InReferencedMetaSoundObject)
 		{
 			using namespace Editor;
+			using namespace Frontend;
 
 			T* MetaSoundObject = NewObject<T>(InParent, InName, InFlags);
 			check(MetaSoundObject);
@@ -27,7 +30,21 @@ namespace Metasound
 				FGraphBuilder::InitMetaSoundPreset(*InReferencedMetaSoundObject, *MetaSoundObject);
 			}
 
-			FGraphBuilder::InitGraph(*MetaSoundObject);
+			FMetasoundAssetBase* MetaSoundAsset = IMetasoundUObjectRegistry::Get().GetObjectAsAssetBase(MetaSoundObject);
+			check(MetaSoundAsset);
+
+			// Initial graph generation is not something to be managed by the transaction
+			// stack, so don't track dirty state until after initial setup if necessary.
+			UMetasoundEditorGraph* Graph = Cast<UMetasoundEditorGraph>(MetaSoundAsset->GetGraph());
+			if (!Graph)
+			{
+				Graph = NewObject<UMetasoundEditorGraph>(MetaSoundObject, FName(), RF_Transactional);
+				Graph->Schema = UMetasoundEditorGraphSchema::StaticClass();
+				MetaSoundAsset->SetGraph(Graph);
+
+				// Has to be done inline to have valid graph initially when opening editor for the first time
+				FGraphBuilder::SynchronizeGraph(*MetaSoundObject);
+			}
 
 			return MetaSoundObject;
 		}
diff --git a/Engine/Plugins/Runtime/Metasound/Source/MetasoundEditor/Private/MetasoundNodeDetailCustomization.cpp b/Engine/Plugins/Runtime/Metasound/Source/MetasoundEditor/Private/MetasoundNodeDetailCustomization.cpp
--- a/Engine/Plugins/Runtime/Metasound/Source/MetasoundEditor/Private/MetasoundNodeDetailCustomization.cpp
+++ b/Engine/Plugins/Runtime/Metasound/Source/MetasoundEditor/Private/MetasoundNodeDetailCustomization.cpp
@@ -14,13 +14,16 @@
 #include "Internationalization/Text.h"
 #include "MetasoundAssetBase.h"
 #include "MetasoundDataReference.h"
+#include "MetasoundDataReferenceMacro.h"
 #include "MetasoundEditorGraphBuilder.h"
 #include "MetasoundEditorGraphNode.h"
-#include "MetasoundEditorGraphInputNodes.h"
+#include "MetasoundEditorGraphInputNode.h"
+#include "MetasoundEditorGraphMemberDefaults.h"
 #include "MetasoundEditorGraphSchema.h"
 #include "MetasoundEditorModule.h"
 #include "MetasoundFrontend.h"
 #include "MetasoundFrontendController.h"
+#include "MetasoundFrontendDataTypeRegistry.h"
 #include "MetasoundFrontendRegistries.h"
 #include "MetasoundUObjectRegistry.h"
 #include "PropertyCustomizationHelpers.h"
@@ -54,12 +57,6 @@ namespace Metasound
 	{
 		namespace MemberCustomizationPrivate
 		{
-			static const FText InputNameText = LOCTEXT("Input_Name", "Input Name");
-			static const FText InputDisplayNameText = LOCTEXT("InputDisplay_Name", "Input Display Name");
-
-			static const FText OutputNameText = LOCTEXT("Output_Name", "Output Name");
-			static const FText OutputDisplayNameText = LOCTEXT("OutputDisplay_Name", "Output Display Name");
-
 			/** Set of input types which are valid registered types, but should
 			 * not show up as an input type option in the MetaSound editor. */
 			static const TSet<FName> HiddenInputTypeNames =
@@ -73,28 +70,28 @@ namespace Metasound
 		{
 			if (FloatLiteral.IsValid())
 			{
-				FloatLiteral->OnClampInputChanged.Remove(OnClampInputChangedDelegateHandle);
-				FloatLiteral->OnRangeChanged.Remove(InputWidgetOnRangeChangedDelegateHandle);
+				FloatLiteral->OnClampChanged.Remove(OnClampChangedDelegateHandle);
+				FloatLiteral->OnRangeChanged.Remove(OnRangeChangedDelegateHandle);
 			}
 		}
 
-		void FMetasoundFloatLiteralCustomization::CustomizeLiteral(UMetasoundEditorGraphInputLiteral& InLiteral, TSharedPtr<IPropertyHandle> InDefaultValueHandle)
+		void FMetasoundFloatLiteralCustomization::CustomizeLiteral(UMetasoundEditorGraphMemberDefaultLiteral& InLiteral, TSharedPtr<IPropertyHandle> InDefaultValueHandle)
 		{
 			check(InputCategoryBuilder);
 
-			UMetasoundEditorGraphInputFloat* InputFloat = Cast<UMetasoundEditorGraphInputFloat>(&InLiteral);
-			if (!ensure(InputFloat))
+			UMetasoundEditorGraphMemberDefaultFloat* DefaultFloat = Cast<UMetasoundEditorGraphMemberDefaultFloat>(&InLiteral);
+			if (!ensure(DefaultFloat))
 			{
 				return;
 			}
-			FloatLiteral = InputFloat;
+			FloatLiteral = DefaultFloat;
 
-			if (IDetailPropertyRow* Row = InputCategoryBuilder->AddExternalObjectProperty(TArray<UObject*>({ InputFloat }), GET_MEMBER_NAME_CHECKED(UMetasoundEditorGraphInputFloat, ClampDefault)))
+			if (IDetailPropertyRow* Row = InputCategoryBuilder->AddExternalObjectProperty(TArray<UObject*>({ DefaultFloat }), GET_MEMBER_NAME_CHECKED(UMetasoundEditorGraphMemberDefaultFloat, ClampDefault)))
 			{
 				// If clamping or using slider, clamp default value to given range 
-				if (InputFloat->ClampDefault || InputFloat->InputWidgetType != EMetasoundInputWidget::None)
+				if (DefaultFloat->ClampDefault || DefaultFloat->WidgetType != EMetasoundMemberDefaultWidget::None)
 				{
-					FVector2D Range = InputFloat->GetRange();
+					FVector2D Range = DefaultFloat->GetRange();
 					InDefaultValueHandle->SetInstanceMetaData("ClampMin", FString::Printf(TEXT("%f"), Range.X));
 					InDefaultValueHandle->SetInstanceMetaData("ClampMax", FString::Printf(TEXT("%f"), Range.Y));
 				}
@@ -104,12 +101,12 @@ namespace Metasound
 					InDefaultValueHandle->SetInstanceMetaData("ClampMax", "");
 				}
 
-				InputFloat->OnClampInputChanged.Remove(OnClampInputChangedDelegateHandle);
-				OnClampInputChangedDelegateHandle = InputFloat->OnClampInputChanged.AddLambda([this](bool ClampInput)
+				DefaultFloat->OnClampChanged.Remove(OnClampChangedDelegateHandle);
+				OnClampChangedDelegateHandle = DefaultFloat->OnClampChanged.AddLambda([this](bool ClampInput)
 				{
 					if (FloatLiteral.IsValid())
 					{
-						FloatLiteral->OnClampInputChanged.Remove(OnClampInputChangedDelegateHandle);
+						FloatLiteral->OnClampChanged.Remove(OnClampChangedDelegateHandle);
 						if (FMetasoundAssetBase* MetasoundAsset = Metasound::IMetasoundUObjectRegistry::Get().GetObjectAsAssetBase(FloatLiteral->GetOutermostObject()))
 						{
 							MetasoundAsset->SetSynchronizationRequired();
@@ -117,56 +114,21 @@ namespace Metasound
 					}
 				});
 			}
-			if (IDetailPropertyRow* Row = InputCategoryBuilder->AddExternalObjectProperty(TArray<UObject*>({ InputFloat }), GET_MEMBER_NAME_CHECKED(UMetasoundEditorGraphInputFloat, Range)))
-			{
-				TSharedPtr<IPropertyHandle> RangeHandle = Row->GetPropertyHandle();
-				if (RangeHandle.IsValid())
-				{
-					TWeakObjectPtr<UMetasoundEditorGraphInput> Input = Cast<UMetasoundEditorGraphInput>(FloatLiteral->GetOuter());
-					FSimpleDelegate UpdateDocumentInput = FSimpleDelegate::CreateLambda([Input]()
-					{
-						if (Input.IsValid())
-						{
-							if (UMetasoundEditorGraphInputLiteral* InputLiteral = Input->Literal)
-							{
-								InputLiteral->UpdateDocumentInputLiteral();
-							}
-						}
-					});
-					RangeHandle->SetOnPropertyValueChanged(UpdateDocumentInput);
-					RangeHandle->SetOnChildPropertyValueChanged(UpdateDocumentInput);
-
-					// If the range is changed, we want to update these details in case we now need to clamp the value
-					if (!InputWidgetOnRangeChangedDelegateHandle.IsValid())
-					{
-						InputWidgetOnRangeChangedDelegateHandle = InputFloat->OnRangeChanged.AddLambda([this](FVector2D Range)
-						{
-							if (FloatLiteral.IsValid())
-							{
-								FloatLiteral->OnRangeChanged.Remove(InputWidgetOnRangeChangedDelegateHandle);
-								if (FMetasoundAssetBase* MetasoundAsset = Metasound::IMetasoundUObjectRegistry::Get().GetObjectAsAssetBase(FloatLiteral->GetOutermostObject()))
-								{
-									MetasoundAsset->SetSynchronizationRequired();
-								}
-							}
-						});
-					}
-				}
-			}
+			InputCategoryBuilder->AddExternalObjectProperty(TArray<UObject*>({ DefaultFloat }), GET_MEMBER_NAME_CHECKED(UMetasoundEditorGraphMemberDefaultFloat, Range));
 		}
 
-		void FMetasoundInputBoolDetailCustomization::CacheProxyData(TSharedPtr<IPropertyHandle> ProxyHandle)
+		void FMetasoundMemberDefaultBoolDetailCustomization::CacheProxyData(TSharedPtr<IPropertyHandle> InPropertyHandle)
 		{
 			DataTypeName = FName();
 
-			const FString* MetadataDataTypeName = ProxyHandle->GetInstanceMetaData(MemberCustomizationPrivate::DataTypeNameIdentifier);
+			const FString* MetadataDataTypeName = InPropertyHandle->GetInstanceMetaData(MemberCustomizationStyle::DataTypeNameIdentifier);
 			if (ensure(MetadataDataTypeName))
 			{
 				DataTypeName = **MetadataDataTypeName;
 			}
 		}
 
-		FText FMetasoundInputBoolDetailCustomization::GetPropertyNameOverride() const
+		FText FMetasoundMemberDefaultBoolDetailCustomization::GetPropertyNameOverride() const
 		{
 			if (DataTypeName == Metasound::GetMetasoundDataTypeName<Metasound::FTrigger>())
 			{
@@ -176,28 +138,42 @@ namespace Metasound
 			return FText::GetEmpty();
 		}
 
-		TSharedRef<SWidget> FMetasoundInputBoolDetailCustomization::CreateStructureWidget(TSharedPtr<IPropertyHandle>& StructPropertyHandle) const
+		TSharedRef<SWidget> FMetasoundMemberDefaultBoolDetailCustomization::CreateStructureWidget(TSharedPtr<IPropertyHandle>& StructPropertyHandle) const
 		{
 			using namespace Frontend;
 
-			if (FMetasoundFrontendRegistryContainer* Registry = FMetasoundFrontendRegistryContainer::Get())
+			TSharedPtr<IPropertyHandle> ValueProperty = StructPropertyHandle->GetChildHandle(GET_MEMBER_NAME_CHECKED(FMetasoundEditorGraphMemberDefaultBoolRef, Value));
+			if (ValueProperty.IsValid())
 			{
-				TSharedPtr<IPropertyHandle> ValueProperty = StructPropertyHandle->GetChildHandle(GET_MEMBER_NAME_CHECKED(FMetasoundEditorGraphInputBoolRef, Value));
-				if (ValueProperty.IsValid())
+				// Not a trigger, so just display as underlying literal type (bool)
+				if (DataTypeName != Metasound::GetMetasoundDataTypeName<Metasound::FTrigger>())
 				{
-					// Not a trigger, so just display as underlying literal type (bool)
-					if (DataTypeName != Metasound::GetMetasoundDataTypeName<Metasound::FTrigger>())
-					{
-						return ValueProperty->CreatePropertyValueWidget();
-					}
+					return ValueProperty->CreatePropertyValueWidget();
+				}
+
+				TAttribute<bool> EnablementAttribute = false;
+				TAttribute<EVisibility> VisibilityAttribute = EVisibility::Visible;
 
-					TArray<UObject*> OuterObjects;
-					ValueProperty->GetOuterObjects(OuterObjects);
-					for (UObject* Object : OuterObjects)
+				TArray<UObject*> OuterObjects;
+				ValueProperty->GetOuterObjects(OuterObjects);
+				if (!OuterObjects.IsEmpty())
+				{
+					if (UMetasoundEditorGraphMemberDefaultLiteral* Literal = Cast<UMetasoundEditorGraphMemberDefaultLiteral>(OuterObjects.Last()))
 					{
-						if (UMetasoundEditorGraphInputLiteral* Literal = Cast<UMetasoundEditorGraphInputLiteral>(Object))
+						if (UMetasoundEditorGraphInput* Input = Cast<UMetasoundEditorGraphInput>(Literal->GetParentMember()))
 						{
-							return SMetasoundGraphNode::CreateTriggerSimulationWidget(*Literal);
+							// Don't display trigger simulation widget if its a trigger
+							// provided by an interface that does not support transmission.
+							const FInterfaceRegistryKey Key = GetInterfaceRegistryKey(Input->GetInterfaceVersion());
+							const IInterfaceRegistryEntry* Entry = IInterfaceRegistry::Get().FindInterfaceRegistryEntry(Key);
+							if (!Entry || Entry->GetRouterName() == Audio::IParameterTransmitter::RouterName)
+							{
+								EnablementAttribute = true;
+								return SMetaSoundGraphNode::CreateTriggerSimulationWidget(*Literal, MoveTemp(VisibilityAttribute), MoveTemp(EnablementAttribute));
+							}
+
+							const FText DisabledToolTip = LOCTEXT("NonTransmittibleInputTriggerSimulationDisabledTooltip", "Trigger simulation disabled: Parent interface does not support being updated by game thread parameters.");
+							return SMetaSoundGraphNode::CreateTriggerSimulationWidget(*Literal, MoveTemp(VisibilityAttribute), MoveTemp(EnablementAttribute), &DisabledToolTip);
 						}
 					}
 				}
@@ -206,24 +182,24 @@ namespace Metasound
 			return SNullWidget::NullWidget;
 		}
 
-		void FMetasoundInputIntDetailCustomization::CacheProxyData(TSharedPtr<IPropertyHandle> ProxyHandle)
+		void FMetasoundMemberDefaultIntDetailCustomization::CacheProxyData(TSharedPtr<IPropertyHandle> ProxyHandle)
 		{
 			DataTypeName = FName();
 
-			const FString* MetadataDataTypeName = ProxyHandle->GetInstanceMetaData(MemberCustomizationPrivate::DataTypeNameIdentifier);
+			const FString* MetadataDataTypeName = ProxyHandle->GetInstanceMetaData(MemberCustomizationStyle::DataTypeNameIdentifier);
 			if (ensure(MetadataDataTypeName))
 			{
 				DataTypeName = **MetadataDataTypeName;
 			}
 		}
 
-		TSharedRef<SWidget> FMetasoundInputIntDetailCustomization::CreateStructureWidget(TSharedPtr<IPropertyHandle>& StructPropertyHandle) const
+		TSharedRef<SWidget> FMetasoundMemberDefaultIntDetailCustomization::CreateStructureWidget(TSharedPtr<IPropertyHandle>& StructPropertyHandle) const
 		{
 			using namespace Frontend;
 
 			if (FMetasoundFrontendRegistryContainer* Registry = FMetasoundFrontendRegistryContainer::Get())
 			{
-				TSharedPtr<IPropertyHandle> ValueProperty = StructPropertyHandle->GetChildHandle(GET_MEMBER_NAME_CHECKED(FMetasoundEditorGraphInputIntRef, Value));
+				TSharedPtr<IPropertyHandle> ValueProperty = StructPropertyHandle->GetChildHandle(GET_MEMBER_NAME_CHECKED(FMetasoundEditorGraphMemberDefaultIntRef, Value));
 				if (ValueProperty.IsValid())
 				{
 					TSharedPtr<const IEnumDataTypeInterface> EnumInterface = IDataTypeRegistry::Get().GetEnumInterfaceForDataType(DataTypeName);
@@ -299,11 +275,11 @@ namespace Metasound
 			return SNullWidget::NullWidget;
 		}
 
-		void FMetasoundInputObjectDetailCustomization::CacheProxyData(TSharedPtr<IPropertyHandle> ProxyHandle)
+		void FMetasoundMemberDefaultObjectDetailCustomization::CacheProxyData(TSharedPtr<IPropertyHandle> ProxyHandle)
 		{
 			ProxyGenClass.Reset();
 
-			const FString* MetadataProxyGenClass = ProxyHandle->GetInstanceMetaData(MemberCustomizationPrivate::ProxyGeneratorClassNameIdentifier);
+			const FString* MetadataProxyGenClass = ProxyHandle->GetInstanceMetaData(MemberCustomizationStyle::ProxyGeneratorClassNameIdentifier);
 			TSharedPtr<IPropertyHandle> MetadataHandle = ProxyHandle->GetParentHandle();
 			if (!ensure(MetadataProxyGenClass))
 			{
@@ -336,9 +312,9 @@ namespace Metasound
 			ensureMsgf(false, TEXT("Failed to find ProxyGeneratorClass. Class not set "));
 		}
 
-		TSharedRef<SWidget> FMetasoundInputObjectDetailCustomization::CreateStructureWidget(TSharedPtr<IPropertyHandle>& StructPropertyHandle) const
+		TSharedRef<SWidget> FMetasoundMemberDefaultObjectDetailCustomization::CreateStructureWidget(TSharedPtr<IPropertyHandle>& StructPropertyHandle) const
 		{
-			TSharedPtr<IPropertyHandle> PropertyHandle = StructPropertyHandle->GetChildHandle(GET_MEMBER_NAME_CHECKED(FMetasoundEditorGraphInputObjectRef, Object));
+			TSharedPtr<IPropertyHandle> PropertyHandle = StructPropertyHandle->GetChildHandle(GET_MEMBER_NAME_CHECKED(FMetasoundEditorGraphMemberDefaultObjectRef, Object));
 
 			const IMetasoundEditorModule& EditorModule = FModuleManager::GetModuleChecked<IMetasoundEditorModule>("MetaSoundEditor");
 			auto FilterAsset = [InEditorModule = &EditorModule, InProxyGenClass = ProxyGenClass](const FAssetData& InAsset)
@@ -398,7 +374,7 @@ namespace Metasound
 			}
 
 			return SNew(STextBlock)
-				.Text(MemberCustomizationPrivate::DefaultPropertyText)
+				.Text(MemberCustomizationStyle::DefaultPropertyText)
 				.Font(IDetailLayoutBuilder::GetDetailFont());
 		}
 
@@ -476,7 +452,7 @@ namespace Metasound
 			CacheProxyData(ProxyProperty);
 
 			TSharedRef<SWidget> ValueWidget = CreateValueWidget(ParentArrayProperty, StructPropertyHandle, bIsInArray);
-			FDetailWidgetRow& ValueRow = ChildBuilder.AddCustomRow(MemberCustomizationPrivate::DefaultPropertyText);
+			FDetailWidgetRow& ValueRow = ChildBuilder.AddCustomRow(MemberCustomizationStyle::DefaultPropertyText);
 			if (bIsInArray)
 			{
 				ValueRow.NameContent()
@@ -503,20 +479,18 @@ namespace Metasound
 				}
 			}
 
-			FSimpleDelegate UpdateDocumentInput = FSimpleDelegate::CreateLambda([InInputs = Inputs]()
+			FSimpleDelegate UpdateFrontendDefaultLiteral = FSimpleDelegate::CreateLambda([InInputs = Inputs]()
 			{
 				for (const TWeakObjectPtr<UMetasoundEditorGraphInput>& GraphInput : InInputs)
 				{
 					if (GraphInput.IsValid())
 					{
-						if (UMetasoundEditorGraphInputLiteral* InputLiteral = GraphInput->Literal)
-						{
-							InputLiteral->UpdateDocumentInputLiteral();
-						}
+						constexpr bool bPostTransaction = true;
+						GraphInput->UpdateFrontendDefaultLiteral(bPostTransaction);
 					}
 				}
 			});
-			StructPropertyHandle->SetOnChildPropertyValueChanged(UpdateDocumentInput);
+			StructPropertyHandle->SetOnChildPropertyValueChanged(UpdateFrontendDefaultLiteral);
 
 			ValueRow.ValueContent()
 			[
@@ -532,7 +506,7 @@ namespace Metasound
 		{
 			if (GraphMember.IsValid())
 			{
-				return GraphMember->TypeName;
+				return GraphMember->GetDataType();
 			}
 
 			return FName();
@@ -540,8 +514,8 @@ namespace Metasound
 
 		void FMetasoundDataTypeSelector::OnDataTypeSelected(FName InSelectedTypeName)
 		{
-			FName ArrayDataTypeName = FName(*InSelectedTypeName.ToString() + MemberCustomizationPrivate::ArrayIdentifier);
 			FName NewDataTypeName;
+			FName ArrayDataTypeName = CreateArrayTypeNameFromElementTypeName(InSelectedTypeName);
 
 			// Update data type based on "Is Array" checkbox and support for arrays.
 			// If an array type is not supported, default to the base data type.
@@ -554,7 +528,7 @@ namespace Metasound
 				}
 				else
 				{
-					check(EditorModule.IsRegisteredDataType(InSelectedTypeName));
+					ensure(EditorModule.IsRegisteredDataType(InSelectedTypeName));
 					NewDataTypeName = InSelectedTypeName;
 				}
 			}
@@ -566,12 +540,12 @@ namespace Metasound
 				}
 				else
 				{
-					check(EditorModule.IsRegisteredDataType(ArrayDataTypeName));
+					ensure(EditorModule.IsRegisteredDataType(ArrayDataTypeName));
 					NewDataTypeName = ArrayDataTypeName;
 				}
 			}
 
-			if (NewDataTypeName == GraphMember->TypeName)
+			if (NewDataTypeName == GraphMember->GetDataType())
 			{
 				return;
 			}
@@ -583,55 +557,61 @@ namespace Metasound
 			if (GraphMember.IsValid())
 			{
 				GraphMember->SetDataType(NewDataTypeName);
-			}
 
-			// Required to rebuild the literal details customization.
-			// This is seemingly dangerous (as the Builder's raw ptr is cached),
-			// but the builder cannot be accessed any other way and instances of
-			// this type are always built from and managed by the parent DetailLayoutBuilder.
-			check(DetailLayoutBuilder);
-			DetailLayoutBuilder->ForceRefreshDetails();
+				if (FMetasoundAssetBase* MetasoundAsset = Metasound::IMetasoundUObjectRegistry::Get().GetObjectAsAssetBase(GraphMember->GetOutermostObject()))
+				{
+					MetasoundAsset->SetUpdateDetailsOnSynchronization();
+				}
+			}
 		}
 
 		void FMetasoundDataTypeSelector::AddDataTypeSelector(IDetailLayoutBuilder& InDetailLayout, const FText& InRowName, TWeakObjectPtr<UMetasoundEditorGraphMember> InGraphMember, bool bIsEnabled)
 		{
+			using namespace Frontend;
+
 			if (!InGraphMember.IsValid())
 			{
 				return;
 			}
 
-			DetailLayoutBuilder = &InDetailLayout;
 			GraphMember = InGraphMember;
 
-			FString CurrentTypeName = GraphMember->TypeName.ToString();
+			FDataTypeRegistryInfo DataTypeInfo;
+			if (!ensure(IDataTypeRegistry::Get().GetDataTypeInfo(InGraphMember->GetDataType(), DataTypeInfo)))
+			{
+				return;
+			}
 
-			bool bCurrentTypeIsArray = CurrentTypeName.EndsWith(MemberCustomizationPrivate::ArrayIdentifier);
-			if (bCurrentTypeIsArray)
+			const bool bIsArrayType = DataTypeInfo.IsArrayType();
+			if (bIsArrayType)
+			{
+				ArrayTypeName = GraphMember->GetDataType();
+				BaseTypeName = CreateElementTypeNameFromArrayTypeName(InGraphMember->GetDataType());
+			}
+			else
 			{
-				CurrentTypeName.LeftChopInline(MemberCustomizationPrivate::ArrayIdentifier.Len());
+				ArrayTypeName = CreateArrayTypeNameFromElementTypeName(InGraphMember->GetDataType());
+				BaseTypeName = GraphMember->GetDataType();
 			}
 
 			IMetasoundEditorModule& EditorModule = FModuleManager::GetModuleChecked<IMetasoundEditorModule>("MetaSoundEditor");
 
 			// Not all types have an equivalent array type. Base types without array
-			// types should have the "Is Array" checkbox disabled. 
-			const FName ArrayType = *(CurrentTypeName + MemberCustomizationPrivate::ArrayIdentifier);
-			const bool bIsArrayTypeRegistered = EditorModule.IsRegisteredDataType(ArrayType);
-			const bool bIsArrayTypeRegisteredHidden = MemberCustomizationPrivate::HiddenInputTypeNames.Contains(ArrayType);
+			// types should have the "Is Array" checkbox disabled.
+			const bool bIsArrayTypeRegistered = bIsArrayType || IDataTypeRegistry::Get().IsRegistered(ArrayTypeName);
+			const bool bIsArrayTypeRegisteredHidden = MemberCustomizationPrivate::HiddenInputTypeNames.Contains(ArrayTypeName);
 
 			TArray<FName> BaseDataTypes;
-			EditorModule.IterateDataTypes([&](const FEditorDataType& EditorDataType)
+			IDataTypeRegistry::Get().IterateDataTypeInfo([&BaseDataTypes](const FDataTypeRegistryInfo& RegistryInfo)
 			{
-				const FName& TypeName = EditorDataType.RegistryInfo.DataTypeName;
-
 				// Hide the type from the combo selector if any of the following is true;
-				const bool bIsArrayType = EditorDataType.RegistryInfo.IsArrayType();
-				const bool bIsVariable = EditorDataType.RegistryInfo.bIsVariable;
-				const bool bIsHiddenType = MemberCustomizationPrivate::HiddenInputTypeNames.Contains(TypeName);
+				const bool bIsArrayType = RegistryInfo.IsArrayType();
+				const bool bIsVariable = RegistryInfo.bIsVariable;
+				const bool bIsHiddenType = MemberCustomizationPrivate::HiddenInputTypeNames.Contains(RegistryInfo.DataTypeName);
 				const bool bHideBaseType = bIsArrayType || bIsVariable || bIsHiddenType;
 				if (!bHideBaseType)
 				{
-					BaseDataTypes.Add(TypeName);
+					BaseDataTypes.Add(RegistryInfo.DataTypeName);
 				}
 			});
 
@@ -678,7 +658,7 @@ namespace Metasound
 						SNew(STextBlock)
 						.Text_Lambda([this]()
 						{
-							return FText::FromName(GetDataType());
+							return FText::FromName(BaseTypeName);
 						})
 					]
 				]
@@ -706,7 +686,7 @@ namespace Metasound
 				]
 			];
 
-			auto NameMatchesPredicate = [CurrentTypeName](const TSharedPtr<FString>& Item) { return *Item == CurrentTypeName; };
+			auto NameMatchesPredicate = [TypeString = BaseTypeName.ToString()](const TSharedPtr<FString>& Item) { return *Item == TypeString; };
 			const TSharedPtr<FString>* SelectedItem = ComboOptions.FindByPredicate(NameMatchesPredicate);
 			if (ensure(SelectedItem))
 			{
@@ -715,161 +695,19 @@ namespace Metasound
 		}
 
 		ECheckBoxState FMetasoundDataTypeSelector::OnGetDataTypeArrayCheckState(TWeakObjectPtr<UMetasoundEditorGraphMember> InGraphMember) const
-		{
-			if (InGraphMember.IsValid())
-			{
-				FString CurrentTypeName = InGraphMember->TypeName.ToString();
-				bool bCurrentTypeIsArray = CurrentTypeName.EndsWith(MemberCustomizationPrivate::ArrayIdentifier);
-				return bCurrentTypeIsArray ? ECheckBoxState::Checked : ECheckBoxState::Unchecked;
-			}
-
-			return ECheckBoxState::Undetermined;
-		}
-
-		void FMetasoundInputDetailCustomization::CustomizeDetails(IDetailLayoutBuilder& DetailLayout)
 		{
 			using namespace Frontend;
 
-			TMetasoundGraphMemberDetailCustomization<UMetasoundEditorGraphInput>::CustomizeDetails(DetailLayout);
-
-			if (!GraphMember.IsValid())
-			{
-				return;
-			}
-
-			IDetailCategoryBuilder& CategoryBuilder = DetailLayout.EditCategory("General");
-
-			const bool bIsInterfaceMember = CastChecked<UMetasoundEditorGraphVertex>(GraphMember)->IsInterfaceMember();
-			const bool bIsGraphEditable = IsGraphEditable();
-			NameEditableTextBox = SNew(SEditableTextBox)
-				.Text(this, &FMetasoundInputDetailCustomization::GetName)
-				.OnTextChanged(this, &FMetasoundInputDetailCustomization::OnNameChanged)
-				.OnTextCommitted(this, &FMetasoundInputDetailCustomization::OnNameCommitted)
-				.IsReadOnly(bIsInterfaceMember || !bIsGraphEditable)
-				.Font(IDetailLayoutBuilder::GetDetailFont());
-
-			DisplayNameEditableTextBox = SNew(SEditableTextBox)
-				.Text(this, &FMetasoundInputDetailCustomization::GetDisplayName)
-				.OnTextCommitted(this, &FMetasoundInputDetailCustomization::OnDisplayNameCommitted)
-				.IsReadOnly(bIsInterfaceMember || !bIsGraphEditable)
-				.Font(IDetailLayoutBuilder::GetDetailFont());
-
-			CategoryBuilder.AddCustomRow(MemberCustomizationPrivate::InputNameText)
-			.EditCondition(!bIsInterfaceMember && bIsGraphEditable, nullptr)
-			.NameContent()
-			[
-				SNew(STextBlock)
-				.Font(IDetailLayoutBuilder::GetDetailFontBold())
-				.Text(MemberCustomizationPrivate::InputNameText)
-				.ToolTipText(LOCTEXT("InputName_Description", "Name used by external systems to identify input. Used as DisplayName within MetaSound Graph Editor if no DisplayName is provided."))
-			]
-			.ValueContent()
-			[
-				NameEditableTextBox.ToSharedRef()
-			];
-
-			// TODO: Enable and use proper FText property editor
-// 			CategoryBuilder.AddCustomRow(MemberCustomizationPrivate::InputDisplayNameText)
-// 			.EditCondition(!bIsInterfaceMember && bIsGraphEditable, nullptr)
-// 			.NameContent()
-// 			[
-// 				SNew(STextBlock)
-// 				.Font(IDetailLayoutBuilder::GetDetailFontBold())
-// 				.Text(MemberCustomizationPrivate::InputDisplayNameText)
-// 				.ToolTipText(LOCTEXT("InputDisplayName_Description", "Optional, localized name used within the MetaSounds editor(s) to describe the given input."))
-// 			]
-// 			.ValueContent()
-// 			[
-// 				DisplayNameEditableTextBox.ToSharedRef()
-// 			];
-
-			CategoryBuilder.AddCustomRow(MemberCustomizationPrivate::NodeTooltipText)
-			.EditCondition(!bIsInterfaceMember && bIsGraphEditable, nullptr)
-			.NameContent()
-			[
-				SNew(STextBlock)
-				.Font(IDetailLayoutBuilder::GetDetailFontBold())
-				.Text(MemberCustomizationPrivate::NodeTooltipText)
-			]
-			.ValueContent()
-			[
-				SNew(SMultiLineEditableTextBox)
-				.Text(this, &FMetasoundInputDetailCustomization::GetTooltip)
-				.OnTextCommitted(this, &FMetasoundInputDetailCustomization::OnTooltipCommitted)
-				.IsReadOnly(bIsInterfaceMember || !bIsGraphEditable)
-				.ModiferKeyForNewLine(EModifierKey::Shift)
-				.RevertTextOnEscape(true)
-				.WrapTextAt(MemberCustomizationPrivate::DetailsTitleMaxWidth - MemberCustomizationPrivate::DetailsTitleWrapPadding)
-				.Font(IDetailLayoutBuilder::GetDetailFont())
-			];
-
-			DataTypeSelector->AddDataTypeSelector(DetailLayout, MemberCustomizationPrivate::DataTypeNameText, GraphMember, !bIsInterfaceMember && bIsGraphEditable);
-
-			IDetailCategoryBuilder& DefaultCategoryBuilder = DetailLayout.EditCategory("DefaultValue");
-			TSharedPtr<IPropertyHandle> LiteralHandle = DetailLayout.GetProperty(GET_MEMBER_NAME_CHECKED(UMetasoundEditorGraphInput, Literal));
-			if (ensure(GraphMember.IsValid()) && ensure(LiteralHandle.IsValid()))
+			if (InGraphMember.IsValid())
 			{
-				UObject* LiteralObject = nullptr;
-				if (LiteralHandle->GetValue(LiteralObject) == FPropertyAccess::Success)
+				FDataTypeRegistryInfo DataTypeInfo;
+				if (ensure(IDataTypeRegistry::Get().GetDataTypeInfo(InGraphMember->GetDataType(), DataTypeInfo)))
 				{
-					if (ensure(LiteralObject))
-					{
-						LiteralHandle->MarkHiddenByCustomization();
-
-						TSharedPtr<IPropertyHandle> DefaultValueHandle;
-
-						if (IDetailPropertyRow* Row = DefaultCategoryBuilder.AddExternalObjectProperty(TArray<UObject*>({ LiteralObject }), "Default"))
-						{
-							DefaultValueHandle = Row->GetPropertyHandle();
-							if (DefaultValueHandle.IsValid())
-							{
-								SetDefaultPropertyMetaData(DefaultValueHandle.ToSharedRef());
-
-								FSimpleDelegate UpdateDocumentInput = FSimpleDelegate::CreateLambda([GraphMember = this->GraphMember]()
-								{
-									if (GraphMember.IsValid())
-									{
-										if (UMetasoundEditorGraphInputLiteral* InputLiteral = GraphMember->Literal)
-										{
-											InputLiteral->UpdateDocumentInputLiteral();
-										}
-									}
-								});
-
-								DefaultValueHandle->SetOnPropertyValueChanged(UpdateDocumentInput);
-								DefaultValueHandle->SetOnChildPropertyValueChanged(UpdateDocumentInput);
-
-								TSharedPtr<IPropertyHandleArray> DefaultValueArray = DefaultValueHandle->AsArray();
-								if (DefaultValueArray.IsValid())
-								{
-									DefaultValueArray->SetOnNumElementsChanged(UpdateDocumentInput);
-								}
-							}
-						}
-
-						IMetasoundEditorModule& EditorModule = FModuleManager::GetModuleChecked<IMetasoundEditorModule>("MetaSoundEditor");
-						LiteralCustomization = EditorModule.CreateInputLiteralCustomization(*LiteralObject->GetClass(), DefaultCategoryBuilder);
-						if (LiteralCustomization.IsValid())
-						{
-							LiteralCustomization->CustomizeLiteral(*CastChecked<UMetasoundEditorGraphInputLiteral>(LiteralObject), DefaultValueHandle);
-						}
-
-						UMetasoundEditorGraphInputFloat* InputFloat = Cast<UMetasoundEditorGraphInputFloat>(LiteralObject);
-						if (InputFloat)
-						{
-							// add input widget properties 
-							IDetailCategoryBuilder& WidgetCategoryBuilder = DetailLayout.EditCategory("Widget");
-							WidgetCategoryBuilder.AddExternalObjectProperty(TArray<UObject*>({ InputFloat }), GET_MEMBER_NAME_CHECKED(UMetasoundEditorGraphInputFloat, InputWidgetType));
-							WidgetCategoryBuilder.AddExternalObjectProperty(TArray<UObject*>({ InputFloat }), GET_MEMBER_NAME_CHECKED(UMetasoundEditorGraphInputFloat, InputWidgetOrientation));
-							WidgetCategoryBuilder.AddExternalObjectProperty(TArray<UObject*>({ InputFloat }), GET_MEMBER_NAME_CHECKED(UMetasoundEditorGraphInputFloat, InputWidgetValueType));
-						}
-					}
-					else
-					{
-						DefaultCategoryBuilder.AddProperty(LiteralHandle);
-					}
+					return DataTypeInfo.IsArrayType() ? ECheckBoxState::Checked : ECheckBoxState::Unchecked;
 				}
 			}
+
+			return ECheckBoxState::Undetermined;
 		}
 
 		void FMetasoundDataTypeSelector::OnDataTypeArrayChanged(TWeakObjectPtr<UMetasoundEditorGraphMember> InGraphMember, ECheckBoxState InNewState)
@@ -879,212 +717,55 @@ namespace Metasound
 				TSharedPtr<FString> DataTypeRoot = DataTypeComboBox->GetSelectedItem();
 				if (ensure(DataTypeRoot.IsValid()))
 				{
-					FString DataTypeString = *DataTypeRoot.Get();
-					if (InNewState == ECheckBoxState::Checked)
-					{
-						DataTypeString += MemberCustomizationPrivate::ArrayIdentifier;
-					}
-
 					// Have to stop playback to avoid attempting to change live edit data on invalid input type.
 					check(GEditor);
 					GEditor->ResetPreviewAudioComponent();
 
-					InGraphMember->SetDataType(FName(DataTypeString));
+					const FName DataType = InNewState == ECheckBoxState::Checked ? ArrayTypeName : BaseTypeName;
+					InGraphMember->SetDataType(DataType);
 
-					// Required to rebuild the literal details customization.
-					// This is seemingly dangerous (as the Builder's raw ptr is cached),
-					// but the builder cannot be accessed any other way and instances of
-					// this type are always built from and managed by the parent DetailLayoutBuilder.
-					check(DetailLayoutBuilder);
-					DetailLayoutBuilder->ForceRefreshDetails();
+					if (FMetasoundAssetBase* MetasoundAsset = Metasound::IMetasoundUObjectRegistry::Get().GetObjectAsAssetBase(GraphMember->GetOutermostObject()))
+					{
+						MetasoundAsset->SetUpdateDetailsOnSynchronization();
+					}
 				}
 			}
 		}
 
-		void FMetasoundInputDetailCustomization::SetDefaultPropertyMetaData(TSharedRef<IPropertyHandle> InDefaultPropertyHandle) const
-		{
-			using namespace Frontend;
-
-			if (!GraphMember.IsValid())
-			{
-				return;
-			}
-
-			FMetasoundFrontendRegistryContainer* Registry = FMetasoundFrontendRegistryContainer::Get();
-			if (!ensure(Registry))
-			{
-				return;
-			}
-
-			const FName TypeName = GetLiteralDataType();
-			if (TypeName.IsNone())
-			{
-				return;
-			}
-
-			FString TypeNameString = TypeName.ToString();
-			if (TypeNameString.EndsWith(MemberCustomizationPrivate::ArrayIdentifier))
-			{
-				TypeNameString = TypeNameString.LeftChop(MemberCustomizationPrivate::ArrayIdentifier.Len());
-			}
-			InDefaultPropertyHandle->SetInstanceMetaData(MemberCustomizationPrivate::DataTypeNameIdentifier, TypeNameString);
-
-			FDataTypeRegistryInfo DataTypeInfo;
-			if (!ensure(IDataTypeRegistry::Get().GetDataTypeInfo(TypeName, DataTypeInfo)))
-			{
-				return;
-			}
-
-			const EMetasoundFrontendLiteralType LiteralType = GetMetasoundFrontendLiteralType(DataTypeInfo.PreferredLiteralType);
-			if (LiteralType != EMetasoundFrontendLiteralType::UObject && LiteralType != EMetasoundFrontendLiteralType::UObjectArray)
-			{
-				return;
-			}
-
-			UClass* ProxyGenClass = DataTypeInfo.ProxyGeneratorClass;
-			if (ProxyGenClass)
-			{
-				const FString ClassName = ProxyGenClass->GetName();
-				InDefaultPropertyHandle->SetInstanceMetaData(MemberCustomizationPrivate::ProxyGeneratorClassNameIdentifier, ClassName);
-			}
-		}
-
-		FName FMetasoundInputDetailCustomization::GetLiteralDataType() const
-		{
-			return GraphMember->TypeName;
-		}
+		const FText FMetasoundInputDetailCustomization::MemberNameText = LOCTEXT("InputGraphMemberLabel", "Input");
 
-		void FMetasoundOutputDetailCustomization::CustomizeDetails(IDetailLayoutBuilder& DetailLayout)
+		bool FMetasoundInputDetailCustomization::IsInterfaceMember() const
 		{
-			using namespace Frontend;
-
-			TMetasoundGraphMemberDetailCustomization<UMetasoundEditorGraphOutput>::CustomizeDetails(DetailLayout);
-			if (!GraphMember.IsValid())
+			if (GraphMember.IsValid())
 			{
-				return;
+				return CastChecked<UMetasoundEditorGraphVertex>(GraphMember)->IsInterfaceMember();
 			}
 
-			IDetailCategoryBuilder& CategoryBuilder = DetailLayout.EditCategory("General");
-
-			const bool bIsInterfaceMember = CastChecked<UMetasoundEditorGraphVertex>(GraphMember)->IsInterfaceMember();
-			const bool bIsGraphEditable = IsGraphEditable();
-
-			NameEditableTextBox = SNew(SEditableTextBox)
-				.Text(this, &FMetasoundOutputDetailCustomization::GetName)
-				.OnTextChanged(this, &FMetasoundOutputDetailCustomization::OnNameChanged)
-				.OnTextCommitted(this, &FMetasoundOutputDetailCustomization::OnNameCommitted)
-				.IsReadOnly(bIsInterfaceMember || !bIsGraphEditable)
-				.Font(IDetailLayoutBuilder::GetDetailFont());
-
-			DisplayNameEditableTextBox = SNew(SEditableTextBox)
-				.Text(this, &FMetasoundOutputDetailCustomization::GetDisplayName)
-				.OnTextCommitted(this, &FMetasoundOutputDetailCustomization::OnDisplayNameCommitted)
-				.IsReadOnly(bIsInterfaceMember || !bIsGraphEditable)
-				.Font(IDetailLayoutBuilder::GetDetailFont());
-
-			CategoryBuilder.AddCustomRow(MemberCustomizationPrivate::OutputNameText)
-			.EditCondition(!bIsInterfaceMember && bIsGraphEditable, nullptr)
-			.NameContent()
-			[
-				SNew(STextBlock)
-				.Font(IDetailLayoutBuilder::GetDetailFontBold())
-				.Text(MemberCustomizationPrivate::OutputNameText)
-				.ToolTipText(LOCTEXT("OutputName_Description", "Name used by external systems to identify output. Used as DisplayName within MetaSound Graph Editor if no DisplayName is provided."))
-			]
-			.ValueContent()
-			[
-				NameEditableTextBox.ToSharedRef()
-			];
-
-			// TODO: Enable and use proper FText property editor
-// 			CategoryBuilder.AddCustomRow(MemberCustomizationPrivate::OutputDisplayNameText)
-// 			.EditCondition(!bIsInterfaceMember && bIsGraphEditable, nullptr)
-// 			.NameContent()
-// 			[
-// 				SNew(STextBlock)
-// 				.Font(IDetailLayoutBuilder::GetDetailFontBold())
-// 				.Text(MemberCustomizationPrivate::OutputDisplayNameText)
-// 				.ToolTipText(LOCTEXT("OutputDisplayName_Description", "Optional, localized name used within the MetaSounds editor(s) to describe the given output."))
-// 			]
-// 			.ValueContent()
-// 			[
-// 				DisplayNameEditableTextBox.ToSharedRef()
-// 			];
-
-			CategoryBuilder.AddCustomRow(MemberCustomizationPrivate::NodeTooltipText)
-			.EditCondition(!bIsInterfaceMember && bIsGraphEditable, nullptr)
-			.NameContent()
-			[
-				SNew(STextBlock)
-				.Font(IDetailLayoutBuilder::GetDetailFontBold())
-				.Text(MemberCustomizationPrivate::NodeTooltipText)
-			]
-			.ValueContent()
-			[
-				SNew(SMultiLineEditableTextBox)
-				.Text(this, &FMetasoundOutputDetailCustomization::GetTooltip)
-				.OnTextCommitted(this, &FMetasoundOutputDetailCustomization::OnTooltipCommitted)
-				.IsReadOnly(bIsInterfaceMember || !bIsGraphEditable)
-				.ModiferKeyForNewLine(EModifierKey::Shift)
-				.RevertTextOnEscape(true)
-				.WrapTextAt(MemberCustomizationPrivate::DetailsTitleMaxWidth - MemberCustomizationPrivate::DetailsTitleWrapPadding)
-				.Font(IDetailLayoutBuilder::GetDetailFont())
-			];
-
-			DataTypeSelector->AddDataTypeSelector(DetailLayout, MemberCustomizationPrivate::DataTypeNameText, GraphMember, !bIsInterfaceMember && bIsGraphEditable);
+			return false;
 		}
 
-		void FMetasoundOutputDetailCustomization::SetDefaultPropertyMetaData(TSharedRef<IPropertyHandle> InDefaultPropertyHandle) const
+		void FMetasoundInputDetailCustomization::CustomizeDetails(IDetailLayoutBuilder& DetailLayout)
 		{
 			using namespace Frontend;
 
-			if (!GraphMember.IsValid())
-			{
-				return;
-			}
-
-			FMetasoundFrontendRegistryContainer* Registry = FMetasoundFrontendRegistryContainer::Get();
-			if (!ensure(Registry))
-			{
-				return;
-			}
-
-			const FName TypeName = GetLiteralDataType();
-			if (TypeName.IsNone())
-			{
-				return;
-			}
-
-			FString TypeNameString = TypeName.ToString();
-			if (TypeNameString.EndsWith(MemberCustomizationPrivate::ArrayIdentifier))
-			{
-				TypeNameString = TypeNameString.LeftChop(MemberCustomizationPrivate::ArrayIdentifier.Len());
-			}
-			InDefaultPropertyHandle->SetInstanceMetaData(MemberCustomizationPrivate::DataTypeNameIdentifier, TypeNameString);
+			TMetasoundGraphMemberDetailCustomization<UMetasoundEditorGraphInput, FMetasoundInputDetailCustomization>::CustomizeDetails(DetailLayout);
+		}
 
-			FDataTypeRegistryInfo DataTypeInfo;
-			if (!ensure(IDataTypeRegistry::Get().GetDataTypeInfo(TypeName, DataTypeInfo)))
-			{
-				return;
-			}
+		const FText FMetasoundOutputDetailCustomization::MemberNameText = LOCTEXT("OutputGraphMemberLabel", "Output");
 
-			const EMetasoundFrontendLiteralType LiteralType = GetMetasoundFrontendLiteralType(DataTypeInfo.PreferredLiteralType);
-			if (LiteralType != EMetasoundFrontendLiteralType::UObject && LiteralType != EMetasoundFrontendLiteralType::UObjectArray)
+		bool FMetasoundOutputDetailCustomization::IsInterfaceMember() const
+		{
+			if (GraphMember.IsValid())
 			{
-				return;
+				return CastChecked<UMetasoundEditorGraphVertex>(GraphMember)->IsInterfaceMember();
 			}
 
-			UClass* ProxyGenClass = DataTypeInfo.ProxyGeneratorClass;
-			if (ProxyGenClass)
-			{
-				const FString ClassName = ProxyGenClass->GetName();
-				InDefaultPropertyHandle->SetInstanceMetaData(MemberCustomizationPrivate::ProxyGeneratorClassNameIdentifier, ClassName);
-			}
+			return false;
 		}
 
-		FName FMetasoundOutputDetailCustomization::GetLiteralDataType() const
+		void FMetasoundOutputDetailCustomization::CustomizeDetails(IDetailLayoutBuilder& DetailLayout)
 		{
-			return GraphMember->TypeName;
+			TMetasoundGraphMemberDetailCustomization<UMetasoundEditorGraphOutput, FMetasoundOutputDetailCustomization>::CustomizeDetails(DetailLayout);
 		}
 	} // namespace Editor
 } // namespace Metasound
diff --git a/Engine/Plugins/Runtime/Metasound/Source/MetasoundEditor/Private/MetasoundNodeDetailCustomization.h b/Engine/Plugins/Runtime/Metasound/Source/MetasoundEditor/Private/MetasoundNodeDetailCustomization.h
--- a/Engine/Plugins/Runtime/Metasound/Source/MetasoundEditor/Private/MetasoundNodeDetailCustomization.h
+++ b/Engine/Plugins/Runtime/Metasound/Source/MetasoundEditor/Private/MetasoundNodeDetailCustomization.h
@@ -2,69 +2,64 @@
 #pragma once
 
 #include "Components/Widget.h"
-#include "CoreMinimal.h"
+#include "DetailCategoryBuilder.h"
 #include "DetailLayoutBuilder.h"
+#include "DetailWidgetRow.h"
 #include "EdGraph/EdGraphSchema.h"
 #include "IDetailCustomization.h"
 #include "IPropertyTypeCustomization.h"
 #include "Layout/Visibility.h"
 #include "MetasoundAssetBase.h"
+#include "MetasoundEditorGraph.h"
+#include "MetasoundEditorGraphMemberDefaults.h"
 #include "MetasoundEditorGraphNode.h"
-#include "MetasoundEditorGraphInputNodes.h"
+#include "MetasoundFrontendLiteral.h"
 #include "MetasoundEditorModule.h"
 #include "MetasoundUObjectRegistry.h"
 #include "PropertyHandle.h"
+#include "PropertyRestriction.h"
 #include "ScopedTransaction.h"
+#include "SSearchableComboBox.h"
 #include "Types/SlateEnums.h"
 #include "UObject/NameTypes.h"
+#include "Widgets/Input/SCheckBox.h"
 #include "Widgets/Input/SEditableTextBox.h"
+#include "Widgets/Input/SMultiLineEditableTextBox.h"
+#include "Widgets/Input/STextComboBox.h"
 #include "WorkflowOrientedApp/SModeWidget.h"
 
-// Forward Declarations
-class FDetailWidgetRow;
-class FPropertyRestriction;
-class IDetailLayoutBuilder;
-class IPropertyHandle;
-class SCheckBox;
-class STextComboBox;
-class SSearchableComboBox;
-
 #define LOCTEXT_NAMESPACE "MetaSoundEditor"
 
 namespace Metasound
 {
 	namespace Editor
 	{
-		namespace MemberCustomizationPrivate
+		// TODO: Move to actual style
+		namespace MemberCustomizationStyle
 		{
-			/** Minimum size of the details title panel */
-			static constexpr float DetailsTitleMinWidth = 125.f;
 			/** Maximum size of the details title panel */
 			static constexpr float DetailsTitleMaxWidth = 300.f;
 			/** magic number retrieved from SGraphNodeComment::GetWrapAt() */
 			static constexpr float DetailsTitleWrapPadding = 32.0f;
 
-			static const FString ArrayIdentifier = TEXT(":Array");
-
 			static const FText DataTypeNameText = LOCTEXT("Node_DataTypeName", "Type");
 			static const FText DefaultPropertyText = LOCTEXT("Node_DefaultPropertyName", "Default Value");
 			static const FText NodeTooltipText = LOCTEXT("Node_Tooltip", "Tooltip");
 
 			static const FName DataTypeNameIdentifier = "DataTypeName";
 			static const FName ProxyGeneratorClassNameIdentifier = "GeneratorClass";
+		} // namespace MemberCustomizationStyle
 
-		} // namespace MemberCustomizationPrivate
-
-		class FMetasoundFloatLiteralCustomization : public IMetaSoundInputLiteralCustomization
+		class FMetasoundFloatLiteralCustomization : public IMemberDefaultLiteralCustomization
 		{
 			IDetailCategoryBuilder* InputCategoryBuilder = nullptr;
-			TWeakObjectPtr<UMetasoundEditorGraphInputFloat> FloatLiteral;
+			TWeakObjectPtr<UMetasoundEditorGraphMemberDefaultFloat> FloatLiteral;
 
 			// Delegate for updating the clamp min/max of the input value when the slider range is changed 
-			FDelegateHandle InputWidgetOnRangeChangedDelegateHandle;
+			FDelegateHandle OnRangeChangedDelegateHandle;
 
 			// Delegate for clamping the input value or not
-			FDelegateHandle OnClampInputChangedDelegateHandle;
+			FDelegateHandle OnClampChangedDelegateHandle;
 
 		public:
 			FMetasoundFloatLiteralCustomization(IDetailCategoryBuilder& InInputCategoryBuilder)
@@ -74,15 +69,15 @@ namespace Metasound
 
 			virtual ~FMetasoundFloatLiteralCustomization();
 
-			virtual void CustomizeLiteral(UMetasoundEditorGraphInputLiteral& InLiteral, TSharedPtr<IPropertyHandle> InDefaultValueHandle) override;
+			virtual void CustomizeLiteral(UMetasoundEditorGraphMemberDefaultLiteral& InLiteral, TSharedPtr<IPropertyHandle> InDefaultValueHandle) override;
 		};
 
-		class FMetasoundFloatLiteralCustomizationFactory : public IMetaSoundInputLiteralCustomizationFactory
+		class FMetasoundFloatLiteralCustomizationFactory : public IMemberDefaultLiteralCustomizationFactory
 		{
 		public:
-			virtual TUniquePtr<IMetaSoundInputLiteralCustomization> CreateLiteralCustomization(IDetailCategoryBuilder& DefaultCategoryBuilder) const override
+			virtual TUniquePtr<IMemberDefaultLiteralCustomization> CreateLiteralCustomization(IDetailCategoryBuilder& DefaultCategoryBuilder) const override
 			{
-				return TUniquePtr<IMetaSoundInputLiteralCustomization>(new FMetasoundFloatLiteralCustomization(DefaultCategoryBuilder));
+				return TUniquePtr<IMemberDefaultLiteralCustomization>(new FMetasoundFloatLiteralCustomization(DefaultCategoryBuilder));
 			}
 		};
 
@@ -104,7 +99,7 @@ namespace Metasound
 			TSharedRef<SWidget> CreateValueWidget(TSharedPtr<IPropertyHandleArray> ParentArrayProperty, TSharedPtr<IPropertyHandle> StructPropertyHandle, bool bIsInArray) const;
 		};
 
-		class FMetasoundInputBoolDetailCustomization : public FMetasoundInputArrayDetailCustomizationBase
+		class FMetasoundMemberDefaultBoolDetailCustomization : public FMetasoundInputArrayDetailCustomizationBase
 		{
 		protected:
 			virtual FText GetPropertyNameOverride() const override;
@@ -115,7 +110,7 @@ namespace Metasound
 			FName DataTypeName;
 		};
 
-		class FMetasoundInputIntDetailCustomization : public FMetasoundInputArrayDetailCustomizationBase
+		class FMetasoundMemberDefaultIntDetailCustomization : public FMetasoundInputArrayDetailCustomizationBase
 		{
 		protected:
 			virtual TSharedRef<SWidget> CreateStructureWidget(TSharedPtr<IPropertyHandle>& StructPropertyHandle) const override;
@@ -125,7 +120,7 @@ namespace Metasound
 			FName DataTypeName;
 		};
 
-		class FMetasoundInputObjectDetailCustomization : public FMetasoundInputArrayDetailCustomizationBase
+		class FMetasoundMemberDefaultObjectDetailCustomization : public FMetasoundInputArrayDetailCustomizationBase
 		{
 		protected:
 			virtual TSharedRef<SWidget> CreateStructureWidget(TSharedPtr<IPropertyHandle>& StructPropertyHandle) const override;
@@ -154,74 +149,259 @@ namespace Metasound
 			TSharedPtr<SSearchableComboBox> DataTypeComboBox;
 			TArray<TSharedPtr<FString>> ComboOptions;
 
-			IDetailLayoutBuilder* DetailLayoutBuilder = nullptr;
+			FName BaseTypeName;
+			FName ArrayTypeName;
 		};
 
-		template <typename GraphMemberType>
+		template <typename TMemberType, typename TChildClass>
 		class TMetasoundGraphMemberDetailCustomization : public IDetailCustomization
 		{
 		public:
-			TMetasoundGraphMemberDetailCustomization(const FText& InGraphMemberLabel)
+			TMetasoundGraphMemberDetailCustomization()
 				: IDetailCustomization()
-				, GraphMemberLabel(InGraphMemberLabel)
 			{
 				DataTypeSelector = MakeShared<FMetasoundDataTypeSelector>();
 			}
 
-		protected:
-			FText GraphMemberLabel;
+			virtual ~TMetasoundGraphMemberDetailCustomization() = default;
 
-			TWeakObjectPtr<GraphMemberType> GraphMember;
+		protected:
+			TWeakObjectPtr<TMemberType> GraphMember;
 			TSharedPtr<SEditableTextBox> NameEditableTextBox;
 			TSharedPtr<SEditableTextBox> DisplayNameEditableTextBox;
 			TSharedPtr<FMetasoundDataTypeSelector> DataTypeSelector;
 
 			bool bIsNameInvalid = false;
 
+			static const FText& GetMemberDisplayNameText()
+			{
+				static const FText MemberDisplayNameTextFormat = LOCTEXT("GraphMember_NameFormat", "{0} Display Name");
+				static const FText DisplayName = FText::Format(MemberDisplayNameTextFormat, TChildClass::MemberNameText);
+				return DisplayName;
+			}
+
+			static const FText& GetMemberDisplayNameDescriptionText()
+			{
+				static const FText MemberDisplayNameDescriptionFormat = LOCTEXT("GraphMember_DisplayNameDescriptionFormat", "Optional, localized name used within the MetaSounds editor(s) to describe the given {0}.");
+				static const FText DisplayName = FText::Format(MemberDisplayNameDescriptionFormat, TChildClass::MemberNameText);
+				return DisplayName;
+			}
+
+			static const FText& GetMemberNameText()
+			{
+				static const FText MemberNameTextFormat = LOCTEXT("GraphMember_DisplayNameFormat", "{0} Name");
+				static const FText Name = FText::Format(MemberNameTextFormat, TChildClass::MemberNameText);
+				return Name;
+			}
+
+			static const FText& GetMemberNameDescriptionText()
+			{
+				static const FText MemberNameDescriptionFormat = LOCTEXT("GraphMember_DisplayDescriptionFormat", "Name used by external systems/referencing MetaSounds to identify {0}. Used as DisplayName within MetaSound Graph Editor if no DisplayName is provided.");
+				static const FText Name = FText::Format(MemberNameDescriptionFormat, TChildClass::MemberNameText);
+				return Name;
+			}
+
+			virtual bool IsInterfaceMember() const
+			{
+				return false;
+			}
+
 			// IDetailCustomization interface
-			virtual void CustomizeDetails(IDetailLayoutBuilder& DetailLayout) override
+			virtual void CustomizeDetails(IDetailLayoutBuilder& InDetailLayout) override
 			{
+				using namespace Frontend;
+
 				TArray<TWeakObjectPtr<UObject>> Objects;
-				DetailLayout.GetObjectsBeingCustomized(Objects);
+				InDetailLayout.GetObjectsBeingCustomized(Objects);
 				if (Objects.IsEmpty())
 				{
 					return;
 				}
 
-				GraphMember = Cast<GraphMemberType>(Objects[0].Get());
+				GraphMember = Cast<TMemberType>(Objects.Last().Get());
+				if (!GraphMember.IsValid())
+				{
+					return;
+				}
+
+				IDetailCategoryBuilder& CategoryBuilder = InDetailLayout.EditCategory("General");
+
+				const bool bIsInterfaceMember = IsInterfaceMember();
+				const bool bIsGraphEditable = IsGraphEditable();
+
+				NameEditableTextBox = SNew(SEditableTextBox)
+					.Text(this, &TChildClass::GetName)
+					.OnTextChanged(this, &TChildClass::OnNameChanged)
+					.OnTextCommitted(this, &TChildClass::OnNameCommitted)
+					.IsReadOnly(bIsInterfaceMember || !bIsGraphEditable)
+					.Font(IDetailLayoutBuilder::GetDetailFont());
+
+				DisplayNameEditableTextBox = SNew(SEditableTextBox)
+					.Text(this, &TChildClass::GetDisplayName)
+					.OnTextCommitted(this, &TChildClass::OnDisplayNameCommitted)
+					.IsReadOnly(bIsInterfaceMember || !bIsGraphEditable)
+					.Font(IDetailLayoutBuilder::GetDetailFont());
+
+				CategoryBuilder.AddCustomRow(TChildClass::MemberNameText)
+					.EditCondition(!bIsInterfaceMember && bIsGraphEditable, nullptr)
+					.NameContent()
+					[
+						SNew(STextBlock)
+						.Font(IDetailLayoutBuilder::GetDetailFontBold())
+					.Text(TChildClass::MemberNameText)
+					.ToolTipText(GetMemberNameDescriptionText())
+					]
+					.ValueContent()
+					[
+						NameEditableTextBox.ToSharedRef()
+					];
+
+				// TODO: Enable and use proper FText property editor
+// 				CategoryBuilder.AddCustomRow(GetMemberDisplayNameText())
+// 				.EditCondition(!bIsInterfaceMember && bIsGraphEditable, nullptr)
+// 				.NameContent()
+// 				[
+// 					SNew(STextBlock)
+// 					.Font(IDetailLayoutBuilder::GetDetailFontBold())
+// 					.Text(GetMemberDisplayNameText())
+// 					.ToolTipText(GetMemberDisplayNameDescriptionText())
+// 				]
+// 				.ValueContent()
+// 				[
+// 					DisplayNameEditableTextBox.ToSharedRef()
+// 				];
+
+				CategoryBuilder.AddCustomRow(MemberCustomizationStyle::NodeTooltipText)
+				.EditCondition(!bIsInterfaceMember && bIsGraphEditable, nullptr)
+				.NameContent()
+				[
+					SNew(STextBlock)
+					.Font(IDetailLayoutBuilder::GetDetailFontBold())
+				.Text(MemberCustomizationStyle::NodeTooltipText)
+				]
+				.ValueContent()
+				[
+					SNew(SMultiLineEditableTextBox)
+					.Text(this, &TChildClass::GetTooltip)
+					.OnTextCommitted(this, &TChildClass::OnTooltipCommitted)
+					.IsReadOnly(bIsInterfaceMember || !bIsGraphEditable)
+					.ModiferKeyForNewLine(EModifierKey::Shift)
+					.RevertTextOnEscape(true)
+					.WrapTextAt(MemberCustomizationStyle::DetailsTitleMaxWidth - MemberCustomizationStyle::DetailsTitleWrapPadding)
+					.Font(IDetailLayoutBuilder::GetDetailFont())
+				];
+
+				DataTypeSelector->AddDataTypeSelector(InDetailLayout, MemberCustomizationStyle::DataTypeNameText, GraphMember, !bIsInterfaceMember && bIsGraphEditable);
+
+				IDetailCategoryBuilder& DefaultCategoryBuilder = InDetailLayout.EditCategory("DefaultValue");
+				TSharedPtr<IPropertyHandle> LiteralHandle = InDetailLayout.GetProperty(GET_MEMBER_NAME_CHECKED(TMemberType, Literal));
+				if (ensure(GraphMember.IsValid()) && ensure(LiteralHandle.IsValid()))
+				{
+					UObject* LiteralObject = nullptr;
+					if (LiteralHandle->GetValue(LiteralObject) == FPropertyAccess::Success)
+					{
+						if (LiteralObject)
+						{
+							LiteralHandle->MarkHiddenByCustomization();
+
+							TSharedPtr<IPropertyHandle> DefaultValueHandle;
+
+							if (IDetailPropertyRow* Row = DefaultCategoryBuilder.AddExternalObjectProperty(TArray<UObject*>({ LiteralObject }), "Default"))
+							{
+								DefaultValueHandle = Row->GetPropertyHandle();
+								if (DefaultValueHandle.IsValid())
+								{
+									SetDefaultPropertyMetaData(DefaultValueHandle.ToSharedRef());
+								}
+							}
+
+							IMetasoundEditorModule& EditorModule = FModuleManager::GetModuleChecked<IMetasoundEditorModule>("MetaSoundEditor");
+							LiteralCustomization = EditorModule.CreateMemberDefaultLiteralCustomization(*LiteralObject->GetClass(), DefaultCategoryBuilder);
+							if (LiteralCustomization.IsValid())
+							{
+								LiteralCustomization->CustomizeLiteral(*CastChecked<UMetasoundEditorGraphMemberDefaultLiteral>(LiteralObject), DefaultValueHandle);
+							}
+
+							UMetasoundEditorGraphMemberDefaultFloat* DefaultFloat = Cast<UMetasoundEditorGraphMemberDefaultFloat>(LiteralObject);
+							if (DefaultFloat)
+							{
+								// add input widget properties 
+								IDetailCategoryBuilder& WidgetCategoryBuilder = InDetailLayout.EditCategory("Widget");
+								WidgetCategoryBuilder.AddExternalObjectProperty(TArray<UObject*>({ DefaultFloat }), GET_MEMBER_NAME_CHECKED(UMetasoundEditorGraphMemberDefaultFloat, WidgetType));
+								WidgetCategoryBuilder.AddExternalObjectProperty(TArray<UObject*>({ DefaultFloat }), GET_MEMBER_NAME_CHECKED(UMetasoundEditorGraphMemberDefaultFloat, WidgetOrientation));
+								WidgetCategoryBuilder.AddExternalObjectProperty(TArray<UObject*>({ DefaultFloat }), GET_MEMBER_NAME_CHECKED(UMetasoundEditorGraphMemberDefaultFloat, WidgetValueType));
+							}
+						}
+					}
+				}
 			}
 			// End of IDetailCustomization interface
 
-			void OnNameChanged(const FText& InNewName)
+			void SetDefaultPropertyMetaData(TSharedRef<IPropertyHandle> InDefaultPropertyHandle) const
 			{
 				using namespace Frontend;
 
-				bIsNameInvalid = false;
-				DisplayNameEditableTextBox->SetError(FText::GetEmpty());
+				if (!GraphMember.IsValid())
+				{
+					return;
+				}
 
-				if (!ensure(GraphMember.IsValid()))
+				FMetasoundFrontendRegistryContainer* Registry = FMetasoundFrontendRegistryContainer::Get();
+				if (!ensure(Registry))
 				{
 					return;
 				}
 
-				FText Error;
-				if (!GraphMember->CanRename(InNewName, Error))
+				FName TypeName = GraphMember->GetDataType();
+				if (TypeName.IsNone())
 				{
-					bIsNameInvalid = true;
-					DisplayNameEditableTextBox->SetError(Error);
+					return;
+				}
+
+				FDataTypeRegistryInfo DataTypeInfo;
+				if (!ensure(IDataTypeRegistry::Get().GetDataTypeInfo(TypeName, DataTypeInfo)))
+				{
+					return;
+				}
+
+				if (DataTypeInfo.IsArrayType())
+				{
+					TypeName = CreateElementTypeNameFromArrayTypeName(TypeName);
+				}
+				InDefaultPropertyHandle->SetInstanceMetaData(MemberCustomizationStyle::DataTypeNameIdentifier, TypeName.ToString());
+
+				const EMetasoundFrontendLiteralType LiteralType = GetMetasoundFrontendLiteralType(DataTypeInfo.PreferredLiteralType);
+				if (LiteralType != EMetasoundFrontendLiteralType::UObject && LiteralType != EMetasoundFrontendLiteralType::UObjectArray)
+				{
+					return;
+				}
+
+				UClass* ProxyGenClass = DataTypeInfo.ProxyGeneratorClass;
+				if (ProxyGenClass)
+				{
+					const FString ClassName = ProxyGenClass->GetName();
+					InDefaultPropertyHandle->SetInstanceMetaData(MemberCustomizationStyle::ProxyGeneratorClassNameIdentifier, ClassName);
 				}
 			}
 
-			FText GetDisplayName() const
+			void OnNameChanged(const FText& InNewName)
 			{
 				using namespace Frontend;
 
-				if (GraphMember.IsValid())
+				bIsNameInvalid = false;
+				DisplayNameEditableTextBox->SetError(FText::GetEmpty());
+
+				if (!ensure(GraphMember.IsValid()))
 				{
-					return GraphMember->GetDisplayName();
+					return;
 				}
 
-				return FText::GetEmpty();
+				FText Error;
+				if (!GraphMember->CanRename(InNewName, Error))
+				{
+					bIsNameInvalid = true;
+					DisplayNameEditableTextBox->SetError(Error);
+				}
 			}
 
 			FText GetName() const
@@ -247,22 +427,13 @@ namespace Metasound
 				return false;
 			}
 
-			void OnTooltipCommitted(const FText& InNewText, ETextCommit::Type InTextCommit)
+			FText GetDisplayName() const
 			{
 				using namespace Frontend;
 
 				if (GraphMember.IsValid())
 				{
-					constexpr bool bPostTransaction = true;
-					GraphMember->SetDescription(InNewText, bPostTransaction);
-				}
-			}
-
-			FText GetTooltip() const
-			{
-				if (GraphMember.IsValid())
-				{
-					return GraphMember->GetDescription();
+					return GraphMember->GetDisplayName();
 				}
 
 				return FText::GetEmpty();
@@ -282,13 +453,34 @@ namespace Metasound
 				bIsNameInvalid = false;
 			}
 
+			void OnTooltipCommitted(const FText& InNewText, ETextCommit::Type InTextCommit)
+			{
+				using namespace Frontend;
+
+				if (GraphMember.IsValid())
+				{
+					constexpr bool bPostTransaction = true;
+					GraphMember->SetDescription(InNewText, bPostTransaction);
+				}
+			}
+
+			FText GetTooltip() const
+			{
+				if (GraphMember.IsValid())
+				{
+					return GraphMember->GetDescription();
+				}
+
+				return FText::GetEmpty();
+			}
+
 			void OnNameCommitted(const FText& InNewName, ETextCommit::Type InTextCommit)
 			{
 				using namespace Frontend;
 
 				if (!bIsNameInvalid && GraphMember.IsValid())
 				{
-					const FText TransactionLabel = FText::Format(LOCTEXT("Rename Graph Member", "Set MetaSound {0}'s Name"), GraphMember->GetGraphMemberLabel());
+					const FText TransactionLabel = FText::Format(LOCTEXT("RenameGraphMember_Format", "Set MetaSound {0}'s Name"), TChildClass::MemberNameText);
 					const FScopedTransaction Transaction(TransactionLabel);
 
 					constexpr bool bPostTransaction = false;
@@ -310,58 +502,53 @@ namespace Metasound
 				return ECheckBoxState::Unchecked;
 			}
 
-			void OnPrivateChanged(ECheckBoxState InNewState)
-			{
-				if (GraphMember.IsValid())
-				{
-					const bool bIsChecked = InNewState == ECheckBoxState::Checked;
-					Frontend::FNodeHandle NodeHandle = GraphMember->GetNodeHandle();
-					FMetasoundFrontendNodeStyle NodeStyle = NodeHandle->GetNodeStyle();
-					NodeStyle.bIsPrivate = bIsChecked;
-					NodeHandle->SetNodeStyle(NodeStyle);
-				}
-			}
+		private:
+			TUniquePtr<IMemberDefaultLiteralCustomization> LiteralCustomization;
 		};
 
-		class FMetasoundInputDetailCustomization : public TMetasoundGraphMemberDetailCustomization<UMetasoundEditorGraphInput>
+		class FMetasoundInputDetailCustomization : public TMetasoundGraphMemberDetailCustomization<UMetasoundEditorGraphInput, FMetasoundInputDetailCustomization>
 		{
 		public:
 			FMetasoundInputDetailCustomization()
-				: TMetasoundGraphMemberDetailCustomization<UMetasoundEditorGraphInput>(LOCTEXT("InputGraphMemberLabel", "Input"))
+				: TMetasoundGraphMemberDetailCustomization<UMetasoundEditorGraphInput, FMetasoundInputDetailCustomization>()
 			{
 			}
+
 			virtual ~FMetasoundInputDetailCustomization() = default;
 
+			virtual bool IsInterfaceMember() const override;
+
 			// IDetailCustomization interface
-			virtual void CustomizeDetails(IDetailLayoutBuilder& DetailLayout) override;
+			virtual void CustomizeDetails(IDetailLayoutBuilder& InDetailLayout) override;
 			// End of IDetailCustomization interface
 
-		private:
-			void SetDefaultPropertyMetaData(TSharedRef<IPropertyHandle> InDefaultPropertyHandle) const;
+			static const FText MemberNameText;
 
-			FName GetLiteralDataType() const;
-
-			TUniquePtr<IMetaSoundInputLiteralCustomization> LiteralCustomization;
+		private:
+			TUniquePtr<IMemberDefaultLiteralCustomization> LiteralCustomization;
 		};
 
-		class FMetasoundOutputDetailCustomization : public TMetasoundGraphMemberDetailCustomization<UMetasoundEditorGraphOutput>
+		class FMetasoundOutputDetailCustomization : public TMetasoundGraphMemberDetailCustomization<UMetasoundEditorGraphOutput, FMetasoundOutputDetailCustomization>
 		{
 		public:
+			static const FText MemberNameText;
+
 			FMetasoundOutputDetailCustomization()
-				: TMetasoundGraphMemberDetailCustomization<UMetasoundEditorGraphOutput>(LOCTEXT("OutputGraphMemberLabel", "Output"))
+				: TMetasoundGraphMemberDetailCustomization<UMetasoundEditorGraphOutput, FMetasoundOutputDetailCustomization>()
 			{
 			}
 
+			virtual ~FMetasoundOutputDetailCustomization() = default;
+
+			virtual bool IsInterfaceMember() const override;
+
 			// IDetailCustomization interface
-			virtual void CustomizeDetails(IDetailLayoutBuilder& DetailLayout) override;
+			virtual void CustomizeDetails(IDetailLayoutBuilder& InDetailLayout) override;
 			// End of IDetailCustomization interface
 
 		private:
-			void SetDefaultPropertyMetaData(TSharedRef<IPropertyHandle> InDefaultPropertyHandle) const;
-
-			FName GetLiteralDataType() const;
+			TUniquePtr<IMemberDefaultLiteralCustomization> LiteralCustomization;
 		};
-
 	} // namespace Editor
 } // namespace Metasound
 #undef LOCTEXT_NAMESPACE
diff --git a/Engine/Plugins/Runtime/Metasound/Source/MetasoundEditor/Private/MetasoundVariableDetailCustomization.cpp b/Engine/Plugins/Runtime/Metasound/Source/MetasoundEditor/Private/MetasoundVariableDetailCustomization.cpp
--- a/Engine/Plugins/Runtime/Metasound/Source/MetasoundEditor/Private/MetasoundVariableDetailCustomization.cpp
+++ b/Engine/Plugins/Runtime/Metasound/Source/MetasoundEditor/Private/MetasoundVariableDetailCustomization.cpp
@@ -17,195 +17,11 @@ namespace Metasound
 {
 	namespace Editor
 	{
-		namespace MemberCustomizationPrivate
-		{
-			static const FText VariableNameText = LOCTEXT("Variable_Name", "Variable Name");
-			static const FText VariableDisplayNameText = LOCTEXT("VariableDisplay_Name", "Variable Display Name");
-		}
-
-		FMetasoundVariableDetailCustomization::FMetasoundVariableDetailCustomization()
-			: TMetasoundGraphMemberDetailCustomization<UMetasoundEditorGraphVariable>(LOCTEXT("VariableGraphMemberLabel", "Variable"))
-		{
-		}
+		const FText FMetasoundVariableDetailCustomization::MemberNameText = LOCTEXT("Variable_Name", "Variable");
 
 		void FMetasoundVariableDetailCustomization::CustomizeDetails(IDetailLayoutBuilder& DetailLayout)
 		{
-			using namespace Frontend;
-
-			TMetasoundGraphMemberDetailCustomization<UMetasoundEditorGraphVariable>::CustomizeDetails(DetailLayout);
-
-			if (!GraphMember.IsValid())
-			{
-				return;
-			}
-
-			IDetailCategoryBuilder& CategoryBuilder = DetailLayout.EditCategory("General");
-
-			const bool bIsGraphEditable = IsGraphEditable();
-			NameEditableTextBox = SNew(SEditableTextBox)
-				.Text(this, &FMetasoundVariableDetailCustomization::GetName)
-				.OnTextChanged(this, &FMetasoundVariableDetailCustomization::OnNameChanged)
-				.OnTextCommitted(this, &FMetasoundVariableDetailCustomization::OnNameCommitted)
-				.IsReadOnly(!bIsGraphEditable)
-				.Font(IDetailLayoutBuilder::GetDetailFont());
-
-			DisplayNameEditableTextBox = SNew(SEditableTextBox)
-				.Text(this, &FMetasoundVariableDetailCustomization::GetDisplayName)
-				.OnTextCommitted(this, &FMetasoundVariableDetailCustomization::OnDisplayNameCommitted)
-				.IsReadOnly(!bIsGraphEditable)
-				.Font(IDetailLayoutBuilder::GetDetailFont());
-
-			CategoryBuilder.AddCustomRow(MemberCustomizationPrivate::VariableNameText)
-			.EditCondition(bIsGraphEditable, nullptr)
-			.NameContent()
-			[
-				SNew(STextBlock)
-				.Font(IDetailLayoutBuilder::GetDetailFontBold())
-				.Text(MemberCustomizationPrivate::VariableNameText)
-				.ToolTipText(LOCTEXT("VariableName_Description", "Name used by external systems to identify input. Used as DisplayName within MetaSound Graph Editor if no DisplayName is provided."))
-			]
-			.ValueContent()
-			[
-				NameEditableTextBox.ToSharedRef()
-			];
-
-			// TODO: Enable and use proper FText property editor
-// 			CategoryBuilder.AddCustomRow(MemberCustomizationPrivate::VariableDisplayNameText)
-// 			.EditCondition(!bIsInterfaceMember && bIsGraphEditable, nullptr)
-// 			.NameContent()
-// 			[
-// 				SNew(STextBlock)
-// 				.Font(IDetailLayoutBuilder::GetDetailFontBold())
-// 				.Text(MemberCustomizationPrivate::VariableDisplayNameText)
-// 				.ToolTipText(LOCTEXT("VariableDisplayName_Description", "Optional, localized name used within the MetaSounds editor(s) to describe the given input."))
-// 			]
-// 			.ValueContent()
-// 			[
-// 				DisplayNameEditableTextBox.ToSharedRef()
-// 			];
-
-			CategoryBuilder.AddCustomRow(MemberCustomizationPrivate::NodeTooltipText)
-			.EditCondition(bIsGraphEditable, nullptr)
-			.NameContent()
-			[
-				SNew(STextBlock)
-				.Font(IDetailLayoutBuilder::GetDetailFontBold())
-				.Text(MemberCustomizationPrivate::NodeTooltipText)
-			]
-			.ValueContent()
-			[
-				SNew(SMultiLineEditableTextBox)
-				.Text(this, &FMetasoundVariableDetailCustomization::GetTooltip)
-				.OnTextCommitted(this, &FMetasoundVariableDetailCustomization::OnTooltipCommitted)
-				.IsReadOnly(!bIsGraphEditable)
-				.ModiferKeyForNewLine(EModifierKey::Shift)
-				.RevertTextOnEscape(true)
-				.WrapTextAt(MemberCustomizationPrivate::DetailsTitleMaxWidth - MemberCustomizationPrivate::DetailsTitleWrapPadding)
-				.Font(IDetailLayoutBuilder::GetDetailFont())
-			];
-
-			DataTypeSelector->AddDataTypeSelector(DetailLayout, MemberCustomizationPrivate::DataTypeNameText, GraphMember, bIsGraphEditable);
-
-			IDetailCategoryBuilder& DefaultCategoryBuilder = DetailLayout.EditCategory("DefaultValue");
-			TSharedPtr<IPropertyHandle> LiteralHandle = DetailLayout.GetProperty(GET_MEMBER_NAME_CHECKED(UMetasoundEditorGraphVariable, Literal));
-			if (ensure(GraphMember.IsValid()) && ensure(LiteralHandle.IsValid()))
-			{
-				UObject* LiteralObject = nullptr;
-				if (LiteralHandle->GetValue(LiteralObject) == FPropertyAccess::Success)
-				{
-					if (ensure(LiteralObject))
-					{
-						LiteralHandle->MarkHiddenByCustomization();
-
-						TSharedPtr<IPropertyHandle> DefaultValueHandle;
-
-						if (IDetailPropertyRow* Row = DefaultCategoryBuilder.AddExternalObjectProperty(TArray<UObject*>({ LiteralObject }), "Default"))
-						{
-							DefaultValueHandle = Row->GetPropertyHandle();
-							if (DefaultValueHandle.IsValid())
-							{
-								SetDefaultPropertyMetaData(DefaultValueHandle.ToSharedRef());
-
-								FSimpleDelegate UpdateDocumentVariable = FSimpleDelegate::CreateLambda([GraphMember = this->GraphMember]()
-								{
-									if (GraphMember.IsValid())
-									{
-										GraphMember->UpdateDocumentVariable();
-									}
-								});
-
-								DefaultValueHandle->SetOnPropertyValueChanged(UpdateDocumentVariable);
-								DefaultValueHandle->SetOnChildPropertyValueChanged(UpdateDocumentVariable);
-
-								TSharedPtr<IPropertyHandleArray> DefaultValueArray = DefaultValueHandle->AsArray();
-								if (DefaultValueArray.IsValid())
-								{
-									DefaultValueArray->SetOnNumElementsChanged(UpdateDocumentVariable);
-								}
-							}
-						}
-
-						IMetasoundEditorModule& EditorModule = FModuleManager::GetModuleChecked<IMetasoundEditorModule>("MetaSoundEditor");
-						LiteralCustomization = EditorModule.CreateInputLiteralCustomization(*LiteralObject->GetClass(), DefaultCategoryBuilder);
-						if (LiteralCustomization.IsValid())
-						{
-							LiteralCustomization->CustomizeLiteral(*CastChecked<UMetasoundEditorGraphInputLiteral>(LiteralObject), DefaultValueHandle);
-						}
-					}
-					else
-					{
-						DefaultCategoryBuilder.AddProperty(LiteralHandle);
-					}
-				}
-			}
-		}
-
-		void FMetasoundVariableDetailCustomization::SetDefaultPropertyMetaData(TSharedRef<IPropertyHandle> InDefaultPropertyHandle) const
-		{
-			using namespace Frontend;
-
-			if (!GraphMember.IsValid())
-			{
-				return;
-			}
-
-			FMetasoundFrontendRegistryContainer* Registry = FMetasoundFrontendRegistryContainer::Get();
-			if (!ensure(Registry))
-			{
-				return;
-			}
-
-			const FName TypeName = GraphMember->TypeName;
-			if (TypeName.IsNone())
-			{
-				return;
-			}
-
-			FString TypeNameString = TypeName.ToString();
-			if (TypeNameString.EndsWith(MemberCustomizationPrivate::ArrayIdentifier))
-			{
-				TypeNameString = TypeNameString.LeftChop(MemberCustomizationPrivate::ArrayIdentifier.Len());
-			}
-			InDefaultPropertyHandle->SetInstanceMetaData(MemberCustomizationPrivate::DataTypeNameIdentifier, TypeNameString);
-
-			FDataTypeRegistryInfo DataTypeInfo;
-			if (!ensure(IDataTypeRegistry::Get().GetDataTypeInfo(TypeName, DataTypeInfo)))
-			{
-				return;
-			}
-
-			const EMetasoundFrontendLiteralType LiteralType = GetMetasoundFrontendLiteralType(DataTypeInfo.PreferredLiteralType);
-			if (LiteralType != EMetasoundFrontendLiteralType::UObject && LiteralType != EMetasoundFrontendLiteralType::UObjectArray)
-			{
-				return;
-			}
-
-			UClass* ProxyGenClass = DataTypeInfo.ProxyGeneratorClass;
-			if (ProxyGenClass)
-			{
-				const FString ClassName = ProxyGenClass->GetName();
-				InDefaultPropertyHandle->SetInstanceMetaData(MemberCustomizationPrivate::ProxyGeneratorClassNameIdentifier, ClassName);
-			}
+			TMetasoundGraphMemberDetailCustomization<UMetasoundEditorGraphVariable, FMetasoundVariableDetailCustomization>::CustomizeDetails(DetailLayout);
 		}
 	} // namespace Editor
 } // namespace Metasound
diff --git a/Engine/Plugins/Runtime/Metasound/Source/MetasoundEditor/Private/MetasoundVariableDetailCustomization.h b/Engine/Plugins/Runtime/Metasound/Source/MetasoundEditor/Private/MetasoundVariableDetailCustomization.h
--- a/Engine/Plugins/Runtime/Metasound/Source/MetasoundEditor/Private/MetasoundVariableDetailCustomization.h
+++ b/Engine/Plugins/Runtime/Metasound/Source/MetasoundEditor/Private/MetasoundVariableDetailCustomization.h
@@ -8,7 +8,7 @@
 #include "Templates/UniquePtr.h"
 
 class IDetailLayoutBuilder;
-class IMetaSoundInputLiteralCustomization;
+class IMemberDefaultLiteralCustomization;
 
 #define LOCTEXT_NAMESPACE "MetaSoundEditor"
 
@@ -16,22 +16,21 @@ namespace Metasound
 {
 	namespace Editor
 	{
-		class FMetasoundVariableDetailCustomization : public TMetasoundGraphMemberDetailCustomization<UMetasoundEditorGraphVariable>
+		class FMetasoundVariableDetailCustomization : public TMetasoundGraphMemberDetailCustomization<UMetasoundEditorGraphVariable, FMetasoundVariableDetailCustomization>
 		{
 		public:
-			FMetasoundVariableDetailCustomization();
+			static const FText MemberNameText;
+
+			FMetasoundVariableDetailCustomization()
+				: TMetasoundGraphMemberDetailCustomization<UMetasoundEditorGraphVariable, FMetasoundVariableDetailCustomization>()
+			{
+			}
+
 			virtual ~FMetasoundVariableDetailCustomization() = default;
 
 			// IDetailCustomization interface
 			virtual void CustomizeDetails(IDetailLayoutBuilder& DetailLayout) override;
 			// End of IDetailCustomization interface
-
-		private:
-			void SetDefaultPropertyMetaData(TSharedRef<IPropertyHandle> InDefaultPropertyHandle) const;
-
-			FName GetLiteralDataType() const;
-
-			TUniquePtr<IMetaSoundInputLiteralCustomization> LiteralCustomization;
 		};
 	} // namespace Editor
 } // namespace Metasound
diff --git a/Engine/Plugins/Runtime/Metasound/Source/MetasoundEditor/Private/SMetasoundGraphNode.cpp b/Engine/Plugins/Runtime/Metasound/Source/MetasoundEditor/Private/SMetasoundGraphNode.cpp
--- a/Engine/Plugins/Runtime/Metasound/Source/MetasoundEditor/Private/SMetasoundGraphNode.cpp
+++ b/Engine/Plugins/Runtime/Metasound/Source/MetasoundEditor/Private/SMetasoundGraphNode.cpp
@@ -5,6 +5,7 @@
 #include "Components/AudioComponent.h"
 #include "EditorStyleSet.h"
 #include "GraphEditorSettings.h"
+#include "IAudioParameterTransmitter.h"
 #include "IDocumentation.h"
 #include "KismetPins/SGraphPinBool.h"
 #include "KismetPins/SGraphPinExec.h"
@@ -14,10 +15,12 @@
 #include "KismetPins/SGraphPinString.h"
 #include "MetasoundEditorGraph.h"
 #include "MetasoundEditorGraphBuilder.h"
-#include "MetasoundEditorGraphInputNodes.h"
+#include "MetasoundEditorGraphInputNode.h"
+#include "MetasoundEditorGraphMemberDefaults.h"
 #include "MetasoundEditorGraphNode.h"
 #include "MetasoundEditorGraphSchema.h"
 #include "MetasoundEditorModule.h"
+#include "MetasoundFrontendArchetypeRegistry.h"
 #include "MetasoundFrontendRegistries.h"
 #include "MetasoundTrigger.h"
 #include "NodeFactory.h"
@@ -47,701 +50,820 @@
 
 #define LOCTEXT_NAMESPACE "MetasoundGraphNode"
 
-SMetasoundGraphNode::~SMetasoundGraphNode()
+
+namespace Metasound
 {
-	UMetasoundEditorGraphNode& Node = GetMetasoundNode();
-	if (UMetasoundEditorGraphInputNode* InputNode = Cast<UMetasoundEditorGraphInputNode>(&Node))
+	namespace Editor
 	{
-		if (UMetasoundEditorGraphInput* GraphInput = InputNode->Input)
+		SMetaSoundGraphNode::~SMetaSoundGraphNode()
 		{
-			if (UMetasoundEditorGraphInputFloat* InputFloat = Cast<UMetasoundEditorGraphInputFloat>(GraphInput->Literal))
+			UMetasoundEditorGraphNode& Node = GetMetaSoundNode();
+			if (UMetasoundEditorGraphMemberNode* MemberNode = Cast<UMetasoundEditorGraphMemberNode>(&Node))
 			{
-				InputFloat->OnDefaultValueChanged.Remove(InputSliderOnValueChangedDelegateHandle);
-				InputFloat->OnRangeChanged.Remove(InputSliderOnRangeChangedDelegateHandle);
+				if (UMetasoundEditorGraphMember* GraphMember = MemberNode->GetMember())
+				{
+					if (UMetasoundEditorGraphMemberDefaultFloat* DefaultFloat = Cast<UMetasoundEditorGraphMemberDefaultFloat>(GraphMember->GetLiteral()))
+					{
+						DefaultFloat->OnDefaultValueChanged.Remove(InputSliderOnValueChangedDelegateHandle);
+						DefaultFloat->OnRangeChanged.Remove(InputSliderOnRangeChangedDelegateHandle);
+					}
+				}
 			}
 		}
-	}
-}
 
-void SMetasoundGraphNode::Construct(const FArguments& InArgs, class UEdGraphNode* InNode)
-{
-	GraphNode = InNode;
-	SetCursor(EMouseCursor::CardinalCross);
-	UpdateGraphNode();
-}
-
-void SMetasoundGraphNode::ExecuteInputTrigger(UMetasoundEditorGraphInputLiteral& Literal)
-{
-	UMetasoundEditorGraphInput* Input = Cast<UMetasoundEditorGraphInput>(Literal.GetOuter());
-	if (!ensure(Input))
-	{
-		return;
-	}
-
-	// If modifying graph not currently being previewed, do not forward request
-	if (UMetasoundEditorGraph* Graph = Cast<UMetasoundEditorGraph>(Input->GetOuter()))
-	{
-		if (!Graph->IsPreviewing())
+		void SMetaSoundGraphNode::Construct(const FArguments& InArgs, class UEdGraphNode* InNode)
 		{
-			return;
+			GraphNode = InNode;
+			SetCursor(EMouseCursor::CardinalCross);
+			UpdateGraphNode();
 		}
-	}
 
-	if (UAudioComponent* PreviewComponent = GEditor->GetPreviewAudioComponent())
-	{
-		PreviewComponent->SetTriggerParameter(Input->GetMemberName());
-	}
-}
-
-TSharedRef<SWidget> SMetasoundGraphNode::CreateTriggerSimulationWidget(UMetasoundEditorGraphInputLiteral& InputLiteral)
-{
-	return SNew(SHorizontalBox)
-	+ SHorizontalBox::Slot()
-	.Padding(2.0f, 0.0f, 0.0f, 0.0f)
-	.HAlign(HAlign_Left)
-	.VAlign(VAlign_Center)
-	[
-		SNew(SButton)
-		.ButtonStyle(FAppStyle::Get(), "SimpleButton")
-		.OnClicked_Lambda([LiteralPtr = TWeakObjectPtr<UMetasoundEditorGraphInputLiteral>(&InputLiteral)]()
+		void SMetaSoundGraphNode::ExecuteTrigger(UMetasoundEditorGraphMemberDefaultLiteral& Literal)
 		{
-			if (LiteralPtr.IsValid())
+			UMetasoundEditorGraphMember* Member = Cast<UMetasoundEditorGraphMember>(Literal.GetOuter());
+			if (!ensure(Member))
 			{
-				ExecuteInputTrigger(*LiteralPtr.Get());
+				return;
 			}
-			return FReply::Handled();
-		})
-		.ToolTipText(LOCTEXT("TriggerTestToolTip", "Executes trigger if currently previewing MetaSound."))
-		.ForegroundColor(FSlateColor::UseForeground())
-		.ContentPadding(0)
-		.IsFocusable(false)
-		[
-			SNew(SImage)
-			.Image(FAppStyle::Get().GetBrush("Icons.CircleArrowDown"))
-			.ColorAndOpacity(FSlateColor::UseForeground())
-		]
-	];
-}
-
-void SMetasoundGraphNode::CreateInputSideAddButton(TSharedPtr<SVerticalBox> InputBox)
-{
-	TSharedRef<SWidget> AddPinButton = AddPinButtonContent(
-		LOCTEXT("MetasoundGraphNode_AddPinInputButton", "Add Input"),
-		LOCTEXT("MetasoundGraphNode_AddPinInputButton_Tooltip", "Add an input to the parent Metasound node.")
-	);
-
-	FMargin AddPinPadding = Settings->GetOutputPinPadding();
-	AddPinPadding.Top += 6.0f;
-
-	InputBox->AddSlot()
-	.AutoHeight()
-	.VAlign(VAlign_Center)
-	.Padding(AddPinPadding)
-	[
-		AddPinButton
-	];
-}
-
-void SMetasoundGraphNode::CreateOutputSideAddButton(TSharedPtr<SVerticalBox> OutputBox)
-{
-	TSharedRef<SWidget> AddPinButton = AddPinButtonContent(
-		LOCTEXT("MetasoundGraphNode_AddPinOutputButton", "Add Output"),
-		LOCTEXT("MetasoundGraphNode_AddPinOutputButton_Tooltip", "Add an output to the parent Metasound node.")
-	);
-
-	FMargin AddPinPadding = Settings->GetOutputPinPadding();
-	AddPinPadding.Top += 6.0f;
-
-	OutputBox->AddSlot()
-	.AutoHeight()
-	.VAlign(VAlign_Center)
-	.Padding(AddPinPadding)
-	[
-		AddPinButton
-	];
-}
-
-UMetasoundEditorGraphNode& SMetasoundGraphNode::GetMetasoundNode()
-{
-	return *CastChecked<UMetasoundEditorGraphNode>(GraphNode);
-}
-
-const UMetasoundEditorGraphNode& SMetasoundGraphNode::GetMetasoundNode() const
-{
-	check(GraphNode);
-	return *Cast<UMetasoundEditorGraphNode>(GraphNode);
-}
-
-TSharedPtr<SGraphPin> SMetasoundGraphNode::CreatePinWidget(UEdGraphPin* InPin) const
-{
-	using namespace Metasound::Editor;
-	using namespace Metasound::Frontend;
 
-	TSharedPtr<SGraphPin> PinWidget;
+			if (UMetasoundEditorGraph* Graph = Member->GetOwningGraph())
+			{
+				if (!Graph->IsPreviewing())
+				{
+					TSharedPtr<FEditor> MetaSoundEditor = FGraphBuilder::GetEditorForMetasound(Graph->GetMetasoundChecked());
+					if (!MetaSoundEditor.IsValid())
+					{
+						return;
+					}
+					MetaSoundEditor->Play();
+				}
+			}
 
-	if (const UMetasoundEditorGraphSchema* GraphSchema = Cast<const UMetasoundEditorGraphSchema>(InPin->GetSchema()))
-	{
-		// Don't show default value field for container types
-		if (InPin->PinType.ContainerType != EPinContainerType::None)
-		{
-			PinWidget = SNew(SMetasoundGraphPin, InPin)
-				.ToolTipText(this, &SMetasoundGraphNode::GetPinTooltip, InPin);
+			if (UAudioComponent* PreviewComponent = GEditor->GetPreviewAudioComponent())
+			{
+				PreviewComponent->SetTriggerParameter(Member->GetMemberName());
+			}
 		}
 
-		else if (InPin->PinType.PinCategory == FGraphBuilder::PinCategoryAudio)
+		TAttribute<EVisibility> SMetaSoundGraphNode::GetSimulationVisibilityAttribute() const
 		{
-			PinWidget = SNew(SMetasoundGraphPin, InPin)
-				.ToolTipText(this, &SMetasoundGraphNode::GetPinTooltip, InPin);
+			return TAttribute<EVisibility>::CreateLambda([this]()
+			{
+				using namespace Frontend;
+
+				if (const UMetasoundEditorGraphMemberNode* Node = Cast<UMetasoundEditorGraphMemberNode>(&GetMetaSoundNode()))
+				{
+					if (const UMetasoundEditorGraphVertex* Vertex = Cast<UMetasoundEditorGraphVertex>(Node->GetMember()))
+					{
+						if (const UMetasoundEditorGraph* Graph = Vertex->GetOwningGraph())
+						{
+							if (!Graph->IsPreviewing())
+							{
+								return EVisibility::Hidden;
+							}
+						}
+
+						// Don't enable trigger simulation widget if its a trigger provided by an interface
+						// that does not support transmission.
+						const FInterfaceRegistryKey Key = GetInterfaceRegistryKey(Vertex->GetInterfaceVersion());
+						const IInterfaceRegistryEntry* Entry = IInterfaceRegistry::Get().FindInterfaceRegistryEntry(Key);
+						if (Entry && Entry->GetRouterName() != Audio::IParameterTransmitter::RouterName)
+						{
+							return EVisibility::Hidden;
+						}
+						else if (const UMetasoundEditorGraphMemberDefaultLiteral* Literal = Vertex->GetLiteral())
+						{
+							if (!Literal)
+							{
+								return EVisibility::Hidden;
+							}
+						}
+					}
+				}
+
+				return EVisibility::Visible;
+			});
 		}
 
-		else if (InPin->PinType.PinCategory == FGraphBuilder::PinCategoryBoolean)
+		TSharedRef<SWidget> SMetaSoundGraphNode::CreateTriggerSimulationWidget(UMetasoundEditorGraphMemberDefaultLiteral& InputLiteral, TAttribute<EVisibility>&& InVisibility, TAttribute<bool>&& InEnablement, const FText* InToolTip)
 		{
-			PinWidget = SNew(SMetasoundGraphPinBool, InPin)
-				.ToolTipText(this, &SMetasoundGraphNode::GetPinTooltip, InPin);
+			const FText ToolTip = InToolTip
+				? *InToolTip
+				: LOCTEXT("TriggerTestToolTip", "Executes trigger if currently previewing MetaSound.");
+
+			TSharedPtr<SButton> SimulationButton;
+			TSharedRef<SWidget> SimulationWidget = SNew(SHorizontalBox)
+			+ SHorizontalBox::Slot()
+			.Padding(2.0f, 0.0f, 0.0f, 0.0f)
+			.HAlign(HAlign_Left)
+			.VAlign(VAlign_Center)
+			[
+				SAssignNew(SimulationButton, SButton)
+				.ButtonStyle(FAppStyle::Get(), "SimpleButton")
+				.OnClicked_Lambda([LiteralPtr = TWeakObjectPtr<UMetasoundEditorGraphMemberDefaultLiteral>(&InputLiteral)]()
+				{
+					if (LiteralPtr.IsValid())
+					{
+						ExecuteTrigger(*LiteralPtr.Get());
+					}
+					return FReply::Handled();
+				})
+				.ToolTipText(ToolTip)
+				.ForegroundColor(FSlateColor::UseForeground())
+				.ContentPadding(0)
+				.IsFocusable(false)
+				[
+					SNew(SImage)
+					.Image(FAppStyle::Get().GetBrush("Icons.CircleArrowDown"))
+					.ColorAndOpacity(FSlateColor::UseForeground())
+				]
+				.Visibility(MoveTemp(InVisibility))
+			];
+
+			SimulationButton->SetEnabled(MoveTemp(InEnablement));
+
+			return SimulationWidget;
 		}
 
-		else if (InPin->PinType.PinCategory == FGraphBuilder::PinCategoryFloat)
+		void SMetaSoundGraphNode::CreateInputSideAddButton(TSharedPtr<SVerticalBox> InputBox)
 		{
-			PinWidget = SNew(SMetasoundGraphPinFloat, InPin)
-				.ToolTipText(this, &SMetasoundGraphNode::GetPinTooltip, InPin);
+			TSharedRef<SWidget> AddPinButton = AddPinButtonContent(
+				LOCTEXT("MetasoundGraphNode_AddPinInputButton", "Add Input"),
+				LOCTEXT("MetasoundGraphNode_AddPinInputButton_Tooltip", "Add an input to the parent Metasound node.")
+			);
+
+			FMargin AddPinPadding = Settings->GetOutputPinPadding();
+			AddPinPadding.Top += 6.0f;
+
+			InputBox->AddSlot()
+			.AutoHeight()
+			.VAlign(VAlign_Center)
+			.Padding(AddPinPadding)
+			[
+				AddPinButton
+			];
 		}
 
-		else if (InPin->PinType.PinCategory == FGraphBuilder::PinCategoryInt32)
+		void SMetaSoundGraphNode::CreateOutputSideAddButton(TSharedPtr<SVerticalBox> OutputBox)
 		{
-			if (SMetasoundGraphEnumPin::FindEnumInterfaceFromPin(InPin))
-			{
-				PinWidget = SNew(SMetasoundGraphEnumPin, InPin)
-					.ToolTipText(this, &SMetasoundGraphNode::GetPinTooltip, InPin);
-			}
-			else
-			{
-				PinWidget = SNew(SMetasoundGraphPinInteger, InPin)
-					.ToolTipText(this, &SMetasoundGraphNode::GetPinTooltip, InPin);
-			}
+			TSharedRef<SWidget> AddPinButton = AddPinButtonContent(
+				LOCTEXT("MetasoundGraphNode_AddPinOutputButton", "Add Output"),
+				LOCTEXT("MetasoundGraphNode_AddPinOutputButton_Tooltip", "Add an output to the parent Metasound node.")
+			);
+
+			FMargin AddPinPadding = Settings->GetOutputPinPadding();
+			AddPinPadding.Top += 6.0f;
+
+			OutputBox->AddSlot()
+			.AutoHeight()
+			.VAlign(VAlign_Center)
+			.Padding(AddPinPadding)
+			[
+				AddPinButton
+			];
 		}
 
-		else if (InPin->PinType.PinCategory == FGraphBuilder::PinCategoryObject)
+		UMetasoundEditorGraphNode& SMetaSoundGraphNode::GetMetaSoundNode()
 		{
-			PinWidget = SNew(SMetasoundGraphPinObject, InPin)
-				.ToolTipText(this, &SMetasoundGraphNode::GetPinTooltip, InPin);
+			return *CastChecked<UMetasoundEditorGraphNode>(GraphNode);
 		}
 
-		else if (InPin->PinType.PinCategory == FGraphBuilder::PinCategoryString)
+		const UMetasoundEditorGraphNode& SMetaSoundGraphNode::GetMetaSoundNode() const
 		{
-			PinWidget = SNew(SMetasoundGraphPinString, InPin)
-				.ToolTipText(this, &SMetasoundGraphNode::GetPinTooltip, InPin);
+			check(GraphNode);
+			return *Cast<UMetasoundEditorGraphNode>(GraphNode);
 		}
 
-		else if (InPin->PinType.PinCategory == FGraphBuilder::PinCategoryTrigger)
+		TSharedPtr<SGraphPin> SMetaSoundGraphNode::CreatePinWidget(UEdGraphPin* InPin) const
 		{
-			PinWidget = SNew(SMetasoundGraphPin, InPin)
-				.ToolTipText(this, &SMetasoundGraphNode::GetPinTooltip, InPin);
+			using namespace Frontend;
 
-			if (const ISlateStyle* MetasoundStyle = FSlateStyleRegistry::FindSlateStyle("MetaSoundStyle"))
+			TSharedPtr<SGraphPin> PinWidget;
+
+			if (const UMetasoundEditorGraphSchema* GraphSchema = Cast<const UMetasoundEditorGraphSchema>(InPin->GetSchema()))
 			{
-				const FSlateBrush* PinConnectedBrush = MetasoundStyle->GetBrush(TEXT("MetasoundEditor.Graph.TriggerPin.Connected"));
-				const FSlateBrush* PinDisconnectedBrush = MetasoundStyle->GetBrush(TEXT("MetasoundEditor.Graph.TriggerPin.Disconnected"));
-				PinWidget->SetCustomPinIcon(PinConnectedBrush, PinDisconnectedBrush);
-			}
-		}
-	}
+				// Don't show default value field for container types
+				if (InPin->PinType.ContainerType != EPinContainerType::None)
+				{
+					PinWidget = SNew(SMetasoundGraphPin, InPin)
+						.ToolTipText(this, &SMetaSoundGraphNode::GetPinTooltip, InPin);
+				}
 
-	if (!PinWidget.IsValid())
-	{
-		PinWidget = SNew(SMetasoundGraphPin, InPin)
-			.ToolTipText(this, &SMetasoundGraphNode::GetPinTooltip, InPin);
-	}
+				else if (InPin->PinType.PinCategory == FGraphBuilder::PinCategoryAudio)
+				{
+					PinWidget = SNew(SMetasoundGraphPin, InPin)
+						.ToolTipText(this, &SMetaSoundGraphNode::GetPinTooltip, InPin);
+				}
 
-	return PinWidget;
-}
+				else if (InPin->PinType.PinCategory == FGraphBuilder::PinCategoryBoolean)
+				{
+					PinWidget = SNew(SMetasoundGraphPinBool, InPin)
+						.ToolTipText(this, &SMetaSoundGraphNode::GetPinTooltip, InPin);
+				}
 
-void SMetasoundGraphNode::CreateStandardPinWidget(UEdGraphPin* InPin)
-{
-	const bool bShowPin = ShouldPinBeHidden(InPin);
-	if (bShowPin)
-	{
-		TSharedPtr<SGraphPin> NewPin = CreatePinWidget(InPin);
-		check(NewPin.IsValid());
+				else if (InPin->PinType.PinCategory == FGraphBuilder::PinCategoryFloat)
+				{
+					PinWidget = SNew(SMetasoundGraphPinFloat, InPin)
+						.ToolTipText(this, &SMetaSoundGraphNode::GetPinTooltip, InPin);
+				}
 
-		Metasound::Frontend::FNodeHandle NodeHandle = GetMetasoundNode().GetNodeHandle();
-		if (InPin->Direction == EGPD_Input)
-		{
-			if (!NodeHandle->GetClassStyle().Display.bShowInputNames)
+				else if (InPin->PinType.PinCategory == FGraphBuilder::PinCategoryInt32)
+				{
+					if (SMetasoundGraphEnumPin::FindEnumInterfaceFromPin(InPin))
+					{
+						PinWidget = SNew(SMetasoundGraphEnumPin, InPin)
+							.ToolTipText(this, &SMetaSoundGraphNode::GetPinTooltip, InPin);
+					}
+					else
+					{
+						PinWidget = SNew(SMetasoundGraphPinInteger, InPin)
+							.ToolTipText(this, &SMetaSoundGraphNode::GetPinTooltip, InPin);
+					}
+				}
+
+				else if (InPin->PinType.PinCategory == FGraphBuilder::PinCategoryObject)
+				{
+					PinWidget = SNew(SMetasoundGraphPinObject, InPin)
+						.ToolTipText(this, &SMetaSoundGraphNode::GetPinTooltip, InPin);
+				}
+
+				else if (InPin->PinType.PinCategory == FGraphBuilder::PinCategoryString)
+				{
+					PinWidget = SNew(SMetasoundGraphPinString, InPin)
+						.ToolTipText(this, &SMetaSoundGraphNode::GetPinTooltip, InPin);
+				}
+
+				else if (InPin->PinType.PinCategory == FGraphBuilder::PinCategoryTrigger)
+				{
+					PinWidget = SNew(SMetasoundGraphPin, InPin)
+						.ToolTipText(this, &SMetaSoundGraphNode::GetPinTooltip, InPin);
+
+					if (const ISlateStyle* MetasoundStyle = FSlateStyleRegistry::FindSlateStyle("MetaSoundStyle"))
+					{
+						const FSlateBrush* PinConnectedBrush = MetasoundStyle->GetBrush(TEXT("MetasoundEditor.Graph.TriggerPin.Connected"));
+						const FSlateBrush* PinDisconnectedBrush = MetasoundStyle->GetBrush(TEXT("MetasoundEditor.Graph.TriggerPin.Disconnected"));
+						PinWidget->SetCustomPinIcon(PinConnectedBrush, PinDisconnectedBrush);
+					}
+				}
+			}
+
+			if (!PinWidget.IsValid())
 			{
-				NewPin->SetShowLabel(false);
+				PinWidget = SNew(SMetasoundGraphPin, InPin)
+					.ToolTipText(this, &SMetaSoundGraphNode::GetPinTooltip, InPin);
 			}
+
+			return PinWidget;
 		}
-		else if (InPin->Direction == EGPD_Output)
+
+		void SMetaSoundGraphNode::CreateStandardPinWidget(UEdGraphPin* InPin)
 		{
-			if (!NodeHandle->GetClassStyle().Display.bShowOutputNames)
+			const bool bShowPin = ShouldPinBeHidden(InPin);
+			if (bShowPin)
 			{
-				NewPin->SetShowLabel(false);
+				TSharedPtr<SGraphPin> NewPin = CreatePinWidget(InPin);
+				check(NewPin.IsValid());
+
+				Metasound::Frontend::FNodeHandle NodeHandle = GetMetaSoundNode().GetNodeHandle();
+				if (InPin->Direction == EGPD_Input)
+				{
+					if (!NodeHandle->GetClassStyle().Display.bShowInputNames)
+					{
+						NewPin->SetShowLabel(false);
+					}
+				}
+				else if (InPin->Direction == EGPD_Output)
+				{
+					if (!NodeHandle->GetClassStyle().Display.bShowOutputNames)
+					{
+						NewPin->SetShowLabel(false);
+					}
+				}
+
+				AddPin(NewPin.ToSharedRef());
 			}
 		}
 
-		AddPin(NewPin.ToSharedRef());
-	}
-}
+		FText SMetaSoundGraphNode::GetPinTooltip(UEdGraphPin* InPin) const
+		{
+			using namespace Frontend;
 
-FText SMetasoundGraphNode::GetPinTooltip(UEdGraphPin* InPin) const
-{
-	using namespace Metasound::Editor;
-	using namespace Metasound::Frontend;
+			if (InPin->Direction == EGPD_Input)
+			{
+				FConstInputHandle InputHandle = FGraphBuilder::GetConstInputHandleFromPin(InPin);
+				return InputHandle->GetTooltip();
+			}
+			else
+			{
+				FConstOutputHandle OutputHandle = FGraphBuilder::GetConstOutputHandleFromPin(InPin);
+				return OutputHandle->GetTooltip();
+			}
+		}
 
-	if (InPin->Direction == EGPD_Input)
-	{
-		FConstInputHandle InputHandle = FGraphBuilder::GetConstInputHandleFromPin(InPin);
-		return InputHandle->GetTooltip();
-	}
-	else
-	{
-		FConstOutputHandle OutputHandle = FGraphBuilder::GetConstOutputHandleFromPin(InPin);
-		return OutputHandle->GetTooltip();
-	}
-}
+		TSharedRef<SWidget> SMetaSoundGraphNode::CreateTitleWidget(TSharedPtr<SNodeTitle> NodeTitle)
+		{
+			Metasound::Frontend::FNodeHandle NodeHandle = GetMetaSoundNode().GetNodeHandle();
+			if (!NodeHandle->GetClassStyle().Display.bShowName)
+			{
+				return SNullWidget::NullWidget;
+			}
 
-TSharedRef<SWidget> SMetasoundGraphNode::CreateTitleWidget(TSharedPtr<SNodeTitle> NodeTitle)
-{
-	Metasound::Frontend::FNodeHandle NodeHandle = GetMetasoundNode().GetNodeHandle();
-	if (!NodeHandle->GetClassStyle().Display.bShowName)
-	{
-		return SNullWidget::NullWidget;
-	}
+			TSharedPtr<SHorizontalBox> TitleBoxWidget = SNew(SHorizontalBox);
 
-	TSharedPtr<SHorizontalBox> TitleBoxWidget = SNew(SHorizontalBox);
+			FSlateIcon NodeIcon = GetMetaSoundNode().GetNodeTitleIcon();
+			if (const FSlateBrush* IconBrush = NodeIcon.GetIcon())
+			{
+				if (IconBrush != FStyleDefaults::GetNoBrush())
+				{
+					TSharedPtr<SImage> Image;
+					TitleBoxWidget->AddSlot()
+					.AutoWidth()
+					.VAlign(VAlign_Center)
+					.HAlign(HAlign_Right)
+					[
+						SNew(SHorizontalBox)
+						+ SHorizontalBox::Slot()
+						.Padding(0.0f, 0.0f, 4.0f, 0.0f)
+						.VAlign(VAlign_Center)
+						.HAlign(HAlign_Right)
+						[
+							SAssignNew(Image, SImage)
+						]
+					];
+					Image->SetColorAndOpacity(TAttribute<FSlateColor>::CreateLambda([this]() { return FSlateColor(GetNodeTitleColorOverride()); }));
+					Image->SetImage(IconBrush);
+				}
+			}
 
-	FSlateIcon NodeIcon = GetMetasoundNode().GetNodeTitleIcon();
-	if (const FSlateBrush* IconBrush = NodeIcon.GetIcon())
-	{
-		if (IconBrush != FStyleDefaults::GetNoBrush())
-		{
-			TSharedPtr<SImage> Image;
 			TitleBoxWidget->AddSlot()
 			.AutoWidth()
 			.VAlign(VAlign_Center)
-			.HAlign(HAlign_Right)
 			[
-				SNew(SHorizontalBox)
-				+ SHorizontalBox::Slot()
-				.Padding(0.0f, 0.0f, 4.0f, 0.0f)
-				.VAlign(VAlign_Center)
-				.HAlign(HAlign_Right)
-				[
-					SAssignNew(Image, SImage)
-				]
+				SGraphNode::CreateTitleWidget(NodeTitle)
 			];
-			Image->SetColorAndOpacity(TAttribute<FSlateColor>::CreateLambda([this]() { return FSlateColor(GetNodeTitleColorOverride()); }));
-			Image->SetImage(IconBrush);
-		}
-	}
 
-	TitleBoxWidget->AddSlot()
-	.AutoWidth()
-	[
-		SGraphNode::CreateTitleWidget(NodeTitle)
-	];
+			InlineEditableText->SetColorAndOpacity(TAttribute<FLinearColor>::Create(TAttribute<FLinearColor>::FGetter::CreateSP(this, &SMetaSoundGraphNode::GetNodeTitleColorOverride)));
 
-	InlineEditableText->SetColorAndOpacity(TAttribute<FLinearColor>::Create(TAttribute<FLinearColor>::FGetter::CreateSP(this, &SMetasoundGraphNode::GetNodeTitleColorOverride)));
+			return TitleBoxWidget.ToSharedRef();
+		}
 
-	return TitleBoxWidget.ToSharedRef();
-}
+		FLinearColor SMetaSoundGraphNode::GetNodeTitleColorOverride() const
+		{
+			FLinearColor ReturnTitleColor = GraphNode->IsDeprecated() ? FLinearColor::Red : GetNodeObj()->GetNodeTitleColor();
 
-FLinearColor SMetasoundGraphNode::GetNodeTitleColorOverride() const
-{
-	FLinearColor ReturnTitleColor = GraphNode->IsDeprecated() ? FLinearColor::Red : GetNodeObj()->GetNodeTitleColor();
+			if (!GraphNode->IsNodeEnabled() || GraphNode->IsDisplayAsDisabledForced() || GraphNode->IsNodeUnrelated())
+			{
+				ReturnTitleColor *= FLinearColor(0.5f, 0.5f, 0.5f, 0.4f);
+			}
+			else
+			{
+				ReturnTitleColor.A = FadeCurve.GetLerp();
+			}
 
-	if (!GraphNode->IsNodeEnabled() || GraphNode->IsDisplayAsDisabledForced() || GraphNode->IsNodeUnrelated())
-	{
-		ReturnTitleColor *= FLinearColor(0.5f, 0.5f, 0.5f, 0.4f);
-	}
-	else
-	{
-		ReturnTitleColor.A = FadeCurve.GetLerp();
-	}
+			return ReturnTitleColor;
+		}
 
-	return ReturnTitleColor;
-}
+		void SMetaSoundGraphNode::SetDefaultTitleAreaWidget(TSharedRef<SOverlay> DefaultTitleAreaWidget)
+		{
+			SGraphNode::SetDefaultTitleAreaWidget(DefaultTitleAreaWidget);
 
-void SMetasoundGraphNode::SetDefaultTitleAreaWidget(TSharedRef<SOverlay> DefaultTitleAreaWidget)
-{
-	SGraphNode::SetDefaultTitleAreaWidget(DefaultTitleAreaWidget);
+			Metasound::Frontend::FNodeHandle NodeHandle = GetMetaSoundNode().GetNodeHandle();
+			if (NodeHandle->GetClassStyle().Display.bShowName)
+			{
+				DefaultTitleAreaWidget->ClearChildren();
+				TSharedPtr<SNodeTitle> NodeTitle = SNew(SNodeTitle, GraphNode);
 
-	Metasound::Frontend::FNodeHandle NodeHandle = GetMetasoundNode().GetNodeHandle();
-	if (NodeHandle->GetClassStyle().Display.bShowName)
-	{
-		DefaultTitleAreaWidget->ClearChildren();
-		TSharedPtr<SNodeTitle> NodeTitle = SNew(SNodeTitle, GraphNode);
-
-		DefaultTitleAreaWidget->AddSlot()
-		.HAlign(HAlign_Fill)
-		.VAlign(VAlign_Center)
-		[
-			SNew(SHorizontalBox)
-			+ SHorizontalBox::Slot()
-			.HAlign(HAlign_Fill)
-			[
-				SNew(SBorder)
-				.BorderImage(FEditorStyle::GetBrush("NoBorder"))
+				DefaultTitleAreaWidget->AddSlot()
+				.HAlign(HAlign_Fill)
+				.VAlign(VAlign_Center)
 				[
 					SNew(SHorizontalBox)
 					+ SHorizontalBox::Slot()
+					.HAlign(HAlign_Fill)
 					[
-						SNew(SVerticalBox)
-						+ SVerticalBox::Slot()
-						.AutoHeight()
-						.HAlign(HAlign_Center)
+						SNew(SBorder)
+						.BorderImage(FEditorStyle::GetBrush("NoBorder"))
 						[
-							CreateTitleWidget(NodeTitle)
-						]
-						+ SVerticalBox::Slot()
-						.AutoHeight()
-						[
-							NodeTitle.ToSharedRef()
+							SNew(SHorizontalBox)
+							+ SHorizontalBox::Slot()
+							[
+								SNew(SVerticalBox)
+								+ SVerticalBox::Slot()
+								.AutoHeight()
+								.HAlign(HAlign_Center)
+								[
+									CreateTitleWidget(NodeTitle)
+								]
+								+ SVerticalBox::Slot()
+								.AutoHeight()
+								[
+									NodeTitle.ToSharedRef()
+								]
+							]
 						]
 					]
-				]
-			]
-			+ SHorizontalBox::Slot()
-			.HAlign(HAlign_Right)
-			.VAlign(VAlign_Center)
-			.Padding(0, 0, 5, 0)
-			.AutoWidth()
-			[
-				CreateTitleRightWidget()
-			]
-		];
-
-		DefaultTitleAreaWidget->AddSlot()
-		.VAlign(VAlign_Top)
-		[
-			SNew(SBorder)
-			.Visibility(EVisibility::HitTestInvisible)
-			.BorderImage( FEditorStyle::GetBrush( "Graph.Node.TitleHighlight" ) )
-			.BorderBackgroundColor( this, &SGraphNode::GetNodeTitleIconColor )
-			[
-				SNew(SSpacer)
-				.Size(FVector2D(20,20))
-			]
-		];
+					+ SHorizontalBox::Slot()
+					.HAlign(HAlign_Right)
+					.VAlign(VAlign_Center)
+					.Padding(0, 0, 5, 0)
+					.AutoWidth()
+					[
+						CreateTitleRightWidget()
+					]
+				];
 
-	}
-	else
-	{
-		DefaultTitleAreaWidget->SetVisibility(EVisibility::Collapsed);
-	}
-}
+				DefaultTitleAreaWidget->AddSlot()
+				.VAlign(VAlign_Top)
+				[
+					SNew(SBorder)
+					.Visibility(EVisibility::HitTestInvisible)
+					.BorderImage( FEditorStyle::GetBrush( "Graph.Node.TitleHighlight" ) )
+					.BorderBackgroundColor( this, &SGraphNode::GetNodeTitleIconColor )
+					[
+						SNew(SSpacer)
+						.Size(FVector2D(20,20))
+					]
+				];
 
-void SMetasoundGraphNode::MoveTo(const FVector2D& NewPosition, FNodeSet& NodeFilter, bool bMarkDirty)
-{
-	SGraphNode::MoveTo(NewPosition, NodeFilter, bMarkDirty);
+			}
+			else
+			{
+				DefaultTitleAreaWidget->SetVisibility(EVisibility::Collapsed);
+			}
+		}
 
-	GetMetasoundNode().UpdatePosition();
-}
+		void SMetaSoundGraphNode::MoveTo(const FVector2D& NewPosition, FNodeSet& NodeFilter, bool bMarkDirty)
+		{
+			SGraphNode::MoveTo(NewPosition, NodeFilter, bMarkDirty);
 
-const FSlateBrush* SMetasoundGraphNode::GetNodeBodyBrush() const
-{
-// 	if (const ISlateStyle* MetasoundStyle = FSlateStyleRegistry::FindSlateStyle("MetaSoundStyle"))
-// 	{
-// 		if (GetMetasoundNode().GetNodeHandle()->GetClassType() == EMetasoundFrontendClassType::Input)
-// 		{
-// 			if (const FSlateBrush* InputBrush = MetasoundStyle->GetBrush("MetasoundEditor.Graph.Node.Body.Input"))
-// 			{
-// 				return InputBrush;
-// 			}
-// 		}
-// 
-// 		if (const FSlateBrush* DefaultBrush = MetasoundStyle->GetBrush("MetasoundEditor.Graph.Node.Body.Default"))
-// 		{
-// 			return DefaultBrush;
-// 		}
-// 	}
-
-	return FEditorStyle::GetBrush("Graph.Node.Body");
-}
-
-EVisibility SMetasoundGraphNode::IsAddPinButtonVisible() const
-{
-	EVisibility DefaultVisibility = SGraphNode::IsAddPinButtonVisible();
-	if (DefaultVisibility == EVisibility::Visible)
-	{
-		if (!GetMetasoundNode().CanAddInputPin())
+			GetMetaSoundNode().UpdatePosition();
+		}
+
+		const FSlateBrush* SMetaSoundGraphNode::GetNodeBodyBrush() const
 		{
-			return EVisibility::Collapsed;
+		// 	if (const ISlateStyle* MetasoundStyle = FSlateStyleRegistry::FindSlateStyle("MetaSoundStyle"))
+		// 	{
+		// 		if (GetMetaSoundNode().GetNodeHandle()->GetClassType() == EMetasoundFrontendClassType::Input)
+		// 		{
+		// 			if (const FSlateBrush* InputBrush = MetasoundStyle->GetBrush("MetasoundEditor.Graph.Node.Body.Input"))
+		// 			{
+		// 				return InputBrush;
+		// 			}
+		// 		}
+		// 
+		// 		if (const FSlateBrush* DefaultBrush = MetasoundStyle->GetBrush("MetasoundEditor.Graph.Node.Body.Default"))
+		// 		{
+		// 			return DefaultBrush;
+		// 		}
+		// 	}
+
+			return FEditorStyle::GetBrush("Graph.Node.Body");
 		}
-	}
 
-	return DefaultVisibility;
-}
+		EVisibility SMetaSoundGraphNode::IsAddPinButtonVisible() const
+		{
+			EVisibility DefaultVisibility = SGraphNode::IsAddPinButtonVisible();
+			if (DefaultVisibility == EVisibility::Visible)
+			{
+				if (!GetMetaSoundNode().CanAddInputPin())
+				{
+					return EVisibility::Collapsed;
+				}
+			}
 
-FReply SMetasoundGraphNode::OnAddPin()
-{
-	GetMetasoundNode().CreateInputPin();
+			return DefaultVisibility;
+		}
 
-	return FReply::Handled();
-}
+		FReply SMetaSoundGraphNode::OnAddPin()
+		{
+			GetMetaSoundNode().CreateInputPin();
 
-FName SMetasoundGraphNode::GetLiteralDataType() const
-{
-	using namespace Metasound::Frontend;
+			return FReply::Handled();
+		}
 
-	FName TypeName;
+		FName SMetaSoundGraphNode::GetLiteralDataType() const
+		{
+			using namespace Frontend;
 
-	// Just take last type.  If more than one, all types are the same.
-	const UMetasoundEditorGraphNode& Node = GetMetasoundNode();
-	Node.GetNodeHandle()->IterateConstOutputs([InTypeName = &TypeName](FConstOutputHandle OutputHandle)
-	{
-		*InTypeName = OutputHandle->GetDataType();
-	});
+			FName TypeName;
 
-	return TypeName;
-}
+			// Just take last type.  If more than one, all types are the same.
+			const UMetasoundEditorGraphNode& Node = GetMetaSoundNode();
+			Node.GetNodeHandle()->IterateConstOutputs([InTypeName = &TypeName](FConstOutputHandle OutputHandle)
+			{
+				*InTypeName = OutputHandle->GetDataType();
+			});
 
-TSharedRef<SWidget> SMetasoundGraphNode::CreateTitleRightWidget()
-{
-	const FName TypeName = GetLiteralDataType();
-	if (TypeName == Metasound::GetMetasoundDataTypeName<Metasound::FTrigger>())
-	{
-		if (UMetasoundEditorGraphInputNode* Node = Cast<UMetasoundEditorGraphInputNode>(&GetMetasoundNode()))
+			return TypeName;
+		}
+
+		TSharedRef<SWidget> SMetaSoundGraphNode::CreateTitleRightWidget()
 		{
-			if (UMetasoundEditorGraphInput* Input = Node->Input)
+			using namespace Frontend;
+
+			const FName TypeName = GetLiteralDataType();
+			if (TypeName == Metasound::GetMetasoundDataTypeName<Metasound::FTrigger>())
 			{
-				if (Input->Literal)
+				if (UMetasoundEditorGraphMemberNode* Node = Cast<UMetasoundEditorGraphMemberNode>(&GetMetaSoundNode()))
 				{
-					return CreateTriggerSimulationWidget(*Input->Literal);
+					if (UMetasoundEditorGraphVertex* Vertex = Cast<UMetasoundEditorGraphVertex>(Node->GetMember()))
+					{
+						if (UMetasoundEditorGraphMemberDefaultLiteral* Literal = Vertex->GetLiteral())
+						{
+							TAttribute<EVisibility> SimVisibility = GetSimulationVisibilityAttribute();
+							TAttribute<bool> SimEnablement = true;
+							return CreateTriggerSimulationWidget(*Literal, MoveTemp(SimVisibility), MoveTemp(SimEnablement));
+						}
+					}
 				}
 			}
+
+			return SGraphNode::CreateTitleRightWidget();
 		}
-	}
 
-	return SGraphNode::CreateTitleRightWidget();
-}
+		UMetasoundEditorGraphMember* SMetaSoundGraphNode::GetMetaSoundMember()
+		{
+			if (UMetasoundEditorGraphMemberNode* MemberNode = GetMetaSoundMemberNode())
+			{
+				return MemberNode->GetMember();
+			}
 
-TSharedRef<SWidget> SMetasoundGraphNode::CreateNodeContentArea()
-{
-	using namespace Metasound::Editor;
-	using namespace Metasound::Frontend;
-
-	FNodeHandle NodeHandle = GetMetasoundNode().GetNodeHandle();
-	const FMetasoundFrontendClassStyleDisplay& StyleDisplay = NodeHandle->GetClassStyle().Display;
-	TSharedRef<SHorizontalBox> ContentBox = SNew(SHorizontalBox);
-	TSharedPtr<SWidget> OuterContentBox; // currently only used for input float nodes to accommodate the input widget
-
-	// If float input node, add slider input widget 
-	bool IsFloatInputNode = false;
-	UMetasoundEditorGraphNode& Node = GetMetasoundNode();
-	if (UMetasoundEditorGraphInputNode* InputNode = Cast<UMetasoundEditorGraphInputNode>(&Node))
-	{
-		if (UMetasoundEditorGraphInput* GraphInput = InputNode->Input)
+			return nullptr;
+		}
+
+		UMetasoundEditorGraphMemberNode* SMetaSoundGraphNode::GetMetaSoundMemberNode()
+		{
+			return Cast<UMetasoundEditorGraphMemberNode>(&GetMetaSoundNode());
+		}
+
+		TSharedRef<SWidget> SMetaSoundGraphNode::CreateNodeContentArea()
 		{
-			if (UMetasoundEditorGraphInputFloat* InputFloat = Cast<UMetasoundEditorGraphInputFloat>(GraphInput->Literal))
+			using namespace Frontend;
+
+			FNodeHandle NodeHandle = GetMetaSoundNode().GetNodeHandle();
+			const FMetasoundFrontendClassStyleDisplay& StyleDisplay = NodeHandle->GetClassStyle().Display;
+			TSharedPtr<SHorizontalBox> ContentBox = SNew(SHorizontalBox);
+			TSharedPtr<SWidget> OuterContentBox; // currently only used for input float nodes to accommodate the input widget
+
+			// If float input node, check if custom widget required
+			bool IsFloatMemberNode = false;
+			if (UMetasoundEditorGraphMember* GraphMember = GetMetaSoundMember())
 			{
-				if (InputFloat->InputWidgetType != EMetasoundInputWidget::None)
+				if (UMetasoundEditorGraphMemberDefaultFloat* DefaultFloat = Cast<UMetasoundEditorGraphMemberDefaultFloat>(GraphMember->GetLiteral()))
 				{
-					const float InputWidgetPadding = 3.0f;
-					const FVector2D SliderDesiredSizeVertical = FVector2D(30.0f, 250.0f);
-					const FVector2D RadialSliderDesiredSize = FVector2D(56.0f, 87.0f);
-					IsFloatInputNode = true;
-
-					auto OnValueChangedLambda = [InputFloat, GraphInput, this](float Value)
+					if (DefaultFloat->WidgetType != EMetasoundMemberDefaultWidget::None)
 					{
-						if (InputWidget.IsValid())
+						constexpr float WidgetPadding = 3.0f;
+						static const FVector2D SliderDesiredSizeVertical = FVector2D(30.0f, 250.0f);
+						static const FVector2D RadialSliderDesiredSize = FVector2D(56.0f, 87.0f);
+
+						IsFloatMemberNode = true;
+
+						auto OnValueChangedLambda = [DefaultFloat, GraphMember, this](float Value)
 						{
-							float Output = InputWidget->GetOutputValue(Value);
-							InputFloat->SetDefault(Output);
-							if (UMetasoundEditorGraphInputLiteral* InputLiteral = GraphInput->Literal)
+							if (InputWidget.IsValid())
 							{
-								InputLiteral->UpdateDocumentInputLiteral();
+								constexpr bool bPostTransaction = false;
+								float Output = InputWidget->GetOutputValue(Value);
+								DefaultFloat->SetDefault(Output);
+								GraphMember->UpdateFrontendDefaultLiteral(bPostTransaction);
 							}
-						}
-					};
-					if (InputFloat->InputWidgetType == EMetasoundInputWidget::Slider)
-					{
-						// Create slider 
-						if (InputFloat->InputWidgetValueType == EMetasoundInputWidgetValueType::Frequency)
-						{
-							SAssignNew(InputWidget, SAudioFrequencySlider)
-								.OnValueChanged_Lambda(OnValueChangedLambda);
-						}
-						else if (InputFloat->InputWidgetValueType == EMetasoundInputWidgetValueType::Volume)
+						};
+
+						auto OnValueCommittedLambda = [DefaultFloat, GraphMember, this](float Value)
 						{
-							SAssignNew(InputWidget, SAudioVolumeSlider)
-								.OnValueChanged_Lambda(OnValueChangedLambda);
-						}
-						else
+							if (InputWidget.IsValid())
+							{
+								constexpr bool bPostTransaction = true;
+								float Output = InputWidget->GetOutputValue(Value);
+								DefaultFloat->SetDefault(Output);
+								GraphMember->UpdateFrontendDefaultLiteral(bPostTransaction);
+
+								GraphMember->MarkNodesForRefresh();
+								if (UMetasoundEditorGraph* Graph = GraphMember->GetOwningGraph())
+								{
+									Graph->SetSynchronizationRequired();
+								}
+							}
+						};
+
+						if (DefaultFloat->WidgetType == EMetasoundMemberDefaultWidget::Slider)
 						{
-							SAssignNew(InputWidget, SAudioSlider)
-								.OnValueChanged_Lambda(OnValueChangedLambda);
-							InputWidget->SetShowUnitsText(false);
+							// Create slider 
+							if (DefaultFloat->WidgetValueType == EMetasoundMemberDefaultWidgetValueType::Frequency)
+							{
+								SAssignNew(InputWidget, SAudioFrequencySlider)
+									.OnValueChanged_Lambda(OnValueChangedLambda)
+									.OnValueCommitted_Lambda(OnValueCommittedLambda);
+							}
+							else if (DefaultFloat->WidgetValueType == EMetasoundMemberDefaultWidgetValueType::Volume)
+							{
+								SAssignNew(InputWidget, SAudioVolumeSlider)
+									.OnValueChanged_Lambda(OnValueChangedLambda)
+									.OnValueCommitted_Lambda(OnValueCommittedLambda);
+							}
+							else
+							{
+								SAssignNew(InputWidget, SAudioSlider)
+									.OnValueChanged_Lambda(OnValueChangedLambda)
+									.OnValueCommitted_Lambda(OnValueCommittedLambda);
+								InputWidget->SetShowUnitsText(false);
+							}
+							// Slider layout 
+							if (DefaultFloat->WidgetOrientation == Orient_Vertical)
+							{
+								SAssignNew(OuterContentBox, SVerticalBox)
+									+ SVerticalBox::Slot()
+									.HAlign(HAlign_Right)
+									.VAlign(VAlign_Center)
+									.AutoHeight()
+									[
+										ContentBox.ToSharedRef()
+									]
+								+ SVerticalBox::Slot()
+									.HAlign(HAlign_Fill)
+									.VAlign(VAlign_Top)
+									.Padding(WidgetPadding, 0.0f, WidgetPadding, WidgetPadding)
+									.AutoHeight()
+									[
+										InputWidget.ToSharedRef()
+									];
+								InputWidget->SetDesiredSizeOverride(SliderDesiredSizeVertical);
+							}
+							else // horizontal orientation
+							{
+								UMetasoundEditorGraphMemberNode* MemberNode = GetMetaSoundMemberNode();
+								TSharedPtr<SWidget> Slot1;
+								TSharedPtr<SWidget> Slot2;
+								if (MemberNode->IsA<UMetasoundEditorGraphInputNode>())
+								{
+									Slot1 = InputWidget;
+									Slot2 = ContentBox;
+								}
+								else
+								{
+									Slot1 = ContentBox;
+									Slot2 = InputWidget;
+								}
+
+								SAssignNew(OuterContentBox, SHorizontalBox)
+									+ SHorizontalBox::Slot()
+									.HAlign(HAlign_Fill)
+									.VAlign(VAlign_Center)
+									.AutoWidth()
+									[
+										Slot1.ToSharedRef()
+									]
+								+ SHorizontalBox::Slot()
+									.HAlign(HAlign_Center)
+									.VAlign(VAlign_Fill)
+									.Padding(WidgetPadding, 0.0f, 0.0f, WidgetPadding)
+									.AutoWidth()
+									[
+										Slot2.ToSharedRef()
+									];
+								InputWidget->SetDesiredSizeOverride(FVector2D(SliderDesiredSizeVertical.Y, SliderDesiredSizeVertical.X));
+							}
+							// safe downcast because the ptr was just assigned above 
+							StaticCastSharedPtr<SAudioSliderBase>(InputWidget)->SetOrientation(DefaultFloat->WidgetOrientation);
 						}
-						// Slider layout 
-						if (InputFloat->InputWidgetOrientation == Orient_Vertical)
+						else if (DefaultFloat->WidgetType == EMetasoundMemberDefaultWidget::RadialSlider)
 						{
+							// Create slider 
+							if (DefaultFloat->WidgetValueType == EMetasoundMemberDefaultWidgetValueType::Frequency)
+							{
+								SAssignNew(InputWidget, SAudioFrequencyRadialSlider)
+									.OnValueChanged_Lambda(OnValueChangedLambda);
+							}
+							else if (DefaultFloat->WidgetValueType == EMetasoundMemberDefaultWidgetValueType::Volume)
+							{
+								SAssignNew(InputWidget, SAudioVolumeRadialSlider)
+									.OnValueChanged_Lambda(OnValueChangedLambda);
+							}
+							else
+							{
+								SAssignNew(InputWidget, SAudioRadialSlider)
+									.OnValueChanged_Lambda(OnValueChangedLambda);
+								InputWidget->SetShowUnitsText(false);
+							}
+							// Only vertical layout for radial slider
 							SAssignNew(OuterContentBox, SVerticalBox)
+								+ SVerticalBox::Slot()
+								.HAlign(HAlign_Right)
+								.VAlign(VAlign_Center)
+								.AutoHeight()
+								[
+									ContentBox.ToSharedRef()
+								]
 							+ SVerticalBox::Slot()
-							.HAlign(HAlign_Right)
-							.VAlign(VAlign_Center)
-							.AutoHeight()
-							[
-								ContentBox
-							]
-							+ SVerticalBox::Slot()
-							.HAlign(HAlign_Fill)
-							.VAlign(VAlign_Top)
-							.Padding(InputWidgetPadding, 0.0f, InputWidgetPadding, InputWidgetPadding)
-							.AutoHeight()
-							[
-								InputWidget.ToSharedRef()
-							];
-							InputWidget->SetDesiredSizeOverride(SliderDesiredSizeVertical);
+								.HAlign(HAlign_Fill)
+								.VAlign(VAlign_Top)
+								.Padding(WidgetPadding, 0.0f, WidgetPadding, WidgetPadding)
+								.AutoHeight()
+								[
+									InputWidget.ToSharedRef()
+								];
+							InputWidget->SetDesiredSizeOverride(RadialSliderDesiredSize);
 						}
-						else // horizontal orientation
-						{
-							SAssignNew(OuterContentBox, SHorizontalBox)
-							+ SHorizontalBox::Slot()
-							.HAlign(HAlign_Fill)
-							.VAlign(VAlign_Center)
-							.AutoWidth()
-							[
-								InputWidget.ToSharedRef()
-							]
-							+SHorizontalBox::Slot()
-							.HAlign(HAlign_Center)
-							.VAlign(VAlign_Fill)
-							.Padding(InputWidgetPadding, 0.0f, 0.0f, InputWidgetPadding)
-							.AutoWidth()
-							[
-								ContentBox
-							];
-							InputWidget->SetDesiredSizeOverride(FVector2D(SliderDesiredSizeVertical.Y, SliderDesiredSizeVertical.X));
-						}
-						// safe downcast because the ptr was just assigned above 
-						StaticCastSharedPtr<SAudioSliderBase>(InputWidget)->SetOrientation(InputFloat->InputWidgetOrientation);
-					}
-					else if (InputFloat->InputWidgetType == EMetasoundInputWidget::RadialSlider)
-					{
-						// Create slider 
-						if (InputFloat->InputWidgetValueType == EMetasoundInputWidgetValueType::Frequency)
+
+						InputWidget->SetOutputRange(DefaultFloat->GetRange());
+						InputWidget->SetUnitsTextReadOnly(true);
+						InputWidget->SetValue(InputWidget->GetLinValue(DefaultFloat->GetDefault()));
+						InputWidget->SetVisibility(TAttribute<EVisibility>::Create([this]()
 						{
-							SAssignNew(InputWidget, SAudioFrequencyRadialSlider)
-								.OnValueChanged_Lambda(OnValueChangedLambda);
-						}
-						else if (InputFloat->InputWidgetValueType == EMetasoundInputWidgetValueType::Volume)
+							if (UMetasoundEditorGraphMemberNode* Node = GetMetaSoundMemberNode())
+							{
+								return Node->EnableInteractWidgets() ? EVisibility::Visible : EVisibility::Collapsed;
+							}
+							return EVisibility::Collapsed;
+						}));
+
+						// Setup & clear delegate if necessary (ex. if was just saved)
+						if (InputSliderOnValueChangedDelegateHandle.IsValid())
 						{
-							SAssignNew(InputWidget, SAudioVolumeRadialSlider)
-								.OnValueChanged_Lambda(OnValueChangedLambda);
+							DefaultFloat->OnDefaultValueChanged.Remove(InputSliderOnValueChangedDelegateHandle);
+							InputSliderOnValueChangedDelegateHandle.Reset();
 						}
-						else
+
+						InputSliderOnValueChangedDelegateHandle = DefaultFloat->OnDefaultValueChanged.AddLambda([Widget = InputWidget](float Value)
 						{
-							SAssignNew(InputWidget, SAudioRadialSlider)
-								.OnValueChanged_Lambda(OnValueChangedLambda);
-							InputWidget->SetShowUnitsText(false);
-						}
-						// Only vertical layout for radial slider
-						SAssignNew(OuterContentBox, SVerticalBox)
-						+ SVerticalBox::Slot()
-						.HAlign(HAlign_Right)
-						.VAlign(VAlign_Center)
-						.AutoHeight()
-						[
-							ContentBox
-						]
-						+ SVerticalBox::Slot()
-						.HAlign(HAlign_Fill)
-						.VAlign(VAlign_Top)
-						.Padding(InputWidgetPadding, 0.0f, InputWidgetPadding, InputWidgetPadding)
-						.AutoHeight()
-						[
-							InputWidget.ToSharedRef()
-						];
-						InputWidget->SetDesiredSizeOverride(RadialSliderDesiredSize);
-					}
-					InputWidget->SetOutputRange(InputFloat->GetRange());
-					InputWidget->SetUnitsTextReadOnly(true);
-					InputWidget->SetValue(InputWidget->GetLinValue(InputFloat->GetDefault()));
-					// setup delegates 
-					// clear delegate if necessary (ex. if was just saved) 
-					if (InputSliderOnValueChangedDelegateHandle.IsValid())
-					{
-						InputFloat->OnDefaultValueChanged.Remove(InputSliderOnValueChangedDelegateHandle);
-						InputSliderOnValueChangedDelegateHandle.Reset();
-					}
-					InputSliderOnValueChangedDelegateHandle = InputFloat->OnDefaultValueChanged.AddLambda([InputWidget = this->InputWidget](float Value)
-					{
-						if (InputWidget.IsValid())
+							if (Widget.IsValid())
+							{
+								const float LinValue = Widget->GetLinValue(Value);
+								Widget->SetValue(LinValue);
+							}
+						});
+
+						if (InputSliderOnRangeChangedDelegateHandle.IsValid())
 						{
-							const float LinValue = InputWidget->GetLinValue(Value);
-							InputWidget->SetValue(LinValue);
+							DefaultFloat->OnRangeChanged.Remove(InputSliderOnRangeChangedDelegateHandle);
+							InputSliderOnRangeChangedDelegateHandle.Reset();
 						}
-					});
 
-					if (InputSliderOnRangeChangedDelegateHandle.IsValid())
-					{
-						InputFloat->OnRangeChanged.Remove(InputSliderOnRangeChangedDelegateHandle);
-						InputSliderOnRangeChangedDelegateHandle.Reset();
-					}
-					InputSliderOnRangeChangedDelegateHandle = InputFloat->OnRangeChanged.AddLambda([InputWidget = this->InputWidget](FVector2D Range)
-					{
-						if (InputWidget.IsValid())
+						InputSliderOnRangeChangedDelegateHandle = DefaultFloat->OnRangeChanged.AddLambda([Widget = InputWidget](FVector2D Range)
 						{
-							InputWidget->SetOutputRange(Range);
-						}
-					});
+							if (Widget.IsValid())
+							{
+								Widget->SetOutputRange(Range);
+							}
+						});
+					}
 				}
 			}
-		}
-	}
 	
-	if (StyleDisplay.ImageName.IsNone())
-	{
-		ContentBox->AddSlot()
-			.HAlign(HAlign_Left)
-			.VAlign(VAlign_Top)
-			.AutoWidth()
-			[
-				SAssignNew(LeftNodeBox, SVerticalBox)
-			];
-	}
-	else
-	{
-		ContentBox->AddSlot()
-			.HAlign(HAlign_Left)
-			.VAlign(VAlign_Top)
-			.AutoWidth()
-			[
-				SAssignNew(LeftNodeBox, SVerticalBox)
-			];
+			if (StyleDisplay.ImageName.IsNone())
+			{
+				ContentBox->AddSlot()
+					.HAlign(HAlign_Left)
+					.VAlign(VAlign_Top)
+					.AutoWidth()
+					[
+						SAssignNew(LeftNodeBox, SVerticalBox)
+					];
+			}
+			else
+			{
+				ContentBox->AddSlot()
+					.HAlign(HAlign_Left)
+					.VAlign(VAlign_Top)
+					.AutoWidth()
+					[
+						SAssignNew(LeftNodeBox, SVerticalBox)
+					];
+
+				if (const ISlateStyle* MetasoundStyle = FSlateStyleRegistry::FindSlateStyle("MetaSoundStyle"))
+				{
+					const FSlateBrush* ImageBrush = MetasoundStyle->GetBrush(StyleDisplay.ImageName);
+					ContentBox->AddSlot()
+					.AutoWidth()
+					.HAlign(HAlign_Center)
+					.VAlign(VAlign_Center)
+					[
+						SNew(SImage)
+						.Image(ImageBrush)
+						.ColorAndOpacity(FSlateColor::UseForeground())
+						.DesiredSizeOverride(FVector2D(20, 20))
+					];
+				}
+			}
 
-		if (const ISlateStyle* MetasoundStyle = FSlateStyleRegistry::FindSlateStyle("MetaSoundStyle"))
-		{
-			const FSlateBrush* ImageBrush = MetasoundStyle->GetBrush(StyleDisplay.ImageName);
 			ContentBox->AddSlot()
-			.AutoWidth()
-			.HAlign(HAlign_Center)
-			.VAlign(VAlign_Center)
-			[
-				SNew(SImage)
-				.Image(ImageBrush)
-				.ColorAndOpacity(FSlateColor::UseForeground())
-				.DesiredSizeOverride(FVector2D(20, 20))
-			];
+				.AutoWidth()
+				.HAlign(HAlign_Right)
+				.VAlign(VAlign_Center)
+				[
+					SAssignNew(RightNodeBox, SVerticalBox)
+				];
+
+			return SNew(SBorder)
+				.BorderImage(FEditorStyle::GetBrush("NoBorder"))
+				.HAlign(HAlign_Fill)
+				.VAlign(VAlign_Fill)
+				.Padding(FMargin(0,3))
+				[
+					(IsFloatMemberNode ? OuterContentBox : ContentBox).ToSharedRef()
+				];
 		}
-	}
-
-	ContentBox->AddSlot()
-		.AutoWidth()
-		.HAlign(HAlign_Right)
-		.VAlign(VAlign_Center)
-		[
-			SAssignNew(RightNodeBox, SVerticalBox)
-		];
-
-	return SNew(SBorder)
-		.BorderImage(FEditorStyle::GetBrush("NoBorder"))
-		.HAlign(HAlign_Fill)
-		.VAlign(VAlign_Fill)
-		.Padding(FMargin(0,3))
-		[
-			IsFloatInputNode ? OuterContentBox.ToSharedRef() : ContentBox
-		];
-}
+	} // namespace Editor
+} // namespace Metasound
 #undef LOCTEXT_NAMESPACE // MetasoundGraphNode
diff --git a/Engine/Plugins/Runtime/Metasound/Source/MetasoundEditor/Private/SMetasoundGraphNode.h b/Engine/Plugins/Runtime/Metasound/Source/MetasoundEditor/Private/SMetasoundGraphNode.h
--- a/Engine/Plugins/Runtime/Metasound/Source/MetasoundEditor/Private/SMetasoundGraphNode.h
+++ b/Engine/Plugins/Runtime/Metasound/Source/MetasoundEditor/Private/SMetasoundGraphNode.h
@@ -4,64 +4,81 @@
 
 #include "Input/Reply.h"
 #include "Layout/Visibility.h"
+#include "Misc/Attribute.h"
 #include "SGraphNode.h"
 #include "Widgets/DeclarativeSyntaxSupport.h"
 #include "Widgets/SOverlay.h"
 
 
 // Forward Declarations
+class SAudioInputWidget;
 class SGraphPin;
 class SVerticalBox;
-class UMetasoundEditorGraphInputLiteral;
+class UMetasoundEditorGraphMember;
+class UMetasoundEditorGraphMemberDefaultLiteral;
+class UMetasoundEditorGraphMemberNode;
 class UMetasoundEditorGraphNode;
-class SAudioInputWidget;
 
 
-class SMetasoundGraphNode : public SGraphNode
+namespace Metasound
 {
-public:
-	SLATE_BEGIN_ARGS(SMetasoundGraphNode)
+	namespace Editor
 	{
-	}
-
-	SLATE_END_ARGS()
-	virtual ~SMetasoundGraphNode();
-
-	void Construct(const FArguments& InArgs, class UEdGraphNode* InNode);
-
-protected:
-
-	// SGraphNode Interface
-	virtual void CreateInputSideAddButton(TSharedPtr<SVerticalBox> InputBox) override;
-	virtual void CreateOutputSideAddButton(TSharedPtr<SVerticalBox> OutputBox) override;
-	virtual TSharedPtr<SGraphPin> CreatePinWidget(UEdGraphPin* InPin) const override;
-	virtual void CreateStandardPinWidget(UEdGraphPin* InPin) override;
-	virtual TSharedRef<SWidget> CreateNodeContentArea() override;
-	virtual TSharedRef<SWidget> CreateTitleWidget(TSharedPtr<SNodeTitle> NodeTitle) override;
-	virtual const FSlateBrush* GetNodeBodyBrush() const override;
-	virtual TSharedRef<SWidget> CreateTitleRightWidget() override;
-	virtual EVisibility IsAddPinButtonVisible() const override;
-	virtual FReply OnAddPin() override;
-	virtual void MoveTo(const FVector2D& NewPosition, FNodeSet& NodeFilter, bool bMarkDirty = true) override;
-	virtual void SetDefaultTitleAreaWidget(TSharedRef<SOverlay> DefaultTitleAreaWidget) override;
-
-	FLinearColor GetNodeTitleColorOverride() const;
-
-	FName GetLiteralDataType() const;
-	UMetasoundEditorGraphNode& GetMetasoundNode();
-	const UMetasoundEditorGraphNode& GetMetasoundNode() const;
-
-public:
-	static void ExecuteInputTrigger(UMetasoundEditorGraphInputLiteral& Literal);
-	static TSharedRef<SWidget> CreateTriggerSimulationWidget(UMetasoundEditorGraphInputLiteral& Literal);
-
-private:
-	FText GetPinTooltip(UEdGraphPin* InPin) const;
-
-	// Slider widget for float input 
-	TSharedPtr<SAudioInputWidget> InputWidget;
-	// Handle for on value changed delegate for input slider 
-	FDelegateHandle InputSliderOnValueChangedDelegateHandle;
-	// Handle for on input slider range changed  
-	FDelegateHandle InputSliderOnRangeChangedDelegateHandle;
-};
+		class SMetaSoundGraphNode : public SGraphNode
+		{
+			public:
+			SLATE_BEGIN_ARGS(SMetaSoundGraphNode)
+			{
+			}
+
+			SLATE_END_ARGS()
+			virtual ~SMetaSoundGraphNode();
+
+			void Construct(const FArguments& InArgs, class UEdGraphNode* InNode);
+
+		protected:
+
+			// SGraphNode Interface
+			virtual void CreateInputSideAddButton(TSharedPtr<SVerticalBox> InputBox) override;
+			virtual void CreateOutputSideAddButton(TSharedPtr<SVerticalBox> OutputBox) override;
+			virtual TSharedPtr<SGraphPin> CreatePinWidget(UEdGraphPin* InPin) const override;
+			virtual void CreateStandardPinWidget(UEdGraphPin* InPin) override;
+			virtual TSharedRef<SWidget> CreateNodeContentArea() override;
+			virtual TSharedRef<SWidget> CreateTitleWidget(TSharedPtr<SNodeTitle> NodeTitle) override;
+			virtual const FSlateBrush* GetNodeBodyBrush() const override;
+			virtual TSharedRef<SWidget> CreateTitleRightWidget() override;
+			virtual EVisibility IsAddPinButtonVisible() const override;
+			virtual FReply OnAddPin() override;
+			virtual void MoveTo(const FVector2D& NewPosition, FNodeSet& NodeFilter, bool bMarkDirty = true) override;
+			virtual void SetDefaultTitleAreaWidget(TSharedRef<SOverlay> DefaultTitleAreaWidget) override;
+
+			FLinearColor GetNodeTitleColorOverride() const;
+
+			FName GetLiteralDataType() const;
+			UMetasoundEditorGraphNode& GetMetaSoundNode();
+			const UMetasoundEditorGraphNode& GetMetaSoundNode() const;
+
+		public:
+			static void ExecuteTrigger(UMetasoundEditorGraphMemberDefaultLiteral& Literal);
+			static TSharedRef<SWidget> CreateTriggerSimulationWidget(UMetasoundEditorGraphMemberDefaultLiteral& Literal, TAttribute<EVisibility>&& InVisibility, TAttribute<bool>&& InEnablement, const FText* InToolTip = nullptr);
+
+		private:
+			// If this node represents a graph member node, returns corresponding member.
+			UMetasoundEditorGraphMember* GetMetaSoundMember();
+
+			// If this node represents a graph member node, returns cast node.
+			UMetasoundEditorGraphMemberNode* GetMetaSoundMemberNode();
+
+			TAttribute<EVisibility> GetSimulationVisibilityAttribute() const;
+
+			FText GetPinTooltip(UEdGraphPin* InPin) const;
+
+			// Slider widget for float input
+			TSharedPtr<SAudioInputWidget> InputWidget;
+			// Handle for on value changed delegate for input slider 
+			FDelegateHandle InputSliderOnValueChangedDelegateHandle;
+			// Handle for on input slider range changed  
+			FDelegateHandle InputSliderOnRangeChangedDelegateHandle;
+		};
+	} // namespace Editor
+} // namespace Metasound
diff --git a/Engine/Plugins/Runtime/Metasound/Source/MetasoundEditor/Private/SMetasoundGraphPin.h b/Engine/Plugins/Runtime/Metasound/Source/MetasoundEditor/Private/SMetasoundGraphPin.h
--- a/Engine/Plugins/Runtime/Metasound/Source/MetasoundEditor/Private/SMetasoundGraphPin.h
+++ b/Engine/Plugins/Runtime/Metasound/Source/MetasoundEditor/Private/SMetasoundGraphPin.h
@@ -9,11 +9,13 @@
 #include "KismetPins/SGraphPinString.h"
 #include "MetasoundEditorGraph.h"
 #include "MetasoundEditorGraphBuilder.h"
+#include "MetasoundEditorGraphMemberDefaults.h"
 #include "MetasoundEditorGraphNode.h"
 #include "MetasoundFrontendController.h"
 #include "SGraphPin.h"
 #include "Widgets/Input/SButton.h"
 #include "Widgets/SBoxPanel.h"
+#include "Widgets/SNullWidget.h"
 #include "Widgets/SWidget.h"
 
 #define LOCTEXT_NAMESPACE "MetasoundEditor"
@@ -72,11 +74,45 @@ namespace Metasound
 				return IInputController::GetInvalidHandle();
 			}
 
+			bool ShowDefaultValueWidget() const
+			{
+				UEdGraphPin* Pin = ParentPinType::GetPinObj();
+				if (!Pin)
+				{
+					return true;
+				}
+
+				UMetasoundEditorGraphMemberNode* Node = Cast<UMetasoundEditorGraphMemberNode>(Pin->GetOwningNode());
+				if (!Node)
+				{
+					return true;
+				}
+
+				UMetasoundEditorGraphMember* Member = Node->GetMember();
+				if (!Member)
+				{
+					return true;
+				}
+
+				UMetasoundEditorGraphMemberDefaultFloat* DefaultFloat = Cast<UMetasoundEditorGraphMemberDefaultFloat>(Member->GetLiteral());
+				if (!DefaultFloat)
+				{
+					return true;
+				}
+
+				return DefaultFloat->WidgetType == EMetasoundMemberDefaultWidget::None;
+			}
+
 			virtual TSharedRef<SWidget> GetDefaultValueWidget() override
 			{
 				using namespace Frontend;
 				TSharedRef<SWidget> DefaultWidget = ParentPinType::GetDefaultValueWidget();
 
+				if (!ShowDefaultValueWidget())
+				{
+					return SNullWidget::NullWidget;
+				}
+
 				// For now, arrays do not support literals.
 				// TODO: Support array literals by displaying
 				// default literals (non-array too) in inspector window.
@@ -123,22 +159,38 @@ namespace Metasound
 
 							if (UEdGraphPin* Pin = ParentPinType::GetPinObj())
 							{
-								if (UEdGraphNode* Node = Pin->GetOwningNode())
+								if (UMetasoundEditorGraphNode* Node = Cast<UMetasoundEditorGraphNode>(Pin->GetOwningNode()))
 								{
 									if (UMetasoundEditorGraph* MetaSoundGraph = CastChecked<UMetasoundEditorGraph>(Node->GetGraph()))
 									{
 										UObject& MetaSound = MetaSoundGraph->GetMetasoundChecked();
 
 										{
-											const FScopedTransaction Transaction(LOCTEXT("MetaSoundEditorSetLiteralToClassDefault", "Set Literal to Class Default"));
+											const FScopedTransaction Transaction(LOCTEXT("MetaSoundEditorResetToClassDefault", "Reset to Class Default"));
 											MetaSound.Modify();
 											MetaSoundGraph->Modify();
 
-											FInputHandle InputHandle = GetInputHandle();
-											InputHandle->ClearLiteral();
+											if (UMetasoundEditorGraphMemberNode* MemberNode = Cast<UMetasoundEditorGraphMemberNode>(Node))
+											{
+												UMetasoundEditorGraphMember* Member = MemberNode->GetMember();
+												if (ensure(Member))
+												{
+													Member->ResetToClassDefault();
+												}
+											}
+											else
+											{
+												FInputHandle InputHandle = GetInputHandle();
+												InputHandle->ClearLiteral();
+											}
 										}
 
-										FGraphBuilder::SynchronizePinLiteral(*Pin);
+										// Mark node for refresh in case non-pin widget on node needs updating (ex. sliders, knobs)
+										Node->bRefreshNode = true;
+
+										// Full node synchronization required as custom
+										// node-level widgets may need to be refreshed
+										MetaSoundGraph->SetSynchronizationRequired();
 									}
 								}
 							}
diff --git a/Engine/Plugins/Runtime/Metasound/Source/MetasoundEditor/Public/IMetasoundEditor.h b/Engine/Plugins/Runtime/Metasound/Source/MetasoundEditor/Public/IMetasoundEditor.h
--- a/Engine/Plugins/Runtime/Metasound/Source/MetasoundEditor/Public/IMetasoundEditor.h
+++ b/Engine/Plugins/Runtime/Metasound/Source/MetasoundEditor/Public/IMetasoundEditor.h
@@ -1,11 +1,10 @@
 // Copyright Epic Games, Inc. All Rights Reserved.
-
 #pragma once
 
-#include "CoreMinimal.h"
-
+#include "Containers/Array.h"
 #include "Layout/SlateRect.h"
 #include "Toolkits/AssetEditorToolkit.h"
+#include "UObject/NoExportTypes.h"
 #include "UObject/Object.h"
 
 
@@ -14,4 +13,6 @@ class IMetasoundEditor : public FAssetEditorToolkit
 	virtual UObject* GetMetasoundObject() const = 0;
 	virtual void SetSelection(const TArray<UObject*>& SelectedObjects) = 0;
 	virtual bool GetBoundsForSelectedNodes(FSlateRect& Rect, float Padding) = 0;
+	virtual void Play() = 0;
+	virtual void Stop() = 0;
 };
diff --git a/Engine/Plugins/Runtime/Metasound/Source/MetasoundEditor/Public/MetasoundEditorGraph.h b/Engine/Plugins/Runtime/Metasound/Source/MetasoundEditor/Public/MetasoundEditorGraph.h
--- a/Engine/Plugins/Runtime/Metasound/Source/MetasoundEditor/Public/MetasoundEditorGraph.h
+++ b/Engine/Plugins/Runtime/Metasound/Source/MetasoundEditor/Public/MetasoundEditorGraph.h
@@ -1,8 +1,6 @@
 // Copyright Epic Games, Inc. All Rights Reserved.
 #pragma once
 
-#include "CoreMinimal.h"
-
 #include "AudioParameterControllerInterface.h"
 #include "MetasoundEditor.h"
 #include "MetasoundEditorGraphNode.h"
@@ -20,6 +18,8 @@
 
 // Forward Declarations
 struct FMetasoundFrontendDocument;
+struct FPropertyChangedChainEvent;
+struct FPropertyChangedEvent;
 class ITargetPlatform;
 class UMetasoundEditorGraphInputNode;
 class UMetasoundEditorGraphNode;
@@ -36,7 +36,7 @@ namespace Metasound
 DECLARE_MULTICAST_DELEGATE_OneParam(FOnMetasoundMemberNameChanged, FGuid /* NodeID */);
 
 UCLASS()
-class METASOUNDEDITOR_API UMetasoundEditorGraphInputLiteral : public UObject
+class METASOUNDEDITOR_API UMetasoundEditorGraphMemberDefaultLiteral : public UObject
 {
 	GENERATED_BODY()
 
@@ -59,8 +59,11 @@ public:
 	{
 	}
 
+#if WITH_EDITOR
+	virtual void PostEditChangeProperty(FPropertyChangedEvent& InPropertyChangedEvent) override;
+	virtual void PostEditChangeChainProperty(FPropertyChangedChainEvent& InPropertyChangedEvent) override;
 	virtual void PostEditUndo() override;
-	void UpdateDocumentInputLiteral(bool bPostTransaction = true);
+#endif // WITH_EDITOR
 
 	const UMetasoundEditorGraphMember* GetParentMember() const
 	{
@@ -81,10 +84,6 @@ class METASOUNDEDITOR_API UMetasoundEditorGraphMember : public UObject
 	GENERATED_BODY()
 
 public:
-	/** Metasound Data Type. */
-	UPROPERTY()
-	FName TypeName;
-
 	/** Delegate called when the name of the associated Frontend Node is changed */
 	FOnMetasoundMemberNameChanged NameChanged;
 
@@ -95,7 +94,7 @@ public:
 	virtual TArray<UMetasoundEditorGraphNode*> GetNodes() const PURE_VIRTUAL(UMetasoundEditorGraphMember::GetNodes, return TArray<UMetasoundEditorGraphNode*>(); );
 
 	/** Sets the datatype on the member. */
-	virtual void SetDataType(FName InNewType, bool bPostTransaction = true, bool bRegisterParentGraph = true) PURE_VIRTUAL(UMetasoundEditorGraphMember::SetDataType, );
+	virtual void SetDataType(FName InNewType, bool bPostTransaction = true) PURE_VIRTUAL(UMetasoundEditorGraphMember::SetDataType, );
 	
 	/** If the Member Name can be changed to InNewName, returns true,
 	 * otherwise returns false with an error. */
@@ -125,6 +124,14 @@ public:
 	/** Returns the label of the derived member type (e.g. Input/Output/Variable) */
 	virtual const FText& GetGraphMemberLabel() const PURE_VIRTUAL(UMetasoundEditorGraphMember::GetGraphMemberLabel, return FText::GetEmpty(); );
 
+	/** Resets the member to the class default. */
+	virtual void ResetToClassDefault() PURE_VIRTUAL(UMetasoundEditorGraphMember::ResetToClassDefault, );
+
+	/** Update the frontend with the given member's default UObject value.
+	 * @param bPostTransaction - Post as editor transaction if true
+	 */
+	virtual void UpdateFrontendDefaultLiteral(bool bPostTransaction) PURE_VIRTUAL(UMetasoundEditorGraphMember::UpdateFrontendDefaultLiteral, );
+
 	/** Returns the parent MetaSound Graph. If the Outer object of the member is non
 	 * a UMetasoundEditorGraph, returns a nullptr. */
 	UMetasoundEditorGraph* GetOwningGraph();
@@ -135,6 +142,28 @@ public:
 
 	/** Marks all member's nodes as requiring a display refresh. */
 	void MarkNodesForRefresh();
+
+#if WITH_EDITOR
+	virtual void PostEditUndo() override;
+#endif // WITH_EDITOR
+
+	/** Returns the current data type */
+	FName GetDataType() const;
+
+	/** Returns literal associated with the given member */
+	virtual UMetasoundEditorGraphMemberDefaultLiteral* GetLiteral() const PURE_VIRTUAL(UMetasoundEditorGraphMember::GetLiteral, return nullptr; );
+
+protected:
+	virtual void SetLiteral(UMetasoundEditorGraphMemberDefaultLiteral* InLiteral) PURE_VIRTUAL (UMetasoundEditorGraphMember::SetLiteral, );
+
+	/** Conforms literal object type to member's DataType */
+	void ConformLiteralDataType();
+
+	/** Metasound Data Type. */
+	UPROPERTY()
+	FName TypeName;
+
+friend class UMetasoundEditorGraph;
 };
 
 /** Base class for an input or output of the graph. */
@@ -148,6 +177,9 @@ protected:
 	virtual Metasound::Frontend::FNodeHandle AddNodeHandle(const FName& InNodeName, FName InDataType) PURE_VIRTUAL(UMetasoundEditorGraphVertex::AddNodeHandle, return Metasound::Frontend::INodeController::GetInvalidHandle(); )
 
 public:
+	/** Initializes all properties with the given parameters required to identify the frontend member from this editor graph member. */
+	void InitMember(FName InDataType, const FMetasoundFrontendLiteral& InDefaultLiteral, FGuid InNodeID, FMetasoundFrontendClassName&& InClassName);
+
 	/** ID of Metasound Frontend node. */
 	UPROPERTY()
 	FGuid NodeID;
@@ -167,18 +199,15 @@ public:
 	virtual void SetDescription(const FText& InDescription, bool bPostTransaction) override;
 	virtual void SetMemberName(const FName& InNewName, bool bPostTransaction) override;
 	virtual void SetDisplayName(const FText& InNewName, bool bPostTransaction) override;
-	virtual void SetDataType(FName InNewType, bool bPostTransaction = true, bool bRegisterParentGraph = true) override;
+	virtual void SetDataType(FName InNewType, bool bPostTransaction = true) override;
 	/* ~End UMetasoundEditorGraphMember interface */
 
 	/** Version of interface membership, or invalid version if not an interface member. */
 	virtual const FMetasoundFrontendVersion& GetInterfaceVersion() const;
 
-	/** If true, this member cannot be modified by the user. */
+	/** Returns true if member is part of an interface. */
 	virtual bool IsInterfaceMember() const;
 
-	/** Called when the data type changes. */
-	virtual void OnDataTypeChanged();
-
 	/** Returns the Metasound class type of the associated node */
 	virtual EMetasoundFrontendClassType GetClassType() const { return EMetasoundFrontendClassType::Invalid; }
 
@@ -187,6 +216,11 @@ public:
 
 	/** Returns the node handle associated with the vertex. */
 	Metasound::Frontend::FConstNodeHandle GetConstNodeHandle() const;
+
+	virtual UMetasoundEditorGraphMemberDefaultLiteral* GetLiteral() const PURE_VIRTUAL(UMetasoundEditorGraphMember::GetLiteral, return nullptr; );
+
+protected:
+	virtual void SetLiteral(UMetasoundEditorGraphMemberDefaultLiteral* InLiteral) PURE_VIRTUAL (UMetasoundEditorGraphMember::SetLiteral, );
 };
 
 UCLASS()
@@ -194,25 +228,23 @@ class METASOUNDEDITOR_API UMetasoundEditorGraphInput : public UMetasoundEditorGr
 {
 	GENERATED_BODY()
 
-protected:
-	virtual Metasound::Frontend::FNodeHandle AddNodeHandle(const FName& InNodeName, FName InDataType) override;
-	virtual EMetasoundFrontendClassType GetClassType() const override { return EMetasoundFrontendClassType::Input; }
-
 public:
+	/** Default literal value of member */
 	UPROPERTY(VisibleAnywhere, Category = DefaultValue)
-	UMetasoundEditorGraphInputLiteral* Literal;
-
-	void OnDataTypeChanged() override;
+	UMetasoundEditorGraphMemberDefaultLiteral* Literal;
 
 	virtual const FText& GetGraphMemberLabel() const override;
-	virtual Metasound::Editor::ENodeSection GetSectionID() const override;
+	virtual void ResetToClassDefault() override;
+	virtual void SetMemberName(const FName& InNewName, bool bPostTransaction) override;
+	virtual void UpdateFrontendDefaultLiteral(bool bPostTransaction) override;
 
-#if WITH_EDITOR
-	virtual void PostEditUndo() override;
-#endif // WITH_EDITOR
-private:
+	virtual UMetasoundEditorGraphMemberDefaultLiteral* GetLiteral() const { return Literal; }
 
-	void UpdateEditorLiteralType();
+protected:
+	virtual void SetLiteral(UMetasoundEditorGraphMemberDefaultLiteral* InLiteral) { Literal = InLiteral; }
+	virtual Metasound::Frontend::FNodeHandle AddNodeHandle(const FName& InNodeName, FName InDataType) override;
+	virtual EMetasoundFrontendClassType GetClassType() const override { return EMetasoundFrontendClassType::Input; }
+	virtual Metasound::Editor::ENodeSection GetSectionID() const override;
 };
 
 
@@ -222,9 +254,18 @@ class METASOUNDEDITOR_API UMetasoundEditorGraphOutput : public UMetasoundEditorG
 	GENERATED_BODY()
 
 public:
+	/** Default literal value of member */
+	UPROPERTY(VisibleAnywhere, Category = DefaultValue)
+	UMetasoundEditorGraphMemberDefaultLiteral* Literal;
+
 	virtual const FText& GetGraphMemberLabel() const override;
+	virtual void ResetToClassDefault() override;
+	virtual void UpdateFrontendDefaultLiteral(bool bPostTransaction) override;
+
+	virtual UMetasoundEditorGraphMemberDefaultLiteral* GetLiteral() const { return Literal; }
 
 protected:
+	virtual void SetLiteral(UMetasoundEditorGraphMemberDefaultLiteral* InLiteral) { Literal = InLiteral; }
 	virtual Metasound::Frontend::FNodeHandle AddNodeHandle(const FName& InNodeName, FName InDataType) override;
 	virtual EMetasoundFrontendClassType GetClassType() const override { return EMetasoundFrontendClassType::Output; }
 	virtual Metasound::Editor::ENodeSection GetSectionID() const override;
@@ -239,13 +280,15 @@ class METASOUNDEDITOR_API UMetasoundEditorGraphVariable : public UMetasoundEdito
 	FGuid VariableID;
 
 public:
-
+	/** Default literal value of member */
 	UPROPERTY(VisibleAnywhere, Category = DefaultValue)
-	UMetasoundEditorGraphInputLiteral* Literal;
+	UMetasoundEditorGraphMemberDefaultLiteral* Literal;
+
+	void InitMember(FName InDataType, const FMetasoundFrontendLiteral& InDefaultLiteral, FGuid InVariableID);
 
 	/* ~Begin UMetasoundEditorGraphMember interface */
 	virtual Metasound::Editor::ENodeSection GetSectionID() const override;
-	virtual void SetDataType(FName InNewType, bool bPostTransaction = true, bool bRegisterParentGraph = false) override;
+	virtual void SetDataType(FName InNewType, bool bPostTransaction = true) override;
 
 	virtual FText GetDescription() const override;
 	virtual void SetDescription(const FText& InDescription, bool bPostTransaction) override;
@@ -258,18 +301,24 @@ public:
 	virtual FText GetDisplayName() const override;
 	virtual void SetDisplayName(const FText& InNewName, bool bPostTransaction) override;
 
+	virtual void ResetToClassDefault() override;
+	virtual void UpdateFrontendDefaultLiteral(bool bPostTransaction) override;
+
 	virtual TArray<UMetasoundEditorGraphNode*> GetNodes() const override;
 
 	virtual const FText& GetGraphMemberLabel() const override;
 	/* ~EndUMetasoundEditorGraphMember interface */
 
-	void SetFrontendVariable(const Metasound::Frontend::FConstVariableHandle& InVariable);
 	const FGuid& GetVariableID() const;
 
-	void UpdateDocumentVariable(bool bPostTransaction = true);
-#if WITH_EDITOR
-	virtual void PostEditUndo() override;
-#endif // WITH_EDITOR
+	Metasound::Frontend::FVariableHandle GetVariableHandle();
+	Metasound::Frontend::FConstVariableHandle GetConstVariableHandle() const;
+
+	virtual UMetasoundEditorGraphMemberDefaultLiteral* GetLiteral() const { return Literal; }
+
+protected:
+	virtual void SetLiteral(UMetasoundEditorGraphMemberDefaultLiteral* InLiteral) { Literal = InLiteral; }
+
 private:
 	struct FVariableEditorNodes
 	{
@@ -286,8 +335,6 @@ private:
 		TArray<FVector2D> DeferredAccessorLocations;
 	};
 
-	Metasound::Frontend::FVariableHandle GetVariableHandle();
-	Metasound::Frontend::FConstVariableHandle GetConstVariableHandle() const;
 	FVariableEditorNodes GetVariableNodes() const;
 	FVariableNodeLocations GetVariableNodeLocations() const;
 	void AddVariableNodes(UObject& InMetasound, Metasound::Frontend::FGraphHandle& InFrontendGraph, const FVariableNodeLocations& InNodeLocs);
diff --git a/Engine/Plugins/Runtime/Metasound/Source/MetasoundEditor/Public/MetasoundEditorGraphInputNode.h b/Engine/Plugins/Runtime/Metasound/Source/MetasoundEditor/Public/MetasoundEditorGraphInputNode.h
new file mode 100644
--- /dev/null
+++ b/Engine/Plugins/Runtime/Metasound/Source/MetasoundEditor/Public/MetasoundEditorGraphInputNode.h
@@ -0,0 +1,42 @@
+// Copyright Epic Games, Inc. All Rights Reserved.
+#pragma once
+
+#include "AudioParameterControllerInterface.h"
+#include "MetasoundEditorGraph.h"
+#include "MetasoundEditorGraphBuilder.h"
+#include "MetasoundEditorGraphNode.h"
+#include "MetasoundVertex.h"
+#include "Misc/Guid.h"
+#include "Textures/SlateIcon.h"
+#include "UObject/NoExportTypes.h"
+#include "UObject/ScriptInterface.h"
+
+#include "MetasoundEditorGraphInputNode.generated.h"
+
+
+UCLASS()
+class METASOUNDEDITOR_API UMetasoundEditorGraphInputNode : public UMetasoundEditorGraphMemberNode
+{
+	GENERATED_BODY()
+
+public:
+	virtual ~UMetasoundEditorGraphInputNode() = default;
+
+	UPROPERTY()
+	TObjectPtr<UMetasoundEditorGraphInput> Input;
+
+public:
+	virtual UMetasoundEditorGraphMember* GetMember() const override;
+	virtual void UpdatePreviewInstance(const Metasound::FVertexName& InParameterName, TScriptInterface<IAudioParameterControllerInterface>& InParameterInterface) const;
+
+	virtual FMetasoundFrontendClassName GetClassName() const;
+
+	virtual FGuid GetNodeID() const override;
+	virtual FLinearColor GetNodeTitleColor() const override;
+	virtual FSlateIcon GetNodeTitleIcon() const override;
+
+protected:
+	virtual void SetNodeID(FGuid InNodeID) override;
+
+	friend class Metasound::Editor::FGraphBuilder;
+};
diff --git a/Engine/Plugins/Runtime/Metasound/Source/MetasoundEditor/Public/MetasoundEditorGraphInputNodes.h b/Engine/Plugins/Runtime/Metasound/Source/MetasoundEditor/Public/MetasoundEditorGraphInputNodes.h
deleted file mode 100644
diff --git a/Engine/Plugins/Runtime/Metasound/Source/MetasoundEditor/Public/MetasoundEditorGraphMemberDefaults.h b/Engine/Plugins/Runtime/Metasound/Source/MetasoundEditor/Public/MetasoundEditorGraphMemberDefaults.h
new file mode 100644
--- /dev/null
+++ b/Engine/Plugins/Runtime/Metasound/Source/MetasoundEditor/Public/MetasoundEditorGraphMemberDefaults.h
@@ -0,0 +1,273 @@
+// Copyright Epic Games, Inc. All Rights Reserved.
+#pragma once
+
+#include "Algo/Count.h"
+#include "Algo/Transform.h"
+#include "AudioParameterControllerInterface.h"
+#include "Components/Widget.h"
+#include "Delegates/Delegate.h"
+#include "EdGraph/EdGraphNode.h"
+#include "GraphEditorSettings.h"
+#include "MetasoundDataReference.h"
+#include "MetasoundEditorGraph.h"
+#include "MetasoundEditorGraphNode.h"
+#include "MetasoundFrontend.h"
+#include "MetasoundFrontendController.h"
+#include "MetasoundFrontendDocument.h"
+#include "MetasoundFrontendLiteral.h"
+#include "MetasoundFrontendRegistries.h"
+#include "MetasoundPrimitives.h"
+#include "MetasoundVertex.h"
+#include "Misc/Guid.h"
+#include "Sound/SoundWave.h"
+#include "UObject/ObjectMacros.h"
+#include "UObject/SoftObjectPath.h"
+
+#include "MetasoundEditorGraphMemberDefaults.generated.h"
+
+
+// Broken out to be able to customize and swap enum behavior for basic integer literal behavior
+USTRUCT()
+struct FMetasoundEditorGraphMemberDefaultBoolRef
+{
+	GENERATED_BODY()
+
+	UPROPERTY(EditAnywhere, Category = DefaultValue)
+	bool Value = false;
+};
+
+UCLASS(MinimalAPI)
+class UMetasoundEditorGraphMemberDefaultBool : public UMetasoundEditorGraphMemberDefaultLiteral
+{
+	GENERATED_BODY()
+
+public:
+	virtual ~UMetasoundEditorGraphMemberDefaultBool() = default;
+
+	UPROPERTY(EditAnywhere, Category = DefaultValue)
+	FMetasoundEditorGraphMemberDefaultBoolRef Default;
+
+	virtual FMetasoundFrontendLiteral GetDefault() const override;
+	virtual EMetasoundFrontendLiteralType GetLiteralType() const override;
+	virtual void SetFromLiteral(const FMetasoundFrontendLiteral& InLiteral) override;
+	virtual void UpdatePreviewInstance(const Metasound::FVertexName& InParameterName, TScriptInterface<IAudioParameterControllerInterface>& InParameterInterface) const override;
+};
+
+UCLASS(MinimalAPI)
+class UMetasoundEditorGraphMemberDefaultBoolArray : public UMetasoundEditorGraphMemberDefaultLiteral
+{
+	GENERATED_BODY()
+
+public:
+	virtual ~UMetasoundEditorGraphMemberDefaultBoolArray() = default;
+
+	UPROPERTY(EditAnywhere, Category = DefaultValue)
+	TArray<FMetasoundEditorGraphMemberDefaultBoolRef> Default;
+
+	virtual FMetasoundFrontendLiteral GetDefault() const override;
+	virtual EMetasoundFrontendLiteralType GetLiteralType() const override;
+	virtual void SetFromLiteral(const FMetasoundFrontendLiteral& InLiteral) override;
+	virtual void UpdatePreviewInstance(const Metasound::FVertexName& InParameterName, TScriptInterface<IAudioParameterControllerInterface>& InParameterInterface) const override;
+};
+
+// Broken out to be able to customize and swap enum behavior for basic integer literal behavior
+USTRUCT()
+struct FMetasoundEditorGraphMemberDefaultIntRef
+{
+	GENERATED_BODY()
+
+	UPROPERTY(EditAnywhere, Category = DefaultValue)
+	int32 Value = 0;
+};
+
+UCLASS(MinimalAPI)
+class UMetasoundEditorGraphMemberDefaultInt : public UMetasoundEditorGraphMemberDefaultLiteral
+{
+	GENERATED_BODY()
+
+public:
+	virtual ~UMetasoundEditorGraphMemberDefaultInt() = default;
+
+	UPROPERTY(EditAnywhere, Category = DefaultValue)
+	FMetasoundEditorGraphMemberDefaultIntRef Default;
+
+	virtual FMetasoundFrontendLiteral GetDefault() const override;
+	virtual EMetasoundFrontendLiteralType GetLiteralType() const override;
+	virtual void SetFromLiteral(const FMetasoundFrontendLiteral& InLiteral) override;
+	virtual void UpdatePreviewInstance(const Metasound::FVertexName& InParameterName, TScriptInterface<IAudioParameterControllerInterface>& InParameterInterface) const override;
+};
+
+UCLASS(MinimalAPI)
+class UMetasoundEditorGraphMemberDefaultIntArray : public UMetasoundEditorGraphMemberDefaultLiteral
+{
+	GENERATED_BODY()
+
+public:
+	virtual ~UMetasoundEditorGraphMemberDefaultIntArray() = default;
+
+	UPROPERTY(EditAnywhere, Category = DefaultValue)
+	TArray<FMetasoundEditorGraphMemberDefaultIntRef> Default;
+
+	virtual FMetasoundFrontendLiteral GetDefault() const override;
+	virtual EMetasoundFrontendLiteralType GetLiteralType() const override;
+	virtual void SetFromLiteral(const FMetasoundFrontendLiteral& InLiteral) override;
+	virtual void UpdatePreviewInstance(const Metasound::FVertexName& InParameterName, TScriptInterface<IAudioParameterControllerInterface>& InParameterInterface) const override;
+};
+
+// For input widget
+DECLARE_MULTICAST_DELEGATE_OneParam(FOnMetasoundInputValueChangedEvent, float);
+DECLARE_MULTICAST_DELEGATE_OneParam(FOnMetasoundRangeChangedEvent, FVector2D);
+DECLARE_MULTICAST_DELEGATE_OneParam(FOnMetasoundInputClampDefaultChangedEvent, bool);
+
+UENUM()
+enum class EMetasoundMemberDefaultWidget : uint8
+{
+	None,
+	Slider,
+	RadialSlider UMETA(DisplayName = "Knob"),
+};
+
+UENUM()
+enum class EMetasoundMemberDefaultWidgetValueType : uint8
+{
+	Linear,
+	Frequency UMETA(DisplayName = "Frequency (Log)"),
+	Volume
+};
+
+UCLASS(MinimalAPI)
+class UMetasoundEditorGraphMemberDefaultFloat : public UMetasoundEditorGraphMemberDefaultLiteral
+{
+	GENERATED_BODY()
+
+private:
+	UPROPERTY(EditAnywhere, Category = DefaultValue)
+	float Default = 0.f;
+
+public:
+	virtual ~UMetasoundEditorGraphMemberDefaultFloat() = default;
+
+	UPROPERTY(EditAnywhere, Category = DefaultValue, meta=(EditCondition = "WidgetType == EMetasoundMemberDefaultWidget::None", EditConditionHides))
+	bool ClampDefault = false;
+
+	UPROPERTY(EditAnywhere, Category = DefaultValue, meta = (EditCondition = "ClampDefault || WidgetType != EMetasoundMemberDefaultWidget::None", EditConditionHides))
+	FVector2D Range = FVector2D(FMath::Min(0.0f, Default), FMath::Max(0.0f, Default));
+
+	UPROPERTY(EditAnywhere, Category = Widget, meta=(DisplayName = "Widget"))
+	EMetasoundMemberDefaultWidget WidgetType = EMetasoundMemberDefaultWidget::None;
+	
+	UPROPERTY(EditAnywhere, Category = Widget, meta=(DisplayName = "Orientation", EditCondition = "WidgetType == EMetasoundMemberDefaultWidget::Slider", EditConditionHides))
+	TEnumAsByte<EOrientation> WidgetOrientation = EOrientation::Orient_Horizontal;
+
+	UPROPERTY(EditAnywhere, Category = Widget, meta=(DisplayName = "Value Type", EditCondition = "WidgetType != EMetasoundMemberDefaultWidget::None", EditConditionHides))
+	EMetasoundMemberDefaultWidgetValueType WidgetValueType = EMetasoundMemberDefaultWidgetValueType::Linear;
+
+	FOnMetasoundInputValueChangedEvent OnDefaultValueChanged;
+	FOnMetasoundRangeChangedEvent OnRangeChanged;
+	FOnMetasoundInputClampDefaultChangedEvent OnClampChanged;
+
+	virtual FMetasoundFrontendLiteral GetDefault() const override;
+	virtual EMetasoundFrontendLiteralType GetLiteralType() const override;
+	virtual void SetFromLiteral(const FMetasoundFrontendLiteral& InLiteral) override;
+	virtual void UpdatePreviewInstance(const Metasound::FVertexName& InParameterName, TScriptInterface<IAudioParameterControllerInterface>& InParameterInterface) const override;
+	virtual void PostEditChangeProperty(FPropertyChangedEvent& PropertyChangedEvent) override;
+
+	void SetDefault(const float InDefault);
+	float GetDefault();
+	FVector2D GetRange();
+};
+
+UCLASS(MinimalAPI)
+class UMetasoundEditorGraphMemberDefaultFloatArray : public UMetasoundEditorGraphMemberDefaultLiteral
+{
+	GENERATED_BODY()
+
+public:
+	virtual ~UMetasoundEditorGraphMemberDefaultFloatArray() = default;
+
+	UPROPERTY(EditAnywhere, Category = DefaultValue)
+	TArray<float> Default;
+
+	virtual FMetasoundFrontendLiteral GetDefault() const override;
+	virtual EMetasoundFrontendLiteralType GetLiteralType() const override;
+	virtual void SetFromLiteral(const FMetasoundFrontendLiteral& InLiteral) override;
+	virtual void UpdatePreviewInstance(const Metasound::FVertexName& InParameterName, TScriptInterface<IAudioParameterControllerInterface>& InParameterInterface) const override;
+};
+
+UCLASS(MinimalAPI)
+class UMetasoundEditorGraphMemberDefaultString : public UMetasoundEditorGraphMemberDefaultLiteral
+{
+	GENERATED_BODY()
+
+public:
+	virtual ~UMetasoundEditorGraphMemberDefaultString() = default;
+
+	UPROPERTY(EditAnywhere, Category = DefaultValue)
+	FString Default;
+
+	virtual FMetasoundFrontendLiteral GetDefault() const override;
+	virtual EMetasoundFrontendLiteralType GetLiteralType() const override;
+	virtual void SetFromLiteral(const FMetasoundFrontendLiteral& InLiteral) override;
+	virtual void UpdatePreviewInstance(const Metasound::FVertexName& InParameterName, TScriptInterface<IAudioParameterControllerInterface>& InParameterInterface) const override;
+};
+
+UCLASS(MinimalAPI)
+class UMetasoundEditorGraphMemberDefaultStringArray : public UMetasoundEditorGraphMemberDefaultLiteral
+{
+	GENERATED_BODY()
+
+public:
+	virtual ~UMetasoundEditorGraphMemberDefaultStringArray() = default;
+
+	UPROPERTY(EditAnywhere, Category = DefaultValue)
+	TArray<FString> Default;
+
+	virtual FMetasoundFrontendLiteral GetDefault() const override;
+	virtual EMetasoundFrontendLiteralType GetLiteralType() const override;
+	virtual void SetFromLiteral(const FMetasoundFrontendLiteral& InLiteral) override;
+	virtual void UpdatePreviewInstance(const Metasound::FVertexName& InParameterName, TScriptInterface<IAudioParameterControllerInterface>& InParameterInterface) const override;
+};
+
+// Broken out to be able to customize and swap AllowedClass based on provided object proxy
+USTRUCT()
+struct FMetasoundEditorGraphMemberDefaultObjectRef
+{
+	GENERATED_BODY()
+
+	UPROPERTY(EditAnywhere, Category = DefaultValue)
+	UObject* Object = nullptr;
+};
+
+UCLASS(MinimalAPI)
+class UMetasoundEditorGraphMemberDefaultObject : public UMetasoundEditorGraphMemberDefaultLiteral
+{
+	GENERATED_BODY()
+
+public:
+	virtual ~UMetasoundEditorGraphMemberDefaultObject() = default;
+
+	UPROPERTY(EditAnywhere, Category = DefaultValue)
+	FMetasoundEditorGraphMemberDefaultObjectRef Default;
+
+	virtual FMetasoundFrontendLiteral GetDefault() const override;
+	virtual EMetasoundFrontendLiteralType GetLiteralType() const override;
+	virtual void SetFromLiteral(const FMetasoundFrontendLiteral& InLiteral) override;
+	virtual void UpdatePreviewInstance(const Metasound::FVertexName& InParameterName, TScriptInterface<IAudioParameterControllerInterface>& InParameterInterface) const override;
+};
+
+UCLASS(MinimalAPI)
+class UMetasoundEditorGraphMemberDefaultObjectArray : public UMetasoundEditorGraphMemberDefaultLiteral
+{
+	GENERATED_BODY()
+
+public:
+	virtual ~UMetasoundEditorGraphMemberDefaultObjectArray() = default;
+
+	UPROPERTY(EditAnywhere, Category = DefaultValue)
+	TArray<FMetasoundEditorGraphMemberDefaultObjectRef> Default;
+
+	virtual FMetasoundFrontendLiteral GetDefault() const override;
+	virtual EMetasoundFrontendLiteralType GetLiteralType() const override;
+	virtual void SetFromLiteral(const FMetasoundFrontendLiteral& InLiteral) override;
+	virtual void UpdatePreviewInstance(const Metasound::FVertexName& InParameterName, TScriptInterface<IAudioParameterControllerInterface>& InParameterInterface) const override;
+};
\ No newline at end of file
diff --git a/Engine/Plugins/Runtime/Metasound/Source/MetasoundEditor/Public/MetasoundEditorGraphNode.h b/Engine/Plugins/Runtime/Metasound/Source/MetasoundEditor/Public/MetasoundEditorGraphNode.h
--- a/Engine/Plugins/Runtime/Metasound/Source/MetasoundEditor/Public/MetasoundEditorGraphNode.h
+++ b/Engine/Plugins/Runtime/Metasound/Source/MetasoundEditor/Public/MetasoundEditorGraphNode.h
@@ -20,6 +20,7 @@
 class UEdGraphPin;
 class UMetaSound;
 class UMetasoundEditorGraphOutput;
+class UMetasoundEditorGraphMember;
 class UMetasoundEditorGraphVariable;
 
 namespace Metasound
@@ -97,7 +98,7 @@ public:
 	virtual FGuid GetNodeID() const { return FGuid(); }
 	virtual FText GetDisplayName() const;
 
-	// Whether or not to refresh given node in associated graph editor.
+	// Marks node to be reconstructed on next editor tick.
 	bool bRefreshNode = false;
 
 protected:
@@ -106,8 +107,26 @@ protected:
 	friend class Metasound::Editor::FGraphBuilder;
 };
 
+/** Node that represents a graph member */
+UCLASS(Abstract, MinimalAPI)
+class UMetasoundEditorGraphMemberNode : public UMetasoundEditorGraphNode
+{
+	GENERATED_BODY()
+
+public:
+	// Whether or not the member node supports interactivity on the graph node
+	virtual UMetasoundEditorGraphMember* GetMember() const PURE_VIRTUAL(UMetasoundEditorGraphMemberNode::GetMember, return nullptr;)
+
+	// Whether or not the member node supports interact widgets on the visual node (ex. float manipulation widgets)
+	virtual bool EnableInteractWidgets() const
+	{
+		return true;
+	}
+};
+
+/** Node that represents a graph output */
 UCLASS(MinimalAPI)
-class UMetasoundEditorGraphOutputNode : public UMetasoundEditorGraphNode
+class UMetasoundEditorGraphOutputNode : public UMetasoundEditorGraphMemberNode
 {
 	GENERATED_BODY()
 
@@ -117,13 +136,19 @@ public:
 
 	virtual FMetasoundFrontendClassName GetClassName() const override;
 	virtual FGuid GetNodeID() const override;
-	
+	virtual UMetasoundEditorGraphMember* GetMember() const override;
+
 	// Disallow deleting outputs as they require being connected to some
 	// part of the graph by the Frontend Graph Builder (which is enforced
 	// even when the Editor Graph Node does not have a visible input by
 	// way of a literal input.
 	virtual bool CanUserDeleteNode() const override;
 
+	virtual void PinDefaultValueChanged(UEdGraphPin* InPin) override;
+
+	// Disables interact widgets (ex. sliders, knobs) when input is connected
+	virtual bool EnableInteractWidgets() const override;
+
 protected:
 	virtual FLinearColor GetNodeTitleColor() const override;
 	virtual FSlateIcon GetNodeTitleIcon() const override;
@@ -182,12 +207,12 @@ protected:
 	friend class Metasound::Editor::FGraphBuilder;
 };
 
-/** UMetasoundEditorGraphVariableNode represents any of the several variable node
- * types (Accessor, DeferredAccessor, Mutator). */
+/** Represents any of the several variable node types (Accessor, DeferredAccessor, Mutator). */
 UCLASS(MinimalAPI)
-class UMetasoundEditorGraphVariableNode : public UMetasoundEditorGraphNode
+class UMetasoundEditorGraphVariableNode : public UMetasoundEditorGraphMemberNode
 {
 	GENERATED_BODY()
+
 protected:
 	// Class type of the frontend node (Accessor, DeferredAccessor or Mutator)
 	UPROPERTY()
@@ -197,7 +222,7 @@ protected:
 	UPROPERTY()
 	FMetasoundFrontendClassName ClassName;
 
-	// ID of the frontend node. 
+	// ID of the frontend node.
 	UPROPERTY()
 	FGuid NodeID;
 
@@ -206,9 +231,13 @@ public:
 	UPROPERTY()
 	UMetasoundEditorGraphVariable* Variable;
 
+	virtual UMetasoundEditorGraphMember* GetMember() const override;
+	virtual bool EnableInteractWidgets() const override;
 	virtual FMetasoundFrontendClassName GetClassName() const override;
 	virtual FGuid GetNodeID() const override;
 	
+	virtual void PinDefaultValueChanged(UEdGraphPin* Pin) override;
+
 protected:
 	virtual FLinearColor GetNodeTitleColor() const override;
 	virtual FSlateIcon GetNodeTitleIcon() const override;
diff --git a/Engine/Plugins/Runtime/Metasound/Source/MetasoundEditor/Public/MetasoundEditorGraphSchema.h b/Engine/Plugins/Runtime/Metasound/Source/MetasoundEditor/Public/MetasoundEditorGraphSchema.h
--- a/Engine/Plugins/Runtime/Metasound/Source/MetasoundEditor/Public/MetasoundEditorGraphSchema.h
+++ b/Engine/Plugins/Runtime/Metasound/Source/MetasoundEditor/Public/MetasoundEditorGraphSchema.h
@@ -1,10 +1,8 @@
 // Copyright Epic Games, Inc. All Rights Reserved.
-
 #pragma once
 
 #include "AssetData.h"
 #include "ConnectionDrawingPolicy.h"
-#include "CoreMinimal.h"
 #include "EdGraph/EdGraphSchema.h"
 #include "EdGraphUtilities.h"
 #include "MetasoundFrontend.h"
@@ -116,7 +114,7 @@ struct METASOUNDEDITOR_API FMetasoundGraphSchemaAction_NewInput : public FMetaso
 
 /** Promotes an input to a graph input, using its respective literal value as the default value */
 USTRUCT()
-struct METASOUNDEDITOR_API FMetasoundGraphSchemaAction_PromoteToInput : public FMetasoundGraphSchemaAction_NewInput
+struct METASOUNDEDITOR_API FMetasoundGraphSchemaAction_PromoteToInput : public FMetasoundGraphSchemaAction
 {
 	GENERATED_USTRUCT_BODY();
 
@@ -155,7 +153,7 @@ struct METASOUNDEDITOR_API FMetasoundGraphSchemaAction_NewOutput : public FMetas
 
 /** Promotes a node output to a graph output */
 USTRUCT()
-struct METASOUNDEDITOR_API FMetasoundGraphSchemaAction_PromoteToOutput : public FMetasoundGraphSchemaAction_NewOutput
+struct METASOUNDEDITOR_API FMetasoundGraphSchemaAction_PromoteToOutput : public FMetasoundGraphSchemaAction
 {
 	GENERATED_USTRUCT_BODY();
 
@@ -192,10 +190,11 @@ struct METASOUNDEDITOR_API FMetasoundGraphSchemaAction_NewVariableNode : public
 	//~ End FMetasoundGraphSchemaAction Interface
 	
 protected:
-
 	// Derived classes should override this method to create the desired frontend node.
-	virtual Metasound::Frontend::FNodeHandle CreateFrontendVariableNode(const Metasound::Frontend::FGraphHandle& InFrontendGraph, const FGuid& InVariableID) const { return Metasound::Frontend::INodeController::GetInvalidHandle(); }
-
+	virtual Metasound::Frontend::FNodeHandle CreateFrontendVariableNode(const Metasound::Frontend::FGraphHandle& InFrontendGraph, const FGuid& InVariableID) const
+	{
+		return Metasound::Frontend::INodeController::GetInvalidHandle();
+	}
 };
 
 /** Adds a variable node to the graph */
@@ -247,6 +246,46 @@ protected:
 	virtual Metasound::Frontend::FNodeHandle CreateFrontendVariableNode(const Metasound::Frontend::FGraphHandle& InFrontendGraph, const FGuid& InVariableID) const override;
 };
 
+/** Promotes an input to a graph variable & respective getter node, using its respective literal value as the default value */
+USTRUCT()
+struct METASOUNDEDITOR_API FMetasoundGraphSchemaAction_PromoteToVariable_AccessorNode : public FMetasoundGraphSchemaAction
+{
+	GENERATED_USTRUCT_BODY();
+
+	FMetasoundGraphSchemaAction_PromoteToVariable_AccessorNode();
+
+	//~ Begin FEdGraphSchemaAction Interface
+	virtual UEdGraphNode* PerformAction(UEdGraph* ParentGraph, UEdGraphPin* FromPin, const FVector2D Location, bool bSelectNewNode = true) override;
+	//~ End FEdGraphSchemaAction Interface
+};
+
+/** Promotes an output to a graph variable & respective setter node, using its respective literal value as the default value */
+USTRUCT()
+struct METASOUNDEDITOR_API FMetasoundGraphSchemaAction_PromoteToVariable_MutatorNode : public FMetasoundGraphSchemaAction
+{
+	GENERATED_USTRUCT_BODY();
+
+	FMetasoundGraphSchemaAction_PromoteToVariable_MutatorNode();
+
+	//~ Begin FEdGraphSchemaAction Interface
+	virtual UEdGraphNode* PerformAction(UEdGraph* ParentGraph, UEdGraphPin* FromPin, const FVector2D Location, bool bSelectNewNode = true) override;
+	//~ End FEdGraphSchemaAction Interface
+};
+
+/** Promotes an input to a graph variable & respective deferred getter node, using its respective literal value as the default value */
+USTRUCT()
+struct METASOUNDEDITOR_API FMetasoundGraphSchemaAction_PromoteToVariable_DeferredAccessorNode : public FMetasoundGraphSchemaAction
+{
+	GENERATED_USTRUCT_BODY();
+
+	FMetasoundGraphSchemaAction_PromoteToVariable_DeferredAccessorNode();
+
+	//~ Begin FEdGraphSchemaAction Interface
+	virtual UEdGraphNode* PerformAction(UEdGraph* ParentGraph, UEdGraphPin* FromPin, const FVector2D Location, bool bSelectNewNode = true) override;
+	//~ End FEdGraphSchemaAction Interface
+};
+
+
 /** Action to add a node to the graph */
 USTRUCT()
 struct METASOUNDEDITOR_API FMetasoundGraphSchemaAction_NewNode : public FMetasoundGraphSchemaAction
diff --git a/Engine/Plugins/Runtime/Metasound/Source/MetasoundEditor/Public/MetasoundEditorModule.h b/Engine/Plugins/Runtime/Metasound/Source/MetasoundEditor/Public/MetasoundEditorModule.h
--- a/Engine/Plugins/Runtime/Metasound/Source/MetasoundEditor/Public/MetasoundEditorModule.h
+++ b/Engine/Plugins/Runtime/Metasound/Source/MetasoundEditor/Public/MetasoundEditorModule.h
@@ -14,7 +14,7 @@
 
 
 class UMetasoundEditorGraph;
-class UMetasoundEditorGraphInputLiteral;
+class UMetasoundEditorGraphMemberDefaultLiteral;
 
 DECLARE_LOG_CATEGORY_EXTERN(LogMetasoundEditor, Log, All);
 
@@ -37,20 +37,20 @@ namespace Metasound
 			}
 		};
 
-		class METASOUNDEDITOR_API IMetaSoundInputLiteralCustomization
+		class METASOUNDEDITOR_API IMemberDefaultLiteralCustomization
 		{
 		public:
-			virtual ~IMetaSoundInputLiteralCustomization() = default;
+			virtual ~IMemberDefaultLiteralCustomization() = default;
 
-			virtual void CustomizeLiteral(UMetasoundEditorGraphInputLiteral& InLiteral, TSharedPtr<IPropertyHandle> InDefaultValueHandle) = 0;
+			virtual void CustomizeLiteral(UMetasoundEditorGraphMemberDefaultLiteral& InLiteral, TSharedPtr<IPropertyHandle> InDefaultValueHandle) = 0;
 		};
 
-		class METASOUNDEDITOR_API IMetaSoundInputLiteralCustomizationFactory
+		class METASOUNDEDITOR_API IMemberDefaultLiteralCustomizationFactory
 		{
 		public:
-			virtual ~IMetaSoundInputLiteralCustomizationFactory() = default;
+			virtual ~IMemberDefaultLiteralCustomizationFactory() = default;
 
-			virtual TUniquePtr<IMetaSoundInputLiteralCustomization> CreateLiteralCustomization(IDetailCategoryBuilder& DefaultCategoryBuilder) const = 0;
+			virtual TUniquePtr<IMemberDefaultLiteralCustomization> CreateLiteralCustomization(IDetailCategoryBuilder& DefaultCategoryBuilder) const = 0;
 		};
 
 		class METASOUNDEDITOR_API IMetasoundEditorModule : public IModuleInterface
@@ -73,9 +73,9 @@ namespace Metasound
 
 			virtual void IterateDataTypes(TUniqueFunction<void(const FEditorDataType&)> InDataTypeFunction) const = 0;
 
-			virtual TUniquePtr<IMetaSoundInputLiteralCustomization> CreateInputLiteralCustomization(UClass& InClass, IDetailCategoryBuilder& DefaultCategoryBuilder) const = 0;
+			virtual TUniquePtr<IMemberDefaultLiteralCustomization> CreateMemberDefaultLiteralCustomization(UClass& InClass, IDetailCategoryBuilder& DefaultCategoryBuilder) const = 0;
 
-			virtual const TSubclassOf<UMetasoundEditorGraphInputLiteral> FindInputLiteralClass(EMetasoundFrontendLiteralType InLiteralType) const = 0;
+			virtual const TSubclassOf<UMetasoundEditorGraphMemberDefaultLiteral> FindDefaultLiteralClass(EMetasoundFrontendLiteralType InLiteralType) const = 0;
 		};
 	} // namespace Editor
 } // namespace Metasound
diff --git a/Engine/Plugins/Runtime/Metasound/Source/MetasoundFrontend/Private/MetasoundAssetBase.cpp b/Engine/Plugins/Runtime/Metasound/Source/MetasoundFrontend/Private/MetasoundAssetBase.cpp
--- a/Engine/Plugins/Runtime/Metasound/Source/MetasoundFrontend/Private/MetasoundAssetBase.cpp
+++ b/Engine/Plugins/Runtime/Metasound/Source/MetasoundFrontend/Private/MetasoundAssetBase.cpp
@@ -390,7 +390,7 @@ bool FMetasoundAssetBase::VersionAsset()
 }
 
 #if WITH_EDITORONLY_DATA
-bool FMetasoundAssetBase::GetSynchronizationPending() const
+bool FMetasoundAssetBase::GetSynchronizationRequired() const
 {
 	return bSynchronizationRequired;
 }
@@ -400,9 +400,9 @@ bool FMetasoundAssetBase::GetSynchronizationClearUpdateNotes() const
 	return bSynchronizationClearUpdateNotes;
 }
 
-bool FMetasoundAssetBase::GetSynchronizationInterfacesUpdated() const
+bool FMetasoundAssetBase::GetSynchronizationUpdateDetails() const
 {
-	return bSynchronizationInterfacesUpdated;
+	return bSynchronizationUpdateDetails;
 }
 
 void FMetasoundAssetBase::SetSynchronizationRequired()
@@ -413,17 +413,19 @@ void FMetasoundAssetBase::SetSynchronizationRequired()
 void FMetasoundAssetBase::SetClearNodeNotesOnSynchronization()
 {
 	bSynchronizationClearUpdateNotes = true;
+	bSynchronizationRequired = true;
 }
 
-void FMetasoundAssetBase::SetInterfacesUpdatedOnSynchronization()
+void FMetasoundAssetBase::SetUpdateDetailsOnSynchronization()
 {
-	bSynchronizationInterfacesUpdated = true;
+	bSynchronizationUpdateDetails = true;
+	bSynchronizationRequired = true;
 }
 
 void FMetasoundAssetBase::ResetSynchronizationState()
 {
 	bSynchronizationClearUpdateNotes = false;
-	bSynchronizationInterfacesUpdated = false;
+	bSynchronizationUpdateDetails = false;
 	bSynchronizationRequired = false;
 }
 #endif // WITH_EDITORONLY_DATA
diff --git a/Engine/Plugins/Runtime/Metasound/Source/MetasoundFrontend/Private/MetasoundFrontendDataTypeRegistry.cpp b/Engine/Plugins/Runtime/Metasound/Source/MetasoundFrontend/Private/MetasoundFrontendDataTypeRegistry.cpp
--- a/Engine/Plugins/Runtime/Metasound/Source/MetasoundFrontend/Private/MetasoundFrontendDataTypeRegistry.cpp
+++ b/Engine/Plugins/Runtime/Metasound/Source/MetasoundFrontend/Private/MetasoundFrontendDataTypeRegistry.cpp
@@ -1,8 +1,7 @@
 // Copyright Epic Games, Inc. All Rights Reserved.
-
-
 #include "MetasoundFrontendDataTypeRegistry.h"
 
+
 namespace Metasound
 {
 	namespace Frontend
@@ -463,6 +462,8 @@ namespace Metasound
 				virtual bool GetDataTypeInfo(const UObject* InObject, FDataTypeRegistryInfo& OutInfo) const override;
 				virtual bool GetDataTypeInfo(const FName& InDataType, FDataTypeRegistryInfo& OutInfo) const override;
 
+				virtual void IterateDataTypeInfo(TFunctionRef<void(const FDataTypeRegistryInfo&)> InFunction) const override;
+
 				// Return the enum interface for a data type. If the data type does not have 
 				// an enum interface, returns a nullptr.
 				virtual TSharedPtr<const IEnumDataTypeInterface> GetEnumInterfaceForDataType(const FName& InDataType) const override;
@@ -593,6 +594,14 @@ namespace Metasound
 				return false;
 			}
 
+			void FDataTypeRegistry::IterateDataTypeInfo(TFunctionRef<void(const FDataTypeRegistryInfo&)> InFunction) const
+			{
+				for (const TPair<FName, TUniquePtr<IDataTypeRegistryEntry>>& Entry : RegisteredDataTypes)
+				{
+					InFunction(Entry.Value->GetDataTypeInfo());
+				}
+			}
+
 			bool FDataTypeRegistry::IsRegistered(const FName& InDataType) const
 			{
 				return RegisteredDataTypes.Contains(InDataType);
diff --git a/Engine/Plugins/Runtime/Metasound/Source/MetasoundFrontend/Private/MetasoundFrontendGraphController.cpp b/Engine/Plugins/Runtime/Metasound/Source/MetasoundFrontend/Private/MetasoundFrontendGraphController.cpp
--- a/Engine/Plugins/Runtime/Metasound/Source/MetasoundFrontend/Private/MetasoundFrontendGraphController.cpp
+++ b/Engine/Plugins/Runtime/Metasound/Source/MetasoundFrontend/Private/MetasoundFrontendGraphController.cpp
@@ -359,6 +359,13 @@ namespace Metasound
 						SetNode = AddNode(SetNodeClass.Metadata, FGuid::NewGuid());
 						if (SetNode->IsValid())
 						{
+							// Initialize set default literal value to that of the variable
+							FInputHandle InputHandle = SetNode->GetInputWithVertexName(VariableNames::GetInputDataName());
+							if (ensure(InputHandle->IsValid()))
+							{
+								InputHandle->SetLiteral(Variable->Literal);
+							}
+
 							Variable->MutatorNodeID = SetNode->GetID();
 							FGuid SourceVariableNodeID = Variable->VariableNodeID;
 
diff --git a/Engine/Plugins/Runtime/Metasound/Source/MetasoundFrontend/Private/MetasoundFrontendTransform.cpp b/Engine/Plugins/Runtime/Metasound/Source/MetasoundFrontend/Private/MetasoundFrontendTransform.cpp
--- a/Engine/Plugins/Runtime/Metasound/Source/MetasoundFrontend/Private/MetasoundFrontendTransform.cpp
+++ b/Engine/Plugins/Runtime/Metasound/Source/MetasoundFrontend/Private/MetasoundFrontendTransform.cpp
@@ -482,7 +482,7 @@ namespace Metasound
 				FMetasoundFrontendVersionNumber UpdateVersion = NodeHandle->FindHighestMinorVersionInRegistry();
 				if (UpdateVersion.IsValid() && UpdateVersion > ClassMetadata.GetVersion())
 				{
-					UE_LOG(LogMetaSound, Display, TEXT("Auto-Updating node class '%s': Newer minor version '%s' found."), *ClassMetadata.GetDisplayName().ToString(), *UpdateVersion.ToString());
+					UE_LOG(LogMetaSound, Display, TEXT("Auto-Updating node class '%s': Newer version '%s' found."), *ClassMetadata.GetDisplayName().ToString(), *UpdateVersion.ToString());
 				}
 				else if (InterfaceUpdates.ContainsChanges())
 				{
diff --git a/Engine/Plugins/Runtime/Metasound/Source/MetasoundFrontend/Public/MetasoundAssetBase.h b/Engine/Plugins/Runtime/Metasound/Source/MetasoundFrontend/Public/MetasoundAssetBase.h
--- a/Engine/Plugins/Runtime/Metasound/Source/MetasoundFrontend/Public/MetasoundAssetBase.h
+++ b/Engine/Plugins/Runtime/Metasound/Source/MetasoundFrontend/Public/MetasoundAssetBase.h
@@ -127,12 +127,15 @@ public:
 	bool VersionAsset();
 
 #if WITH_EDITORONLY_DATA
-	bool GetSynchronizationPending() const;
+	// TODO: These flags & associated functions are highly UE editor-specific.
+	// Split synchronization requirement flag into synchronization required &
+	// object type refresh or checking frontend class guids when synchronizing.
+	bool GetSynchronizationRequired() const;
 	bool GetSynchronizationClearUpdateNotes() const;
-	bool GetSynchronizationInterfacesUpdated() const;
+	bool GetSynchronizationUpdateDetails() const;
 	void ResetSynchronizationState();
 	void SetClearNodeNotesOnSynchronization();
-	void SetInterfacesUpdatedOnSynchronization();
+	void SetUpdateDetailsOnSynchronization();
 	void SetSynchronizationRequired();
 #endif // WITH_EDITORONLY_DATA
 
@@ -175,7 +178,7 @@ protected:
 #if WITH_EDITORONLY_DATA
 	bool bSynchronizationRequired = true;
 	bool bSynchronizationClearUpdateNotes = false;
-	bool bSynchronizationInterfacesUpdated = false;
+	bool bSynchronizationUpdateDetails = false;
 #endif // WITH_EDITORONLY_DATA
 
 private:
diff --git a/Engine/Plugins/Runtime/Metasound/Source/MetasoundFrontend/Public/MetasoundFrontendDataTypeRegistry.h b/Engine/Plugins/Runtime/Metasound/Source/MetasoundFrontend/Public/MetasoundFrontendDataTypeRegistry.h
--- a/Engine/Plugins/Runtime/Metasound/Source/MetasoundFrontend/Public/MetasoundFrontendDataTypeRegistry.h
+++ b/Engine/Plugins/Runtime/Metasound/Source/MetasoundFrontend/Public/MetasoundFrontendDataTypeRegistry.h
@@ -1,11 +1,18 @@
 // Copyright Epic Games, Inc. All Rights Reserved.
-
 #pragma once
 
+#include "Containers/Array.h"
 #include "MetasoundEnum.h"
 #include "MetasoundFrontendDocument.h"
 #include "MetasoundFrontendRegistries.h"
 #include "MetasoundLiteral.h"
+#include "Misc/Optional.h"
+#include "Templates/Function.h"
+#include "Templates/SharedPointer.h"
+#include "Templates/UniquePtr.h"
+#include "UObject/Class.h"
+#include "UObject/NameTypes.h"
+
 
 namespace Metasound
 {
@@ -197,15 +204,22 @@ namespace Metasound
 			 */
 			virtual bool RegisterDataType(TUniquePtr<IDataTypeRegistryEntry>&& InEntry) = 0;
 
+			/** Provides all names of registered DataTypes. */
 			virtual void GetRegisteredDataTypeNames(TArray<FName>& OutNames) const = 0;
 
+			/** Returns DataType info associated with the provided object. */
 			virtual bool GetDataTypeInfo(const UObject* InObject, FDataTypeRegistryInfo& OutInfo) const = 0;
+
+			/** Returns DataType info associated with the provided DataType name. */
 			virtual bool GetDataTypeInfo(const FName& InDataType, FDataTypeRegistryInfo& OutInfo) const = 0;
 
+			/** Iterates all registered data type info */
+			virtual void IterateDataTypeInfo(TFunctionRef<void(const FDataTypeRegistryInfo&)> InFunction) const = 0;
+
+			/** Returns whether or not a DataType is registered with the given name. */
 			virtual bool IsRegistered(const FName& InDataType) const = 0;
 
-			// Return the enum interface for a data type. If the data type does not have 
-			// an enum interface, returns a nullptr.
+			/** Return the enum interface for a data type. If the data type does not have an enum interface, returns a nullptr. */
 			virtual TSharedPtr<const IEnumDataTypeInterface> GetEnumInterfaceForDataType(const FName& InDataType) const = 0;
 
 			virtual ELiteralType GetDesiredLiteralType(const FName& InDataType) const = 0;
diff --git a/Engine/Plugins/Runtime/Metasound/Source/MetasoundFrontend/Public/MetasoundRouter.h b/Engine/Plugins/Runtime/Metasound/Source/MetasoundFrontend/Public/MetasoundRouter.h
--- a/Engine/Plugins/Runtime/Metasound/Source/MetasoundFrontend/Public/MetasoundRouter.h
+++ b/Engine/Plugins/Runtime/Metasound/Source/MetasoundFrontend/Public/MetasoundRouter.h
@@ -493,13 +493,6 @@ namespace Metasound
 			: IReceiver(InDataChannel, GetDataTypeName())
 			, OperatorSettings(InitParams.OperatorSettings)
 		{
-			// Executable DataTypes (currently just triggers) do not
-			// sync to past state provided by the data channel,
-			// so mark them as up-to-date.
-			if (TExecutableDataType<TDataType>::bIsExecutable)
-			{
-				LastPayloadID = InDataChannel->GetPayloadID();
-			}
 		}
 
 		void PushPayload(int64 InPayloadID, const TDataType& InDataPayload)
diff --git a/Engine/Plugins/Runtime/Metasound/Source/MetasoundGraphCore/Private/MetasoundDataReference.cpp b/Engine/Plugins/Runtime/Metasound/Source/MetasoundGraphCore/Private/MetasoundDataReference.cpp
--- a/Engine/Plugins/Runtime/Metasound/Source/MetasoundGraphCore/Private/MetasoundDataReference.cpp
+++ b/Engine/Plugins/Runtime/Metasound/Source/MetasoundGraphCore/Private/MetasoundDataReference.cpp
@@ -1,6 +1,7 @@
 // Copyright Epic Games, Inc. All Rights Reserved.
 
 #include "MetasoundDataReference.h"
+#include "MetasoundDataReferenceMacro.h"
 
 namespace Metasound
 {
@@ -13,4 +14,16 @@ namespace Metasound
 	const void* const TDataReferenceTypeInfo<void>::TypeId = static_cast<const void* const>(TDataReferenceTypeInfo<void>::TypePtr);
 
 	const FName IDataReference::RouterName = "DataReference";
+
+	FName CreateArrayTypeNameFromElementTypeName(const FName InTypeName)
+	{
+		return FName(InTypeName.ToString() + TEXT(METASOUND_DATA_TYPE_NAME_ARRAY_TYPE_SPECIFIER));
+	}
+
+	FName CreateElementTypeNameFromArrayTypeName(const FName InArrayTypeName)
+	{
+		auto GetSuffixLength = []() { return FString(TEXT(METASOUND_DATA_TYPE_NAME_ARRAY_TYPE_SPECIFIER)).Len(); };
+		static const int32 SuffixLength = GetSuffixLength();
+		return *InArrayTypeName.ToString().LeftChop(SuffixLength);
+	}
 }
diff --git a/Engine/Plugins/Runtime/Metasound/Source/MetasoundGraphCore/Public/MetasoundDataReference.h b/Engine/Plugins/Runtime/Metasound/Source/MetasoundGraphCore/Public/MetasoundDataReference.h
--- a/Engine/Plugins/Runtime/Metasound/Source/MetasoundGraphCore/Public/MetasoundDataReference.h
+++ b/Engine/Plugins/Runtime/Metasound/Source/MetasoundGraphCore/Public/MetasoundDataReference.h
@@ -39,7 +39,6 @@ namespace Metasound
 		static_assert(TSpecializationHelper<DataType>::Value, "TDataReferenceTypeInfo is not specialized.  Use macro DECLARE_METASOUND_DATA_REFERENCE_TYPES to declare a new type, or ensure that an existing DECLARE_METASOUND_DATA_REFERENCE_TYPES exists in the include path.");
 	};
 
-
 	/** Return the data type FName for a registered data type. */
 	template<typename DataType>
 	const FName& GetMetasoundDataTypeName()
@@ -75,6 +74,12 @@ namespace Metasound
 		return TDataReferenceTypeInfo<std::decay_t<DataType>>::TypeId;
 	}
 
+	/** Returns array type associated with the base datatype provided(ex. 'Float:Array' if 'Float' is provided) */
+	METASOUNDGRAPHCORE_API FName CreateArrayTypeNameFromElementTypeName(const FName InTypeName);
+
+	/** Returns the base data type with the array extension(ex. 'Float' if 'Float:Array' is provided) */
+	METASOUNDGRAPHCORE_API FName CreateElementTypeNameFromArrayTypeName(const FName InArrayTypeName);
+
 	/** Specialize void data type for internal use. */
 	template<>
 	struct TDataReferenceTypeInfo<void>
diff --git a/Engine/Source/Runtime/Engine/Private/AudioDevice.cpp b/Engine/Source/Runtime/Engine/Private/AudioDevice.cpp
--- a/Engine/Source/Runtime/Engine/Private/AudioDevice.cpp
+++ b/Engine/Source/Runtime/Engine/Private/AudioDevice.cpp
@@ -20,6 +20,7 @@
 #include "HAL/FileManager.h"
 #include "HAL/LowLevelMemTracker.h"
 #include "IAudioExtensionPlugin.h"
+#include "IAudioParameterTransmitter.h"
 #include "Misc/ConfigCacheIni.h"
 #include "Misc/App.h"
 #include "Misc/OutputDeviceArchiveWrapper.h"
@@ -5794,12 +5795,17 @@ void FAudioDevice::UnlinkActiveSoundFromComponent(const FActiveSound& InActiveSo
 		{
 			for (int32 i = ActiveSoundsInComponent->Num() - 1; i >= 0; --i)
 			{
-				const FActiveSound* ActiveSound = (*ActiveSoundsInComponent)[i];
-				
+				FActiveSound* ActiveSound = (*ActiveSoundsInComponent)[i];
 				if (ensure(ActiveSound))
 				{
 					if (ActiveSound->GetInstanceID() == InActiveSound.GetInstanceID())
 					{
+						// Reset the transmitter prior to unlinking to avoid associated
+						// state potentially kept in plugins to not become stale.
+						if (Audio::IParameterTransmitter* Transmitter = ActiveSound->GetTransmitter())
+						{
+							Transmitter->Reset();
+						}
 						ActiveSoundsInComponent->RemoveAtSwap(i, 1, false);
 						break;
 					}
