Tooltip fix-up
#rnx

diff --git a/Engine/Plugins/Runtime/Metasound/Source/MetasoundEditor/Private/MetasoundEditorGraphSchema.cpp b/Engine/Plugins/Runtime/Metasound/Source/MetasoundEditor/Private/MetasoundEditorGraphSchema.cpp
--- a/Engine/Plugins/Runtime/Metasound/Source/MetasoundEditor/Private/MetasoundEditorGraphSchema.cpp
+++ b/Engine/Plugins/Runtime/Metasound/Source/MetasoundEditor/Private/MetasoundEditorGraphSchema.cpp
@@ -449,7 +449,7 @@ FMetasoundGraphSchemaAction_PromoteToVariable_AccessorNode::FMetasoundGraphSchem
 	: FMetasoundGraphSchemaAction(
 		FText(),
 		LOCTEXT("PromoteToVariableGetterName", "Promote To Graph Variable"),
-		LOCTEXT("PromoteToInputTooltip", "Promotes node input to graph variable and creates a setter node"),
+		LOCTEXT("PromoteToInputTooltip", "Promotes node input to graph variable and creates a getter node"),
 		Metasound::Editor::EPrimaryContextGroup::Common)
 {
 }
