Non-Unity Fix

diff --git a/Engine/Plugins/Runtime/Metasound/Source/MetasoundEditor/Private/MetasoundEditorGraphBuilder.h b/Engine/Plugins/Runtime/Metasound/Source/MetasoundEditor/Private/MetasoundEditorGraphBuilder.h
--- a/Engine/Plugins/Runtime/Metasound/Source/MetasoundEditor/Private/MetasoundEditorGraphBuilder.h
+++ b/Engine/Plugins/Runtime/Metasound/Source/MetasoundEditor/Private/MetasoundEditorGraphBuilder.h
@@ -65,8 +65,8 @@ namespace Metasound
 			static UMetasoundEditorGraphNode* AddNode(UObject& InMetaSound, Frontend::FNodeHandle InNodeHandle, FVector2D InLocation, bool bInSelectNewNode = true);
 
 			// Convenience functions for retrieving the editor for the given MetaSound/EdGraph
-			static TSharedPtr<FEditor> GetEditorForMetasound(const UObject& MetaSound);
-			static TSharedPtr<FEditor> GetEditorForGraph(const UEdGraph& EdGraph);
+			static TSharedPtr<Metasound::Editor::FEditor> GetEditorForMetasound(const UObject& MetaSound);
+			static TSharedPtr<Metasound::Editor::FEditor> GetEditorForGraph(const UEdGraph& EdGraph);
 
 			// Wraps RegisterGraphWithFrontend logic in Frontend with any additional logic required to refresh editor & respective editor object state.
 			static void RegisterGraphWithFrontend(UObject& InMetaSound);
