Adding new interface for one-shots, upgrading old metasounds to use new interface if they hooked up the on-finished trigger.

Great test-case for interfaces and our upgrade code paths.

diff --git a/Engine/Plugins/Runtime/Metasound/Source/MetasoundEditor/Private/MetasoundDetailCustomization.cpp b/Engine/Plugins/Runtime/Metasound/Source/MetasoundEditor/Private/MetasoundDetailCustomization.cpp
--- a/Engine/Plugins/Runtime/Metasound/Source/MetasoundEditor/Private/MetasoundDetailCustomization.cpp
+++ b/Engine/Plugins/Runtime/Metasound/Source/MetasoundEditor/Private/MetasoundDetailCustomization.cpp
@@ -277,30 +277,34 @@ namespace Metasound
 			if (const FMetasoundAssetBase* MetaSoundAsset = IMetasoundUObjectRegistry::Get().GetObjectAsAssetBase(MetaSound.Get()))
 			{
 				auto GetVersionName = [](const FMetasoundFrontendVersion& Version) { return Version.Name; };
-				auto IsTransmittableVersion = [](const FMetasoundFrontendVersion& Version)
+				auto CanAddOrRemoveInterface = [](const FMetasoundFrontendVersion& Version)
 				{
 					using namespace Metasound::Frontend;
 
 					const FInterfaceRegistryKey Key = GetInterfaceRegistryKey(Version);
 					if (const IInterfaceRegistryEntry* Entry = IInterfaceRegistry::Get().FindInterfaceRegistryEntry(Key))
 					{
-						return Entry->GetRouterName() == Audio::IParameterTransmitter::RouterName;
+						return Entry->EditorCanAddOrRemove();
 					}
 
 					return false;
 				};
 
 				const TSet<FMetasoundFrontendVersion>& ImplementedInterfaces = MetaSoundAsset->GetDocumentChecked().Interfaces;
-				Algo::TransformIf(ImplementedInterfaces, ImplementedInterfaceNames, IsTransmittableVersion, GetVersionName);
+				Algo::TransformIf(ImplementedInterfaces, ImplementedInterfaceNames, CanAddOrRemoveInterface, GetVersionName);
 
-				Audio::IAudioParameterInterfaceRegistry::Get().IterateInterfaces([this](Audio::FParameterInterfacePtr Interface)
+				TArray<FMetasoundFrontendInterface> Interfaces = ISearchEngine::Get().FindAllInterfaces();
+				for (const FMetasoundFrontendInterface& Interface : Interfaces)
 				{
-					if (!ImplementedInterfaceNames.Contains(Interface->GetName()))
+					if (!ImplementedInterfaceNames.Contains(Interface.Version.Name))
 					{
-						FString Name = Interface->GetName().ToString();
-						AddableInterfaceNames.Add(MakeShared<FString>(MoveTemp(Name)));
+						if (CanAddOrRemoveInterface(Interface.Version))
+						{
+							FString Name = Interface.Version.Name.ToString();
+							AddableInterfaceNames.Add(MakeShared<FString>(MoveTemp(Name)));
+						}
 					}
-				});
+				}
 			}
 		}
 
diff --git a/Engine/Plugins/Runtime/Metasound/Source/MetasoundEngine/Private/MetasoundInterface.cpp b/Engine/Plugins/Runtime/Metasound/Source/MetasoundEngine/Private/MetasoundInterface.cpp
--- a/Engine/Plugins/Runtime/Metasound/Source/MetasoundEngine/Private/MetasoundInterface.cpp
+++ b/Engine/Plugins/Runtime/Metasound/Source/MetasoundEngine/Private/MetasoundInterface.cpp
@@ -28,6 +28,7 @@ namespace Metasound
 		struct FInterfaceRegistryOptions
 		{
 			bool bIsDefault = false;
+			bool bEditorCanAddOrRemove = true;
 			FName InputSystemName;
 			FName UClassName;
 		};
@@ -47,6 +48,11 @@ namespace Metasound
 			{
 			}
 
+			virtual bool EditorCanAddOrRemove() const override
+			{
+				return Options.bEditorCanAddOrRemove;
+			}
+
 			virtual bool UClassIsSupported(FName InUClassName) const override
 			{
 				if (Options.UClassName.IsNone())
@@ -89,13 +95,14 @@ namespace Metasound
 		};
 
 		template <typename UClassType>
-		void RegisterInterface(const FMetasoundFrontendInterface& InInterface, TUniquePtr<Frontend::IDocumentTransform>&& InUpdateTransform, bool bInIsDefault, FName InRouterName)
+		void RegisterInterface(const FMetasoundFrontendInterface& InInterface, TUniquePtr<Frontend::IDocumentTransform>&& InUpdateTransform, bool bInIsDefault, bool bEditorCanAddOrRemove, FName InRouterName)
 		{
 			using namespace Frontend;
 
 			FInterfaceRegistryOptions Options
 			{
 				bInIsDefault,
+				bEditorCanAddOrRemove,
 				InRouterName,
 				UClassType::StaticClass()->GetFName()
 			};
@@ -104,7 +111,7 @@ namespace Metasound
 			IInterfaceRegistry::Get().RegisterInterface(MakeUnique<FInterfaceRegistryEntry>(InInterface, MoveTemp(InUpdateTransform), MoveTemp(Options)));
 		}
 
-		void RegisterInterface(Audio::FParameterInterfacePtr Interface, TUniquePtr<Frontend::IDocumentTransform>&& InUpdateTransform, bool bInIsDefault, FName InRouterName)
+		void RegisterInterface(Audio::FParameterInterfacePtr Interface, TUniquePtr<Frontend::IDocumentTransform>&& InUpdateTransform, bool bInIsDefault, bool bEditorCanAddOrRemove, FName InRouterName)
 		{
 			using namespace Frontend;
 
@@ -170,7 +177,7 @@ namespace Metasound
 			const bool bInterfaceSupportsSource = SourceClass->IsChildOf(&Interface->GetType());
 			if (ensureAlwaysMsgf(bInterfaceSupportsSource, TEXT("Interfaces defined using Audio::FParameterInterface currently only supported by UMetaSoundSource asset type.")))
 			{
-				RegisterInterface<UMetaSoundSource>(FrontendInterface, MoveTemp(InUpdateTransform), bInIsDefault, InRouterName);
+				RegisterInterface<UMetaSoundSource>(FrontendInterface, MoveTemp(InUpdateTransform), bInIsDefault, bEditorCanAddOrRemove, InRouterName);
 				UE_LOG(LogMetaSound, Verbose, TEXT("Interface '%s' registered for asset type '%s'."), *Interface->GetName().ToString(), *SourceClass->GetName());
 			}
 		}
@@ -182,38 +189,48 @@ namespace Metasound
 			// Register Default Internal Interfaces (Not managed directly by end-user & added by default when creating new MetaSound assets).
 			{
 				constexpr bool bIsDefault = true;
-				RegisterInterface<UMetaSound>(MetasoundV1_0::GetInterface(), nullptr, bIsDefault, IDataReference::RouterName);
+				constexpr bool bEditorCanAddOrRemove = false;
+				RegisterInterface<UMetaSound>(MetasoundV1_0::GetInterface(), nullptr, bIsDefault, bEditorCanAddOrRemove, IDataReference::RouterName);
 
-				RegisterInterface(SourceInterface::CreateInterface(*UMetaSoundSource::StaticClass()), nullptr, bIsDefault, IDataReference::RouterName);
-				RegisterInterface(OutputFormatMonoInterface::CreateInterface(*UMetaSoundSource::StaticClass()), nullptr, bIsDefault, IDataReference::RouterName);
+				RegisterInterface(SourceInterface::CreateInterface(*UMetaSoundSource::StaticClass()), MakeUnique<SourceInterface::FUpdateInterface>(), bIsDefault, bEditorCanAddOrRemove, IDataReference::RouterName);
+				RegisterInterface(OutputFormatMonoInterface::CreateInterface(*UMetaSoundSource::StaticClass()), nullptr, bIsDefault, bEditorCanAddOrRemove, IDataReference::RouterName);
 			}
 
 			// Register Non-Default Internal Interfaces (Not managed directly by end-user & not added by default when creating new MetaSound assets).
 			{
 				constexpr bool bIsDefault = false;
+				constexpr bool bEditorCanAddOrRemove = false;
+				
+				// Added for upgrading old metasounds
+				RegisterInterface(SourceInterfaceV1_0::CreateInterface(*UMetaSoundSource::StaticClass()), nullptr, bIsDefault, bEditorCanAddOrRemove, IDataReference::RouterName);
 
 				// Set default interface with unset version to use base UMetaSound class implementation (legacy requirement for 5.0 alpha).
-				RegisterInterface<UMetaSound>(FMetasoundFrontendInterface(), nullptr, bIsDefault, IDataReference::RouterName);
+				RegisterInterface<UMetaSound>(FMetasoundFrontendInterface(), nullptr, bIsDefault, bEditorCanAddOrRemove, IDataReference::RouterName);
 
-				RegisterInterface<UMetaSoundSource>(MetasoundOutputFormatStereoV1_0::GetInterface(), nullptr, bIsDefault, IDataReference::RouterName);
-				RegisterInterface<UMetaSoundSource>(MetasoundOutputFormatStereoV1_1::GetInterface(), MakeUnique<MetasoundOutputFormatStereoV1_1::FUpdateInterface>(), bIsDefault, IDataReference::RouterName);
-				RegisterInterface<UMetaSoundSource>(MetasoundOutputFormatStereoV1_2::GetInterface(), MakeUnique<MetasoundOutputFormatStereoV1_2::FUpdateInterface>(), bIsDefault, IDataReference::RouterName);
+				RegisterInterface<UMetaSoundSource>(MetasoundOutputFormatStereoV1_0::GetInterface(), nullptr, bIsDefault, bEditorCanAddOrRemove, IDataReference::RouterName);
+				RegisterInterface<UMetaSoundSource>(MetasoundOutputFormatStereoV1_1::GetInterface(), MakeUnique<MetasoundOutputFormatStereoV1_1::FUpdateInterface>(), bIsDefault, bEditorCanAddOrRemove, IDataReference::RouterName);
+				RegisterInterface<UMetaSoundSource>(MetasoundOutputFormatStereoV1_2::GetInterface(), MakeUnique<MetasoundOutputFormatStereoV1_2::FUpdateInterface>(), bIsDefault, bEditorCanAddOrRemove, IDataReference::RouterName);
 
-				RegisterInterface<UMetaSoundSource>(MetasoundOutputFormatMonoV1_0::GetInterface(), nullptr, bIsDefault, IDataReference::RouterName);
-				RegisterInterface<UMetaSoundSource>(MetasoundOutputFormatMonoV1_1::GetInterface(), MakeUnique<MetasoundOutputFormatMonoV1_1::FUpdateInterface>(), bIsDefault, IDataReference::RouterName);
-				RegisterInterface<UMetaSoundSource>(MetasoundOutputFormatMonoV1_2::GetInterface(), MakeUnique<MetasoundOutputFormatMonoV1_2::FUpdateInterface>(), bIsDefault, IDataReference::RouterName);
+				RegisterInterface<UMetaSoundSource>(MetasoundOutputFormatMonoV1_0::GetInterface(), nullptr, bIsDefault, bEditorCanAddOrRemove, IDataReference::RouterName);
+				RegisterInterface<UMetaSoundSource>(MetasoundOutputFormatMonoV1_1::GetInterface(), MakeUnique<MetasoundOutputFormatMonoV1_1::FUpdateInterface>(), bIsDefault, bEditorCanAddOrRemove, IDataReference::RouterName);
+				RegisterInterface<UMetaSoundSource>(MetasoundOutputFormatMonoV1_2::GetInterface(), MakeUnique<MetasoundOutputFormatMonoV1_2::FUpdateInterface>(), bIsDefault, bEditorCanAddOrRemove, IDataReference::RouterName);
 
-				RegisterInterface(OutputFormatStereoInterface::CreateInterface(*UMetaSoundSource::StaticClass()), nullptr, bIsDefault, IDataReference::RouterName);
+				RegisterInterface(OutputFormatStereoInterface::CreateInterface(*UMetaSoundSource::StaticClass()), nullptr, bIsDefault, bEditorCanAddOrRemove, IDataReference::RouterName);
+				
+				RegisterInterface(SourceOneShotInterface::CreateInterface(*UMetaSoundSource::StaticClass()), nullptr, bIsDefault, true, IDataReference::RouterName);
 			}
 
 			// Register External Interfaces (Interfaces defined externally & can be managed directly by end-user).
 			auto RegisterExternalInterface = [](Audio::FParameterInterfacePtr Interface)
 			{
 				// Currently, no externally defined interfaces can be added as default for protection against undesired
-				// interfaces automatically being added when creating a new MetaSound asset through the editor.
+				// interfaces automatically being added when creating a new MetaSound asset through the editor. Also,
+				// all parameter interfaces are enabled in the editor for addition/removal.
 				constexpr bool bIsDefault = false;
+				constexpr bool bEditorCanAddOrRemove = true;
+
 				TUniquePtr<Frontend::IDocumentTransform> NullTransform;
-				RegisterInterface(Interface, MoveTemp(NullTransform), bIsDefault, Audio::IParameterTransmitter::RouterName);
+				RegisterInterface(Interface, MoveTemp(NullTransform), bIsDefault, bEditorCanAddOrRemove, Audio::IParameterTransmitter::RouterName);
 			};
 
 			Audio::IAudioParameterInterfaceRegistry::Get().IterateInterfaces(RegisterExternalInterface);
diff --git a/Engine/Plugins/Runtime/Metasound/Source/MetasoundEngine/Private/MetasoundSource.cpp b/Engine/Plugins/Runtime/Metasound/Source/MetasoundEngine/Private/MetasoundSource.cpp
--- a/Engine/Plugins/Runtime/Metasound/Source/MetasoundEngine/Private/MetasoundSource.cpp
+++ b/Engine/Plugins/Runtime/Metasound/Source/MetasoundEngine/Private/MetasoundSource.cpp
@@ -689,6 +689,14 @@ bool UMetaSoundSource::IsParameterValid(const FAudioParameter& InParameter, cons
 	return bIsValid;
 }
 
+bool UMetaSoundSource::IsOneShot()
+{
+	using namespace Metasound::Frontend;
+
+	// If the metasound source implements the one-shot interface, then it's a one-shot metasound
+	return IsInterfaceDeclared(SourceOneShotInterface::GetVersion());
+}
+
 TUniquePtr<Audio::IParameterTransmitter> UMetaSoundSource::CreateParameterTransmitter(Audio::FParameterTransmitterInitParams&& InParams) const
 {
 	Metasound::FMetaSoundParameterTransmitter::FInitParams InitParams(GetOperatorSettings(InParams.SampleRate), InParams.InstanceID);
diff --git a/Engine/Plugins/Runtime/Metasound/Source/MetasoundEngine/Public/MetasoundSource.h b/Engine/Plugins/Runtime/Metasound/Source/MetasoundEngine/Public/MetasoundSource.h
--- a/Engine/Plugins/Runtime/Metasound/Source/MetasoundEngine/Public/MetasoundSource.h
+++ b/Engine/Plugins/Runtime/Metasound/Source/MetasoundEngine/Public/MetasoundSource.h
@@ -168,6 +168,7 @@ public:
 	virtual ISoundGeneratorPtr CreateSoundGenerator(const FSoundGeneratorInitParams& InParams) override;
 	virtual TUniquePtr<Audio::IParameterTransmitter> CreateParameterTransmitter(Audio::FParameterTransmitterInitParams&& InParams) const override;
 	virtual bool IsParameterValid(const FAudioParameter& InParameter) const override;
+	virtual bool IsOneShot() override;
 
 protected:
 	/** Gets all the default parameters for this Asset.  */
diff --git a/Engine/Plugins/Runtime/Metasound/Source/MetasoundEngineTest/Private/EngineTestMetasoundSource.cpp b/Engine/Plugins/Runtime/Metasound/Source/MetasoundEngineTest/Private/EngineTestMetasoundSource.cpp
--- a/Engine/Plugins/Runtime/Metasound/Source/MetasoundEngineTest/Private/EngineTestMetasoundSource.cpp
+++ b/Engine/Plugins/Runtime/Metasound/Source/MetasoundEngineTest/Private/EngineTestMetasoundSource.cpp
@@ -77,7 +77,12 @@ namespace EngineTestMetasoundSourcePrivate
 		check(RootGraph->IsValid());
 
 		// Add default source & mono interface members (OnPlay, OnFinished & Mono Out)
-		FModifyRootGraphInterfaces InterfaceTransform({ SourceInterface::GetVersion(), OutputFormatMonoInterface::GetVersion() }, { });
+		FModifyRootGraphInterfaces InterfaceTransform(
+		{	
+			SourceInterface::GetVersion(), 
+			SourceOneShotInterface::GetVersion(),
+			OutputFormatMonoInterface::GetVersion() 
+		}, { });
 		InterfaceTransform.Transform(DocumentHandle);
 
 		// Input on Play
@@ -94,7 +99,7 @@ namespace EngineTestMetasoundSourcePrivate
 		check(FrequencyInputNode->IsValid());
 
 		// Output On Finished
-		FNodeHandle OnFinishedOutputNode = RootGraph->GetOutputNodeWithName(SourceInterface::Outputs::OnFinished);
+		FNodeHandle OnFinishedOutputNode = RootGraph->GetOutputNodeWithName(SourceOneShotInterface::Outputs::OnFinished);
 		check(OnFinishedOutputNode->IsValid());
 
 		// Output Audio
diff --git a/Engine/Plugins/Runtime/Metasound/Source/MetasoundFrontend/Private/MetasoundSourceInterface.cpp b/Engine/Plugins/Runtime/Metasound/Source/MetasoundFrontend/Private/MetasoundSourceInterface.cpp
--- a/Engine/Plugins/Runtime/Metasound/Source/MetasoundFrontend/Private/MetasoundSourceInterface.cpp
+++ b/Engine/Plugins/Runtime/Metasound/Source/MetasoundFrontend/Private/MetasoundSourceInterface.cpp
@@ -6,15 +6,59 @@
 #include "MetasoundTrigger.h"
 #include "UObject/Class.h"
 #include "Templates/SharedPointer.h"
-
+#include "AudioParameter.h"
 
 #define LOCTEXT_NAMESPACE "Metasound"
-#define AUDIO_PARAMETER_INTERFACE_NAMESPACE "UE.Source"
+
 namespace Metasound
 {
 	namespace Frontend
 	{
-		namespace SourceInterface
+	
+#define AUDIO_PARAMETER_INTERFACE_NAMESPACE "UE.Source.OneShot"
+
+		namespace SourceOneShotInterface
+		{
+			const FMetasoundFrontendVersion& GetVersion()
+			{
+				static const FMetasoundFrontendVersion Version = { AUDIO_PARAMETER_INTERFACE_NAMESPACE, { 1, 0 } };
+				return Version;
+			}
+
+			namespace Outputs
+			{
+				const FName OnFinished = AUDIO_PARAMETER_INTERFACE_MEMBER_DEFINE("OnFinished");
+			}
+
+			Audio::FParameterInterfacePtr CreateInterface(const UClass& InUClass)
+			{
+				struct FInterface : public Audio::FParameterInterface
+				{
+					FInterface(const UClass& InAssetClass)
+						: FParameterInterface(SourceOneShotInterface::GetVersion().Name, SourceOneShotInterface::GetVersion().Number.ToInterfaceVersion(), InAssetClass)
+					{
+						Outputs =
+						{
+							{
+								LOCTEXT("OnFinished", "On Finished"),
+								LOCTEXT("OnFinishedDescription", "Trigger executed to initiate stopping the source."),
+								GetMetasoundDataTypeName<FTrigger>(),
+								Outputs::OnFinished
+							}
+						};
+					}
+				};
+
+				return MakeShared<FInterface>(InUClass);
+			}
+		} // namespace SourceOneShotInterface
+
+#undef AUDIO_PARAMETER_INTERFACE_NAMESPACE
+
+
+#define AUDIO_PARAMETER_INTERFACE_NAMESPACE "UE.Source"
+
+		namespace SourceInterfaceV1_0
 		{
 			const FMetasoundFrontendVersion& GetVersion()
 			{
@@ -46,7 +90,7 @@ namespace Metasound
 				struct FInterface : public Audio::FParameterInterface
 				{
 					FInterface(const UClass& InAssetClass)
-						: FParameterInterface(SourceInterface::GetVersion().Name, SourceInterface::GetVersion().Number.ToInterfaceVersion(), InAssetClass)
+						: FParameterInterface(SourceInterfaceV1_0::GetVersion().Name, SourceInterfaceV1_0::GetVersion().Number.ToInterfaceVersion(), InAssetClass)
 					{
 						Inputs =
 						{
@@ -106,8 +150,128 @@ namespace Metasound
 
 				return MakeShared<FInterface>(InUClass);
 			}
-		} // namespace SourceInterface
+		} // namespace SourceInterfaceV1_0
+
+
+		namespace SourceInterface
+		{
+			const FMetasoundFrontendVersion& GetVersion()
+			{
+				static const FMetasoundFrontendVersion Version = { AUDIO_PARAMETER_INTERFACE_NAMESPACE, { 1, 1 } };
+				return Version;
+			}
+
+			namespace Inputs
+			{
+				const FName OnPlay = AUDIO_PARAMETER_INTERFACE_MEMBER_DEFINE("OnPlay");
+			}
+
+			namespace Environment
+			{
+				const FName DeviceID = AUDIO_PARAMETER_INTERFACE_MEMBER_DEFINE("AudioDeviceID");
+				const FName GraphName = AUDIO_PARAMETER_INTERFACE_MEMBER_DEFINE("GraphName");
+				const FName IsPreview = AUDIO_PARAMETER_INTERFACE_MEMBER_DEFINE("IsPreviewSound");
+				const FName SoundUniqueID = AUDIO_PARAMETER_INTERFACE_MEMBER_DEFINE("SoundUniqueID");
+				const FName TransmitterID = AUDIO_PARAMETER_INTERFACE_MEMBER_DEFINE("TransmitterID");
+			}
+
+			Audio::FParameterInterfacePtr CreateInterface(const UClass& InUClass)
+			{
+				struct FInterface : public Audio::FParameterInterface
+				{
+					FInterface(const UClass& InAssetClass)
+						: FParameterInterface(SourceInterface::GetVersion().Name, SourceInterface::GetVersion().Number.ToInterfaceVersion(), InAssetClass)
+					{
+						Inputs =
+						{
+							{
+								LOCTEXT("OnPlay", "On Play"),
+								LOCTEXT("OnPlayDescription", "Trigger executed when the source is played."),
+								GetMetasoundDataTypeName<FTrigger>(),
+								{ Inputs::OnPlay, false }
+							}
+						};
+
+						Environment =
+						{
+							{
+								LOCTEXT("AudioDeviceIDDisplayName", "Audio Device ID"),
+								LOCTEXT("AudioDeviceIDDescription", "ID of AudioDevice source is played from."),
+								FName(), // TODO: Align environment data types with environment (ex. this is actually set/get as a uint32)
+								Environment::DeviceID
+							},
+							{
+								LOCTEXT("GraphNameDisplayName", "Graph Name"),
+								LOCTEXT("AudioDeviceIDDescription", "Name of source graph (for debugging/logging)."),
+								GetMetasoundDataTypeName<FString>(),
+								Environment::GraphName
+							},
+							{
+								LOCTEXT("IsPreviewSoundDisplayName", "Is Preview Sound"),
+								LOCTEXT("IsPreviewSoundDescription", "Whether source is being played as a previewed sound."),
+								GetMetasoundDataTypeName<bool>(),
+								Environment::IsPreview
+							},
+							{
+								LOCTEXT("TransmitterIDDisplayName", "Transmitter ID"),
+								LOCTEXT("TransmitterIDDescription", "ID used by Transmission System to generate a unique send address for each source instance."),
+								FName(), // TODO: Align environment data types with environment (ex. this is actually set/get as a uint64)
+								Environment::TransmitterID
+							},
+							{
+								LOCTEXT("AudioDeviceIDDisplayName", "Sound Unique ID"),
+								LOCTEXT("AudioDeviceIDDescription", "ID of unique source instance."),
+								FName(), // TODO: Align environment data types with environment (ex. this is actually set/get as a uint32)
+								Environment::SoundUniqueID
+							}
+						};
+					}
+				};
+
+				return MakeShared<FInterface>(InUClass);
+			}
+
+			bool FUpdateInterface::Transform(Frontend::FDocumentHandle InDocument) const
+			{
+				using namespace Frontend;
+
+				// When upgrading, we only want to add the one-shot interface if the MetaSound actually has the OnFinished trigger connected.
+				bool bIsOnFinishedConnected = false;
+				InDocument->GetRootGraph()->IterateConstNodes([&](FConstNodeHandle NodeHandle)
+				{
+					NodeHandle->IterateConstInputs([&](FConstInputHandle InputHandle)
+					{
+						if (InputHandle->GetName() == SourceInterfaceV1_0::Outputs::OnFinished)
+						{
+							bIsOnFinishedConnected = InputHandle->IsConnected();
+						}
+					});
+				}, EMetasoundFrontendClassType::Output);
+
+				const TArray<FMetasoundFrontendVersion> InterfacesToRemove
+				{
+					SourceInterfaceV1_0::GetVersion()
+				};
+
+				TArray<FMetasoundFrontendVersion> InterfacesToAdd
+				{
+					SourceInterface::GetVersion()
+				};
+
+				if (bIsOnFinishedConnected)
+				{
+					InterfacesToAdd.Add(SourceOneShotInterface::GetVersion());
+				}
+
+				FModifyRootGraphInterfaces InterfaceTransform(InterfacesToRemove, InterfacesToAdd);
+				return InterfaceTransform.Transform(InDocument);
+			}
+
+} // namespace SourceInterface
+
+#undef AUDIO_PARAMETER_INTERFACE_NAMESPACE
+
 	} // namespace Frontend
 } // namespace Metasound
-#undef AUDIO_PARAMETER_INTERFACE_NAMESPACE
+
 #define LOCTEXT_NAMESPACE "Metasound"
diff --git a/Engine/Plugins/Runtime/Metasound/Source/MetasoundFrontend/Public/MetasoundFrontendArchetypeRegistry.h b/Engine/Plugins/Runtime/Metasound/Source/MetasoundFrontend/Public/MetasoundFrontendArchetypeRegistry.h
--- a/Engine/Plugins/Runtime/Metasound/Source/MetasoundFrontend/Public/MetasoundFrontendArchetypeRegistry.h
+++ b/Engine/Plugins/Runtime/Metasound/Source/MetasoundFrontend/Public/MetasoundFrontendArchetypeRegistry.h
@@ -29,6 +29,9 @@ namespace Metasound
 			// asset of the supported UClass is created.
 			virtual bool IsDefault() const = 0;
 
+			// Whether or not the interface can be added or removed by the editor.
+			virtual bool EditorCanAddOrRemove() const = 0;
+
 			// Name of routing system used to update interface inputs (ex. ParameterInterface or DataReference).
 			virtual FName GetRouterName() const = 0;
 
diff --git a/Engine/Plugins/Runtime/Metasound/Source/MetasoundFrontend/Public/MetasoundSourceInterface.h b/Engine/Plugins/Runtime/Metasound/Source/MetasoundFrontend/Public/MetasoundSourceInterface.h
--- a/Engine/Plugins/Runtime/Metasound/Source/MetasoundFrontend/Public/MetasoundSourceInterface.h
+++ b/Engine/Plugins/Runtime/Metasound/Source/MetasoundFrontend/Public/MetasoundSourceInterface.h
@@ -3,9 +3,10 @@
 
 #include "IAudioParameterInterfaceRegistry.h"
 #include "MetasoundFrontendDocument.h"
+#include "MetasoundFrontendTransform.h"
 #include "UObject/NameTypes.h"
-
-
+//#include "MetasoundEngineArchetypes.h"
+#include "MetasoundFrontendController.h"
 // Forward Declarations
 struct FMetasoundFrontendVersion;
 
@@ -13,7 +14,18 @@ namespace Metasound
 {
 	namespace Frontend
 	{
-		namespace SourceInterface
+		namespace SourceOneShotInterface
+		{
+			namespace Outputs
+			{
+				METASOUNDFRONTEND_API const extern FName OnFinished;
+			}
+
+			METASOUNDFRONTEND_API const FMetasoundFrontendVersion& GetVersion();
+			METASOUNDFRONTEND_API Audio::FParameterInterfacePtr CreateInterface(const UClass& InClass);
+		}
+
+		namespace SourceInterfaceV1_0
 		{
 			namespace Inputs
 			{
@@ -37,5 +49,34 @@ namespace Metasound
 			METASOUNDFRONTEND_API const FMetasoundFrontendVersion& GetVersion();
 			METASOUNDFRONTEND_API Audio::FParameterInterfacePtr CreateInterface(const UClass& InClass);
 		}
+
+		namespace SourceInterface
+		{
+			namespace Inputs
+			{
+				METASOUNDFRONTEND_API const extern FName OnPlay;
+			}
+
+			namespace Environment
+			{
+				METASOUNDFRONTEND_API const extern FName DeviceID;
+				METASOUNDFRONTEND_API const extern FName GraphName;
+				METASOUNDFRONTEND_API const extern FName IsPreview;
+				METASOUNDFRONTEND_API const extern FName SoundUniqueID;
+				METASOUNDFRONTEND_API const extern FName TransmitterID;
+			}
+
+			METASOUNDFRONTEND_API const FMetasoundFrontendVersion& GetVersion();
+			METASOUNDFRONTEND_API Audio::FParameterInterfacePtr CreateInterface(const UClass& InClass);
+
+			class METASOUNDFRONTEND_API FUpdateInterface : public Frontend::IDocumentTransform
+			{
+			public:
+				virtual bool Transform(Frontend::FDocumentHandle InDocument) const override;
+			};
+		}
+
+
+
 	} // namespace Frontend
 } // namespace Metasound
diff --git a/Engine/Plugins/Runtime/Metasound/Source/MetasoundGenerator/Private/MetasoundGenerator.cpp b/Engine/Plugins/Runtime/Metasound/Source/MetasoundGenerator/Private/MetasoundGenerator.cpp
--- a/Engine/Plugins/Runtime/Metasound/Source/MetasoundGenerator/Private/MetasoundGenerator.cpp
+++ b/Engine/Plugins/Runtime/Metasound/Source/MetasoundGenerator/Private/MetasoundGenerator.cpp
@@ -88,7 +88,7 @@ namespace Metasound
 			// References must be cached before moving the operator to the InitParams
 			FDataReferenceCollection Inputs = GraphOperator->GetInputs();
 			FTriggerWriteRef PlayTrigger = Inputs.GetDataWriteReferenceOrConstruct<FTrigger>(SourceInterface::Inputs::OnPlay, InitParams.OperatorSettings, false);
-			FTriggerReadRef FinishTrigger = Outputs.GetDataReadReferenceOrConstruct<FTrigger>(SourceInterface::Outputs::OnFinished, InitParams.OperatorSettings, false);
+			FTriggerReadRef FinishTrigger = Outputs.GetDataReadReferenceOrConstruct<FTrigger>(SourceOneShotInterface::Outputs::OnFinished, InitParams.OperatorSettings, false);
 
 			FMetasoundGeneratorData GeneratorData
 			{
diff --git a/Engine/Source/Runtime/Engine/Classes/Sound/SoundBase.h b/Engine/Source/Runtime/Engine/Classes/Sound/SoundBase.h
--- a/Engine/Source/Runtime/Engine/Classes/Sound/SoundBase.h
+++ b/Engine/Source/Runtime/Engine/Classes/Sound/SoundBase.h
@@ -252,9 +252,12 @@ public:
 	*/
 	virtual class UCurveTable* GetCurveData() const { return nullptr; }
 
-	/** Returns whether or not this sound is looping. */
+	/** Returns whether or not this sound is looping. TODO: Deprecate this to only use IsOneshot() in a MetaSound world. */
 	bool IsLooping();
 
+	/** Query if it's one shot. One shot is defined as a sound which is intended to have a fixed duration. */
+	virtual bool IsOneShot();
+
 	/** Parses the Sound to generate the WaveInstances to play. */
 	virtual void Parse( class FAudioDevice* AudioDevice, const UPTRINT NodeWaveInstanceHash, FActiveSound& ActiveSound, const FSoundParseParameters& ParseParams, TArray<FWaveInstance*>& WaveInstances ) { }
 
diff --git a/Engine/Source/Runtime/Engine/Private/SoundBase.cpp b/Engine/Source/Runtime/Engine/Private/SoundBase.cpp
--- a/Engine/Source/Runtime/Engine/Private/SoundBase.cpp
+++ b/Engine/Source/Runtime/Engine/Private/SoundBase.cpp
@@ -96,6 +96,11 @@ float USoundBase::GetPitchMultiplier()
 	return 1.f;
 }
 
+bool USoundBase::IsOneShot()
+{
+	return !IsLooping();
+}
+
 bool USoundBase::IsLooping()
 {
 	return (GetDuration() >= INDEFINITELY_LOOPING_DURATION);
