Enable F2 for renaming MetaSound member and comment nodes
- Also add rename to action menu when right clicking on node

diff --git a/Engine/Plugins/Runtime/Metasound/Source/MetasoundEditor/Private/MetasoundEditor.cpp b/Engine/Plugins/Runtime/Metasound/Source/MetasoundEditor/Private/MetasoundEditor.cpp
--- a/Engine/Plugins/Runtime/Metasound/Source/MetasoundEditor/Private/MetasoundEditor.cpp
+++ b/Engine/Plugins/Runtime/Metasound/Source/MetasoundEditor/Private/MetasoundEditor.cpp
@@ -1237,6 +1237,10 @@ namespace Metasound
 					FExecuteAction::CreateLambda([this] { DuplicateNodes(); }),
 					FCanExecuteAction::CreateLambda([this]() { return CanDuplicateNodes(); }));
 
+				GraphEditorCommands->MapAction(FGenericCommands::Get().Rename,
+					FExecuteAction::CreateLambda([this] { RenameNode(); }),
+					FCanExecuteAction::CreateLambda([this]() { return CanRenameNodes(); }));
+
 				// Alignment Commands
 				GraphEditorCommands->MapAction(FGraphEditorCommands::Get().AlignNodesTop,
 					FExecuteAction::CreateLambda([this]() { MetasoundGraphEditor->OnAlignTop(); }));
@@ -1836,6 +1840,78 @@ namespace Metasound
 			MetasoundGraphEditor->NotifyGraphChanged();
 		}
 
+		bool FEditor::CanRenameNodes() const
+		{
+			if (!IsGraphEditable())
+			{
+				return false;
+			}
+
+			const FGraphPanelSelectionSet& SelectedNodes = MetasoundGraphEditor->GetSelectedNodes();
+			for (FGraphPanelSelectionSet::TConstIterator SelectedIter(SelectedNodes); SelectedIter; ++SelectedIter)
+			{
+				// Node is directly renameable (comment nodes)
+				const UEdGraphNode* Node = Cast<UEdGraphNode>(*SelectedIter);
+				if (Node && Node->GetCanRenameNode())
+				{
+					return true;
+				}
+
+				// Renameable member nodes
+				if (const UMetasoundEditorGraphMemberNode* MemberNode = Cast<UMetasoundEditorGraphMemberNode>(*SelectedIter))
+				{
+					if (const UMetasoundEditorGraphMember* Member = MemberNode->GetMember()) 
+					{
+						// Variable
+						if (const UMetasoundEditorGraphVariable* Variable = Cast<UMetasoundEditorGraphVariable>(Member))
+						{
+							return true;
+						}
+
+						// Input or output that is not an interface member (does not have a valid interface version)
+						const FMetasoundFrontendVersion* InterfaceVersion = nullptr;
+						if (const UMetasoundEditorGraphVertex* Vertex = Cast<UMetasoundEditorGraphVertex>(Member))
+						{
+							InterfaceVersion = &Vertex->GetInterfaceVersion();
+						}
+
+						return (!InterfaceVersion || !InterfaceVersion->IsValid());
+					}
+				}
+			}
+			return false;
+		}
+
+		void FEditor::RenameNode() 
+		{
+			const FGraphPanelSelectionSet& SelectedNodes = MetasoundGraphEditor->GetSelectedNodes();
+			for (FGraphPanelSelectionSet::TConstIterator SelectedIter(SelectedNodes); SelectedIter; ++SelectedIter)
+			{
+				// Node is directly renameable (comment nodes)
+				UEdGraphNode* Node = Cast<UEdGraphNode>(*SelectedIter);
+				if (Node && Node->GetCanRenameNode())
+				{
+					if (TSharedPtr<SGraphEditor> GraphEditor = GetGraphEditor())
+					{
+						GraphEditor->JumpToNode(Node, /*bRequestRename=*/true);
+						return;
+					}
+				}
+
+				// Renameable member nodes (inputs/outputs/variables)
+				if (UMetasoundEditorGraphMemberNode* MemberNode = Cast<UMetasoundEditorGraphMemberNode>(*SelectedIter))
+				{
+					if (const UMetasoundEditorGraphMember* Member = MemberNode->GetMember()) 
+					{
+						if (Member->OnRenameRequested.IsBound())
+						{
+							Member->OnRenameRequested.Broadcast();
+						}
+					}
+				}
+			}
+		}
+
 		void FEditor::RefreshDetails()
 		{
 			if (MetasoundDetails.IsValid())
diff --git a/Engine/Plugins/Runtime/Metasound/Source/MetasoundEditor/Private/MetasoundEditor.h b/Engine/Plugins/Runtime/Metasound/Source/MetasoundEditor/Private/MetasoundEditor.h
--- a/Engine/Plugins/Runtime/Metasound/Source/MetasoundEditor/Private/MetasoundEditor.h
+++ b/Engine/Plugins/Runtime/Metasound/Source/MetasoundEditor/Private/MetasoundEditor.h
@@ -130,6 +130,12 @@ namespace Metasound
 			void PasteNodes(const FVector2D* InLocation = nullptr);
 			void PasteNodes(const FVector2D* InLocation, const FText& InTransactionText);
 
+			/** Rename node (currently applies to comments and member nodes). */
+			void RenameNode();
+
+			/** Whether at least one of the currently selected node(s) can be renamed. */
+			bool CanRenameNodes() const;
+
 			/** Forces all UX pertaining to the root graph's interface to be refreshed. */
 			void RefreshInterface();
 
diff --git a/Engine/Plugins/Runtime/Metasound/Source/MetasoundEditor/Private/MetasoundEditorGraphSchema.cpp b/Engine/Plugins/Runtime/Metasound/Source/MetasoundEditor/Private/MetasoundEditorGraphSchema.cpp
--- a/Engine/Plugins/Runtime/Metasound/Source/MetasoundEditor/Private/MetasoundEditorGraphSchema.cpp
+++ b/Engine/Plugins/Runtime/Metasound/Source/MetasoundEditor/Private/MetasoundEditorGraphSchema.cpp
@@ -971,6 +971,7 @@ void UMetasoundEditorGraphSchema::GetContextMenuActions(class UToolMenu* Menu, c
 		Section.AddMenuEntry(FGenericCommands::Get().Cut);
 		Section.AddMenuEntry(FGenericCommands::Get().Copy);
 		Section.AddMenuEntry(FGenericCommands::Get().Duplicate);
+		Section.AddMenuEntry(FGenericCommands::Get().Rename);
 		Section.AddMenuEntry(FGraphEditorCommands::Get().BreakNodeLinks);
 
 		// Only display update ability if node is of type external
diff --git a/Engine/Plugins/Runtime/Metasound/Source/MetasoundEditor/Private/MetasoundNodeDetailCustomization.h b/Engine/Plugins/Runtime/Metasound/Source/MetasoundEditor/Private/MetasoundNodeDetailCustomization.h
--- a/Engine/Plugins/Runtime/Metasound/Source/MetasoundEditor/Private/MetasoundNodeDetailCustomization.h
+++ b/Engine/Plugins/Runtime/Metasound/Source/MetasoundEditor/Private/MetasoundNodeDetailCustomization.h
@@ -163,7 +163,10 @@ namespace Metasound
 				DataTypeSelector = MakeShared<FMetasoundDataTypeSelector>();
 			}
 
-			virtual ~TMetasoundGraphMemberDetailCustomization() = default;
+			virtual ~TMetasoundGraphMemberDetailCustomization()
+			{
+				RenameRequestedHandle.Reset();
+			}
 
 		protected:
 			TWeakObjectPtr<TMemberType> GraphMember;
@@ -234,6 +237,7 @@ namespace Metasound
 					.OnTextChanged(this, &TChildClass::OnNameChanged)
 					.OnTextCommitted(this, &TChildClass::OnNameCommitted)
 					.IsReadOnly(bIsInterfaceMember || !bIsGraphEditable)
+					.SelectAllTextWhenFocused(true)
 					.Font(IDetailLayoutBuilder::GetDetailFont());
 
 				DisplayNameEditableTextBox = SNew(SEditableTextBox)
@@ -304,6 +308,29 @@ namespace Metasound
 						{
 							LiteralHandle->MarkHiddenByCustomization();
 
+							// Rename
+							if (UMetasoundEditorGraphMemberDefaultLiteral* MemberDefaultLiteral = CastChecked<UMetasoundEditorGraphMemberDefaultLiteral>(LiteralObject))
+							{
+								if (UMetasoundEditorGraphMember* Member = MemberDefaultLiteral->GetParentMember())
+								{
+									if (const UMetasoundEditorGraph* OwningGraph = GraphMember->GetOwningGraph())
+									{
+										TSharedPtr<FEditor> MetasoundEditor = FGraphBuilder::GetEditorForGraph(*OwningGraph);
+										if (MetasoundEditor->CanRenameNodes())
+										{
+											if (!RenameRequestedHandle.IsValid())
+											{
+												Member->OnRenameRequested.Clear();
+												RenameRequestedHandle = Member->OnRenameRequested.AddLambda([NameEditableTextBox = this->NameEditableTextBox]()
+												{
+													FSlateApplication::Get().SetKeyboardFocus(NameEditableTextBox.ToSharedRef(), EFocusCause::SetDirectly);
+												});
+											}
+										}
+									}
+								}
+							}
+
 							TSharedPtr<IPropertyHandle> DefaultValueHandle;
 
 							if (IDetailPropertyRow* Row = DefaultCategoryBuilder.AddExternalObjectProperty(TArray<UObject*>({ LiteralObject }), "Default"))
@@ -504,6 +531,8 @@ namespace Metasound
 
 		private:
 			TUniquePtr<IMemberDefaultLiteralCustomization> LiteralCustomization;
+
+			FDelegateHandle RenameRequestedHandle;
 		};
 
 		class FMetasoundInputDetailCustomization : public TMetasoundGraphMemberDetailCustomization<UMetasoundEditorGraphInput, FMetasoundInputDetailCustomization>
diff --git a/Engine/Plugins/Runtime/Metasound/Source/MetasoundEditor/Public/MetasoundEditorGraph.h b/Engine/Plugins/Runtime/Metasound/Source/MetasoundEditor/Public/MetasoundEditorGraph.h
--- a/Engine/Plugins/Runtime/Metasound/Source/MetasoundEditor/Public/MetasoundEditorGraph.h
+++ b/Engine/Plugins/Runtime/Metasound/Source/MetasoundEditor/Public/MetasoundEditorGraph.h
@@ -34,6 +34,7 @@ namespace Metasound
 } // namespace Metasound
 
 DECLARE_MULTICAST_DELEGATE_OneParam(FOnMetasoundMemberNameChanged, FGuid /* NodeID */);
+DECLARE_MULTICAST_DELEGATE(FOnMetasoundMemberRenameRequested);
 
 UCLASS()
 class METASOUNDEDITOR_API UMetasoundEditorGraphMemberDefaultLiteral : public UObject
@@ -87,6 +88,9 @@ public:
 	/** Delegate called when the name of the associated Frontend Node is changed */
 	FOnMetasoundMemberNameChanged NameChanged;
 
+	/** Delegate called when a rename is requested on a renameable member node. */
+	FOnMetasoundMemberRenameRequested OnRenameRequested;
+
 	/** Return the section of where this member belongs. */
 	virtual Metasound::Editor::ENodeSection GetSectionID() const PURE_VIRTUAL(UMetasoundEditorGraphMember::GetSectionID, return Metasound::Editor::ENodeSection::None; );
 
