Metasound filter default parameter value fixes -(wasn't using GetDataReadReferenceOrConstructWithVertexDefault() in Operator construction

#p4v-preflight-copy 18480447

diff --git a/Engine/Plugins/Runtime/Metasound/Source/MetasoundStandardNodes/Private/MetasoundBasicFilters.cpp b/Engine/Plugins/Runtime/Metasound/Source/MetasoundStandardNodes/Private/MetasoundBasicFilters.cpp
--- a/Engine/Plugins/Runtime/Metasound/Source/MetasoundStandardNodes/Private/MetasoundBasicFilters.cpp
+++ b/Engine/Plugins/Runtime/Metasound/Source/MetasoundStandardNodes/Private/MetasoundBasicFilters.cpp
@@ -384,12 +384,14 @@ namespace Metasound
 	{
 
 		const FDataReferenceCollection& InputDataRefs = InParams.InputDataReferences;
+		FVertexInterface Interface = DeclareVertexInterface();
+		const FInputVertexInterface& InputInterface = Interface.GetInputInterface();
 
 		// inputs
 		FAudioBufferReadRef AudioIn = InputDataRefs.GetDataReadReferenceOrConstruct<FAudioBuffer>(METASOUND_GET_PARAM_NAME(ParamAudioInput), InParams.OperatorSettings);
-		FFloatReadRef FrequencyIn = InputDataRefs.GetDataReadReferenceOrConstruct<float>(METASOUND_GET_PARAM_NAME(ParamCutoffFrequency));
-		FFloatReadRef ResonanceIn = InputDataRefs.GetDataReadReferenceOrConstruct<float>(METASOUND_GET_PARAM_NAME(ParamResonance));
-		FFloatReadRef PassBandGainCompensationIn = InputDataRefs.GetDataReadReferenceOrConstruct<float>(METASOUND_GET_PARAM_NAME(ParamBandStopControl));
+		FFloatReadRef FrequencyIn = InputDataRefs.GetDataReadReferenceOrConstructWithVertexDefault<float>(InputInterface, METASOUND_GET_PARAM_NAME(ParamCutoffFrequency), InParams.OperatorSettings);
+		FFloatReadRef ResonanceIn = InputDataRefs.GetDataReadReferenceOrConstructWithVertexDefault<float>(InputInterface, METASOUND_GET_PARAM_NAME(ParamResonance), InParams.OperatorSettings);
+		FFloatReadRef PassBandGainCompensationIn = InputDataRefs.GetDataReadReferenceOrConstructWithVertexDefault<float>(InputInterface, METASOUND_GET_PARAM_NAME(ParamBandStopControl), InParams.OperatorSettings);
 
 		return MakeUnique<FStateVariableFilterOperator>(
 			InParams.OperatorSettings
@@ -549,10 +551,12 @@ namespace Metasound
 	{
 
 		const FDataReferenceCollection& InputDataRefs = InParams.InputDataReferences;
+		FVertexInterface Interface = DeclareVertexInterface();
+		const FInputVertexInterface& InputInterface = Interface.GetInputInterface();
 
 		// inputs
 		FAudioBufferReadRef AudioIn = InputDataRefs.GetDataReadReferenceOrConstruct<FAudioBuffer>(METASOUND_GET_PARAM_NAME(ParamAudioInput), InParams.OperatorSettings);
-		FFloatReadRef FrequencyIn = InputDataRefs.GetDataReadReferenceOrConstruct<float>(METASOUND_GET_PARAM_NAME(ParamCutoffFrequency));
+		FFloatReadRef FrequencyIn = InputDataRefs.GetDataReadReferenceOrConstructWithVertexDefault<float>(InputInterface, METASOUND_GET_PARAM_NAME(ParamCutoffFrequency), InParams.OperatorSettings);
 
 		return MakeUnique<FOnePoleLowPassFilterOperator>(
 			InParams.OperatorSettings
@@ -684,10 +688,12 @@ namespace Metasound
 	{
 
 		const FDataReferenceCollection& InputDataRefs = InParams.InputDataReferences;
+		FVertexInterface Interface = DeclareVertexInterface();
+		const FInputVertexInterface& InputInterface = Interface.GetInputInterface();
 
 		// inputs
 		FAudioBufferReadRef AudioIn = InputDataRefs.GetDataReadReferenceOrConstruct<FAudioBuffer>(METASOUND_GET_PARAM_NAME(ParamAudioInput), InParams.OperatorSettings);
-		FFloatReadRef FrequencyIn = InputDataRefs.GetDataReadReferenceOrConstruct<float>(METASOUND_GET_PARAM_NAME(ParamCutoffFrequency));
+		FFloatReadRef FrequencyIn = InputDataRefs.GetDataReadReferenceOrConstructWithVertexDefault<float>(InputInterface, METASOUND_GET_PARAM_NAME(ParamCutoffFrequency), InParams.OperatorSettings);
 
 		return MakeUnique<FOnePoleHighPassFilterOperator>(
 			InParams.OperatorSettings
@@ -841,13 +847,15 @@ namespace Metasound
 	{
 
 		const FDataReferenceCollection& InputDataRefs = InParams.InputDataReferences;
+		FVertexInterface Interface = DeclareVertexInterface();
+		const FInputVertexInterface& InputInterface = Interface.GetInputInterface();
 
 		// inputs
 		FAudioBufferReadRef AudioIn = InputDataRefs.GetDataReadReferenceOrConstruct<FAudioBuffer>(METASOUND_GET_PARAM_NAME(ParamAudioInput), InParams.OperatorSettings);
-		FFloatReadRef FrequencyIn = InputDataRefs.GetDataReadReferenceOrConstruct<float>(METASOUND_GET_PARAM_NAME(ParamCutoffFrequency));
-		FFloatReadRef BandwidthIn = InputDataRefs.GetDataReadReferenceOrConstruct<float>(METASOUND_GET_PARAM_NAME(ParamBandwidth));
-		FFloatReadRef FilterGainDbIn = InputDataRefs.GetDataReadReferenceOrConstruct<float>(METASOUND_GET_PARAM_NAME(ParamGainDb));
-		FEnumBiQuadFilterReadRef FilterType = InputDataRefs.GetDataReadReferenceOrConstruct<FEnumEBiquadFilterType>(METASOUND_GET_PARAM_NAME(ParamFilterType));
+		FFloatReadRef FrequencyIn = InputDataRefs.GetDataReadReferenceOrConstructWithVertexDefault<float>(InputInterface, METASOUND_GET_PARAM_NAME(ParamCutoffFrequency), InParams.OperatorSettings);
+		FFloatReadRef BandwidthIn = InputDataRefs.GetDataReadReferenceOrConstructWithVertexDefault<float>(InputInterface, METASOUND_GET_PARAM_NAME(ParamBandwidth), InParams.OperatorSettings);
+		FFloatReadRef FilterGainDbIn = InputDataRefs.GetDataReadReferenceOrConstructWithVertexDefault<float>(InputInterface, METASOUND_GET_PARAM_NAME(ParamGainDb), InParams.OperatorSettings);
+		FEnumBiQuadFilterReadRef FilterType = InputDataRefs.GetDataReadReferenceOrConstructWithVertexDefault<FEnumEBiquadFilterType>(InputInterface, METASOUND_GET_PARAM_NAME(ParamFilterType), InParams.OperatorSettings);
 
 		return MakeUnique<FBiquadFilterOperator>(
 			InParams.OperatorSettings
