Fixing versioning with metasounds

diff --git a/Engine/Plugins/Runtime/Metasound/Source/MetasoundEngine/Private/MetasoundEngineArchetypes.cpp b/Engine/Plugins/Runtime/Metasound/Source/MetasoundEngine/Private/MetasoundEngineArchetypes.cpp
--- a/Engine/Plugins/Runtime/Metasound/Source/MetasoundEngine/Private/MetasoundEngineArchetypes.cpp
+++ b/Engine/Plugins/Runtime/Metasound/Source/MetasoundEngine/Private/MetasoundEngineArchetypes.cpp
@@ -518,7 +518,7 @@ namespace Metasound
 
 				const TArray<FMetasoundFrontendVersion> InterfacesToAdd
 				{
-					SourceInterface::GetVersion(),
+					SourceInterfaceV1_0::GetVersion(),
 					OutputFormatMonoInterface::GetVersion()
 				};
 
@@ -555,7 +555,7 @@ namespace Metasound
 
 				const TArray<FMetasoundFrontendVersion> InterfacesToAdd
 				{
-					SourceInterface::GetVersion(),
+					SourceInterfaceV1_0::GetVersion(),
 					OutputFormatStereoInterface::GetVersion()
 				};
 
diff --git a/Engine/Plugins/Runtime/Metasound/Source/MetasoundFrontend/Private/MetasoundAssetBase.cpp b/Engine/Plugins/Runtime/Metasound/Source/MetasoundFrontend/Private/MetasoundAssetBase.cpp
--- a/Engine/Plugins/Runtime/Metasound/Source/MetasoundFrontend/Private/MetasoundAssetBase.cpp
+++ b/Engine/Plugins/Runtime/Metasound/Source/MetasoundFrontend/Private/MetasoundAssetBase.cpp
@@ -371,14 +371,24 @@ bool FMetasoundAssetBase::VersionAsset()
 		bDidEdit |= FVersionDocument(AssetName, AssetPath).Transform(DocHandle);
 	}
 
-	// Version Interfaces
+	// Version Interfaces. Has to be re-run until no pass reports an update in case
+	// versions fork (ex. an interface splits into two newly named interfaces).
 	{
 		bool bInterfaceUpdated = false;
-		const TArray<FMetasoundFrontendVersion> Versions = Doc->Interfaces.Array();
-		for (const FMetasoundFrontendVersion& Version : Versions)
+		bool bPassUpdated = true;
+		while (bPassUpdated)
 		{
-			bInterfaceUpdated |= FUpdateRootGraphInterface(Version).Transform(DocHandle);
+			bPassUpdated = false;
+
+			const TArray<FMetasoundFrontendVersion> Versions = Doc->Interfaces.Array();
+			for (const FMetasoundFrontendVersion& Version : Versions)
+			{
+				bPassUpdated |= FUpdateRootGraphInterface(Version).Transform(DocHandle);
+			}
+
+			bInterfaceUpdated |= bPassUpdated;
 		}
+
 		if (bInterfaceUpdated)
 		{
 			ConformObjectDataToInterfaces();
