Minor tool tip fix-up
#rnx

diff --git a/Engine/Plugins/Runtime/Metasound/Source/MetasoundStandardNodes/Private/MetasoundOscillatorNodes.cpp b/Engine/Plugins/Runtime/Metasound/Source/MetasoundStandardNodes/Private/MetasoundOscillatorNodes.cpp
--- a/Engine/Plugins/Runtime/Metasound/Source/MetasoundStandardNodes/Private/MetasoundOscillatorNodes.cpp
+++ b/Engine/Plugins/Runtime/Metasound/Source/MetasoundStandardNodes/Private/MetasoundOscillatorNodes.cpp
@@ -1290,7 +1290,7 @@ namespace Metasound
 		DEFINE_METASOUND_ENUM_ENTRY(ELfoWaveshapeType::Square, LOCTEXT("LfoWaveShapeSquareDescription", "Square"), LOCTEXT("LfoWaveShapeSquareDescriptionTT", "Square shape Low Frequency Oscillator"))
 	DEFINE_METASOUND_ENUM_END()
 
-	// Block rate Oscillator Swiss Army Knife
+	// Blockrate All-Purpose Oscillator
 	class FLfoOperator : public TExecutableOperator<FLfoOperator>
 	{
 	public:
