Allow 0 or nearly 0 Attack Times in the block-rate metasound AD envelope node to skip the attack phase and immediately output 1.

diff --git a/Engine/Plugins/Runtime/Metasound/Source/MetasoundStandardNodes/Private/MetasoundADEnvelopeNode.cpp b/Engine/Plugins/Runtime/Metasound/Source/MetasoundStandardNodes/Private/MetasoundADEnvelopeNode.cpp
--- a/Engine/Plugins/Runtime/Metasound/Source/MetasoundStandardNodes/Private/MetasoundADEnvelopeNode.cpp
+++ b/Engine/Plugins/Runtime/Metasound/Source/MetasoundStandardNodes/Private/MetasoundADEnvelopeNode.cpp
@@ -80,8 +80,17 @@ namespace Metasound
 				// We are in attack
 				if (InState.CurrentSampleIndex < InState.AttackSampleCount)
 				{
-					float AttackFraction = (float)InState.CurrentSampleIndex++ / InState.AttackSampleCount;
-					OutEnvelopeValue = InState.StartingEnvelopeValue + (1.0f - InState.StartingEnvelopeValue) * FMath::Pow(AttackFraction, InState.AttackCurveFactor);
+					if (InState.AttackSampleCount > 1)
+					{
+						float AttackFraction = (float)InState.CurrentSampleIndex++ / InState.AttackSampleCount;
+						OutEnvelopeValue = InState.StartingEnvelopeValue + (1.0f - InState.StartingEnvelopeValue) * FMath::Pow(AttackFraction, InState.AttackCurveFactor);
+					}
+					else
+					{
+						// Attack is effectively 0, skip Attack fade-in
+						InState.CurrentSampleIndex++;
+						OutEnvelopeValue = 1.f;
+					}
 				}
 				else
 				{
