- Add support for MetaSound Editor UObject array drag/drop
- Checkpoint: Clean-up MetaSound Editor Graph Literal Customization Factory Code
#rnx

diff --git a/Engine/Plugins/Runtime/Metasound/Source/MetasoundEditor/MetasoundEditor.Build.cs b/Engine/Plugins/Runtime/Metasound/Source/MetasoundEditor/MetasoundEditor.Build.cs
--- a/Engine/Plugins/Runtime/Metasound/Source/MetasoundEditor/MetasoundEditor.Build.cs
+++ b/Engine/Plugins/Runtime/Metasound/Source/MetasoundEditor/MetasoundEditor.Build.cs
@@ -22,7 +22,8 @@ namespace UnrealBuildTool.Rules
 					"ToolMenus",
 					"ToolWidgets",
 					"AudioWidgets",
-					"AudioSynesthesia"
+					"AudioSynesthesia",
+					"EditorWidgets"
 				}
 			);
 
diff --git a/Engine/Plugins/Runtime/Metasound/Source/MetasoundEditor/Private/MetasoundEditorModule.cpp b/Engine/Plugins/Runtime/Metasound/Source/MetasoundEditor/Private/MetasoundEditorModule.cpp
--- a/Engine/Plugins/Runtime/Metasound/Source/MetasoundEditor/Private/MetasoundEditorModule.cpp
+++ b/Engine/Plugins/Runtime/Metasound/Source/MetasoundEditor/Private/MetasoundEditorModule.cpp
@@ -467,7 +467,7 @@ namespace Metasound
 				return ExplicitProxyClasses.Contains(&InClass);
 			}
 
-			virtual TUniquePtr<IMemberDefaultLiteralCustomization> CreateMemberDefaultLiteralCustomization(UClass& InClass, IDetailCategoryBuilder& InDefaultCategoryBuilder) const override
+			virtual TUniquePtr<FMetasoundDefaultLiteralCustomizationBase> CreateMemberDefaultLiteralCustomization(UClass& InClass, IDetailCategoryBuilder& InDefaultCategoryBuilder) const override
 			{
 				const TUniquePtr<IMemberDefaultLiteralCustomizationFactory>* CustomizationFactory = LiteralCustomizationFactories.Find(&InClass);
 				if (CustomizationFactory && CustomizationFactory->IsValid())
@@ -554,7 +554,9 @@ namespace Metasound
 					"MetaSoundEditorGraphMemberDefaultObjectRef",
 					FOnGetPropertyTypeCustomizationInstance::CreateLambda([]() { return MakeShared<FMetasoundMemberDefaultObjectDetailCustomization>(); }));
 
+				LiteralCustomizationFactories.Add(UMetasoundEditorGraphMemberDefaultLiteral::StaticClass(), MakeUnique<FMetasoundDefaultLiteralCustomizationFactory>());
 				LiteralCustomizationFactories.Add(UMetasoundEditorGraphMemberDefaultFloat::StaticClass(), MakeUnique<FMetasoundFloatLiteralCustomizationFactory>());
+				LiteralCustomizationFactories.Add(UMetasoundEditorGraphMemberDefaultObjectArray::StaticClass(), MakeUnique<FMetasoundObjectArrayLiteralCustomizationFactory>());
 
 				StyleSet = MakeShared<FSlateStyle>();
 
diff --git a/Engine/Plugins/Runtime/Metasound/Source/MetasoundEditor/Private/MetasoundNodeDetailCustomization.cpp b/Engine/Plugins/Runtime/Metasound/Source/MetasoundEditor/Private/MetasoundNodeDetailCustomization.cpp
--- a/Engine/Plugins/Runtime/Metasound/Source/MetasoundEditor/Private/MetasoundNodeDetailCustomization.cpp
+++ b/Engine/Plugins/Runtime/Metasound/Source/MetasoundEditor/Private/MetasoundNodeDetailCustomization.cpp
@@ -30,6 +30,7 @@
 #include "PropertyEditorDelegates.h"
 #include "PropertyHandle.h"
 #include "PropertyRestriction.h"
+#include "SAssetDropTarget.h"
 #include "SlateCore/Public/Styling/SlateColor.h"
 #include "SMetasoundActionMenu.h"
 #include "SMetasoundGraphNode.h"
@@ -64,6 +65,34 @@ namespace Metasound
 				"Audio:Mono",
 				"Audio:Stereo"
 			};
+
+			void GetDataTypeFromElementPropertyHandle(TSharedPtr<IPropertyHandle> ElementPropertyHandle, Frontend::FDataTypeRegistryInfo& OutDataTypeInfo)
+			{
+				using namespace Frontend;
+
+				TArray<UObject*>OuterObjects;
+				ElementPropertyHandle->GetOuterObjects(OuterObjects);
+				if (OuterObjects.Num() == 1)
+				{
+					UObject* Outer = OuterObjects.Last();
+					if (const UMetasoundEditorGraphMemberDefaultLiteral* DefaultLiteral = Cast<UMetasoundEditorGraphMemberDefaultLiteral>(Outer))
+					{
+						if (const UMetasoundEditorGraphMember* Member = Cast<UMetasoundEditorGraphMember>(DefaultLiteral->GetOuter()))
+						{
+							ensure(IDataTypeRegistry::Get().GetDataTypeInfo(Member->GetDataType(), OutDataTypeInfo));
+						}
+					}
+				}
+			}
+
+			// If DataType is an array type, creates & returns the array's
+			// element type. Otherwise, returns this type's DataTypeName.
+			FName GetPrimitiveTypeName(const Frontend::FDataTypeRegistryInfo& InDataTypeInfo)
+			{
+				return InDataTypeInfo.IsArrayType()
+					? CreateElementTypeNameFromArrayTypeName(InDataTypeInfo.DataTypeName)
+					: InDataTypeInfo.DataTypeName;
+			}
 		} // namespace MemberCustomizationPrivate
 
 		FMetasoundFloatLiteralCustomization::~FMetasoundFloatLiteralCustomization()
@@ -75,9 +104,9 @@ namespace Metasound
 			}
 		}
 
-		void FMetasoundFloatLiteralCustomization::CustomizeLiteral(UMetasoundEditorGraphMemberDefaultLiteral& InLiteral, TSharedPtr<IPropertyHandle> InDefaultValueHandle)
+		void FMetasoundFloatLiteralCustomization::CustomizeLiteral(UMetasoundEditorGraphMemberDefaultLiteral& InLiteral, IDetailLayoutBuilder& InDetailLayout)
 		{
-			check(InputCategoryBuilder);
+			check(DefaultCategoryBuilder);
 
 			UMetasoundEditorGraphMemberDefaultFloat* DefaultFloat = Cast<UMetasoundEditorGraphMemberDefaultFloat>(&InLiteral);
 			if (!ensure(DefaultFloat))
@@ -86,19 +115,27 @@ namespace Metasound
 			}
 			FloatLiteral = DefaultFloat;
 
-			if (IDetailPropertyRow* Row = InputCategoryBuilder->AddExternalObjectProperty(TArray<UObject*>({ DefaultFloat }), GET_MEMBER_NAME_CHECKED(UMetasoundEditorGraphMemberDefaultFloat, ClampDefault)))
+			TSharedPtr<IPropertyHandle> DefaultValueHandle;
+			IDetailPropertyRow* Row = DefaultCategoryBuilder->AddExternalObjectProperty(TArray<UObject*>({ DefaultFloat }), UMetasoundEditorGraphMemberDefaultFloat::GetDefaultPropertyName());
+			if (ensure(Row))
+			{
+				DefaultValueHandle = Row->GetPropertyHandle();
+			}
+
+			Row = DefaultCategoryBuilder->AddExternalObjectProperty(TArray<UObject*>({ DefaultFloat }), GET_MEMBER_NAME_CHECKED(UMetasoundEditorGraphMemberDefaultFloat, ClampDefault));
+			if (ensure(Row))
 			{
 				// If clamping or using slider, clamp default value to given range 
 				if (DefaultFloat->ClampDefault || DefaultFloat->WidgetType != EMetasoundMemberDefaultWidget::None)
 				{
 					FVector2D Range = DefaultFloat->GetRange();
-					InDefaultValueHandle->SetInstanceMetaData("ClampMin", FString::Printf(TEXT("%f"), Range.X));
-					InDefaultValueHandle->SetInstanceMetaData("ClampMax", FString::Printf(TEXT("%f"), Range.Y));
+					DefaultValueHandle->SetInstanceMetaData("ClampMin", FString::Printf(TEXT("%f"), Range.X));
+					DefaultValueHandle->SetInstanceMetaData("ClampMax", FString::Printf(TEXT("%f"), Range.Y));
 				}
 				else // Stop clamping
 				{
-					InDefaultValueHandle->SetInstanceMetaData("ClampMin", "");
-					InDefaultValueHandle->SetInstanceMetaData("ClampMax", "");
+					DefaultValueHandle->SetInstanceMetaData("ClampMin", "");
+					DefaultValueHandle->SetInstanceMetaData("ClampMax", "");
 				}
 
 				DefaultFloat->OnClampChanged.Remove(OnClampChangedDelegateHandle);
@@ -114,23 +151,93 @@ namespace Metasound
 					}
 				});
 			}
-			InputCategoryBuilder->AddExternalObjectProperty(TArray<UObject*>({ DefaultFloat }), GET_MEMBER_NAME_CHECKED(UMetasoundEditorGraphMemberDefaultFloat, Range));
+			DefaultCategoryBuilder->AddExternalObjectProperty(TArray<UObject*>({ DefaultFloat }), GET_MEMBER_NAME_CHECKED(UMetasoundEditorGraphMemberDefaultFloat, Range));
+
+			// add input widget properties
+			IDetailCategoryBuilder& WidgetCategoryBuilder = InDetailLayout.EditCategory("EditorOptions");
+			WidgetCategoryBuilder.AddExternalObjectProperty(TArray<UObject*>({ DefaultFloat }), GET_MEMBER_NAME_CHECKED(UMetasoundEditorGraphMemberDefaultFloat, WidgetType));
+			WidgetCategoryBuilder.AddExternalObjectProperty(TArray<UObject*>({ DefaultFloat }), GET_MEMBER_NAME_CHECKED(UMetasoundEditorGraphMemberDefaultFloat, WidgetOrientation));
+			WidgetCategoryBuilder.AddExternalObjectProperty(TArray<UObject*>({ DefaultFloat }), GET_MEMBER_NAME_CHECKED(UMetasoundEditorGraphMemberDefaultFloat, WidgetValueType));
 		}
 
-		void FMetasoundMemberDefaultBoolDetailCustomization::CacheProxyData(TSharedPtr<IPropertyHandle> InPropertyHandle)
+		void FMetasoundObjectArrayLiteralCustomization::CustomizeLiteral(UMetasoundEditorGraphMemberDefaultLiteral& InLiteral, IDetailLayoutBuilder& InDetailLayout)
 		{
-			DataTypeName = FName();
+			check(DefaultCategoryBuilder);
 
-			const FString* MetadataDataTypeName = InPropertyHandle->GetInstanceMetaData(MemberCustomizationStyle::DataTypeNameIdentifier);
-			if (ensure(MetadataDataTypeName))
+			TSharedPtr<IPropertyHandle> DefaultValueHandle;
+			IDetailPropertyRow* Row = DefaultCategoryBuilder->AddExternalObjectProperty(TArray<UObject*>({ &InLiteral }), GET_MEMBER_NAME_CHECKED(UMetasoundEditorGraphMemberDefaultObjectArray, Default));
+			if (ensure(Row))
 			{
-				DataTypeName = **MetadataDataTypeName;
+				DefaultValueHandle = Row->GetPropertyHandle();
 			}
+
+			constexpr bool bShowChildren = true;
+			Row->ShowPropertyButtons(false)
+			.CustomWidget(bShowChildren)
+			.NameContent()
+			[
+				DefaultValueHandle->CreatePropertyNameWidget()
+			]
+			.ValueContent()
+			[
+				SNew(SAssetDropTarget)
+				.bSupportsMultiDrop(true)
+				.OnAreAssetsAcceptableForDropWithReason_Lambda([this, DefaultValueHandle](TArrayView<FAssetData> InAssets, FText& OutReason)
+				{
+					Frontend::FDataTypeRegistryInfo DataTypeInfo;
+					MemberCustomizationPrivate::GetDataTypeFromElementPropertyHandle(DefaultValueHandle, DataTypeInfo);
+
+					const IMetasoundEditorModule& EditorModule = FModuleManager::GetModuleChecked<IMetasoundEditorModule>("MetaSoundEditor");
+					bool bCanDrop = true;
+					for (const FAssetData& AssetData : InAssets)
+					{
+						if (DataTypeInfo.ProxyGeneratorClass)
+						{
+							if (UClass* Class = AssetData.GetClass())
+							{
+								if (EditorModule.IsExplicitProxyClass(*DataTypeInfo.ProxyGeneratorClass))
+								{
+									bCanDrop &= Class == DataTypeInfo.ProxyGeneratorClass;
+								}
+								else
+								{
+									bCanDrop &= Class->IsChildOf(DataTypeInfo.ProxyGeneratorClass);
+								}
+							}
+						}
+					}
+					return true;
+				})
+				.OnAssetsDropped_Lambda([this, DefaultValueHandle](const FDragDropEvent& DragDropEvent, TArrayView<FAssetData> InAssets)
+				{
+					if (DefaultValueHandle.IsValid())
+					{
+						TSharedPtr<IPropertyHandleArray> ArrayHandle = DefaultValueHandle->AsArray();
+						if (ensure(ArrayHandle.IsValid()))
+						{
+							for (const FAssetData& AssetData : InAssets)
+							{
+								uint32 AddIndex = INDEX_NONE;
+								ArrayHandle->GetNumElements(AddIndex);
+								ArrayHandle->AddItem();
+								TSharedPtr<IPropertyHandle> ElementHandle = ArrayHandle->GetElement(static_cast<int32>(AddIndex));
+								TSharedPtr<IPropertyHandle> ObjectHandle = ElementHandle->GetChildHandle(GET_MEMBER_NAME_CHECKED(FMetasoundEditorGraphMemberDefaultObjectRef, Object));
+								ObjectHandle->SetValue(AssetData.GetAsset());
+							}
+						}
+					}
+				})
+				[
+					DefaultValueHandle->CreatePropertyValueWidget()
+				]
+			];
 		}
 
 		FText FMetasoundMemberDefaultBoolDetailCustomization::GetPropertyNameOverride() const
 		{
-			if (DataTypeName == Metasound::GetMetasoundDataTypeName<Metasound::FTrigger>())
+			using namespace MemberCustomizationPrivate;
+
+			if (GetPrimitiveTypeName(DataTypeInfo) == Metasound::GetMetasoundDataTypeName<Metasound::FTrigger>())
 			{
 				return LOCTEXT("TriggerInput_SimulateTitle", "Simulate");
 			}
@@ -141,12 +248,13 @@ namespace Metasound
 		TSharedRef<SWidget> FMetasoundMemberDefaultBoolDetailCustomization::CreateStructureWidget(TSharedPtr<IPropertyHandle>& StructPropertyHandle) const
 		{
 			using namespace Frontend;
+			using namespace MemberCustomizationPrivate;
 
 			TSharedPtr<IPropertyHandle> ValueProperty = StructPropertyHandle->GetChildHandle(GET_MEMBER_NAME_CHECKED(FMetasoundEditorGraphMemberDefaultBoolRef, Value));
 			if (ValueProperty.IsValid())
 			{
 				// Not a trigger, so just display as underlying literal type (bool)
-				if (DataTypeName != Metasound::GetMetasoundDataTypeName<Metasound::FTrigger>())
+				if (GetPrimitiveTypeName(DataTypeInfo) != Metasound::GetMetasoundDataTypeName<Metasound::FTrigger>())
 				{
 					return ValueProperty->CreatePropertyValueWidget();
 				}
@@ -182,27 +290,17 @@ namespace Metasound
 			return SNullWidget::NullWidget;
 		}
 
-		void FMetasoundMemberDefaultIntDetailCustomization::CacheProxyData(TSharedPtr<IPropertyHandle> ProxyHandle)
-		{
-			DataTypeName = FName();
-
-			const FString* MetadataDataTypeName = ProxyHandle->GetInstanceMetaData(MemberCustomizationStyle::DataTypeNameIdentifier);
-			if (ensure(MetadataDataTypeName))
-			{
-				DataTypeName = **MetadataDataTypeName;
-			}
-		}
-
 		TSharedRef<SWidget> FMetasoundMemberDefaultIntDetailCustomization::CreateStructureWidget(TSharedPtr<IPropertyHandle>& StructPropertyHandle) const
 		{
 			using namespace Frontend;
+			using namespace MemberCustomizationPrivate;
 
 			if (FMetasoundFrontendRegistryContainer* Registry = FMetasoundFrontendRegistryContainer::Get())
 			{
 				TSharedPtr<IPropertyHandle> ValueProperty = StructPropertyHandle->GetChildHandle(GET_MEMBER_NAME_CHECKED(FMetasoundEditorGraphMemberDefaultIntRef, Value));
 				if (ValueProperty.IsValid())
 				{
-					TSharedPtr<const IEnumDataTypeInterface> EnumInterface = IDataTypeRegistry::Get().GetEnumInterfaceForDataType(DataTypeName);
+					TSharedPtr<const IEnumDataTypeInterface> EnumInterface = IDataTypeRegistry::Get().GetEnumInterfaceForDataType(GetPrimitiveTypeName(DataTypeInfo));
 
 					// Not an enum, so just display as underlying type (int32)
 					if (!EnumInterface.IsValid())
@@ -275,60 +373,23 @@ namespace Metasound
 			return SNullWidget::NullWidget;
 		}
 
-		void FMetasoundMemberDefaultObjectDetailCustomization::CacheProxyData(TSharedPtr<IPropertyHandle> ProxyHandle)
-		{
-			ProxyGenClass.Reset();
-
-			const FString* MetadataProxyGenClass = ProxyHandle->GetInstanceMetaData(MemberCustomizationStyle::ProxyGeneratorClassNameIdentifier);
-			TSharedPtr<IPropertyHandle> MetadataHandle = ProxyHandle->GetParentHandle();
-			if (!ensure(MetadataProxyGenClass))
-			{
-				return;
-			}
-
-			const FName ClassName = FName(*MetadataProxyGenClass);
-			for (TObjectIterator<UClass> ClassIt; ClassIt; ++ClassIt)
-			{
-				UClass* Class = *ClassIt;
-				if (!Class->IsNative())
-				{
-					continue;
-				}
-				
-				if (Class->HasAnyClassFlags(CLASS_Deprecated | CLASS_NewerVersionExists))
-				{
-					continue;
-				}
-
-				if (ClassIt->GetFName() != ClassName)
-				{
-					continue;
-				}
-
-				ProxyGenClass = *ClassIt;
-				return;
-			}
-
-			ensureMsgf(false, TEXT("Failed to find ProxyGeneratorClass. Class not set "));
-		}
-
 		TSharedRef<SWidget> FMetasoundMemberDefaultObjectDetailCustomization::CreateStructureWidget(TSharedPtr<IPropertyHandle>& StructPropertyHandle) const
 		{
 			TSharedPtr<IPropertyHandle> PropertyHandle = StructPropertyHandle->GetChildHandle(GET_MEMBER_NAME_CHECKED(FMetasoundEditorGraphMemberDefaultObjectRef, Object));
 
 			const IMetasoundEditorModule& EditorModule = FModuleManager::GetModuleChecked<IMetasoundEditorModule>("MetaSoundEditor");
-			auto FilterAsset = [InEditorModule = &EditorModule, InProxyGenClass = ProxyGenClass](const FAssetData& InAsset)
+			auto FilterAsset = [InEditorModule = &EditorModule, InProxyGenClass = DataTypeInfo.ProxyGeneratorClass](const FAssetData& InAsset)
 			{
-				if (InProxyGenClass.IsValid())
+				if (InProxyGenClass)
 				{
 					if (UClass* Class = InAsset.GetClass())
 					{
-						if (InEditorModule->IsExplicitProxyClass(*InProxyGenClass.Get()))
+						if (InEditorModule->IsExplicitProxyClass(*InProxyGenClass))
 						{
-							return Class != InProxyGenClass.Get();
+							return Class != InProxyGenClass;
 						}
 
-						return !Class->IsChildOf(InProxyGenClass.Get());
+						return !Class->IsChildOf(InProxyGenClass);
 					}
 				}
 
@@ -337,7 +398,8 @@ namespace Metasound
 
 			auto ValidateAsset = [FilterAsset](const FAssetData& InAsset)
 			{
-				return !FilterAsset(InAsset);
+				// A null asset reference is a valid default
+				return InAsset.IsValid() ? !FilterAsset(InAsset) : true;
 			};
 
 			auto GetAssetPath = [PropertyHandle = PropertyHandle]()
@@ -352,18 +414,18 @@ namespace Metasound
 
 			return SNew(SObjectPropertyEntryBox)
 				.AllowClear(true)
-				.AllowedClass(ProxyGenClass.Get())
+				.AllowedClass(DataTypeInfo.ProxyGeneratorClass)
 				.DisplayBrowse(true)
 				.DisplayThumbnail(true)
 				.DisplayUseSelected(true)
-				.NewAssetFactories(PropertyCustomizationHelpers::GetNewAssetFactoriesForClasses({ ProxyGenClass.Get() }))
+				.NewAssetFactories(PropertyCustomizationHelpers::GetNewAssetFactoriesForClasses({DataTypeInfo.ProxyGeneratorClass }))
 				.ObjectPath_Lambda(GetAssetPath)
 				.OnShouldFilterAsset_Lambda(FilterAsset)
 				.OnShouldSetAsset_Lambda(ValidateAsset)
 				.PropertyHandle(PropertyHandle);
 		}
 
-		TSharedRef<SWidget> FMetasoundInputArrayDetailCustomizationBase::CreateNameWidget(TSharedPtr<IPropertyHandle> StructPropertyHandle) const
+		TSharedRef<SWidget> FMetasoundDefaultMemberElementDetailCustomizationBase::CreateNameWidget(TSharedPtr<IPropertyHandle> StructPropertyHandle) const
 		{
 			const FText PropertyName = GetPropertyNameOverride();
 			if (!PropertyName.IsEmpty())
@@ -378,39 +440,39 @@ namespace Metasound
 				.Font(IDetailLayoutBuilder::GetDetailFont());
 		}
 
-		TSharedRef<SWidget> FMetasoundInputArrayDetailCustomizationBase::CreateValueWidget(TSharedPtr<IPropertyHandleArray> ParentArrayProperty, TSharedPtr<IPropertyHandle> StructPropertyHandle, bool bIsInArray) const
+		TSharedRef<SWidget> FMetasoundDefaultMemberElementDetailCustomizationBase::CreateValueWidget(TSharedPtr<IPropertyHandleArray> ParentPropertyHandleArray, TSharedPtr<IPropertyHandle> StructPropertyHandle) const
 		{
 			TSharedRef<SWidget> ValueWidget = CreateStructureWidget(StructPropertyHandle);
-			if (!bIsInArray)
+			if (!ParentPropertyHandleArray.IsValid())
 			{
 				return ValueWidget;
 			}
 
 			TSharedPtr<IPropertyHandle> StructPropertyPtr = StructPropertyHandle;
-			FExecuteAction InsertAction = FExecuteAction::CreateLambda([ParentArrayProperty, StructPropertyPtr]
+			FExecuteAction InsertAction = FExecuteAction::CreateLambda([ParentPropertyHandleArray, StructPropertyPtr]
 			{
 				const int32 ArrayIndex = StructPropertyPtr.IsValid() ? StructPropertyPtr->GetIndexInArray() : INDEX_NONE;
-				if (ParentArrayProperty.IsValid() && ArrayIndex >= 0)
+				if (ParentPropertyHandleArray.IsValid() && ArrayIndex >= 0)
 				{
-					ParentArrayProperty->Insert(ArrayIndex);
+					ParentPropertyHandleArray->Insert(ArrayIndex);
 				}
 			});
 
-			FExecuteAction DeleteAction = FExecuteAction::CreateLambda([ParentArrayProperty, StructPropertyPtr]
+			FExecuteAction DeleteAction = FExecuteAction::CreateLambda([ParentPropertyHandleArray, StructPropertyPtr]
 			{
 				const int32 ArrayIndex = StructPropertyPtr.IsValid() ? StructPropertyPtr->GetIndexInArray() : INDEX_NONE;
-				if (ParentArrayProperty.IsValid() && ArrayIndex >= 0)
+				if (ParentPropertyHandleArray.IsValid() && ArrayIndex >= 0)
 				{
-					ParentArrayProperty->DeleteItem(ArrayIndex);
+					ParentPropertyHandleArray->DeleteItem(ArrayIndex);
 				}
 			});
 
-			FExecuteAction DuplicateAction = FExecuteAction::CreateLambda([ParentArrayProperty, StructPropertyPtr]
+			FExecuteAction DuplicateAction = FExecuteAction::CreateLambda([ParentPropertyHandleArray, StructPropertyPtr]
 			{
 				const int32 ArrayIndex = StructPropertyPtr.IsValid() ? StructPropertyPtr->GetIndexInArray() : INDEX_NONE;
-				if (ParentArrayProperty.IsValid() && ArrayIndex >= 0)
+				if (ParentPropertyHandleArray.IsValid() && ArrayIndex >= 0)
 				{
-					ParentArrayProperty->DuplicateItem(ArrayIndex);
+					ParentPropertyHandleArray->DuplicateItem(ArrayIndex);
 				}
 			});
 
@@ -431,29 +493,28 @@ namespace Metasound
 				];
 		}
 
-		void FMetasoundInputArrayDetailCustomizationBase::CustomizeChildren(TSharedRef<IPropertyHandle> StructPropertyHandle, IDetailChildrenBuilder& ChildBuilder, IPropertyTypeCustomizationUtils& StructCustomizationUtils)
+		void FMetasoundDefaultMemberElementDetailCustomizationBase::CustomizeChildren(TSharedRef<IPropertyHandle> StructPropertyHandle, IDetailChildrenBuilder& ChildBuilder, IPropertyTypeCustomizationUtils& StructCustomizationUtils)
 		{
-			bool bIsInArray = false;
-			TSharedPtr<IPropertyHandleArray> ParentArrayProperty;
-			TSharedPtr<IPropertyHandle> ProxyProperty = StructPropertyHandle;
+			TSharedPtr<IPropertyHandleArray> ParentPropertyHandleArray;
+			TSharedPtr<IPropertyHandle> ElementPropertyHandle = StructPropertyHandle;
+			if (ElementPropertyHandle.IsValid())
 			{
-				TSharedPtr<IPropertyHandle> ParentProperty = ProxyProperty->GetParentHandle();
-				if (ProxyProperty.IsValid() && ParentProperty.IsValid())
+				TSharedPtr<IPropertyHandle> ParentProperty = ElementPropertyHandle->GetParentHandle();
+				while (ParentProperty.IsValid() && ParentProperty->GetProperty() != nullptr)
 				{
-					ParentArrayProperty = ParentProperty->AsArray();
-					if (ParentArrayProperty.IsValid())
+					ParentPropertyHandleArray = ParentProperty->AsArray();
+					if (ParentPropertyHandleArray.IsValid())
 					{
-						ProxyProperty = ParentProperty;
-						bIsInArray = true;
+						ElementPropertyHandle = ParentProperty;
+						break;
 					}
 				}
 			}
+			MemberCustomizationPrivate::GetDataTypeFromElementPropertyHandle(ElementPropertyHandle, DataTypeInfo);
 
-			CacheProxyData(ProxyProperty);
-
-			TSharedRef<SWidget> ValueWidget = CreateValueWidget(ParentArrayProperty, StructPropertyHandle, bIsInArray);
+			TSharedRef<SWidget> ValueWidget = CreateValueWidget(ParentPropertyHandleArray, StructPropertyHandle);
 			FDetailWidgetRow& ValueRow = ChildBuilder.AddCustomRow(MemberCustomizationStyle::DefaultPropertyText);
-			if (bIsInArray)
+			if (ParentPropertyHandleArray.IsValid())
 			{
 				ValueRow.NameContent()
 				[
@@ -498,7 +559,7 @@ namespace Metasound
 			];
 		}
 
-		void FMetasoundInputArrayDetailCustomizationBase::CustomizeHeader(TSharedRef<IPropertyHandle> StructPropertyHandle, FDetailWidgetRow& HeaderRow, IPropertyTypeCustomizationUtils& StructCustomizationUtils)
+		void FMetasoundDefaultMemberElementDetailCustomizationBase::CustomizeHeader(TSharedRef<IPropertyHandle> StructPropertyHandle, FDetailWidgetRow& HeaderRow, IPropertyTypeCustomizationUtils& StructCustomizationUtils)
 		{
 		}
 
@@ -744,13 +805,6 @@ namespace Metasound
 			return false;
 		}
 
-		void FMetasoundInputDetailCustomization::CustomizeDetails(IDetailLayoutBuilder& DetailLayout)
-		{
-			using namespace Frontend;
-
-			TMetasoundGraphMemberDetailCustomization<UMetasoundEditorGraphInput, FMetasoundInputDetailCustomization>::CustomizeDetails(DetailLayout);
-		}
-
 		const FText FMetasoundOutputDetailCustomization::MemberNameText = LOCTEXT("OutputGraphMemberLabel", "Output");
 
 		bool FMetasoundOutputDetailCustomization::IsInterfaceMember() const
@@ -762,11 +816,6 @@ namespace Metasound
 
 			return false;
 		}
-
-		void FMetasoundOutputDetailCustomization::CustomizeDetails(IDetailLayoutBuilder& DetailLayout)
-		{
-			TMetasoundGraphMemberDetailCustomization<UMetasoundEditorGraphOutput, FMetasoundOutputDetailCustomization>::CustomizeDetails(DetailLayout);
-		}
 	} // namespace Editor
 } // namespace Metasound
 #undef LOCTEXT_NAMESPACE
diff --git a/Engine/Plugins/Runtime/Metasound/Source/MetasoundEditor/Private/MetasoundNodeDetailCustomization.h b/Engine/Plugins/Runtime/Metasound/Source/MetasoundEditor/Private/MetasoundNodeDetailCustomization.h
--- a/Engine/Plugins/Runtime/Metasound/Source/MetasoundEditor/Private/MetasoundNodeDetailCustomization.h
+++ b/Engine/Plugins/Runtime/Metasound/Source/MetasoundEditor/Private/MetasoundNodeDetailCustomization.h
@@ -6,11 +6,13 @@
 #include "DetailLayoutBuilder.h"
 #include "DetailWidgetRow.h"
 #include "EdGraph/EdGraphSchema.h"
+#include "Framework/Application/SlateApplication.h"
 #include "IDetailCustomization.h"
 #include "IPropertyTypeCustomization.h"
 #include "Layout/Visibility.h"
 #include "MetasoundAssetBase.h"
 #include "MetasoundEditorGraph.h"
+#include "MetasoundEditorGraphBuilder.h"
 #include "MetasoundEditorGraphMemberDefaults.h"
 #include "MetasoundEditorGraphNode.h"
 #include "MetasoundFrontendLiteral.h"
@@ -45,14 +47,10 @@ namespace Metasound
 			static const FText DataTypeNameText = LOCTEXT("Node_DataTypeName", "Type");
 			static const FText DefaultPropertyText = LOCTEXT("Node_DefaultPropertyName", "Default Value");
 			static const FText NodeTooltipText = LOCTEXT("Node_Tooltip", "Tooltip");
-
-			static const FName DataTypeNameIdentifier = "DataTypeName";
-			static const FName ProxyGeneratorClassNameIdentifier = "GeneratorClass";
 		} // namespace MemberCustomizationStyle
 
-		class FMetasoundFloatLiteralCustomization : public IMemberDefaultLiteralCustomization
+		class FMetasoundFloatLiteralCustomization : public FMetasoundDefaultLiteralCustomizationBase
 		{
-			IDetailCategoryBuilder* InputCategoryBuilder = nullptr;
 			TWeakObjectPtr<UMetasoundEditorGraphMemberDefaultFloat> FloatLiteral;
 
 			// Delegate for updating the clamp min/max of the input value when the slider range is changed 
@@ -62,26 +60,59 @@ namespace Metasound
 			FDelegateHandle OnClampChangedDelegateHandle;
 
 		public:
-			FMetasoundFloatLiteralCustomization(IDetailCategoryBuilder& InInputCategoryBuilder)
-				: InputCategoryBuilder(&InInputCategoryBuilder)
+			FMetasoundFloatLiteralCustomization(IDetailCategoryBuilder& InDefaultCategoryBuilder)
+				: FMetasoundDefaultLiteralCustomizationBase(InDefaultCategoryBuilder)
 			{
 			}
-
 			virtual ~FMetasoundFloatLiteralCustomization();
 
-			virtual void CustomizeLiteral(UMetasoundEditorGraphMemberDefaultLiteral& InLiteral, TSharedPtr<IPropertyHandle> InDefaultValueHandle) override;
+			virtual void CustomizeLiteral(UMetasoundEditorGraphMemberDefaultLiteral& InLiteral, IDetailLayoutBuilder& InDetailLayout) override;
+		};
+
+		// Customization to support drag-and-drop of Proxy UObject types on underlying members that are structs.
+		// Struct ownership of objects required to customize asset filters based on dynamic UObject MetaSound Registry DataTypes.
+		class FMetasoundObjectArrayLiteralCustomization : public FMetasoundDefaultLiteralCustomizationBase
+		{
+		public:
+			FMetasoundObjectArrayLiteralCustomization(IDetailCategoryBuilder& InDefaultCategoryBuilder)
+				: FMetasoundDefaultLiteralCustomizationBase(InDefaultCategoryBuilder)
+			{
+			}
+
+			virtual ~FMetasoundObjectArrayLiteralCustomization() = default;
+
+			virtual void CustomizeLiteral(UMetasoundEditorGraphMemberDefaultLiteral& InLiteral, IDetailLayoutBuilder& InDetailLayout) override;
 		};
 
+		class FMetasoundDefaultLiteralCustomizationFactory : public IMemberDefaultLiteralCustomizationFactory
+		{
+		public:
+			virtual TUniquePtr<FMetasoundDefaultLiteralCustomizationBase> CreateLiteralCustomization(IDetailCategoryBuilder& DefaultCategoryBuilder) const override
+			{
+				return TUniquePtr<FMetasoundDefaultLiteralCustomizationBase>(new FMetasoundDefaultLiteralCustomizationBase(DefaultCategoryBuilder));
+			}
+		};
+
+		// Customization to support float widgets (ex. sliders, knobs)
 		class FMetasoundFloatLiteralCustomizationFactory : public IMemberDefaultLiteralCustomizationFactory
 		{
 		public:
-			virtual TUniquePtr<IMemberDefaultLiteralCustomization> CreateLiteralCustomization(IDetailCategoryBuilder& DefaultCategoryBuilder) const override
+			virtual TUniquePtr<FMetasoundDefaultLiteralCustomizationBase> CreateLiteralCustomization(IDetailCategoryBuilder& DefaultCategoryBuilder) const override
 			{
-				return TUniquePtr<IMemberDefaultLiteralCustomization>(new FMetasoundFloatLiteralCustomization(DefaultCategoryBuilder));
+				return TUniquePtr<FMetasoundDefaultLiteralCustomizationBase>(new FMetasoundFloatLiteralCustomization(DefaultCategoryBuilder));
 			}
 		};
 
-		class FMetasoundInputArrayDetailCustomizationBase : public IPropertyTypeCustomization
+		class FMetasoundObjectArrayLiteralCustomizationFactory : public IMemberDefaultLiteralCustomizationFactory
+		{
+		public:
+			virtual TUniquePtr<FMetasoundDefaultLiteralCustomizationBase> CreateLiteralCustomization(IDetailCategoryBuilder& DefaultCategoryBuilder) const override
+			{
+				return TUniquePtr<FMetasoundDefaultLiteralCustomizationBase>(new FMetasoundObjectArrayLiteralCustomization(DefaultCategoryBuilder));
+			}
+		};
+
+		class FMetasoundDefaultMemberElementDetailCustomizationBase : public IPropertyTypeCustomization
 		{
 		public:
 			//~ Begin IPropertyTypeCustomization
@@ -91,43 +122,34 @@ namespace Metasound
 
 		protected:
 			virtual FText GetPropertyNameOverride() const { return FText::GetEmpty(); }
+
+			// TODO: Merge with FMetasoundDefaultLiteralCustomizationBaseFactory
 			virtual TSharedRef<SWidget> CreateStructureWidget(TSharedPtr<IPropertyHandle>& PropertyHandle) const = 0;
-			virtual void CacheProxyData(TSharedPtr<IPropertyHandle> ProxyHandle) { }
+
+			Frontend::FDataTypeRegistryInfo DataTypeInfo;
 
 		private:
 			TSharedRef<SWidget> CreateNameWidget(TSharedPtr<IPropertyHandle> StructPropertyHandle) const;
-			TSharedRef<SWidget> CreateValueWidget(TSharedPtr<IPropertyHandleArray> ParentArrayProperty, TSharedPtr<IPropertyHandle> StructPropertyHandle, bool bIsInArray) const;
+			TSharedRef<SWidget> CreateValueWidget(TSharedPtr<IPropertyHandleArray> ParentArrayProperty, TSharedPtr<IPropertyHandle> StructPropertyHandle) const;
 		};
 
-		class FMetasoundMemberDefaultBoolDetailCustomization : public FMetasoundInputArrayDetailCustomizationBase
+		class FMetasoundMemberDefaultBoolDetailCustomization : public FMetasoundDefaultMemberElementDetailCustomizationBase
 		{
 		protected:
 			virtual FText GetPropertyNameOverride() const override;
 			virtual TSharedRef<SWidget> CreateStructureWidget(TSharedPtr<IPropertyHandle>& StructPropertyHandle) const override;
-			virtual void CacheProxyData(TSharedPtr<IPropertyHandle> ProxyHandle) override;
-
-		private:
-			FName DataTypeName;
 		};
 
-		class FMetasoundMemberDefaultIntDetailCustomization : public FMetasoundInputArrayDetailCustomizationBase
+		class FMetasoundMemberDefaultIntDetailCustomization : public FMetasoundDefaultMemberElementDetailCustomizationBase
 		{
 		protected:
 			virtual TSharedRef<SWidget> CreateStructureWidget(TSharedPtr<IPropertyHandle>& StructPropertyHandle) const override;
-			virtual void CacheProxyData(TSharedPtr<IPropertyHandle> ProxyHandle) override;
-
-		private:
-			FName DataTypeName;
 		};
 
-		class FMetasoundMemberDefaultObjectDetailCustomization : public FMetasoundInputArrayDetailCustomizationBase
+		class FMetasoundMemberDefaultObjectDetailCustomization : public FMetasoundDefaultMemberElementDetailCustomizationBase
 		{
 		protected:
 			virtual TSharedRef<SWidget> CreateStructureWidget(TSharedPtr<IPropertyHandle>& StructPropertyHandle) const override;
-			virtual void CacheProxyData(TSharedPtr<IPropertyHandle> ProxyHandle) override;
-
-		private:
-			TWeakObjectPtr<UClass> ProxyGenClass;
 		};
 
 		class FMetasoundDataTypeSelector : public TSharedFromThis<FMetasoundDataTypeSelector>
@@ -204,6 +226,28 @@ namespace Metasound
 				return Name;
 			}
 
+			void UpdateRenameDelegate(UMetasoundEditorGraphMemberDefaultLiteral& InMemberDefaultLiteral)
+			{
+				if (UMetasoundEditorGraphMember* Member = InMemberDefaultLiteral.GetParentMember())
+				{
+					if (const UMetasoundEditorGraph* OwningGraph = GraphMember->GetOwningGraph())
+					{
+						TSharedPtr<FEditor> MetasoundEditor = FGraphBuilder::GetEditorForGraph(*OwningGraph);
+						if (MetasoundEditor->CanRenameNodes())
+						{
+							if (!RenameRequestedHandle.IsValid())
+							{
+								Member->OnRenameRequested.Clear();
+								RenameRequestedHandle = Member->OnRenameRequested.AddLambda([this]()
+								{
+									FSlateApplication::Get().SetKeyboardFocus(NameEditableTextBox.ToSharedRef(), EFocusCause::SetDirectly);
+								});
+							}
+						}
+					}
+				}
+			}
+
 			virtual bool IsInterfaceMember() const
 			{
 				return false;
@@ -301,62 +345,31 @@ namespace Metasound
 				TSharedPtr<IPropertyHandle> LiteralHandle = InDetailLayout.GetProperty(GET_MEMBER_NAME_CHECKED(TMemberType, Literal));
 				if (ensure(GraphMember.IsValid()) && ensure(LiteralHandle.IsValid()))
 				{
+					// Always hide, even if no customization (LiteralObject isn't found) as this is the case
+					// where the default object is not required (i.e. Default Member is default constructed)
+					LiteralHandle->MarkHiddenByCustomization();
+
 					UObject* LiteralObject = nullptr;
 					if (LiteralHandle->GetValue(LiteralObject) == FPropertyAccess::Success)
 					{
-						if (LiteralObject)
+						if (UMetasoundEditorGraphMemberDefaultLiteral * MemberDefaultLiteral = Cast<UMetasoundEditorGraphMemberDefaultLiteral>(LiteralObject))
 						{
-							LiteralHandle->MarkHiddenByCustomization();
-
-							// Rename
-							if (UMetasoundEditorGraphMemberDefaultLiteral* MemberDefaultLiteral = CastChecked<UMetasoundEditorGraphMemberDefaultLiteral>(LiteralObject))
-							{
-								if (UMetasoundEditorGraphMember* Member = MemberDefaultLiteral->GetParentMember())
-								{
-									if (const UMetasoundEditorGraph* OwningGraph = GraphMember->GetOwningGraph())
-									{
-										TSharedPtr<FEditor> MetasoundEditor = FGraphBuilder::GetEditorForGraph(*OwningGraph);
-										if (MetasoundEditor->CanRenameNodes())
-										{
-											if (!RenameRequestedHandle.IsValid())
-											{
-												Member->OnRenameRequested.Clear();
-												RenameRequestedHandle = Member->OnRenameRequested.AddLambda([NameEditableTextBox = this->NameEditableTextBox]()
-												{
-													FSlateApplication::Get().SetKeyboardFocus(NameEditableTextBox.ToSharedRef(), EFocusCause::SetDirectly);
-												});
-											}
-										}
-									}
-								}
-							}
-
-							TSharedPtr<IPropertyHandle> DefaultValueHandle;
-
-							if (IDetailPropertyRow* Row = DefaultCategoryBuilder.AddExternalObjectProperty(TArray<UObject*>({ LiteralObject }), "Default"))
-							{
-								DefaultValueHandle = Row->GetPropertyHandle();
-								if (DefaultValueHandle.IsValid())
-								{
-									SetDefaultPropertyMetaData(DefaultValueHandle.ToSharedRef());
-								}
-							}
+							UpdateRenameDelegate(*MemberDefaultLiteral);
 
 							IMetasoundEditorModule& EditorModule = FModuleManager::GetModuleChecked<IMetasoundEditorModule>("MetaSoundEditor");
 							LiteralCustomization = EditorModule.CreateMemberDefaultLiteralCustomization(*LiteralObject->GetClass(), DefaultCategoryBuilder);
 							if (LiteralCustomization.IsValid())
 							{
-								LiteralCustomization->CustomizeLiteral(*CastChecked<UMetasoundEditorGraphMemberDefaultLiteral>(LiteralObject), DefaultValueHandle);
+								LiteralCustomization->CustomizeLiteral(*MemberDefaultLiteral, InDetailLayout);
 							}
-
-							UMetasoundEditorGraphMemberDefaultFloat* DefaultFloat = Cast<UMetasoundEditorGraphMemberDefaultFloat>(LiteralObject);
-							if (DefaultFloat)
+							else if (LiteralObject)
 							{
-								// add input widget properties 
-								IDetailCategoryBuilder& WidgetCategoryBuilder = InDetailLayout.EditCategory("Widget");
-								WidgetCategoryBuilder.AddExternalObjectProperty(TArray<UObject*>({ DefaultFloat }), GET_MEMBER_NAME_CHECKED(UMetasoundEditorGraphMemberDefaultFloat, WidgetType));
-								WidgetCategoryBuilder.AddExternalObjectProperty(TArray<UObject*>({ DefaultFloat }), GET_MEMBER_NAME_CHECKED(UMetasoundEditorGraphMemberDefaultFloat, WidgetOrientation));
-								WidgetCategoryBuilder.AddExternalObjectProperty(TArray<UObject*>({ DefaultFloat }), GET_MEMBER_NAME_CHECKED(UMetasoundEditorGraphMemberDefaultFloat, WidgetValueType));
+								IDetailPropertyRow* Row = DefaultCategoryBuilder.AddExternalObjectProperty(TArray<UObject*>({ LiteralObject }), "Default");
+								UClass* LiteralClass = LiteralObject->GetClass();
+								check(LiteralClass);
+								ensureMsgf(Row, TEXT("Class '%s' missing expected 'Default' member."
+									"Either add/rename default member or register customization to display default value/opt out appropriately."),
+									*LiteralClass->GetName());
 							}
 						}
 					}
@@ -364,53 +377,6 @@ namespace Metasound
 			}
 			// End of IDetailCustomization interface
 
-			void SetDefaultPropertyMetaData(TSharedRef<IPropertyHandle> InDefaultPropertyHandle) const
-			{
-				using namespace Frontend;
-
-				if (!GraphMember.IsValid())
-				{
-					return;
-				}
-
-				FMetasoundFrontendRegistryContainer* Registry = FMetasoundFrontendRegistryContainer::Get();
-				if (!ensure(Registry))
-				{
-					return;
-				}
-
-				FName TypeName = GraphMember->GetDataType();
-				if (TypeName.IsNone())
-				{
-					return;
-				}
-
-				FDataTypeRegistryInfo DataTypeInfo;
-				if (!ensure(IDataTypeRegistry::Get().GetDataTypeInfo(TypeName, DataTypeInfo)))
-				{
-					return;
-				}
-
-				if (DataTypeInfo.IsArrayType())
-				{
-					TypeName = CreateElementTypeNameFromArrayTypeName(TypeName);
-				}
-				InDefaultPropertyHandle->SetInstanceMetaData(MemberCustomizationStyle::DataTypeNameIdentifier, TypeName.ToString());
-
-				const EMetasoundFrontendLiteralType LiteralType = GetMetasoundFrontendLiteralType(DataTypeInfo.PreferredLiteralType);
-				if (LiteralType != EMetasoundFrontendLiteralType::UObject && LiteralType != EMetasoundFrontendLiteralType::UObjectArray)
-				{
-					return;
-				}
-
-				UClass* ProxyGenClass = DataTypeInfo.ProxyGeneratorClass;
-				if (ProxyGenClass)
-				{
-					const FString ClassName = ProxyGenClass->GetName();
-					InDefaultPropertyHandle->SetInstanceMetaData(MemberCustomizationStyle::ProxyGeneratorClassNameIdentifier, ClassName);
-				}
-			}
-
 			void OnNameChanged(const FText& InNewName)
 			{
 				using namespace Frontend;
@@ -530,7 +496,7 @@ namespace Metasound
 			}
 
 		private:
-			TUniquePtr<IMemberDefaultLiteralCustomization> LiteralCustomization;
+			TUniquePtr<FMetasoundDefaultLiteralCustomizationBase> LiteralCustomization;
 
 			FDelegateHandle RenameRequestedHandle;
 		};
@@ -547,14 +513,10 @@ namespace Metasound
 
 			virtual bool IsInterfaceMember() const override;
 
-			// IDetailCustomization interface
-			virtual void CustomizeDetails(IDetailLayoutBuilder& InDetailLayout) override;
-			// End of IDetailCustomization interface
-
 			static const FText MemberNameText;
 
 		private:
-			TUniquePtr<IMemberDefaultLiteralCustomization> LiteralCustomization;
+			TUniquePtr<FMetasoundDefaultLiteralCustomizationBase> LiteralCustomization;
 		};
 
 		class FMetasoundOutputDetailCustomization : public TMetasoundGraphMemberDetailCustomization<UMetasoundEditorGraphOutput, FMetasoundOutputDetailCustomization>
@@ -571,12 +533,8 @@ namespace Metasound
 
 			virtual bool IsInterfaceMember() const override;
 
-			// IDetailCustomization interface
-			virtual void CustomizeDetails(IDetailLayoutBuilder& InDetailLayout) override;
-			// End of IDetailCustomization interface
-
 		private:
-			TUniquePtr<IMemberDefaultLiteralCustomization> LiteralCustomization;
+			TUniquePtr<FMetasoundDefaultLiteralCustomizationBase> LiteralCustomization;
 		};
 	} // namespace Editor
 } // namespace Metasound
diff --git a/Engine/Plugins/Runtime/Metasound/Source/MetasoundEditor/Private/MetasoundVariableDetailCustomization.cpp b/Engine/Plugins/Runtime/Metasound/Source/MetasoundEditor/Private/MetasoundVariableDetailCustomization.cpp
--- a/Engine/Plugins/Runtime/Metasound/Source/MetasoundEditor/Private/MetasoundVariableDetailCustomization.cpp
+++ b/Engine/Plugins/Runtime/Metasound/Source/MetasoundEditor/Private/MetasoundVariableDetailCustomization.cpp
@@ -18,11 +18,6 @@ namespace Metasound
 	namespace Editor
 	{
 		const FText FMetasoundVariableDetailCustomization::MemberNameText = LOCTEXT("Variable_Name", "Variable");
-
-		void FMetasoundVariableDetailCustomization::CustomizeDetails(IDetailLayoutBuilder& DetailLayout)
-		{
-			TMetasoundGraphMemberDetailCustomization<UMetasoundEditorGraphVariable, FMetasoundVariableDetailCustomization>::CustomizeDetails(DetailLayout);
-		}
 	} // namespace Editor
 } // namespace Metasound
 #undef LOCTEXT_NAMESPACE
diff --git a/Engine/Plugins/Runtime/Metasound/Source/MetasoundEditor/Private/MetasoundVariableDetailCustomization.h b/Engine/Plugins/Runtime/Metasound/Source/MetasoundEditor/Private/MetasoundVariableDetailCustomization.h
--- a/Engine/Plugins/Runtime/Metasound/Source/MetasoundEditor/Private/MetasoundVariableDetailCustomization.h
+++ b/Engine/Plugins/Runtime/Metasound/Source/MetasoundEditor/Private/MetasoundVariableDetailCustomization.h
@@ -7,8 +7,8 @@
 #include "Templates/SharedPointer.h"
 #include "Templates/UniquePtr.h"
 
+class FMetasoundDefaultLiteralCustomizationBase;
 class IDetailLayoutBuilder;
-class IMemberDefaultLiteralCustomization;
 
 #define LOCTEXT_NAMESPACE "MetaSoundEditor"
 
@@ -27,10 +27,6 @@ namespace Metasound
 			}
 
 			virtual ~FMetasoundVariableDetailCustomization() = default;
-
-			// IDetailCustomization interface
-			virtual void CustomizeDetails(IDetailLayoutBuilder& DetailLayout) override;
-			// End of IDetailCustomization interface
 		};
 	} // namespace Editor
 } // namespace Metasound
diff --git a/Engine/Plugins/Runtime/Metasound/Source/MetasoundEditor/Public/MetasoundEditorGraphMemberDefaults.h b/Engine/Plugins/Runtime/Metasound/Source/MetasoundEditor/Public/MetasoundEditorGraphMemberDefaults.h
--- a/Engine/Plugins/Runtime/Metasound/Source/MetasoundEditor/Public/MetasoundEditorGraphMemberDefaults.h
+++ b/Engine/Plugins/Runtime/Metasound/Source/MetasoundEditor/Public/MetasoundEditorGraphMemberDefaults.h
@@ -162,6 +162,11 @@ public:
 	UPROPERTY(EditAnywhere, Category = Widget, meta=(DisplayName = "Value Type", EditCondition = "WidgetType != EMetasoundMemberDefaultWidget::None", EditConditionHides))
 	EMetasoundMemberDefaultWidgetValueType WidgetValueType = EMetasoundMemberDefaultWidgetValueType::Linear;
 
+	static FName GetDefaultPropertyName()
+	{
+		return GET_MEMBER_NAME_CHECKED(UMetasoundEditorGraphMemberDefaultFloat, Default);
+	}
+
 	FOnMetasoundInputValueChangedEvent OnDefaultValueChanged;
 	FOnMetasoundRangeChangedEvent OnRangeChanged;
 	FOnMetasoundInputClampDefaultChangedEvent OnClampChanged;
diff --git a/Engine/Plugins/Runtime/Metasound/Source/MetasoundEditor/Public/MetasoundEditorModule.h b/Engine/Plugins/Runtime/Metasound/Source/MetasoundEditor/Public/MetasoundEditorModule.h
--- a/Engine/Plugins/Runtime/Metasound/Source/MetasoundEditor/Public/MetasoundEditorModule.h
+++ b/Engine/Plugins/Runtime/Metasound/Source/MetasoundEditor/Public/MetasoundEditorModule.h
@@ -1,7 +1,6 @@
 // Copyright Epic Games, Inc. All Rights Reserved.
 #pragma once
 
-#include "CoreMinimal.h"
 #include "DetailCategoryBuilder.h"
 #include "EdGraph/EdGraphNode.h"
 #include "EdGraph/EdGraphPin.h"
@@ -12,7 +11,8 @@
 #include "PropertyHandle.h"
 #include "Templates/Function.h"
 
-
+// Forward Declarations
+class IDetailLayoutBuilder;
 class UMetasoundEditorGraph;
 class UMetasoundEditorGraphMemberDefaultLiteral;
 
@@ -37,12 +37,20 @@ namespace Metasound
 			}
 		};
 
-		class METASOUNDEDITOR_API IMemberDefaultLiteralCustomization
+		class METASOUNDEDITOR_API FMetasoundDefaultLiteralCustomizationBase
 		{
+		protected:
+			IDetailCategoryBuilder* DefaultCategoryBuilder = nullptr;
+
 		public:
-			virtual ~IMemberDefaultLiteralCustomization() = default;
+			FMetasoundDefaultLiteralCustomizationBase(IDetailCategoryBuilder& InDefaultCategoryBuilder)
+				: DefaultCategoryBuilder(&InDefaultCategoryBuilder)
+			{
+			}
+
+			virtual ~FMetasoundDefaultLiteralCustomizationBase() = default;
 
-			virtual void CustomizeLiteral(UMetasoundEditorGraphMemberDefaultLiteral& InLiteral, TSharedPtr<IPropertyHandle> InDefaultValueHandle) = 0;
+			virtual void CustomizeLiteral(UMetasoundEditorGraphMemberDefaultLiteral& InLiteral, IDetailLayoutBuilder& InDetailLayout) { }
 		};
 
 		class METASOUNDEDITOR_API IMemberDefaultLiteralCustomizationFactory
@@ -50,7 +58,7 @@ namespace Metasound
 		public:
 			virtual ~IMemberDefaultLiteralCustomizationFactory() = default;
 
-			virtual TUniquePtr<IMemberDefaultLiteralCustomization> CreateLiteralCustomization(IDetailCategoryBuilder& DefaultCategoryBuilder) const = 0;
+			virtual TUniquePtr<FMetasoundDefaultLiteralCustomizationBase> CreateLiteralCustomization(IDetailCategoryBuilder& DefaultCategoryBuilder) const = 0;
 		};
 
 		class METASOUNDEDITOR_API IMetasoundEditorModule : public IModuleInterface
@@ -73,7 +81,7 @@ namespace Metasound
 
 			virtual void IterateDataTypes(TUniqueFunction<void(const FEditorDataType&)> InDataTypeFunction) const = 0;
 
-			virtual TUniquePtr<IMemberDefaultLiteralCustomization> CreateMemberDefaultLiteralCustomization(UClass& InClass, IDetailCategoryBuilder& DefaultCategoryBuilder) const = 0;
+			virtual TUniquePtr<FMetasoundDefaultLiteralCustomizationBase> CreateMemberDefaultLiteralCustomization(UClass& InClass, IDetailCategoryBuilder& DefaultCategoryBuilder) const = 0;
 
 			virtual const TSubclassOf<UMetasoundEditorGraphMemberDefaultLiteral> FindDefaultLiteralClass(EMetasoundFrontendLiteralType InLiteralType) const = 0;
 		};
diff --git a/Engine/Plugins/Runtime/Metasound/Source/MetasoundFrontend/Public/MetasoundFrontendDataTypeRegistry.h b/Engine/Plugins/Runtime/Metasound/Source/MetasoundFrontend/Public/MetasoundFrontendDataTypeRegistry.h
--- a/Engine/Plugins/Runtime/Metasound/Source/MetasoundFrontend/Public/MetasoundFrontendDataTypeRegistry.h
+++ b/Engine/Plugins/Runtime/Metasound/Source/MetasoundFrontend/Public/MetasoundFrontendDataTypeRegistry.h
@@ -58,6 +58,7 @@ namespace Metasound
 			// If provided in registration call, UClass this datatype was registered with
 			UClass* ProxyGeneratorClass = nullptr;
 
+			// Returns if DataType supports array parsing.
 			FORCEINLINE bool IsArrayType() const
 			{
 				return bIsDefaultArrayParsable
