Fix init value on metasound math op nodes.

diff --git a/Engine/Plugins/Runtime/Metasound/Source/MetasoundStandardNodes/Private/MetasoundMathNodes.cpp b/Engine/Plugins/Runtime/Metasound/Source/MetasoundStandardNodes/Private/MetasoundMathNodes.cpp
--- a/Engine/Plugins/Runtime/Metasound/Source/MetasoundStandardNodes/Private/MetasoundMathNodes.cpp
+++ b/Engine/Plugins/Runtime/Metasound/Source/MetasoundStandardNodes/Private/MetasoundMathNodes.cpp
@@ -267,6 +267,8 @@ namespace Metasound
 				, AdditionalOperandRefs(InAdditionalOperands)
 				, ValueRef(TDataWriteReferenceFactory<TDataClass>::CreateAny(InSettings))
 			{
+				// Set initial value.
+				TMathOpClass::Calculate(InstanceData, PrimaryOperandRef, AdditionalOperandRefs, ValueRef);
 			}
 
 			virtual FDataReferenceCollection GetInputs() const override
