Fix static analysis warning
#rnx

diff --git a/Engine/Plugins/Runtime/Metasound/Source/MetasoundEditor/Private/MetasoundNodeDetailCustomization.h b/Engine/Plugins/Runtime/Metasound/Source/MetasoundEditor/Private/MetasoundNodeDetailCustomization.h
--- a/Engine/Plugins/Runtime/Metasound/Source/MetasoundEditor/Private/MetasoundNodeDetailCustomization.h
+++ b/Engine/Plugins/Runtime/Metasound/Source/MetasoundEditor/Private/MetasoundNodeDetailCustomization.h
@@ -356,20 +356,21 @@ namespace Metasound
 						{
 							UpdateRenameDelegate(*MemberDefaultLiteral);
 
+							UClass* MemberClass = MemberDefaultLiteral->GetClass();
+							check(MemberClass);
+
 							IMetasoundEditorModule& EditorModule = FModuleManager::GetModuleChecked<IMetasoundEditorModule>("MetaSoundEditor");
-							LiteralCustomization = EditorModule.CreateMemberDefaultLiteralCustomization(*LiteralObject->GetClass(), DefaultCategoryBuilder);
+							LiteralCustomization = EditorModule.CreateMemberDefaultLiteralCustomization(*MemberClass, DefaultCategoryBuilder);
 							if (LiteralCustomization.IsValid())
 							{
 								LiteralCustomization->CustomizeLiteral(*MemberDefaultLiteral, InDetailLayout);
 							}
-							else if (LiteralObject)
+							else
 							{
-								IDetailPropertyRow* Row = DefaultCategoryBuilder.AddExternalObjectProperty(TArray<UObject*>({ LiteralObject }), "Default");
-								UClass* LiteralClass = LiteralObject->GetClass();
-								check(LiteralClass);
+								IDetailPropertyRow* Row = DefaultCategoryBuilder.AddExternalObjectProperty(TArray<UObject*>({ MemberDefaultLiteral }), "Default");
 								ensureMsgf(Row, TEXT("Class '%s' missing expected 'Default' member."
 									"Either add/rename default member or register customization to display default value/opt out appropriately."),
-									*LiteralClass->GetName());
+									*MemberClass->GetName());
 							}
 						}
 					}
