Fix for ensure in metasound source engine test

diff --git a/Engine/Plugins/Runtime/Metasound/Source/MetasoundEngineTest/Private/EngineTestMetasoundSource.cpp b/Engine/Plugins/Runtime/Metasound/Source/MetasoundEngineTest/Private/EngineTestMetasoundSource.cpp
--- a/Engine/Plugins/Runtime/Metasound/Source/MetasoundEngineTest/Private/EngineTestMetasoundSource.cpp
+++ b/Engine/Plugins/Runtime/Metasound/Source/MetasoundEngineTest/Private/EngineTestMetasoundSource.cpp
@@ -78,15 +78,16 @@ namespace EngineTestMetasoundSourcePrivate
 
 		// Add default source & mono interface members (OnPlay, OnFinished & Mono Out)
 		FModifyRootGraphInterfaces InterfaceTransform(
+		{ },
 		{	
 			SourceInterface::GetVersion(), 
 			SourceOneShotInterface::GetVersion(),
 			OutputFormatMonoInterface::GetVersion() 
-		}, { });
+		}); 
 		InterfaceTransform.Transform(DocumentHandle);
 
 		// Input on Play
-		FNodeHandle OnPlayOutputNode = RootGraph->GetOutputNodeWithName(SourceInterface::Inputs::OnPlay);
+		FNodeHandle OnPlayOutputNode = RootGraph->GetInputNodeWithName(SourceInterface::Inputs::OnPlay);
 		check(OnPlayOutputNode->IsValid());
 
 		// Input Frequency
