Miscellaneous MetaSound tooltip fixes

diff --git a/Engine/Plugins/Runtime/AudioModulation/Source/AudioModulation/Private/MetaSoundNodes/MixModulatorsNode.cpp b/Engine/Plugins/Runtime/AudioModulation/Source/AudioModulation/Private/MetaSoundNodes/MixModulatorsNode.cpp
--- a/Engine/Plugins/Runtime/AudioModulation/Source/AudioModulation/Private/MetaSoundNodes/MixModulatorsNode.cpp
+++ b/Engine/Plugins/Runtime/AudioModulation/Source/AudioModulation/Private/MetaSoundNodes/MixModulatorsNode.cpp
@@ -41,7 +41,7 @@ namespace AudioModulation
 					1, // Major Version
 					0, // Minor Version
 					LOCTEXT("MixModulatorsNode_Name", "Mix Modulators"),
-					LOCTEXT("MixModulatorsNode_Description", "Mixes two modulators using the parameterized mix function. Returns 'Normalized' if set to true, or unit space if false"),
+					LOCTEXT("MixModulatorsNode_Description", "Mixes two modulators using the parameterized mix function. Returns the 'Normalized' value (0-1) if true, or the value in unit space (dB, Frequency, etc.) if set to false."),
 					AudioModulation::PluginAuthor,
 					AudioModulation::PluginNodeMissingPrompt,
 					GetDefaultInterface(),
diff --git a/Engine/Plugins/Runtime/Metasound/Source/MetasoundFrontend/Public/MetasoundArrayNodes.h b/Engine/Plugins/Runtime/Metasound/Source/MetasoundFrontend/Public/MetasoundArrayNodes.h
--- a/Engine/Plugins/Runtime/Metasound/Source/MetasoundFrontend/Public/MetasoundArrayNodes.h
+++ b/Engine/Plugins/Runtime/Metasound/Source/MetasoundFrontend/Public/MetasoundArrayNodes.h
@@ -470,7 +470,7 @@ namespace Metasound
 					TInputDataVertexModel<ArrayType>(GetInputRightArrayName(), LOCTEXT("ArrayOpArrayConcatInputRight", "Input Left Array."))
 				),
 				FOutputVertexInterface(
-					TOutputDataVertexModel<ArrayType>(GetOutputArrayName(), LOCTEXT("ArrayOpArrayConcatOutput", "Array after setting."))
+					TOutputDataVertexModel<ArrayType>(GetOutputArrayName(), LOCTEXT("ArrayOpArrayConcatOutput", "Array after concatenation."))
 				)
 			);
 
diff --git a/Engine/Plugins/Runtime/Metasound/Source/MetasoundFrontend/Public/MetasoundArrayRandomNode.h b/Engine/Plugins/Runtime/Metasound/Source/MetasoundFrontend/Public/MetasoundArrayRandomNode.h
--- a/Engine/Plugins/Runtime/Metasound/Source/MetasoundFrontend/Public/MetasoundArrayRandomNode.h
+++ b/Engine/Plugins/Runtime/Metasound/Source/MetasoundFrontend/Public/MetasoundArrayRandomNode.h
@@ -33,9 +33,9 @@ namespace Metasound
 		static const FText InputTriggerResetTooltip = LOCTEXT("TriggerResetTooltip", "Trigger to reset the seed for the randomized array.");
 		static const FText InputRandomArrayTooltip = LOCTEXT("RandomArrayTooltip", "Input array to randomized.");
 		static const FText InputWeightsTooltip = LOCTEXT("WeightsTooltip", "Input array of weights to use for random selection. Will repeat if this array is shorter than the input array to select from.");
-		static const FText InputSeedTooltip = LOCTEXT("SeedTooltip", "Seed to use for the the random shuffle.");
+		static const FText InputSeedTooltip = LOCTEXT("SeedTooltip", "Seed to use for the random shuffle.");
 		static const FText InputNoRepeatOrderTooltip = LOCTEXT("NoRepeatOrderTooltip", "The number of elements to track to avoid repeating in a row.");
-		static const FText InputEnableSharedStateTooltip = LOCTEXT("EnableSharedStateTooltip", "Set to enabled shared state across instances of this metasound.");
+		static const FText InputEnableSharedStateTooltip = LOCTEXT("EnableSharedStateTooltip", "Set to enabled to share state across instances of this metasound.");
 		static const FText OutputTriggerOnNextTooltip = LOCTEXT("TriggerOnNextTooltip", "Triggers when the \"Next\" input is triggered.");
 		static const FText OutputTriggerOnResetTooltip = LOCTEXT("TriggerOnResetTooltip", "Triggers when the \"Shuffle\" input is triggered or if the array is auto-shuffled.");
 		static const FText OutputValueTooltip = LOCTEXT("ValueTooltip", "Value of the current shuffled element.");
diff --git a/Engine/Plugins/Runtime/Metasound/Source/MetasoundStandardNodes/Private/MetasoundBPMToSecondsNode.cpp b/Engine/Plugins/Runtime/Metasound/Source/MetasoundStandardNodes/Private/MetasoundBPMToSecondsNode.cpp
--- a/Engine/Plugins/Runtime/Metasound/Source/MetasoundStandardNodes/Private/MetasoundBPMToSecondsNode.cpp
+++ b/Engine/Plugins/Runtime/Metasound/Source/MetasoundStandardNodes/Private/MetasoundBPMToSecondsNode.cpp
@@ -135,11 +135,11 @@ namespace Metasound
 		static const FVertexInterface Interface(
 			FInputVertexInterface(
 				TInputDataVertexModel<float>(GetInputBPMTempoName(), LOCTEXT("BPMToSecsBPMTT", "Input BPM (Beats Per Minute)."), 90.0f),
-				TInputDataVertexModel<float>(GetInputBeatMultiplierName(), LOCTEXT("BPMToSecsBeatMultiplierTT", "The multipilier of the BPM."), 1.0f),
+				TInputDataVertexModel<float>(GetInputBeatMultiplierName(), LOCTEXT("BPMToSecsBeatMultiplierTT", "The multiplier of the BPM."), 1.0f),
 				TInputDataVertexModel<float>(GetInputDivisionsOfWholeNoteName(), LOCTEXT("BPMToSecsDivOfWholeNoteTT", "Divisions of a whole note."), 4.0f)
 			),
 			FOutputVertexInterface(
-				TOutputDataVertexModel<FTime>(GetOutputTimeSecondsName(), LOCTEXT("BPMToSecsOutputTimeTT", "The output time in seconds give the BPM inputs."))
+				TOutputDataVertexModel<FTime>(GetOutputTimeSecondsName(), LOCTEXT("BPMToSecsOutputTimeTT", "The output time in seconds, given the BPM inputs."))
 			)
 		);
 
diff --git a/Engine/Plugins/Runtime/Metasound/Source/MetasoundStandardNodes/Private/MetasoundCompressorNode.cpp b/Engine/Plugins/Runtime/Metasound/Source/MetasoundStandardNodes/Private/MetasoundCompressorNode.cpp
--- a/Engine/Plugins/Runtime/Metasound/Source/MetasoundStandardNodes/Private/MetasoundCompressorNode.cpp
+++ b/Engine/Plugins/Runtime/Metasound/Source/MetasoundStandardNodes/Private/MetasoundCompressorNode.cpp
@@ -27,7 +27,7 @@ namespace Metasound
 		METASOUND_PARAM(InputAttackTime, "Attack Time", "How long it takes for audio above the threshold to reach its compressed volume level.");
 		METASOUND_PARAM(InputReleaseTime, "Release Time", "How long it takes for audio below the threshold to return to its original volume level.");
 		METASOUND_PARAM(InputKnee, "Knee", "How hard or soft the gain reduction blends from no gain reduction to gain reduction. 0 dB = no blending.");
-		METASOUND_PARAM(InputSidechain, "Sidechain", "(Optional) external audio signal to control the compressor with. If empty, uses the input audio signal.");
+		METASOUND_PARAM(InputSidechain, "Sidechain", "(Optional) External audio signal to control the compressor with. If empty, uses the input audio signal.");
 		METASOUND_PARAM(InputEnvelopeMode, "Envelope Mode", "The envelope-following method the compressor will use for gain detection.");
 		METASOUND_PARAM(InputIsAnalog, "Analog Mode", "Enable Analog Mode for the compressor's envelope follower.");
 		METASOUND_PARAM(InputWetDryMix, "Wet/Dry", "Ratio between the processed/wet signal and the unprocessed/dry signal. 0 is full dry, 1 is full wet, and 0.5 is 50/50.");
diff --git a/Engine/Plugins/Runtime/Metasound/Source/MetasoundStandardNodes/Private/MetasoundDynamicFilterNode.cpp b/Engine/Plugins/Runtime/Metasound/Source/MetasoundStandardNodes/Private/MetasoundDynamicFilterNode.cpp
--- a/Engine/Plugins/Runtime/Metasound/Source/MetasoundStandardNodes/Private/MetasoundDynamicFilterNode.cpp
+++ b/Engine/Plugins/Runtime/Metasound/Source/MetasoundStandardNodes/Private/MetasoundDynamicFilterNode.cpp
@@ -20,7 +20,7 @@ namespace Metasound
 	namespace DynamicFilterNode
 	{
 		METASOUND_PARAM(InputAudio, "Audio", "Incoming audio signal.");
-		METASOUND_PARAM(InputSidechain, "Sidechain", "(Optional) external audio signal to control the filter with. If empty, uses the input audio signal.");
+		METASOUND_PARAM(InputSidechain, "Sidechain", "(Optional) External audio signal to control the filter with. If empty, uses the input audio signal.");
 		METASOUND_PARAM(InputFilterType, "FilterType", "Filter shape to use.");
 		METASOUND_PARAM(InputFrequency, "Frequency", "The center frequency of the filter.");
 		METASOUND_PARAM(InputQ, "Q", "Filter Q, or resonance, controls the steepness of the filter.");
diff --git a/Engine/Plugins/Runtime/Metasound/Source/MetasoundStandardNodes/Private/MetasoundMidSideNodes.cpp b/Engine/Plugins/Runtime/Metasound/Source/MetasoundStandardNodes/Private/MetasoundMidSideNodes.cpp
--- a/Engine/Plugins/Runtime/Metasound/Source/MetasoundStandardNodes/Private/MetasoundMidSideNodes.cpp
+++ b/Engine/Plugins/Runtime/Metasound/Source/MetasoundStandardNodes/Private/MetasoundMidSideNodes.cpp
@@ -26,7 +26,7 @@ namespace Metasound
 		METASOUND_PARAM(InputEqualPower, "Equal Power", "Whether an equal power relationship between the mid and side signals should be maintained.");
 
 		METASOUND_PARAM(OutputAudioMid, "Out Mid", "The Mid content from the audio signal.");
-		METASOUND_PARAM(OutputAudioSide, "Out Side", "the Side content from the audio signal.");
+		METASOUND_PARAM(OutputAudioSide, "Out Side", "The Side content from the audio signal.");
 	}
 
 	// Operator Class
@@ -250,7 +250,7 @@ namespace Metasound
 		METASOUND_PARAM(InputEqualPower, "Equal Power", "Whether an equal power relationship between the mid and side signals should be maintained.");
 	
 		METASOUND_PARAM(OutputAudioLeft, "Out Left", "The left audio channel which has been processed.");
-		METASOUND_PARAM(OutputAudioRight, "Out Right", "the Right audio channel which has been processed.");
+		METASOUND_PARAM(OutputAudioRight, "Out Right", "The Right audio channel which has been processed.");
 	}
 
 	// Operator Class
