Initialize metasound autoconverter nodes on construction

diff --git a/Engine/Plugins/Runtime/Metasound/Source/MetasoundFrontend/Public/MetasoundAutoConverterNode.h b/Engine/Plugins/Runtime/Metasound/Source/MetasoundFrontend/Public/MetasoundAutoConverterNode.h
--- a/Engine/Plugins/Runtime/Metasound/Source/MetasoundFrontend/Public/MetasoundAutoConverterNode.h
+++ b/Engine/Plugins/Runtime/Metasound/Source/MetasoundFrontend/Public/MetasoundAutoConverterNode.h
@@ -128,6 +128,7 @@ namespace Metasound
 				: FromData(InFromDataReference)
 				, ToData(InToDataReference)
 			{
+				Execute();
 			}
 
 			virtual ~FConverterOperator() {}
