Fix for ensure when adding variable set or get node when dragged from pin.

diff --git a/Engine/Plugins/Runtime/Metasound/Source/MetasoundEditor/Private/MetasoundEditorGraphSchema.cpp b/Engine/Plugins/Runtime/Metasound/Source/MetasoundEditor/Private/MetasoundEditorGraphSchema.cpp
--- a/Engine/Plugins/Runtime/Metasound/Source/MetasoundEditor/Private/MetasoundEditorGraphSchema.cpp
+++ b/Engine/Plugins/Runtime/Metasound/Source/MetasoundEditor/Private/MetasoundEditorGraphSchema.cpp
@@ -1,6 +1,7 @@
 // Copyright Epic Games, Inc. All Rights Reserved.
 #include "MetasoundEditorGraphSchema.h"
 
+#include "Algo/AnyOf.h"
 #include "EdGraph/EdGraph.h"
 #include "EdGraphNode_Comment.h"
 #include "EdGraphSchema_K2.h"
@@ -26,6 +27,7 @@
 #include "MetasoundEditorModule.h"
 #include "MetasoundEditorSettings.h"
 #include "MetasoundFrontend.h"
+#include "MetasoundFrontendGraphLinter.h"
 #include "MetasoundFrontendNodesCategories.h"
 #include "MetasoundFrontendRegistries.h"
 #include "MetasoundFrontendSearchEngine.h"
@@ -261,7 +263,62 @@ namespace Metasound
 
 				return nullptr;
 			}
-		}
+
+			bool WillAddingVariableAccessorCauseLoop(const Frontend::IVariableController& InVariable, const Frontend::IInputController& InInput)
+			{
+				using namespace Metasound::Frontend;
+
+				// A variable mutator node must come before a variable accessor node,
+				// or else the nodes will create a loop from the hidden variable pin. 
+				// To determine if adding an accessor node will cause a loop (before actually
+				// adding an accessor node), we check whether an existing mutator can
+				// reach the node upstream which wants to connect to the accessor node.
+				//
+				// Example:
+				// Will cause loop:
+				// 	[AccessorNode]-->[DestinationNode]-->[Node]-->[MutatorNode] 
+				// 	       ^-------------------------------------------|
+				//
+				// Will not cause loop
+				//  [Node]-->[MutatorNode]-->[AccessorNode]-->[DestinationNode]
+				//       |                                        ^ 
+				//       |----------------------------------------|
+				FConstNodeHandle MutatorNode = InVariable.FindMutatorNode();
+				FConstNodeHandle DestinationNode = InInput.GetOwningNode();
+				return FGraphLinter::IsReachableUpstream(*MutatorNode, *DestinationNode);
+			}
+
+			bool WillAddingVariableMutatorCauseLoop(const Frontend::IVariableController& InVariable, const Frontend::IOutputController& InOutput)
+			{
+				using namespace Metasound::Frontend;
+
+				// A variable mutator node must come before a variable accessor node,
+				// or else the nodes will create a loop from the hidden variable pin. 
+				// To determine if adding a mutator node will cause a loop (before actually
+				// adding a mutator node), we check whether any existing accessor can
+				// reach the node downstream which wants to connect to the mutator node.
+				//
+				// Example:
+				// Will cause loop:
+				// 	[AccessorNode]-->[Node]-->[SourceNode]-->[MutatorNode] 
+				// 	     ^---------------------------------------|
+				//
+				// Will not cause loop
+				//  [SourceNode]-->[MutatorNode]-->[AccessorNode]-->[Node]
+				//       |                                            ^ 
+				//       |--------------------------------------------|
+				TArray<FConstNodeHandle> AccessorNodes = InVariable.FindAccessorNodes();
+				FConstNodeHandle SourceNode = InOutput.GetOwningNode();
+
+				auto IsSourceNodeReachableDownstream = [&SourceNode](const FConstNodeHandle& AccessorNode)
+				{
+					return FGraphLinter::IsReachableDownstream(*AccessorNode, *SourceNode);
+				};
+
+				return Algo::AnyOf(AccessorNodes, IsSourceNodeReachableDownstream);
+			}
+
+		} // namespace SchemaPrivate
 	} // namespace Editor
 } // namespace Metasound
 
@@ -1567,6 +1624,7 @@ void UMetasoundEditorGraphSchema::GetVariableActions(FGraphActionMenuBuilder& Ac
 {
 	using namespace Metasound::Frontend;
 	using namespace Metasound::Editor;
+	using namespace Metasound::Editor::SchemaPrivate;
 
 	TArray<FConstVariableHandle> Variables = InGraphHandle->GetVariables();
 
@@ -1574,23 +1632,29 @@ void UMetasoundEditorGraphSchema::GetVariableActions(FGraphActionMenuBuilder& Ac
 	bool bGetDeferredAccessor = true;
 	bool bGetMutator = true;
 	bool bFilterByDataType = false;
+	bool bCheckForLoops = false;
 	FName DataType;
+	FConstInputHandle ConnectingInputHandle = IInputController::GetInvalidHandle();
+	FConstOutputHandle ConnectingOutputHandle = IOutputController::GetInvalidHandle();
 
 	// Determine which variable actions to create.
 	if (const UEdGraphPin* FromPin = ActionMenuBuilder.FromPin)
 	{
 		bFilterByDataType = true;
+		bCheckForLoops = true;
 
 		if (FromPin->Direction == EGPD_Input)
 		{
 			bGetMutator = false;
-			DataType = FGraphBuilder::GetConstInputHandleFromPin(FromPin)->GetDataType();
+			ConnectingInputHandle = FGraphBuilder::GetConstInputHandleFromPin(FromPin);
+			DataType = ConnectingInputHandle->GetDataType();
 		}
 		else if (FromPin->Direction == EGPD_Output)
 		{
 			bGetAccessor = false;
 			bGetDeferredAccessor = false;
-			DataType = FGraphBuilder::GetConstOutputHandleFromPin(FromPin)->GetDataType();
+			ConnectingOutputHandle = FGraphBuilder::GetConstOutputHandleFromPin(FromPin);
+			DataType = ConnectingOutputHandle->GetDataType();
 		}
 	}
 
@@ -1609,9 +1673,13 @@ void UMetasoundEditorGraphSchema::GetVariableActions(FGraphActionMenuBuilder& Ac
 
 		if (bGetAccessor)
 		{
-			FText ActionDisplayName = FText::Format(SchemaPrivate::VariableAccessorDisplayNameFormat, VariableDisplayName);
-			FText ActionTooltip = FText::Format(SchemaPrivate::VariableAccessorTooltipFormat, VariableDisplayName);
-			ActionMenuBuilder.AddAction(MakeShared<FMetasoundGraphSchemaAction_NewVariableAccessorNode>(GroupName, ActionDisplayName, VariableID, ActionTooltip));
+			// Do not add the action if adding an accessor node would cause a loop.
+			if (!(bCheckForLoops && WillAddingVariableAccessorCauseLoop(*Variable, *ConnectingInputHandle)))
+			{
+				FText ActionDisplayName = FText::Format(SchemaPrivate::VariableAccessorDisplayNameFormat, VariableDisplayName);
+				FText ActionTooltip = FText::Format(SchemaPrivate::VariableAccessorTooltipFormat, VariableDisplayName);
+				ActionMenuBuilder.AddAction(MakeShared<FMetasoundGraphSchemaAction_NewVariableAccessorNode>(GroupName, ActionDisplayName, VariableID, ActionTooltip));
+			}
 		}
 
 		if (bGetDeferredAccessor)
@@ -1628,9 +1696,13 @@ void UMetasoundEditorGraphSchema::GetVariableActions(FGraphActionMenuBuilder& Ac
 			bool bMutatorNodeAlreadyExists = Variable->FindMutatorNode()->IsValid();
 			if (!bMutatorNodeAlreadyExists)
 			{
-				FText ActionDisplayName = FText::Format(SchemaPrivate::VariableMutatorDisplayNameFormat, VariableDisplayName);
-				FText ActionTooltip = FText::Format(SchemaPrivate::VariableMutatorTooltipFormat, VariableDisplayName);
-				ActionMenuBuilder.AddAction(MakeShared<FMetasoundGraphSchemaAction_NewVariableMutatorNode>(GroupName, ActionDisplayName, VariableID, ActionTooltip));
+				// Do not add the action if adding a mutator node would cause a loop.
+				if (!(bCheckForLoops && WillAddingVariableMutatorCauseLoop(*Variable, *ConnectingOutputHandle)))
+				{
+					FText ActionDisplayName = FText::Format(SchemaPrivate::VariableMutatorDisplayNameFormat, VariableDisplayName);
+					FText ActionTooltip = FText::Format(SchemaPrivate::VariableMutatorTooltipFormat, VariableDisplayName);
+					ActionMenuBuilder.AddAction(MakeShared<FMetasoundGraphSchemaAction_NewVariableMutatorNode>(GroupName, ActionDisplayName, VariableID, ActionTooltip));
+				}
 			}
 		}
 	}
diff --git a/Engine/Plugins/Runtime/Metasound/Source/MetasoundFrontend/Private/MetasoundFrontendGraphLinter.cpp b/Engine/Plugins/Runtime/Metasound/Source/MetasoundFrontend/Private/MetasoundFrontendGraphLinter.cpp
--- a/Engine/Plugins/Runtime/Metasound/Source/MetasoundFrontend/Private/MetasoundFrontendGraphLinter.cpp
+++ b/Engine/Plugins/Runtime/Metasound/Source/MetasoundFrontend/Private/MetasoundFrontendGraphLinter.cpp
@@ -57,6 +57,53 @@ namespace Metasound
 			return bCausesLoop;
 		}
 
+
+		bool FGraphLinter::IsReachableDownstream(const INodeController& InFromNode, const INodeController& InToNode)
+		{
+			bool bIsReachable = false;
+			const FGuid ToNodeID = InToNode.GetID();
+
+			// Sets bCausesLoop if the FromNode has a path to the ToNode
+			FGraphLinter::DepthFirstTraversal(InFromNode, [&](const INodeController& Node) -> TSet<FGuid>
+			{
+				TSet<FGuid> Children;
+
+				if (ToNodeID == Node.GetID())
+				{
+					bIsReachable = true;
+				}
+
+				if (!bIsReachable)
+				{
+					// Only produce children if bIsReachbable is not set to avoid wasting CPU. 
+					if (Node.IsValid())
+					{
+						TArray<FConstOutputHandle> NodeOutputs = Node.GetConstOutputs();
+						for (const FConstOutputHandle& Output : NodeOutputs)
+						{
+							TArray<FConstInputHandle> ConnectedInputs = Output->GetConstConnectedInputs();
+							for (const FConstInputHandle& Input : ConnectedInputs)
+							{
+								Children.Add(Input->GetOwningNode()->GetID());
+							}
+						}
+						
+					}
+				}
+
+				return Children;
+			});
+			
+			return bIsReachable;
+		}
+
+		bool FGraphLinter::IsReachableUpstream(const INodeController& InFromNode, const INodeController& InToNode)
+		{
+			// Searching upstream is equivalent to swapping the order of the nodes 
+			// and searching downstream.
+			return IsReachableDownstream(InToNode, InFromNode);
+		}
+
 		void FGraphLinter::DepthFirstTraversal(const INodeController& Node, FDepthFirstVisitFunction Visit)
 		{
 			// Non recursive depth first traversal.
diff --git a/Engine/Plugins/Runtime/Metasound/Source/MetasoundFrontend/Public/MetasoundFrontendGraphLinter.h b/Engine/Plugins/Runtime/Metasound/Source/MetasoundFrontend/Public/MetasoundFrontendGraphLinter.h
--- a/Engine/Plugins/Runtime/Metasound/Source/MetasoundFrontend/Public/MetasoundFrontendGraphLinter.h
+++ b/Engine/Plugins/Runtime/Metasound/Source/MetasoundFrontend/Public/MetasoundFrontendGraphLinter.h
@@ -25,6 +25,15 @@ namespace Metasound
 			 */
 			static bool DoesConnectionCauseLoop(const IInputController& InInputController, const IOutputController& InOutputController);
 
+			/** Returns true if the FromNode can reach the ToNode by traversing the graph 
+			 * in the forward direction. */
+			static bool IsReachableDownstream(const INodeController& InFromNode, const INodeController& InToNode);
+
+			/** Returns true if the FromNode can reach the ToNode by traversing the graph backwards
+			 * (aka by traversing the transpose graph).
+			 */
+			static bool IsReachableUpstream(const INodeController& InFromNode, const INodeController& InToNode);
+
 			/** Visits nodes using depth first traversals. */
 			static void DepthFirstTraversal(const INodeController& Node, FDepthFirstVisitFunction Visit);
 		};
