Deprecate send and receive nodes

diff --git a/Engine/Plugins/Runtime/Metasound/Source/MetasoundFrontend/Private/MetasoundFrontendDocument.cpp b/Engine/Plugins/Runtime/Metasound/Source/MetasoundFrontend/Private/MetasoundFrontendDocument.cpp
--- a/Engine/Plugins/Runtime/Metasound/Source/MetasoundFrontend/Private/MetasoundFrontendDocument.cpp
+++ b/Engine/Plugins/Runtime/Metasound/Source/MetasoundFrontend/Private/MetasoundFrontendDocument.cpp
@@ -133,6 +133,7 @@ FMetasoundFrontendClassMetadata::FMetasoundFrontendClassMetadata(const Metasound
 , Author(InNodeClassMetadata.Author)
 , Keywords(InNodeClassMetadata.Keywords)
 , CategoryHierarchy(InNodeClassMetadata.CategoryHierarchy)
+, bIsDeprecated(InNodeClassMetadata.bDeprecated)
 {
 }
 
diff --git a/Engine/Plugins/Runtime/Metasound/Source/MetasoundFrontend/Private/MetasoundFrontendSearchEngine.cpp b/Engine/Plugins/Runtime/Metasound/Source/MetasoundFrontend/Private/MetasoundFrontendSearchEngine.cpp
--- a/Engine/Plugins/Runtime/Metasound/Source/MetasoundFrontend/Private/MetasoundFrontendSearchEngine.cpp
+++ b/Engine/Plugins/Runtime/Metasound/Source/MetasoundFrontend/Private/MetasoundFrontendSearchEngine.cpp
@@ -210,6 +210,19 @@ namespace Metasound
 				}
 			};
 
+			class FRemoveDeprecatedClasses : public IFrontendQueryFilterStep
+			{
+			public:
+				virtual bool Filter(const FFrontendQueryEntry& InEntry) const override
+				{
+					if (InEntry.Value.IsType<FMetasoundFrontendClass>())
+					{
+						return !InEntry.Value.Get<FMetasoundFrontendClass>().Metadata.GetIsDeprecated();
+					}
+					return false;
+				}
+			};
+
 			class FRemoveInterfacesWhichAreNotDefault : public IFrontendQueryFilterStep
 			{
 				virtual bool Filter(const FFrontendQueryEntry& InEntry) const override
@@ -444,6 +457,7 @@ namespace Metasound
 					.AddStep<FMapRegistrationEventsToNodeRegistryKeys>()
 					.AddStep<FReduceRegistrationEventsToCurrentStatus>()
 					.AddStep<FTransformRegistrationEventsToClasses>()
+					.AddStep<FRemoveDeprecatedClasses>()
 					.AddStep<FMapToNull>();
 				return Query;
 			}
@@ -538,6 +552,7 @@ namespace Metasound
 					.AddStep<FMapRegistrationEventsToNodeRegistryKeys>()
 					.AddStep<FReduceRegistrationEventsToCurrentStatus>()
 					.AddStep<FTransformRegistrationEventsToClasses>()
+					.AddStep<FRemoveDeprecatedClasses>()
 					.AddStep<FMapToFullClassName>()
 					.AddStep<FReduceClassesToHighestVersion>();
 				return Query;
@@ -562,7 +577,9 @@ namespace Metasound
 					.AddStep<FMapRegistrationEventsToNodeRegistryKeys>()
 					.AddStep<FReduceRegistrationEventsToCurrentStatus>()
 					.AddStep<FTransformRegistrationEventsToClasses>()
-					.AddStep<FMapClassToFullClassNameAndMajorVersion>();
+					.AddStep<FRemoveDeprecatedClasses>()
+					.AddStep<FMapClassToFullClassNameAndMajorVersion>()
+					.AddStep<FReduceClassesToHighestVersion>();
 				return Query;
 			}
 
diff --git a/Engine/Plugins/Runtime/Metasound/Source/MetasoundFrontend/Public/MetasoundFrontendDocument.h b/Engine/Plugins/Runtime/Metasound/Source/MetasoundFrontend/Public/MetasoundFrontendDocument.h
--- a/Engine/Plugins/Runtime/Metasound/Source/MetasoundFrontend/Public/MetasoundFrontendDocument.h
+++ b/Engine/Plugins/Runtime/Metasound/Source/MetasoundFrontend/Public/MetasoundFrontendDocument.h
@@ -919,6 +919,10 @@ struct METASOUNDFRONTEND_API FMetasoundFrontendClassMetadata
 		UPROPERTY(EditAnywhere, Category = Metasound)
 		TArray<FText> CategoryHierarchy;
 
+		// If true, this node is deprecated and should not be used in new MetaSounds.
+		UPROPERTY(VisibleAnywhere, Category = Metasound)
+		bool bIsDeprecated = false;
+
 		// If true, auto-update will manage (add and remove)
 		// inputs/outputs associated with internally connected
 		// nodes when the interface of the given node is auto-updated.
@@ -1073,6 +1077,20 @@ struct METASOUNDFRONTEND_API FMetasoundFrontendClassMetadata
 		{
 			return ChangeID;
 		}
+
+		bool GetIsDeprecated() const
+		{
+			return bIsDeprecated;
+		}
+
+		void SetIsDeprecated(bool bInIsDeprecated)
+		{
+			if (bInIsDeprecated != bIsDeprecated)
+			{
+				bIsDeprecated = bInIsDeprecated;
+				ChangeID = FGuid::NewGuid();
+			}
+		}
 };
 
 USTRUCT()
diff --git a/Engine/Plugins/Runtime/Metasound/Source/MetasoundFrontend/Public/MetasoundReceiveNode.h b/Engine/Plugins/Runtime/Metasound/Source/MetasoundFrontend/Public/MetasoundReceiveNode.h
--- a/Engine/Plugins/Runtime/Metasound/Source/MetasoundFrontend/Public/MetasoundReceiveNode.h
+++ b/Engine/Plugins/Runtime/Metasound/Source/MetasoundFrontend/Public/MetasoundReceiveNode.h
@@ -60,6 +60,11 @@ namespace Metasound
 				Info.CategoryHierarchy = { LOCTEXT("Metasound_TransmissionNodeCategory", "Transmission") };
 				Info.Keywords = { };
 
+				// Then send & receive nodes do not work as expected, particularly 
+				// around multiple-consumer scenarios. Deprecate them to avoid
+				// metasound assets from relying on send & receive nodes. 
+				Info.bDeprecated = true; 
+
 				return Info;
 			};
 
diff --git a/Engine/Plugins/Runtime/Metasound/Source/MetasoundFrontend/Public/MetasoundSendNode.h b/Engine/Plugins/Runtime/Metasound/Source/MetasoundFrontend/Public/MetasoundSendNode.h
--- a/Engine/Plugins/Runtime/Metasound/Source/MetasoundFrontend/Public/MetasoundSendNode.h
+++ b/Engine/Plugins/Runtime/Metasound/Source/MetasoundFrontend/Public/MetasoundSendNode.h
@@ -63,6 +63,11 @@ namespace Metasound
 				Info.CategoryHierarchy = { LOCTEXT("Metasound_TransmissionNodeCategory", "Transmission") };
 				Info.Keywords = { };
 
+				// Then send & receive nodes do not work as expected, particularly 
+				// around multiple-consumer scenarios. Deprecate them to avoid
+				// metasound assets from relying on send & receive nodes. 
+				Info.bDeprecated = true;
+
 				return Info;
 			};
 
diff --git a/Engine/Plugins/Runtime/Metasound/Source/MetasoundGraphCore/Public/MetasoundNodeInterface.h b/Engine/Plugins/Runtime/Metasound/Source/MetasoundGraphCore/Public/MetasoundNodeInterface.h
--- a/Engine/Plugins/Runtime/Metasound/Source/MetasoundGraphCore/Public/MetasoundNodeInterface.h
+++ b/Engine/Plugins/Runtime/Metasound/Source/MetasoundGraphCore/Public/MetasoundNodeInterface.h
@@ -115,6 +115,9 @@ namespace Metasound
 		/** Display style for node when visualized. */
 		FNodeDisplayStyle DisplayStyle;
 
+		/** If true, the node is deprecated and should not be used in new MetaSounds. */
+		bool bDeprecated = false;
+
 		/** Returns an empty FNodeClassMetadata object. */
 		static const FNodeClassMetadata& GetEmpty()
 		{
