Fix metasound presets not updating correctly.

diff --git a/Engine/Plugins/Runtime/Metasound/Source/MetasoundFrontend/Private/MetasoundFrontendArchetypeRegistry.cpp b/Engine/Plugins/Runtime/Metasound/Source/MetasoundFrontend/Private/MetasoundFrontendArchetypeRegistry.cpp
--- a/Engine/Plugins/Runtime/Metasound/Source/MetasoundFrontend/Private/MetasoundFrontendArchetypeRegistry.cpp
+++ b/Engine/Plugins/Runtime/Metasound/Source/MetasoundFrontend/Private/MetasoundFrontendArchetypeRegistry.cpp
@@ -1,6 +1,8 @@
 // Copyright Epic Games, Inc. All Rights Reserved.
 
 #include "MetasoundFrontendArchetypeRegistry.h"
+
+#include "HAL/PlatformTime.h"
 #include "MetasoundFrontendRegistryTransaction.h"
 
 namespace Metasound
@@ -16,6 +18,7 @@ namespace Metasound
 
 				virtual void RegisterInterface(TUniquePtr<IInterfaceRegistryEntry>&& InEntry) override
 				{
+					FInterfaceRegistryTransaction::FTimeType TransactionTime = FPlatformTime::Cycles64();
 					if (InEntry.IsValid())
 					{
 						FInterfaceRegistryKey Key = GetInterfaceRegistryKey(InEntry->GetInterface());
@@ -25,11 +28,11 @@ namespace Metasound
 							{
 								UE_LOG(LogMetaSound, Warning, TEXT("Registration of interface overwriting previously registered interface [RegistryKey: %s]"), *Key);
 								
-								FInterfaceRegistryTransaction Transaction{FInterfaceRegistryTransaction::ETransactionType::InterfaceUnregistration, Key, Entry->GetInterface().Version};
+								FInterfaceRegistryTransaction Transaction{FInterfaceRegistryTransaction::ETransactionType::InterfaceUnregistration, Key, Entry->GetInterface().Version, TransactionTime};
 								Transactions.Add(MoveTemp(Transaction));
 							}
 							
-							FInterfaceRegistryTransaction Transaction{FInterfaceRegistryTransaction::ETransactionType::InterfaceRegistration, Key, InEntry->GetInterface().Version};
+							FInterfaceRegistryTransaction Transaction{FInterfaceRegistryTransaction::ETransactionType::InterfaceRegistration, Key, InEntry->GetInterface().Version, TransactionTime};
 							Transactions.Add(MoveTemp(Transaction));
 
 							Entries.Add(Key, MoveTemp(InEntry)).Get();
@@ -85,10 +88,11 @@ namespace Metasound
 			return GetInterfaceRegistryKey(InInterface.Version);
 		}
 
-		FInterfaceRegistryTransaction::FInterfaceRegistryTransaction(ETransactionType InType, const FInterfaceRegistryKey& InKey, const FMetasoundFrontendVersion& InInterfaceVersion)
+		FInterfaceRegistryTransaction::FInterfaceRegistryTransaction(ETransactionType InType, const FInterfaceRegistryKey& InKey, const FMetasoundFrontendVersion& InInterfaceVersion, FInterfaceRegistryTransaction::FTimeType InTimestamp)
 		: Type(InType)
 		, Key(InKey)
 		, InterfaceVersion(InInterfaceVersion)
+		, Timestamp(InTimestamp)
 		{
 		}
 
@@ -107,6 +111,11 @@ namespace Metasound
 			return Key;
 		}
 
+		FInterfaceRegistryTransaction::FTimeType FInterfaceRegistryTransaction::GetTimestamp() const
+		{
+			return Timestamp;
+		}
+
 		IInterfaceRegistry& IInterfaceRegistry::Get()
 		{
 			static MetasoundFrontendArchetypeRegistryPrivate::FInterfaceRegistry Registry;
diff --git a/Engine/Plugins/Runtime/Metasound/Source/MetasoundFrontend/Private/MetasoundFrontendQuerySteps.cpp b/Engine/Plugins/Runtime/Metasound/Source/MetasoundFrontend/Private/MetasoundFrontendQuerySteps.cpp
--- a/Engine/Plugins/Runtime/Metasound/Source/MetasoundFrontend/Private/MetasoundFrontendQuerySteps.cpp
+++ b/Engine/Plugins/Runtime/Metasound/Source/MetasoundFrontend/Private/MetasoundFrontendQuerySteps.cpp
@@ -2,9 +2,11 @@
 
 #include "MetasoundFrontendQuerySteps.h"
 
-#include "CoreMinimal.h"
+#include "Algo/MaxElement.h"
 #include "MetasoundFrontend.h"
 #include "MetasoundFrontendDocument.h"
+#include "MetasoundFrontendQuery.h"
+#include "MetasoundFrontendRegistries.h"
 #include "MetasoundFrontendRegistryTransaction.h"
 
 namespace Metasound
@@ -47,33 +49,11 @@ namespace Metasound
 	{
 		using namespace Frontend;
 
-		int32 State = 0;
-		FFrontendQueryEntry* FinalEntry = nullptr;
+		// Get most recent transaction
+		const FFrontendQueryEntry* FinalEntry = Algo::MaxElementBy(InOutEntries, GetTransactionTimestamp);
 
-		for (FFrontendQueryEntry& Entry : InOutEntries)
-		{
-			if (ensure(Entry.Value.IsType<FNodeRegistryTransaction>()))
-			{
-				const FNodeRegistryTransaction& Transaction = Entry.Value.Get<FNodeRegistryTransaction>();
-				
-				switch (Transaction.GetTransactionType())
-				{
-					case FNodeRegistryTransaction::ETransactionType::NodeRegistration:
-						State++;
-						FinalEntry = &Entry;
-						break;
-
-					case FNodeRegistryTransaction::ETransactionType::NodeUnregistration:
-						State--;
-						break;
-
-					default:
-						break;
-				}
-			}
-		}
-
-		if ((nullptr != FinalEntry) && (State > 0))
+		// Check if most recent transaction is valid and not an unregister transaction.
+		if (IsValidTransactionOfType(FNodeRegistryTransaction::ETransactionType::NodeRegistration, FinalEntry))
 		{
 			FFrontendQueryEntry Entry = *FinalEntry;
 			InOutEntries.Reset();
@@ -85,6 +65,31 @@ namespace Metasound
 		}
 	}
 
+	FReduceRegistrationEventsToCurrentStatus::FTimeType FReduceRegistrationEventsToCurrentStatus::GetTransactionTimestamp(const FFrontendQueryEntry& InEntry)
+	{
+		using namespace Frontend;
+
+		if (ensure(InEntry.Value.IsType<FNodeRegistryTransaction>()))
+		{
+			return InEntry.Value.Get<FNodeRegistryTransaction>().GetTimestamp();
+		}
+		return 0;
+	}
+
+	bool FReduceRegistrationEventsToCurrentStatus::IsValidTransactionOfType(Frontend::FNodeRegistryTransaction::ETransactionType InType, const FFrontendQueryEntry* InEntry)
+	{
+		using namespace Frontend;
+
+		if (nullptr != InEntry)
+		{
+			if (InEntry->Value.IsType<FNodeRegistryTransaction>())
+			{
+				return InEntry->Value.Get<FNodeRegistryTransaction>().GetTransactionType() == InType;
+			}
+		}
+		return false;
+	}
+
 	void FTransformRegistrationEventsToClasses::Transform(FFrontendQueryEntry::FValue& InValue) const
 	{
 		using namespace Frontend;
diff --git a/Engine/Plugins/Runtime/Metasound/Source/MetasoundFrontend/Private/MetasoundFrontendRegistries.cpp b/Engine/Plugins/Runtime/Metasound/Source/MetasoundFrontend/Private/MetasoundFrontendRegistries.cpp
--- a/Engine/Plugins/Runtime/Metasound/Source/MetasoundFrontend/Private/MetasoundFrontendRegistries.cpp
+++ b/Engine/Plugins/Runtime/Metasound/Source/MetasoundFrontend/Private/MetasoundFrontendRegistries.cpp
@@ -4,6 +4,7 @@
 
 #include "Algo/ForEach.h"
 #include "CoreMinimal.h"
+#include "HAL/PlatformTime.h"
 #include "MetasoundFrontendDataTypeRegistry.h"
 #include "MetasoundFrontendRegistryTransaction.h"
 #include "MetasoundFrontendSearchEngine.h"
@@ -262,6 +263,8 @@ namespace Metasound
 
 				if (InEntry.IsValid())
 				{
+					FNodeRegistryTransaction::FTimeType Timestamp = FPlatformTime::Cycles64();
+
 					Key = NodeRegistryKey::CreateKey(InEntry->GetClassInfo());
 
 					// check to see if an identical node was already registered, and log
@@ -274,7 +277,7 @@ namespace Metasound
 
 					// Store update to newly registered node in history so nodes
 					// can be queried by transaction ID
-					RegistryTransactionHistory.Add(FNodeRegistryTransaction(FNodeRegistryTransaction::ETransactionType::NodeRegistration, Key, InEntry->GetClassInfo()));
+					RegistryTransactionHistory.Add(FNodeRegistryTransaction(FNodeRegistryTransaction::ETransactionType::NodeRegistration, Key, InEntry->GetClassInfo(), Timestamp));
 
 					// Store registry elements in map so nodes can be queried using registry key.
 					RegisteredNodes.Add(Key, MoveTemp(InEntry));
@@ -289,7 +292,9 @@ namespace Metasound
 				{
 					if (const INodeRegistryEntry* Entry = FindNodeEntry(InKey))
 					{
-						RegistryTransactionHistory.Add(FNodeRegistryTransaction(FNodeRegistryTransaction::ETransactionType::NodeUnregistration, InKey, Entry->GetClassInfo()));
+						FNodeRegistryTransaction::FTimeType Timestamp = FPlatformTime::Cycles64();
+
+						RegistryTransactionHistory.Add(FNodeRegistryTransaction(FNodeRegistryTransaction::ETransactionType::NodeUnregistration, InKey, Entry->GetClassInfo(), Timestamp));
 
 						RegisteredNodes.Remove(InKey);
 						return true;
@@ -428,10 +433,11 @@ namespace Metasound
 			}
 		} // namespace MetasoundFrontendRegistriesPrivate
 
-		FNodeRegistryTransaction::FNodeRegistryTransaction(ETransactionType InType, const FNodeRegistryKey& InKey, const FNodeClassInfo& InNodeClassInfo)
+		FNodeRegistryTransaction::FNodeRegistryTransaction(ETransactionType InType, const FNodeRegistryKey& InKey, const FNodeClassInfo& InNodeClassInfo, FNodeRegistryTransaction::FTimeType InTimestamp)
 		: Type(InType)
 		, Key(InKey)
 		, NodeClassInfo(InNodeClassInfo)
+		, Timestamp(InTimestamp)
 		{
 		}
 
@@ -450,6 +456,11 @@ namespace Metasound
 			return Key;
 		}
 
+		FNodeRegistryTransaction::FTimeType FNodeRegistryTransaction::GetTimestamp() const
+		{
+			return Timestamp;
+		}
+
 		namespace NodeRegistryKey
 		{
 			// All registry keys should be created through this function to ensure consistency.
diff --git a/Engine/Plugins/Runtime/Metasound/Source/MetasoundFrontend/Private/MetasoundFrontendSearchEngine.cpp b/Engine/Plugins/Runtime/Metasound/Source/MetasoundFrontend/Private/MetasoundFrontendSearchEngine.cpp
--- a/Engine/Plugins/Runtime/Metasound/Source/MetasoundFrontend/Private/MetasoundFrontendSearchEngine.cpp
+++ b/Engine/Plugins/Runtime/Metasound/Source/MetasoundFrontend/Private/MetasoundFrontendSearchEngine.cpp
@@ -2,6 +2,7 @@
 
 #include "MetasoundFrontendSearchEngine.h"
 
+#include "Algo/MaxElement.h"
 #include "Algo/Transform.h"
 #include "HAL/IConsoleManager.h"
 #include "MetasoundFrontendArchetypeRegistry.h"
@@ -100,33 +101,11 @@ namespace Metasound
 			public:
 				virtual void Reduce(const FFrontendQueryKey& InKey, FFrontendQueryPartition& InOutEntries) const override
 				{
-					int32 State = 0;
-					FFrontendQueryEntry* FinalEntry = nullptr;
+					// Get most recent transaction
+					const FFrontendQueryEntry* FinalEntry = Algo::MaxElementBy(InOutEntries, GetTransactionTimestamp);
 
-					for (FFrontendQueryEntry& Entry : InOutEntries)
-					{
-						if (ensure(Entry.Value.IsType<FInterfaceRegistryTransaction>()))
-						{
-							const FInterfaceRegistryTransaction& Transaction = Entry.Value.Get<FInterfaceRegistryTransaction>();
-							
-							switch (Transaction.GetTransactionType())
-							{
-								case FInterfaceRegistryTransaction::ETransactionType::InterfaceRegistration:
-									State++;
-									FinalEntry = &Entry;
-									break;
-
-								case FInterfaceRegistryTransaction::ETransactionType::InterfaceUnregistration:
-									State--;
-									break;
-
-								default:
-									break;
-							}
-						}
-					}
-
-					if ((nullptr != FinalEntry) && (State > 0))
+					// Check if most recent transaction is valid and not an unregister transaction.
+					if (IsValidTransactionOfType(FInterfaceRegistryTransaction::ETransactionType::InterfaceRegistration, FinalEntry))
 					{
 						FFrontendQueryEntry Entry = *FinalEntry;
 						InOutEntries.Reset();
@@ -137,6 +116,28 @@ namespace Metasound
 						InOutEntries.Reset();
 					}
 				}
+			private:
+				static FInterfaceRegistryTransaction::FTimeType GetTransactionTimestamp(const FFrontendQueryEntry& InEntry)
+				{
+					if (ensure(InEntry.Value.IsType<FInterfaceRegistryTransaction>()))
+					{
+						return InEntry.Value.Get<FInterfaceRegistryTransaction>().GetTimestamp();
+					}
+					return 0;
+				}
+
+				static bool IsValidTransactionOfType(FInterfaceRegistryTransaction::ETransactionType InType, const FFrontendQueryEntry* InEntry)
+				{
+					if (nullptr != InEntry)
+					{
+						if (InEntry->Value.IsType<FInterfaceRegistryTransaction>())
+						{
+							return InEntry->Value.Get<FInterfaceRegistryTransaction>().GetTransactionType() == InType;
+						}
+					}
+					return false;
+				}
+
 			};
 
 
diff --git a/Engine/Plugins/Runtime/Metasound/Source/MetasoundFrontend/Public/MetasoundFrontendArchetypeRegistry.h b/Engine/Plugins/Runtime/Metasound/Source/MetasoundFrontend/Public/MetasoundFrontendArchetypeRegistry.h
--- a/Engine/Plugins/Runtime/Metasound/Source/MetasoundFrontend/Public/MetasoundFrontendArchetypeRegistry.h
+++ b/Engine/Plugins/Runtime/Metasound/Source/MetasoundFrontend/Public/MetasoundFrontendArchetypeRegistry.h
@@ -45,6 +45,8 @@ namespace Metasound
 		class METASOUNDFRONTEND_API FInterfaceRegistryTransaction
 		{
 		public:
+			using FTimeType = uint64;
+
 			/** Describes the type of transaction. */
 			enum class ETransactionType : uint8
 			{
@@ -53,17 +55,19 @@ namespace Metasound
 				Invalid
 			};
 
-			FInterfaceRegistryTransaction(ETransactionType InType, const FInterfaceRegistryKey& InKey, const FMetasoundFrontendVersion& InInterfaceVersion);
+			FInterfaceRegistryTransaction(ETransactionType InType, const FInterfaceRegistryKey& InKey, const FMetasoundFrontendVersion& InInterfaceVersion, FTimeType InTimestamp);
 
 			ETransactionType GetTransactionType() const;
 			const FMetasoundFrontendVersion& GetInterfaceVersion() const;
 			const FInterfaceRegistryKey& GetInterfaceRegistryKey() const;
+			FTimeType GetTimestamp() const;
 
 		private:
 
 			ETransactionType Type;
 			FInterfaceRegistryKey Key;
 			FMetasoundFrontendVersion InterfaceVersion;
+			FTimeType Timestamp;
 		};
 
 		class METASOUNDFRONTEND_API IInterfaceRegistry
diff --git a/Engine/Plugins/Runtime/Metasound/Source/MetasoundFrontend/Public/MetasoundFrontendQuerySteps.h b/Engine/Plugins/Runtime/Metasound/Source/MetasoundFrontend/Public/MetasoundFrontendQuerySteps.h
--- a/Engine/Plugins/Runtime/Metasound/Source/MetasoundFrontend/Public/MetasoundFrontendQuerySteps.h
+++ b/Engine/Plugins/Runtime/Metasound/Source/MetasoundFrontend/Public/MetasoundFrontendQuerySteps.h
@@ -2,10 +2,11 @@
 
 #pragma once
 
-#include "CoreMinimal.h"
 #include "MetasoundDataReference.h"
 #include "MetasoundFrontend.h"
 #include "MetasoundFrontendQuery.h"
+#include "MetasoundFrontendRegistries.h"
+#include "Misc/Guid.h"
 
 namespace Metasound
 {
@@ -25,7 +26,6 @@ namespace Metasound
 	public:
 		FNodeClassRegistrationEvents();
 		virtual void Stream(TArray<FFrontendQueryValue>& OutValues) override;
-		//void Reset() override;
 
 	private:
 		Frontend::FRegistryTransactionID CurrentTransactionID;
@@ -45,6 +45,10 @@ namespace Metasound
 	{
 	public:
 		virtual void Reduce(const FFrontendQueryKey& InKey, FFrontendQueryPartition& InOutEntries) const override;
+	private:
+		using FTimeType = Frontend::FNodeRegistryTransaction::FTimeType;
+		static FTimeType GetTransactionTimestamp(const FFrontendQueryEntry& InEntry);
+		static bool IsValidTransactionOfType(Frontend::FNodeRegistryTransaction::ETransactionType InType, const FFrontendQueryEntry* InEntry);
 	};
 
 	/** Transforms a registration event into a FMetasoundFrontendClass. */
diff --git a/Engine/Plugins/Runtime/Metasound/Source/MetasoundFrontend/Public/MetasoundFrontendRegistries.h b/Engine/Plugins/Runtime/Metasound/Source/MetasoundFrontend/Public/MetasoundFrontendRegistries.h
--- a/Engine/Plugins/Runtime/Metasound/Source/MetasoundFrontend/Public/MetasoundFrontendRegistries.h
+++ b/Engine/Plugins/Runtime/Metasound/Source/MetasoundFrontend/Public/MetasoundFrontendRegistries.h
@@ -209,6 +209,8 @@ namespace Metasound
 		class METASOUNDFRONTEND_API FNodeRegistryTransaction 
 		{
 		public:
+			using FTimeType = uint64;
+
 			/** Describes the type of transaction. */
 			enum class ETransactionType : uint8
 			{
@@ -217,17 +219,19 @@ namespace Metasound
 				Invalid
 			};
 
-			FNodeRegistryTransaction(ETransactionType InType, const FNodeRegistryKey& InKey, const FNodeClassInfo& InNodeClassInfo);
+			FNodeRegistryTransaction(ETransactionType InType, const FNodeRegistryKey& InKey, const FNodeClassInfo& InNodeClassInfo, FTimeType InTimestamp);
 
 			ETransactionType GetTransactionType() const;
 			const FNodeClassInfo& GetNodeClassInfo() const;
 			const FNodeRegistryKey& GetNodeRegistryKey() const;
+			FTimeType GetTimestamp() const;
 
 		private:
 
 			ETransactionType Type;
 			FNodeRegistryKey Key;
 			FNodeClassInfo NodeClassInfo;
+			FTimeType Timestamp;
 		};
 
 		namespace NodeRegistryKey
