Cleanup metasound deprecation wording in code.
Add visual warning for deprecated metasound nodes.

diff --git a/Engine/Plugins/Runtime/Metasound/Source/MetasoundEditor/Private/MetasoundEditorGraphNode.cpp b/Engine/Plugins/Runtime/Metasound/Source/MetasoundEditor/Private/MetasoundEditorGraphNode.cpp
--- a/Engine/Plugins/Runtime/Metasound/Source/MetasoundEditor/Private/MetasoundEditorGraphNode.cpp
+++ b/Engine/Plugins/Runtime/Metasound/Source/MetasoundEditor/Private/MetasoundEditorGraphNode.cpp
@@ -840,7 +840,23 @@ bool UMetasoundEditorGraphExternalNode::Validate(Metasound::Editor::FGraphNodeVa
 		}
 	}
 
-	// 6. Find all available versions & report if upgrade available
+	// Report of node class is deprecated
+	FMetasoundFrontendClass RegisteredClass;
+	if (FMetasoundFrontendRegistryContainer::Get()->GetFrontendClassFromRegistered(RegistryKey, RegisteredClass))
+	{
+		if (RegisteredClass.Metadata.GetIsDeprecated())
+		{
+			GraphNodePrivate::SetGraphNodeMessage(*this, EMessageSeverity::Warning,
+				FString::Format(TEXT("Class '{0} {1}' is deprecated."),
+				{
+					*RegisteredClass.Metadata.GetClassName().ToString(),
+					*RegisteredClass.Metadata.GetVersion().ToString()
+				}));
+		}
+	}
+
+
+	// Find all available versions & report if upgrade available
 	const Metasound::FNodeClassName NodeClassName = Metadata.GetClassName().ToNodeClassName();
 	const TArray<FMetasoundFrontendClass> SortedClasses = ISearchEngine::Get().FindClassesWithName(NodeClassName, true /* bSortByVersion */);
 	if (SortedClasses.IsEmpty())
diff --git a/Engine/Plugins/Runtime/Metasound/Source/MetasoundEditor/Private/MetasoundEditorGraphSchema.cpp b/Engine/Plugins/Runtime/Metasound/Source/MetasoundEditor/Private/MetasoundEditorGraphSchema.cpp
--- a/Engine/Plugins/Runtime/Metasound/Source/MetasoundEditor/Private/MetasoundEditorGraphSchema.cpp
+++ b/Engine/Plugins/Runtime/Metasound/Source/MetasoundEditor/Private/MetasoundEditorGraphSchema.cpp
@@ -45,10 +45,10 @@
 
 #define LOCTEXT_NAMESPACE "MetaSoundEditor"
 
-static int32 EnableDeprecatedMetaSoundNodeClassCreationCVar = 0;
-FAutoConsoleVariableRef CVarEnableDeprecatedMetaSoundNodeClassCreation(
-	TEXT("au.Debug.Editor.MetaSounds.EnableDeprecatedNodeClassCreation"),
-	EnableDeprecatedMetaSoundNodeClassCreationCVar,
+static int32 EnableAllVersionsMetaSoundNodeClassCreationCVar = 0;
+FAutoConsoleVariableRef CVarEnableAllVersionsMetaSoundNodeClassCreation(
+	TEXT("au.Debug.Editor.MetaSounds.EnableAllVersionsNodeClassCreation"),
+	EnableAllVersionsMetaSoundNodeClassCreationCVar,
 	TEXT("Enable creating nodes major versions of deprecated MetaSound classes in the Editor.\n")
 	TEXT("0: Disabled (default), !0: Enabled"),
 	ECVF_Default);
@@ -1446,8 +1446,8 @@ void UMetasoundEditorGraphSchema::GetConversionActions(FGraphActionMenuBuilder&
 	using namespace Metasound;
 	using namespace Metasound::Editor;
 
-	const bool bIncludeDeprecated = static_cast<bool>(EnableDeprecatedMetaSoundNodeClassCreationCVar);
-	const TArray<FMetasoundFrontendClass> FrontendClasses = Frontend::ISearchEngine::Get().FindAllClasses(bIncludeDeprecated);
+	const bool bIncludeAllVersions = static_cast<bool>(EnableAllVersionsMetaSoundNodeClassCreationCVar);
+	const TArray<FMetasoundFrontendClass> FrontendClasses = Frontend::ISearchEngine::Get().FindAllClasses(bIncludeAllVersions);
 	for (const FMetasoundFrontendClass& FrontendClass : FrontendClasses)
 	{
 		if (InFilters.InputFilterFunction && !FrontendClass.Interface.Inputs.ContainsByPredicate(InFilters.InputFilterFunction))
@@ -1558,8 +1558,8 @@ void UMetasoundEditorGraphSchema::GetFunctionActions(FGraphActionMenuBuilder& Ac
 		ParentAsset = AssetManager.TryLoadAssetFromKey(RegistryKey);
 	}
 
-	const bool bIncludeDeprecated = static_cast<bool>(EnableDeprecatedMetaSoundNodeClassCreationCVar);
-	const TArray<FMetasoundFrontendClass> FrontendClasses = ISearchEngine::Get().FindAllClasses(bIncludeDeprecated);
+	const bool bIncludeAllVersions = static_cast<bool>(EnableAllVersionsMetaSoundNodeClassCreationCVar);
+	const TArray<FMetasoundFrontendClass> FrontendClasses = ISearchEngine::Get().FindAllClasses(bIncludeAllVersions);
 	for (const FMetasoundFrontendClass& FrontendClass : FrontendClasses)
 	{
 		if (InFilters.InputFilterFunction && !FrontendClass.Interface.Inputs.ContainsByPredicate(InFilters.InputFilterFunction))
diff --git a/Engine/Plugins/Runtime/Metasound/Source/MetasoundFrontend/Private/MetasoundFrontendSearchEngine.cpp b/Engine/Plugins/Runtime/Metasound/Source/MetasoundFrontend/Private/MetasoundFrontendSearchEngine.cpp
--- a/Engine/Plugins/Runtime/Metasound/Source/MetasoundFrontend/Private/MetasoundFrontendSearchEngine.cpp
+++ b/Engine/Plugins/Runtime/Metasound/Source/MetasoundFrontend/Private/MetasoundFrontendSearchEngine.cpp
@@ -470,7 +470,7 @@ namespace Metasound
 		};
 
 		// Policy for finding all registered metasound classes, including deprecated classes. 
-		struct FFindAllClassesIncludingDeprecatedQueryPolicy
+		struct FFindAllClassesIncludingAllVersionsQueryPolicy
 		{
 			using ResultType = TArray<FMetasoundFrontendClass>;
 
@@ -642,7 +642,7 @@ namespace Metasound
 		};
 
 		// Policy for finding all registered interfaces (including deprecated). 
-		struct FFindAllInterfacesIncludingDeprecatedQueryPolicy
+		struct FFindAllInterfacesIncludingAllVersionsQueryPolicy
 		{
 			using ResultType = TArray<FMetasoundFrontendInterface>;
 
@@ -724,12 +724,12 @@ namespace Metasound
 			virtual ~FSearchEngine() = default;
 
 			virtual void Prime() override;
-			virtual TArray<FMetasoundFrontendClass> FindAllClasses(bool bInIncludeDeprecated) override;
+			virtual TArray<FMetasoundFrontendClass> FindAllClasses(bool bInIncludeAllVersions) override;
 			virtual TArray<FMetasoundFrontendClass> FindClassesWithName(const FMetasoundFrontendClassName& InName, bool bInSortByVersion) override;
 			virtual bool FindClassWithHighestVersion(const FMetasoundFrontendClassName& InName, FMetasoundFrontendClass& OutClass) override;
 			virtual bool FindClassWithMajorVersion(const FMetasoundFrontendClassName& InName, int32 InMajorVersion, FMetasoundFrontendClass& OutClass) override;
 
-			virtual TArray<FMetasoundFrontendInterface> FindAllInterfaces(bool bInIncludeDeprecated) override;
+			virtual TArray<FMetasoundFrontendInterface> FindAllInterfaces(bool bInIncludeAllVersions) override;
 			virtual TArray<FMetasoundFrontendInterface> FindUClassDefaultInterfaces(FName InUClassName) override;
 			virtual TArray<FMetasoundFrontendVersion> FindAllRegisteredInterfacesWithName(FName InInterfaceName) override;
 			virtual bool FindInterfaceWithHighestVersion(FName InInterfaceName, FMetasoundFrontendInterface& OutInterface) override;
@@ -737,14 +737,14 @@ namespace Metasound
 		private:
 
 			TSearchEngineQuery<FFindAllClassesQueryPolicy> FindAllClassesQuery;
-			TSearchEngineQuery<FFindAllClassesIncludingDeprecatedQueryPolicy> FindAllClassesIncludingDeprecatedQuery;
+			TSearchEngineQuery<FFindAllClassesIncludingAllVersionsQueryPolicy> FindAllClassesIncludingAllVersionsQuery;
 			TSearchEngineQuery<FFindClassesWithNameUnsortedQueryPolicy> FindClassesWithNameUnsortedQuery;
 			TSearchEngineQuery<FFindClassesWithNameSortedQueryPolicy> FindClassesWithNameSortedQuery;
 			TSearchEngineQuery<FFindClassWithHighestVersionQueryPolicy> FindClassWithHighestVersionQuery;
 			TSearchEngineQuery<FFindClassWithMajorVersionQueryPolicy> FindClassWithMajorVersionQuery;
 			TSearchEngineQuery<FFindAllInterfacesQueryPolicy> FindAllInterfacesQuery;
 			TSearchEngineQuery<FFindAllDefaultInterfacesQueryPolicy> FFindAllDefaultInterfacesQuery;
-			TSearchEngineQuery<FFindAllInterfacesIncludingDeprecatedQueryPolicy> FindAllInterfacesIncludingDeprecatedQuery;
+			TSearchEngineQuery<FFindAllInterfacesIncludingAllVersionsQueryPolicy> FindAllInterfacesIncludingAllVersionsQuery;
 			TSearchEngineQuery<FFindAllRegisteredInterfacesWithNameQueryPolicy> FindAllRegisteredInterfacesWithNameQuery;
 			TSearchEngineQuery<FFindInterfaceWithHighestVersionQueryPolicy> FindInterfaceWithHighestVersionQuery;
 		};
@@ -760,25 +760,25 @@ namespace Metasound
 			METASOUND_TRACE_CPUPROFILER_EVENT_SCOPE(metasound::FSearchEngine::Prime);
 
 			FindAllClassesQuery.Prime();
-			FindAllClassesIncludingDeprecatedQuery.Prime();
+			FindAllClassesIncludingAllVersionsQuery.Prime();
 			FindClassesWithNameUnsortedQuery.Prime();
 			FindClassesWithNameSortedQuery.Prime();
 			FindClassWithHighestVersionQuery.Prime();
 			FindClassWithMajorVersionQuery.Prime();
-			FindAllInterfacesIncludingDeprecatedQuery.Prime();
+			FindAllInterfacesIncludingAllVersionsQuery.Prime();
 			FindAllInterfacesQuery.Prime();
 			FindAllRegisteredInterfacesWithNameQuery.Prime();
 			FindInterfaceWithHighestVersionQuery.Prime();
 		}
 
-		TArray<FMetasoundFrontendClass> FSearchEngine::FindAllClasses(bool bInIncludeDeprecated)
+		TArray<FMetasoundFrontendClass> FSearchEngine::FindAllClasses(bool bInIncludeAllVersions)
 		{
 			METASOUND_TRACE_CPUPROFILER_EVENT_SCOPE(metasound::FSearchEngine::FindAllClasses);
 
 			FFrontendQueryKey NullKey;
-			if (bInIncludeDeprecated)
+			if (bInIncludeAllVersions)
 			{
-				return FindAllClassesIncludingDeprecatedQuery.UpdateAndFindResult(NullKey);
+				return FindAllClassesIncludingAllVersionsQuery.UpdateAndFindResult(NullKey);
 			}
 			else
 			{
@@ -818,14 +818,14 @@ namespace Metasound
 			return FindClassWithMajorVersionQuery.UpdateAndFindResult(Key, OutClass);
 		}
 
-		TArray<FMetasoundFrontendInterface> FSearchEngine::FindAllInterfaces(bool bInIncludeDeprecated)
+		TArray<FMetasoundFrontendInterface> FSearchEngine::FindAllInterfaces(bool bInIncludeAllVersions)
 		{
 			METASOUND_TRACE_CPUPROFILER_EVENT_SCOPE(metasound::FSearchEngine::FindAllInterfaces);
 
 			const FFrontendQueryKey NullKey;
-			if (bInIncludeDeprecated)
+			if (bInIncludeAllVersions)
 			{
-				return FindAllInterfacesIncludingDeprecatedQuery.UpdateAndFindResult(NullKey);
+				return FindAllInterfacesIncludingAllVersionsQuery.UpdateAndFindResult(NullKey);
 			}
 			else
 			{
diff --git a/Engine/Plugins/Runtime/Metasound/Source/MetasoundFrontend/Private/MetasoundFrontendTransform.cpp b/Engine/Plugins/Runtime/Metasound/Source/MetasoundFrontend/Private/MetasoundFrontendTransform.cpp
--- a/Engine/Plugins/Runtime/Metasound/Source/MetasoundFrontend/Private/MetasoundFrontendTransform.cpp
+++ b/Engine/Plugins/Runtime/Metasound/Source/MetasoundFrontend/Private/MetasoundFrontendTransform.cpp
@@ -922,8 +922,8 @@ namespace Metasound
 				// with an "Invalid" interface version.
 				if (ensure(Interfaces.IsEmpty()))
 				{
-					constexpr bool bIncludeDeprecatedInterfaces = true;
-					TArray<FMetasoundFrontendInterface> AllInterfaces = ISearchEngine::Get().FindAllInterfaces(bIncludeDeprecatedInterfaces);
+					constexpr bool bIncludeAllVersions = true;
+					TArray<FMetasoundFrontendInterface> AllInterfaces = ISearchEngine::Get().FindAllInterfaces(bIncludeAllVersions);
 
 					const FMetasoundFrontendGraphClass& RootGraph = InDocument->GetRootGraphClass();
 					const TArray<FMetasoundFrontendClass>& Dependencies = InDocument->GetDependencies();
diff --git a/Engine/Plugins/Runtime/Metasound/Source/MetasoundFrontend/Public/MetasoundFrontendSearchEngine.h b/Engine/Plugins/Runtime/Metasound/Source/MetasoundFrontend/Public/MetasoundFrontendSearchEngine.h
--- a/Engine/Plugins/Runtime/Metasound/Source/MetasoundFrontend/Public/MetasoundFrontendSearchEngine.h
+++ b/Engine/Plugins/Runtime/Metasound/Source/MetasoundFrontend/Public/MetasoundFrontendSearchEngine.h
@@ -24,9 +24,9 @@ namespace Metasound
 				virtual void Prime() = 0;
 
 				/** Find all FMetasoundFrontendClasses.
-				  * (Optional) Include deprecated classes (i.e. versions of classes that are not the highest major version).
+				  * (Optional) Include all versions (i.e. deprecated classes and versions of classes that are not the highest major version).
 				  */
-				virtual TArray<FMetasoundFrontendClass> FindAllClasses(bool bInIncludeDeprecated) = 0;
+				virtual TArray<FMetasoundFrontendClass> FindAllClasses(bool bInIncludeAllVersions) = 0;
 
 				/** Find all classes with the given ClassName.
 				  * (Optional) Sort matches based on version.
@@ -39,7 +39,7 @@ namespace Metasound
 				/** Find the class with the given ClassName & Major Version. Returns false if not found, true if found. */
 				virtual bool FindClassWithMajorVersion(const FMetasoundFrontendClassName& InName, int32 InMajorVersion, FMetasoundFrontendClass& OutClass) = 0;
 
-				virtual TArray<FMetasoundFrontendInterface> FindAllInterfaces(bool bInIncludeDeprecated = false) = 0;
+				virtual TArray<FMetasoundFrontendInterface> FindAllInterfaces(bool bInIncludeAllVersions = false) = 0;
 
 				virtual TArray<FMetasoundFrontendInterface> FindUClassDefaultInterfaces(FName InUClassName) = 0;
 
