Fix for Random Node not initializing first output value

diff --git a/Engine/Plugins/Runtime/Metasound/Source/MetasoundStandardNodes/Public/MetasoundRandomNode.h b/Engine/Plugins/Runtime/Metasound/Source/MetasoundStandardNodes/Public/MetasoundRandomNode.h
--- a/Engine/Plugins/Runtime/Metasound/Source/MetasoundStandardNodes/Public/MetasoundRandomNode.h
+++ b/Engine/Plugins/Runtime/Metasound/Source/MetasoundStandardNodes/Public/MetasoundRandomNode.h
@@ -360,11 +360,11 @@ namespace Metasound
 			, bIsDefaultSeeded(*SeedValue == DefaultSeed)
 			, bIsRandomStreamInitialized(false)
 		{
-			// We need to initialize the output value to *something*
-			*OutputValue = *MinValue;
-
 			EvaluateSeedChanges();
 			RandomStream.Reset();
+
+			// We need to initialize the output value to *something*
+			*OutputValue = TRandomNodeSpecialization<ValueType>::GetNextValue(RandomStream, *MinValue, *MaxValue);
 		}
 
 		TRandomNodeOperator(const FOperatorSettings& InSettings,
@@ -382,11 +382,11 @@ namespace Metasound
 			, bIsDefaultSeeded(*SeedValue == DefaultSeed)
 			, bIsRandomStreamInitialized(false)
 		{
-			// We need to initialize the output value to *something*
-			*OutputValue = *MinValue;
-
 			EvaluateSeedChanges();
 			RandomStream.Reset();
+
+			// We need to initialize the output value to *something*
+			*OutputValue = TRandomNodeSpecialization<ValueType>::GetNextValue(RandomStream, *MinValue, *MaxValue);
 		}
 
 		virtual ~TRandomNodeOperator() = default;
