FText Fix-ups in MetaSounds to fix loc build warnings
#rnx

diff --git a/Engine/Plugins/Runtime/Metasound/Source/MetasoundFrontend/Private/MetasoundArrayRandomNode.cpp b/Engine/Plugins/Runtime/Metasound/Source/MetasoundFrontend/Private/MetasoundArrayRandomNode.cpp
--- a/Engine/Plugins/Runtime/Metasound/Source/MetasoundFrontend/Private/MetasoundArrayRandomNode.cpp
+++ b/Engine/Plugins/Runtime/Metasound/Source/MetasoundFrontend/Private/MetasoundArrayRandomNode.cpp
@@ -18,7 +18,6 @@
 
 #include <type_traits>
 
-#define LOCTEXT_NAMESPACE "MetasoundFrontend_RandomArrayGet"
 
 namespace Metasound
 {
@@ -234,5 +233,3 @@ namespace Metasound
 	}
 
 }
-
-#undef LOCTEXT_NAMESPACE
diff --git a/Engine/Plugins/Runtime/Metasound/Source/MetasoundFrontend/Private/MetasoundOutputFormatInterfaces.cpp b/Engine/Plugins/Runtime/Metasound/Source/MetasoundFrontend/Private/MetasoundOutputFormatInterfaces.cpp
--- a/Engine/Plugins/Runtime/Metasound/Source/MetasoundFrontend/Private/MetasoundOutputFormatInterfaces.cpp
+++ b/Engine/Plugins/Runtime/Metasound/Source/MetasoundFrontend/Private/MetasoundOutputFormatInterfaces.cpp
@@ -76,14 +76,14 @@ namespace Metasound
 						Outputs =
 						{
 							{
-								LOCTEXT("GeneratedLeftDisplayName", "Out Left"),
-								LOCTEXT("GeneratedLeftDescription", "The resulting left channel output audio from this source."),
+								LOCTEXT("OutputFormatStereoInterface_GeneratedLeftDisplayName", "Out Left"),
+								LOCTEXT("OutputFormatStereoInterface_GeneratedLeftDescription", "The resulting left channel output audio from this source."),
 								GetMetasoundDataTypeName<FAudioBuffer>(),
 								Outputs::LeftOut
 							},
 							{
-								LOCTEXT("GeneratedLeftDisplayName", "Out Right"),
-								LOCTEXT("GeneratedLeftDescription", "The resulting right channel output audio from this source."),
+								LOCTEXT("OutputFormatStereoInterface_GeneratedRightDisplayName", "Out Right"),
+								LOCTEXT("OutputFormatStereoInterface_GeneratedRightDescription", "The resulting right channel output audio from this source."),
 								GetMetasoundDataTypeName<FAudioBuffer>(),
 								Outputs::RightOut
 							}
diff --git a/Engine/Plugins/Runtime/Metasound/Source/MetasoundFrontend/Public/MetasoundArrayRandomNode.h b/Engine/Plugins/Runtime/Metasound/Source/MetasoundFrontend/Public/MetasoundArrayRandomNode.h
--- a/Engine/Plugins/Runtime/Metasound/Source/MetasoundFrontend/Public/MetasoundArrayRandomNode.h
+++ b/Engine/Plugins/Runtime/Metasound/Source/MetasoundFrontend/Public/MetasoundArrayRandomNode.h
@@ -1,9 +1,7 @@
 // Copyright Epic Games, Inc. All Rights Reserved.
-
 #pragma once
 
 #include "Containers/CircularQueue.h"
-#include "CoreMinimal.h"
 
 #include "MetasoundArrayNodes.h"
 #include "MetasoundExecutableOperator.h"
@@ -12,33 +10,33 @@
 #include "MetasoundNodeInterface.h"
 #include "MetasoundTrigger.h"
 
-#define LOCTEXT_NAMESPACE "MetasoundFrontend_RandomArrayGet"
+#define LOCTEXT_NAMESPACE "MetasoundFrontend"
 
 namespace Metasound
 {
 	namespace ArrayNodeRandomGetVertexNames
 	{
-		static const TCHAR* InputTriggerNextName = TEXT("Next");
-		static const TCHAR* InputTriggerResetName = TEXT("Reset");
-		static const TCHAR* InputRandomArrayName = TEXT("In Array");
-		static const TCHAR* InputWeightsName = TEXT("Weights");
-		static const TCHAR* InputSeedName = TEXT("Seed");
-		static const TCHAR* InputNoRepeatOrderName = TEXT("No Repeats");
-		static const TCHAR* InputEnableSharedStateName = TEXT("Enable Shared State");
-		static const TCHAR* OutputTriggerOnNextName = TEXT("On Next");
-		static const TCHAR* OutputTriggerOnResetName = TEXT("On Reset");
-		static const TCHAR* OutputValueName = TEXT("Value");
-
-		static const FText InputTriggerNextTooltip = LOCTEXT("TriggerNextTooltip", "Trigger to get the next value in the randomized array.");
-		static const FText InputTriggerResetTooltip = LOCTEXT("TriggerResetTooltip", "Trigger to reset the seed for the randomized array.");
-		static const FText InputRandomArrayTooltip = LOCTEXT("RandomArrayTooltip", "Input array to randomized.");
-		static const FText InputWeightsTooltip = LOCTEXT("WeightsTooltip", "Input array of weights to use for random selection. Will repeat if this array is shorter than the input array to select from.");
-		static const FText InputSeedTooltip = LOCTEXT("SeedTooltip", "Seed to use for the random shuffle.");
-		static const FText InputNoRepeatOrderTooltip = LOCTEXT("NoRepeatOrderTooltip", "The number of elements to track to avoid repeating in a row.");
-		static const FText InputEnableSharedStateTooltip = LOCTEXT("EnableSharedStateTooltip", "Set to enabled to share state across instances of this metasound.");
-		static const FText OutputTriggerOnNextTooltip = LOCTEXT("TriggerOnNextTooltip", "Triggers when the \"Next\" input is triggered.");
-		static const FText OutputTriggerOnResetTooltip = LOCTEXT("TriggerOnResetTooltip", "Triggers when the \"Shuffle\" input is triggered or if the array is auto-shuffled.");
-		static const FText OutputValueTooltip = LOCTEXT("ValueTooltip", "Value of the current shuffled element.");
+		static const FName InputTriggerNextName = "Next";
+		static const FName InputTriggerResetName = "Reset";
+		static const FName InputRandomArrayName = "In Array";
+		static const FName InputWeightsName = "Weights";
+		static const FName InputSeedName = "Seed";
+		static const FName InputNoRepeatOrderName = "No Repeats";
+		static const FName InputEnableSharedStateName = "Enable Shared State";
+		static const FName OutputTriggerOnNextName = "On Next";
+		static const FName OutputTriggerOnResetName = "On Reset";
+		static const FName OutputValueName = "Value";
+
+		static const FText InputTriggerNextTooltip = LOCTEXT("RandomArrayGetNode_TriggerNextTooltip", "Trigger to get the next value in the randomized array.");
+		static const FText InputTriggerResetTooltip = LOCTEXT("RandomArrayGetNode_TriggerResetTooltip", "Trigger to reset the seed for the randomized array.");
+		static const FText InputRandomArrayTooltip = LOCTEXT("RandomArrayGetNode_RandomArrayTooltip", "Input array to randomized.");
+		static const FText InputWeightsTooltip = LOCTEXT("RandomArrayGetNode_WeightsTooltip", "Input array of weights to use for random selection. Will repeat if this array is shorter than the input array to select from.");
+		static const FText InputSeedTooltip = LOCTEXT("RandomArrayGetNode_SeedTooltip", "Seed to use for the random shuffle.");
+		static const FText InputNoRepeatOrderTooltip = LOCTEXT("RandomArrayGetNode_NoRepeatOrderTooltip", "The number of elements to track to avoid repeating in a row.");
+		static const FText InputEnableSharedStateTooltip = LOCTEXT("RandomArrayGetNode_EnableSharedStateTooltip", "Set to enabled to share state across instances of this MetaSound.");
+		static const FText OutputTriggerOnNextTooltip = LOCTEXT("RandomArrayGetNode_TriggerOnNextTooltip", "Triggers when the \"Next\" input is triggered.");
+		static const FText OutputTriggerOnResetTooltip = LOCTEXT("RandomArrayGetNode_TriggerOnResetTooltip", "Triggers when the \"Shuffle\" input is triggered or if the array is auto-shuffled.");
+		static const FText OutputValueTooltip = LOCTEXT("RandomArrayGetNode_ValueTooltip", "Value of the current shuffled element.");
 	}
 
 	class METASOUNDFRONTEND_API FArrayRandomGet
@@ -142,8 +140,8 @@ namespace Metasound
 			{
 				FName DataTypeName = GetMetasoundDataTypeName<ArrayType>();
 				FName OperatorName = "Random Get";
-				FText NodeDisplayName = FText::Format(LOCTEXT("RandomGetArrayOpDisplayNamePattern", "Random Get ({0})"), GetMetasoundDataTypeDisplayText<ArrayType>());
-				FText NodeDescription = LOCTEXT("RandomGetArrayDescription", "Randomly retrieve data from input array using the supplied weights.");
+				FText NodeDisplayName = FText::Format(LOCTEXT("RandomArrayGetNode_OpDisplayNamePattern", "Random Get ({0})"), GetMetasoundDataTypeDisplayText<ArrayType>());
+				FText NodeDescription = LOCTEXT("RandomArrayGetNode_Description", "Randomly retrieve data from input array using the supplied weights.");
 				FVertexInterface NodeInterface = GetDefaultInterface();
 
 				return MetasoundArrayNodesPrivate::CreateArrayNodeClassMetadata(DataTypeName, OperatorName, NodeDisplayName, NodeDescription, NodeInterface);
@@ -448,8 +446,5 @@ namespace Metasound
 
 		virtual ~TArrayRandomGetNode() = default;
 	};
-
-
 }
-
 #undef LOCTEXT_NAMESPACE
diff --git a/Engine/Plugins/Runtime/Metasound/Source/MetasoundStandardNodes/Private/MetasoundBPMToSecondsNode.cpp b/Engine/Plugins/Runtime/Metasound/Source/MetasoundStandardNodes/Private/MetasoundBPMToSecondsNode.cpp
--- a/Engine/Plugins/Runtime/Metasound/Source/MetasoundStandardNodes/Private/MetasoundBPMToSecondsNode.cpp
+++ b/Engine/Plugins/Runtime/Metasound/Source/MetasoundStandardNodes/Private/MetasoundBPMToSecondsNode.cpp
@@ -9,7 +9,7 @@
 #include "MetasoundStandardNodesNames.h"
 #include "MetasoundTime.h"
 
-#define LOCTEXT_NAMESPACE "MetasoundStandardNodes_BPMToSeconds"
+#define LOCTEXT_NAMESPACE "MetasoundStandardNodes"
 
 namespace Metasound
 {
@@ -134,12 +134,12 @@ namespace Metasound
 
 		static const FVertexInterface Interface(
 			FInputVertexInterface(
-				TInputDataVertexModel<float>(GetInputBPMTempoName(), LOCTEXT("BPMToSecsBPMTT", "Input BPM (Beats Per Minute)."), 90.0f),
-				TInputDataVertexModel<float>(GetInputBeatMultiplierName(), LOCTEXT("BPMToSecsBeatMultiplierTT", "The multiplier of the BPM."), 1.0f),
-				TInputDataVertexModel<float>(GetInputDivisionsOfWholeNoteName(), LOCTEXT("BPMToSecsDivOfWholeNoteTT", "Divisions of a whole note."), 4.0f)
+				TInputDataVertexModel<float>(GetInputBPMTempoName(), LOCTEXT("BPMToSecondsNode_BPMTT", "Beats Per Minute."), 90.0f),
+				TInputDataVertexModel<float>(GetInputBeatMultiplierName(), LOCTEXT("BPMToSecondsNode_BeatMultiplierTT", "The multiplier of the BPM."), 1.0f),
+				TInputDataVertexModel<float>(GetInputDivisionsOfWholeNoteName(), LOCTEXT("BPMToSecondsNode_DivOfWholeNoteTT", "Divisions of a whole note."), 4.0f)
 			),
 			FOutputVertexInterface(
-				TOutputDataVertexModel<FTime>(GetOutputTimeSecondsName(), LOCTEXT("BPMToSecsOutputTimeTT", "The output time in seconds, given the BPM inputs."))
+				TOutputDataVertexModel<FTime>(GetOutputTimeSecondsName(), LOCTEXT("BPMToSecondsNode_OutputTimeTT", "The output time in seconds."))
 			)
 		);
 
@@ -151,11 +151,11 @@ namespace Metasound
 		auto InitNodeInfo = []() -> FNodeClassMetadata
 		{
 			FNodeClassMetadata Info;
-			Info.ClassName = { StandardNodes::Namespace, TEXT("BPMToSeconds"), TEXT("") };
+			Info.ClassName = { StandardNodes::Namespace, "BPMToSeconds", "" };
 			Info.MajorVersion = 1;
-			Info.MinorVersion = 0;
-			Info.DisplayName = LOCTEXT("BPMToSecondsDisplayName", "BPM To Seconds");
-			Info.Description = LOCTEXT("BPMToSecondsDesc", "Calculates a beat time in seconds from the given BPM, beat multiplier and divisions of a whole note.");
+			Info.MinorVersion = 1;
+			Info.DisplayName = LOCTEXT("BPMToSecondsNode_DisplayName", "BPM To Seconds");
+			Info.Description = LOCTEXT("BPMToSecondsNode_Desc", "Calculates a beat time in seconds from the given BPM, beat multiplier and divisions of a whole note.");
 			Info.Author = PluginAuthor;
 			Info.PromptIfMissing = PluginNodeMissingPrompt;
 			Info.DefaultInterface = GetVertexInterface();
diff --git a/Engine/Plugins/Runtime/Metasound/Source/MetasoundStandardNodes/Private/MetasoundDelayNode.cpp b/Engine/Plugins/Runtime/Metasound/Source/MetasoundStandardNodes/Private/MetasoundDelayNode.cpp
--- a/Engine/Plugins/Runtime/Metasound/Source/MetasoundStandardNodes/Private/MetasoundDelayNode.cpp
+++ b/Engine/Plugins/Runtime/Metasound/Source/MetasoundStandardNodes/Private/MetasoundDelayNode.cpp
@@ -12,18 +12,18 @@
 #include "MetasoundStandardNodesCategories.h"
 #include "MetasoundFacade.h"
 
-#define LOCTEXT_NAMESPACE "MetasoundStandardNodes_Delay"
+#define LOCTEXT_NAMESPACE "MetasoundStandardNodes"
 
 namespace Metasound
 {
 	namespace Delay
 	{
-		static const TCHAR* InParamNameAudioInput = TEXT("In");
-		static const TCHAR* InParamNameDelayTime = TEXT("Delay Time");
-		static const TCHAR* InParamNameDryLevel = TEXT("Dry Level");
-		static const TCHAR* InParamNameWetLevel = TEXT("Wet Level");
-		static const TCHAR* InParamNameFeedbackAmount = TEXT("Feedback");
-		static const TCHAR* OutParamNameAudio = TEXT("Out");
+		static const FName InParamNameAudioInput = "In";
+		static const FName InParamNameDelayTime = "Delay Time";
+		static const FName InParamNameDryLevel = "Dry Level";
+		static const FName InParamNameWetLevel = "Wet Level";
+		static const FName InParamNameFeedbackAmount = "Feedback";
+		static const FName OutParamNameAudio = "Out";
 
 		// TODO: make this a static vertex
 		static float MaxDelaySeconds = 5.0f;
@@ -170,14 +170,14 @@ namespace Metasound
 	{
 		static const FVertexInterface Interface(
 			FInputVertexInterface(
-				TInputDataVertexModel<FAudioBuffer>(Delay::InParamNameAudioInput, LOCTEXT("AudioInputTooltip", "Audio input.")),
-				TInputDataVertexModel<FTime>(Delay::InParamNameDelayTime, LOCTEXT("DelayTimeTooltip", "The amount of time to delay the audio, in seconds."), 1.0f),
-				TInputDataVertexModel<float>(Delay::InParamNameDryLevel, LOCTEXT("DryLevelTooltip", "The dry level of the delay."), 0.0f),
-				TInputDataVertexModel<float>(Delay::InParamNameWetLevel, LOCTEXT("WetlevelTooltip", "The wet level of the delay."), 1.0f),
-				TInputDataVertexModel<float>(Delay::InParamNameFeedbackAmount, LOCTEXT("FeedbackTooltip", "Feedback amount."), 0.0f)
+				TInputDataVertexModel<FAudioBuffer>(Delay::InParamNameAudioInput, LOCTEXT("DelayNode_AudioInputTooltip", "Audio input.")),
+				TInputDataVertexModel<FTime>(Delay::InParamNameDelayTime, LOCTEXT("DelayNode_DelayTimeTooltip", "The amount of time to delay the audio, in seconds."), 1.0f),
+				TInputDataVertexModel<float>(Delay::InParamNameDryLevel, LOCTEXT("DelayNode_DryLevelTooltip", "The dry level of the delay."), 0.0f),
+				TInputDataVertexModel<float>(Delay::InParamNameWetLevel, LOCTEXT("DelayNode_WetlevelTooltip", "The wet level of the delay."), 1.0f),
+				TInputDataVertexModel<float>(Delay::InParamNameFeedbackAmount, LOCTEXT("DelayNode_FeedbackTooltip", "Feedback amount."), 0.0f)
 			),
 			FOutputVertexInterface(
-				TOutputDataVertexModel<FAudioBuffer>(Delay::OutParamNameAudio, LOCTEXT("DelayOutputTooltip", "Audio output."))
+				TOutputDataVertexModel<FAudioBuffer>(Delay::OutParamNameAudio, LOCTEXT("DelayNode_DelayOutputTooltip", "Audio output."))
 			)
 		);
 
@@ -189,11 +189,11 @@ namespace Metasound
 		auto InitNodeInfo = []() -> FNodeClassMetadata
 		{
 			FNodeClassMetadata Info;
-			Info.ClassName = { StandardNodes::Namespace, TEXT("Delay"), StandardNodes::AudioVariant };
+			Info.ClassName = { StandardNodes::Namespace, "Delay", StandardNodes::AudioVariant };
 			Info.MajorVersion = 1;
-			Info.MinorVersion = 0;
-			Info.DisplayName = LOCTEXT("Metasound_DelayDisplayName", "Delay");
-			Info.Description = LOCTEXT("Metasound_DelayNodeDescription", "Delays an audio buffer by the specified amount.");
+			Info.MinorVersion = 1;
+			Info.DisplayName = LOCTEXT("DelayNode_DisplayName", "Delay");
+			Info.Description = LOCTEXT("DelayNode_Description", "Delays an audio buffer by the specified amount.");
 			Info.Author = PluginAuthor;
 			Info.PromptIfMissing = PluginNodeMissingPrompt;
 			Info.DefaultInterface = GetVertexInterface();
diff --git a/Engine/Plugins/Runtime/Metasound/Source/MetasoundStandardNodes/Private/MetasoundRandomNode.cpp b/Engine/Plugins/Runtime/Metasound/Source/MetasoundStandardNodes/Private/MetasoundRandomNode.cpp
--- a/Engine/Plugins/Runtime/Metasound/Source/MetasoundStandardNodes/Private/MetasoundRandomNode.cpp
+++ b/Engine/Plugins/Runtime/Metasound/Source/MetasoundStandardNodes/Private/MetasoundRandomNode.cpp
@@ -5,7 +5,7 @@
 #include "Internationalization/Text.h"
 #include "MetasoundNodeRegistrationMacro.h"
 
-#define LOCTEXT_NAMESPACE "MetasoundStandardNodes_RandomNode"
+#define LOCTEXT_NAMESPACE "MetasoundStandardNodes"
 
 namespace Metasound
 {
@@ -61,42 +61,42 @@ namespace Metasound
 
 		static FText GetNextTriggerDescription()
 		{
-			return LOCTEXT("RandomNodeNextTT", "Trigger to generate the next random integer.");
+			return LOCTEXT("RandomNode_NextTT", "Trigger to generate the next random integer.");
 		}
 
 		static FText GetResetDescription()
 		{
-			return LOCTEXT("RandomNodeResetTT", "Trigger to reset the random sequence with the supplied seed. Useful to get randomized repetition.");
+			return LOCTEXT("RandomNode_ResetTT", "Trigger to reset the random sequence with the supplied seed. Useful to get randomized repetition.");
 		}
 
 		static FText GetSeedDescription()
 		{
-			return LOCTEXT("RandomNodeSeedTT", "The seed value to use for the random node. Set to -1 to use a random seed.");
+			return LOCTEXT("RandomNode_SeedTT", "The seed value to use for the random node. Set to -1 to use a random seed.");
 		}
 
 		static FText GetMinDescription()
 		{
-			return LOCTEXT("RandomNodeMinTT", "Min random value.");
+			return LOCTEXT("RandomNode_MinTT", "Min random value.");
 		}
 
 		static FText GetMaxDescription()
 		{
-			return LOCTEXT("RandNodeMaxTT", "Max random value.");
+			return LOCTEXT("RandomNode_MaxTT", "Max random value.");
 		}
 
 		static FText GetOutputDescription()
 		{
-			return LOCTEXT("RandomNodeOutputTT", "The randomly generated value.");
+			return LOCTEXT("RandomNode_OutputTT", "The randomly generated value.");
 		}
 
 		static FText GetOutputOnNextDescription()
 		{
-			return LOCTEXT("RandomNodeOutputNextTT", "Triggers when next is triggered.");
+			return LOCTEXT("RandomNode_OutputNextTT", "Triggers when next is triggered.");
 		}
 
 		static FText GetOutputOnResetDescription()
 		{
-			return LOCTEXT("RandomNodeOutputNextTT", "Triggers when reset is triggered.");
+			return LOCTEXT("RandomNode_OutputResetTT", "Triggers when reset is triggered.");
 		}
 	}
 
diff --git a/Engine/Plugins/Runtime/Metasound/Source/MetasoundStandardNodes/Public/MetasoundRandomNode.h b/Engine/Plugins/Runtime/Metasound/Source/MetasoundStandardNodes/Public/MetasoundRandomNode.h
--- a/Engine/Plugins/Runtime/Metasound/Source/MetasoundStandardNodes/Public/MetasoundRandomNode.h
+++ b/Engine/Plugins/Runtime/Metasound/Source/MetasoundStandardNodes/Public/MetasoundRandomNode.h
@@ -13,7 +13,7 @@
 #include "MetasoundTime.h"
 #include "MetasoundVertex.h"
 
-#define LOCTEXT_NAMESPACE "MetasoundStandardNodes_RandomNode"
+#define LOCTEXT_NAMESPACE "MetasoundStandardNodes"
 
 namespace Metasound
 {
@@ -58,12 +58,12 @@ namespace Metasound
 
 		static FText GetDisplayName()
 		{
-			return LOCTEXT("Int32RandomValueDisplayName", "Random (Int)");
+			return LOCTEXT("RandomNode_Int32RandomValueDisplayName", "Random (Int)");
 		}
 
 		static FText GetDescription()
 		{
-			return LOCTEXT("Int32RandomDescription", "Generates a seedable random integer in the given value range.");
+			return LOCTEXT("RandomNode_Int32RandomDescription", "Generates a seedable random integer in the given value range.");
 		}
 
 		static bool HasRange()
@@ -107,12 +107,12 @@ namespace Metasound
 
 		static FText GetDisplayName()
 		{
-			return LOCTEXT("FloatRandomValueDisplayName", "Random (Float)");
+			return LOCTEXT("RandomNode_FloatRandomValueDisplayName", "Random (Float)");
 		}
 
 		static FText GetDescription()
 		{
-			return LOCTEXT("FloatRandomDescription", "Generates a seedable random float in the given value range.");
+			return LOCTEXT("RandomNode_FloatRandomDescription", "Generates a seedable random float in the given value range.");
 		}
 
 		static bool HasRange()
@@ -156,12 +156,12 @@ namespace Metasound
 
 		static FText GetDisplayName()
 		{
-			return LOCTEXT("BoolRandomValueDisplayName", "Random (Bool)");
+			return LOCTEXT("RandomNode_BoolRandomValueDisplayName", "Random (Bool)");
 		}
 
 		static FText GetDescription()
 		{
-			return LOCTEXT("BoolRandomDiscription", "Generates a random bool value.");
+			return LOCTEXT("RandomNode_BoolRandomDiscription", "Generates a random bool value.");
 		}
 
 		static bool HasRange()
@@ -205,12 +205,12 @@ namespace Metasound
 
 		static FText GetDisplayName()
 		{
-			return LOCTEXT("TimeRandomValueDisplayName", "Random (Time)");
+			return LOCTEXT("RandomNode_TimeRandomValueDisplayName", "Random (Time)");
 		}
 
 		static FText GetDescription()
 		{
-			return LOCTEXT("TimeRandomDiscription", "Generates a random time value.");
+			return LOCTEXT("RandomNode_TimeRandomDiscription", "Generates a random time value.");
 		}
 
 		static bool HasRange()
@@ -301,11 +301,11 @@ namespace Metasound
 			auto InitNodeInfo = []() -> FNodeClassMetadata
 			{
 				FNodeClassMetadata Info;
-				Info.ClassName = { StandardNodes::Namespace, TRandomNodeSpecialization<ValueType>::GetClassName(), TEXT("") };
+				Info.ClassName = { StandardNodes::Namespace, TRandomNodeSpecialization<ValueType>::GetClassName(), "" };
 				Info.MajorVersion = 1;
-				Info.MinorVersion = 0;
+				Info.MinorVersion = 1;
 				Info.DisplayName = TRandomNodeSpecialization<ValueType>::GetDisplayName();
-				Info.Description = LOCTEXT("RandomOpDescription", "Generates a Random value.");
+				Info.Description = LOCTEXT("RandomNode_Description", "Generates a random value.");
 				Info.Author = PluginAuthor;
 				Info.PromptIfMissing = PluginNodeMissingPrompt;
 				Info.DefaultInterface = GetDefaultInterface();
