Add FTime trigger route nodes from templates

diff --git a/Engine/Plugins/Runtime/Metasound/Source/MetasoundStandardNodes/Private/MetasoundTriggerRouteNode.cpp b/Engine/Plugins/Runtime/Metasound/Source/MetasoundStandardNodes/Private/MetasoundTriggerRouteNode.cpp
--- a/Engine/Plugins/Runtime/Metasound/Source/MetasoundStandardNodes/Private/MetasoundTriggerRouteNode.cpp
+++ b/Engine/Plugins/Runtime/Metasound/Source/MetasoundStandardNodes/Private/MetasoundTriggerRouteNode.cpp
@@ -265,6 +265,14 @@ namespace Metasound
 	REGISTER_TRIGGER_ROUTE_NODE(float, 7)
 	REGISTER_TRIGGER_ROUTE_NODE(float, 8)
 
+	REGISTER_TRIGGER_ROUTE_NODE(FTime, 2)
+	REGISTER_TRIGGER_ROUTE_NODE(FTime, 3)
+	REGISTER_TRIGGER_ROUTE_NODE(FTime, 4)
+	REGISTER_TRIGGER_ROUTE_NODE(FTime, 5)
+	REGISTER_TRIGGER_ROUTE_NODE(FTime, 6)
+	REGISTER_TRIGGER_ROUTE_NODE(FTime, 7)
+	REGISTER_TRIGGER_ROUTE_NODE(FTime, 8)
+
 	REGISTER_TRIGGER_ROUTE_NODE(bool, 2)
 	REGISTER_TRIGGER_ROUTE_NODE(bool, 3)
 	REGISTER_TRIGGER_ROUTE_NODE(bool, 4)
