Mark all open nodes referenced by an asset as needing to be refreshed when renaming an asset
#rnx

diff --git a/Engine/Plugins/Runtime/Metasound/Source/MetasoundEditor/Private/MetasoundEditorGraphBuilder.cpp b/Engine/Plugins/Runtime/Metasound/Source/MetasoundEditor/Private/MetasoundEditorGraphBuilder.cpp
--- a/Engine/Plugins/Runtime/Metasound/Source/MetasoundEditor/Private/MetasoundEditorGraphBuilder.cpp
+++ b/Engine/Plugins/Runtime/Metasound/Source/MetasoundEditor/Private/MetasoundEditorGraphBuilder.cpp
@@ -1337,6 +1337,56 @@ namespace Metasound
 			MetaSoundAsset->UnregisterGraphWithFrontend();
 		}
 
+		void FGraphBuilder::MarkEditorNodesReferencingAssetForRefresh(UObject& InMetaSound)
+		{
+			using namespace Frontend;
+
+			FMetasoundAssetBase* MetaSoundAsset = IMetasoundUObjectRegistry::Get().GetObjectAsAssetBase(&InMetaSound);
+			if (!MetaSoundAsset)
+			{
+				return;
+			}
+
+			if (!GEditor)
+			{
+				return;
+			}
+
+			bool bGraphUpdated = false;
+
+			FMetasoundFrontendClassMetadata AssetClassMetadata = MetaSoundAsset->GetRootGraphHandle()->GetGraphMetadata();
+			AssetClassMetadata.SetType(EMetasoundFrontendClassType::External);
+			const FNodeRegistryKey AssetClassKey = NodeRegistryKey::CreateKey(AssetClassMetadata);
+
+			TArray<UObject*> EditedAssets = GEditor->GetEditorSubsystem<UAssetEditorSubsystem>()->GetAllEditedAssets();
+			for (UObject* EditedAsset : EditedAssets)
+			{
+				if (FMetasoundAssetBase* EditedMetaSoundAsset = IMetasoundUObjectRegistry::Get().GetObjectAsAssetBase(EditedAsset))
+				{
+					const UEdGraph& Graph = EditedMetaSoundAsset->GetGraphChecked();
+					TArray<UMetasoundEditorGraphExternalNode*> ExternalNodes;
+					Graph.GetNodesOfClass<UMetasoundEditorGraphExternalNode>(ExternalNodes);
+					for (UMetasoundEditorGraphExternalNode* Node : ExternalNodes)
+					{
+						FConstNodeHandle NodeHandle = Node->GetConstNodeHandle();
+						const FMetasoundFrontendClassMetadata& ClassMetadata = NodeHandle->GetClassMetadata();
+						const FNodeRegistryKey RegistryKey = NodeRegistryKey::CreateKey(ClassMetadata);
+
+						if (AssetClassKey == RegistryKey)
+						{
+							bGraphUpdated = true;
+							Node->bRefreshNode = true;
+						}
+					}
+				}
+			}
+
+			if (bGraphUpdated)
+			{
+				MetaSoundAsset->SetSynchronizationRequired();
+			}
+		}
+
 		bool FGraphBuilder::IsMatchingInputHandleAndPin(const Frontend::FConstInputHandle& InInputHandle, const UEdGraphPin& InEditorPin)
 		{
 			if (InEditorPin.Direction != EGPD_Input)
diff --git a/Engine/Plugins/Runtime/Metasound/Source/MetasoundEditor/Private/MetasoundEditorGraphBuilder.h b/Engine/Plugins/Runtime/Metasound/Source/MetasoundEditor/Private/MetasoundEditorGraphBuilder.h
--- a/Engine/Plugins/Runtime/Metasound/Source/MetasoundEditor/Private/MetasoundEditorGraphBuilder.h
+++ b/Engine/Plugins/Runtime/Metasound/Source/MetasoundEditor/Private/MetasoundEditorGraphBuilder.h
@@ -66,8 +66,8 @@ namespace Metasound
 			static UMetasoundEditorGraphNode* AddNode(UObject& InMetaSound, Frontend::FNodeHandle InNodeHandle, FVector2D InLocation, bool bInSelectNewNode = true);
 
 			// Convenience functions for retrieving the editor for the given MetaSound/EdGraph
-			static TSharedPtr<Metasound::Editor::FEditor> GetEditorForMetasound(const UObject& MetaSound);
-			static TSharedPtr<Metasound::Editor::FEditor> GetEditorForGraph(const UEdGraph& EdGraph);
+			static TSharedPtr<FEditor> GetEditorForMetasound(const UObject& InMetaSound);
+			static TSharedPtr<FEditor> GetEditorForGraph(const UEdGraph& InEdGraph);
 
 			// Wraps RegisterGraphWithFrontend logic in Frontend with any additional logic required to refresh editor & respective editor object state.
 			static void RegisterGraphWithFrontend(UObject& InMetaSound);
@@ -75,6 +75,8 @@ namespace Metasound
 			// Wraps RegisterGraphWithFrontend logic in Frontend with any additional logic required to refresh editor & respective editor object state.
 			static void UnregisterGraphWithFrontend(UObject& InMetaSound);
 
+			static void MarkEditorNodesReferencingAssetForRefresh(UObject& InMetaSound);
+
 			// Returns a display name for a node. If the node has an empty or whitespace
 			// only DisplayName, then the NodeName is used. 
 			static FText GetDisplayName(const Frontend::INodeController& InFrontendNode, bool bInIncludeNamespace);
diff --git a/Engine/Plugins/Runtime/Metasound/Source/MetasoundEditor/Private/MetasoundEditorModule.cpp b/Engine/Plugins/Runtime/Metasound/Source/MetasoundEditor/Private/MetasoundEditorModule.cpp
--- a/Engine/Plugins/Runtime/Metasound/Source/MetasoundEditor/Private/MetasoundEditorModule.cpp
+++ b/Engine/Plugins/Runtime/Metasound/Source/MetasoundEditor/Private/MetasoundEditorModule.cpp
@@ -290,10 +290,14 @@ namespace Metasound
 					UMetaSoundAssetSubsystem* AssetSubsystem = GEngine->GetEngineSubsystem<UMetaSoundAssetSubsystem>();
 					check(AssetSubsystem);
 
-					AssetSubsystem->RenameAsset(InAssetData, false /* bReregisterWithFrontend */);
+					// Use the FGraphBuilder Register call instead of registering via the
+					// MetaSoundAssetSubsystem so as to properly refresh respective open editors.
+					constexpr bool bReregisterWithFrontend = false;
+					AssetSubsystem->RenameAsset(InAssetData, bReregisterWithFrontend);
 					if (UObject* AssetObject = InAssetData.GetAsset())
 					{
 						FGraphBuilder::RegisterGraphWithFrontend(*AssetObject);
+						FGraphBuilder::MarkEditorNodesReferencingAssetForRefresh(*AssetObject);
 					}
 				}
 			}
