Add Upwards compression to dynamics processor & uses of it. This mainly just changes the gain clamping range to be > 0, and otherwise operates the same as the normal downwards compressor.

diff --git a/Engine/Plugins/Runtime/Metasound/Source/MetasoundStandardNodes/Private/MetasoundCompressorNode.cpp b/Engine/Plugins/Runtime/Metasound/Source/MetasoundStandardNodes/Private/MetasoundCompressorNode.cpp
--- a/Engine/Plugins/Runtime/Metasound/Source/MetasoundStandardNodes/Private/MetasoundCompressorNode.cpp
+++ b/Engine/Plugins/Runtime/Metasound/Source/MetasoundStandardNodes/Private/MetasoundCompressorNode.cpp
@@ -26,10 +26,12 @@ namespace Metasound
 		METASOUND_PARAM(InputThreshold, "Threshold dB", "Amplitude threshold (dB) above which gain will be reduced.");
 		METASOUND_PARAM(InputAttackTime, "Attack Time", "How long it takes for audio above the threshold to reach its compressed volume level.");
 		METASOUND_PARAM(InputReleaseTime, "Release Time", "How long it takes for audio below the threshold to return to its original volume level.");
+		METASOUND_PARAM(InputLookaheadTime, "Lookahead Time", "How long to delay the compressed signal behind the analyzed input signal.");
 		METASOUND_PARAM(InputKnee, "Knee", "How hard or soft the gain reduction blends from no gain reduction to gain reduction. 0 dB = no blending.");
 		METASOUND_PARAM(InputSidechain, "Sidechain", "(Optional) External audio signal to control the compressor with. If empty, uses the input audio signal.");
 		METASOUND_PARAM(InputEnvelopeMode, "Envelope Mode", "The envelope-following method the compressor will use for gain detection.");
 		METASOUND_PARAM(InputIsAnalog, "Analog Mode", "Enable Analog Mode for the compressor's envelope follower.");
+		METASOUND_PARAM(InputIsUpwards, "Upwards Mode", "Enable to switch from a standard downwards compresser to an upwards compressor.");
 		METASOUND_PARAM(InputWetDryMix, "Wet/Dry", "Ratio between the processed/wet signal and the unprocessed/dry signal. 0 is full dry, 1 is full wet, and 0.5 is 50/50.");
 
 		METASOUND_PARAM(OutputAudio, "Audio", "The output audio signal.");
@@ -47,29 +49,33 @@ namespace Metasound
 			const FFloatReadRef& InThresholdDb,
 			const FTimeReadRef& InAttackTime,
 			const FTimeReadRef& InReleaseTime,
+			const FTimeReadRef& InLookaheadTime,
 			const FFloatReadRef& InKnee,
 			const bool& bInUseSidechain,
 			const FAudioBufferReadRef& InSidechain,
 			const FEnvelopePeakModeReadRef& InEnvelopeMode,
 			const FBoolReadRef& bInIsAnalog,
+			const FBoolReadRef& bInIsUpwards,
 			const FFloatReadRef& InWetDryMix)
 			: AudioInput(InAudio)
 			, RatioInput(InRatio)
 			, ThresholdDbInput(InThresholdDb)
 			, AttackTimeInput(InAttackTime)
 			, ReleaseTimeInput(InReleaseTime)
+			, LookaheadTimeInput(InLookaheadTime)
 			, KneeInput(InKnee)
 			, SidechainInput(InSidechain)
 			, EnvelopeModeInput(InEnvelopeMode)
 			, bIsAnalogInput(bInIsAnalog)
+			, bIsUpwardsInput(bInIsUpwards)
 			, WetDryMixInput(InWetDryMix)
 			, AudioOutput(FAudioBufferWriteRef::CreateNew(InSettings))
 			, EnvelopeOutput(FAudioBufferWriteRef::CreateNew(InSettings))
 			, Compressor()
 			, bUseSidechain(bInUseSidechain)
-			, PrevThresholdDb(*InThresholdDb)
 			, PrevAttackTime(FMath::Max(FTime::ToMilliseconds(*InAttackTime), 0.0))
 			, PrevReleaseTime(FMath::Max(FTime::ToMilliseconds(*InReleaseTime), 0.0))
+			, PrevLookaheadTime(FMath::Max(FTime::ToMilliseconds(*InLookaheadTime), 0.0))
 		{
 			Compressor.Init(InSettings.GetSampleRate(), 1);
 			Compressor.SetKeyNumChannels(1);
@@ -77,8 +83,18 @@ namespace Metasound
 			Compressor.SetThreshold(*ThresholdDbInput);
 			Compressor.SetAttackTime(PrevAttackTime);
 			Compressor.SetReleaseTime(PrevReleaseTime);
+			Compressor.SetLookaheadMsec(PrevLookaheadTime);
 			Compressor.SetKneeBandwidth(*KneeInput);
 
+			if (*bIsUpwardsInput)
+			{
+				Compressor.SetProcessingMode(Audio::EDynamicsProcessingMode::UpwardsCompressor);
+			}
+			else
+			{
+				Compressor.SetProcessingMode(Audio::EDynamicsProcessingMode::Compressor);
+			}
+
 			switch (*EnvelopeModeInput)
 			{
 			default:
@@ -134,10 +150,12 @@ namespace Metasound
 					TInputDataVertexModel<float>(METASOUND_GET_PARAM_NAME_AND_TT(InputThreshold), -6.0f),
 					TInputDataVertexModel<FTime>(METASOUND_GET_PARAM_NAME_AND_TT(InputAttackTime), 0.01f),
 					TInputDataVertexModel<FTime>(METASOUND_GET_PARAM_NAME_AND_TT(InputReleaseTime), 0.1f),
+					TInputDataVertexModel<FTime>(METASOUND_GET_PARAM_NAME_AND_TT(InputLookaheadTime), 0.01f),
 					TInputDataVertexModel<float>(METASOUND_GET_PARAM_NAME_AND_TT(InputKnee), 10.0f),
 					TInputDataVertexModel<FAudioBuffer>(METASOUND_GET_PARAM_NAME_AND_TT(InputSidechain)),
 					TInputDataVertexModel<FEnumEnvelopePeakMode>(METASOUND_GET_PARAM_NAME_AND_TT(InputEnvelopeMode)),
 					TInputDataVertexModel<bool>(METASOUND_GET_PARAM_NAME_AND_TT(InputIsAnalog), true),
+					TInputDataVertexModel<bool>(METASOUND_GET_PARAM_NAME_AND_TT(InputIsUpwards), false),
 					TInputDataVertexModel<float>(METASOUND_GET_PARAM_NAME_AND_TT(InputWetDryMix), 1.0f)
 				),
 				FOutputVertexInterface(
@@ -160,9 +178,11 @@ namespace Metasound
 			InputDataReferences.AddDataReadReference(METASOUND_GET_PARAM_NAME(InputThreshold), ThresholdDbInput);
 			InputDataReferences.AddDataReadReference(METASOUND_GET_PARAM_NAME(InputAttackTime), AttackTimeInput);
 			InputDataReferences.AddDataReadReference(METASOUND_GET_PARAM_NAME(InputReleaseTime), ReleaseTimeInput);
+			InputDataReferences.AddDataReadReference(METASOUND_GET_PARAM_NAME(InputLookaheadTime), LookaheadTimeInput);
 			InputDataReferences.AddDataReadReference(METASOUND_GET_PARAM_NAME(InputKnee), KneeInput);
 			InputDataReferences.AddDataReadReference(METASOUND_GET_PARAM_NAME(InputSidechain), SidechainInput);
 			InputDataReferences.AddDataReadReference(METASOUND_GET_PARAM_NAME(InputEnvelopeMode), EnvelopeModeInput);
+			InputDataReferences.AddDataReadReference(METASOUND_GET_PARAM_NAME(InputIsUpwards), bIsUpwardsInput);
 			InputDataReferences.AddDataReadReference(METASOUND_GET_PARAM_NAME(InputIsAnalog), bIsAnalogInput);
 			InputDataReferences.AddDataReadReference(METASOUND_GET_PARAM_NAME(InputWetDryMix), WetDryMixInput);
 
@@ -193,15 +213,17 @@ namespace Metasound
 			FFloatReadRef ThresholdDbIn = Inputs.GetDataReadReferenceOrConstructWithVertexDefault<float>(InputInterface, METASOUND_GET_PARAM_NAME(InputThreshold), InParams.OperatorSettings);
 			FTimeReadRef AttackTimeIn = Inputs.GetDataReadReferenceOrConstructWithVertexDefault<FTime>(InputInterface, METASOUND_GET_PARAM_NAME(InputAttackTime), InParams.OperatorSettings);
 			FTimeReadRef ReleaseTimeIn = Inputs.GetDataReadReferenceOrConstructWithVertexDefault<FTime>(InputInterface, METASOUND_GET_PARAM_NAME(InputReleaseTime), InParams.OperatorSettings);
+			FTimeReadRef LookaheadTimeIn = Inputs.GetDataReadReferenceOrConstructWithVertexDefault<FTime>(InputInterface, METASOUND_GET_PARAM_NAME(InputLookaheadTime), InParams.OperatorSettings);
 			FFloatReadRef KneeIn = Inputs.GetDataReadReferenceOrConstructWithVertexDefault<float>(InputInterface, METASOUND_GET_PARAM_NAME(InputKnee), InParams.OperatorSettings);
 			FAudioBufferReadRef SidechainIn = Inputs.GetDataReadReferenceOrConstruct<FAudioBuffer>(METASOUND_GET_PARAM_NAME(InputSidechain), InParams.OperatorSettings);
 			FEnvelopePeakModeReadRef EnvelopeModeIn = Inputs.GetDataReadReferenceOrConstruct<FEnumEnvelopePeakMode>(METASOUND_GET_PARAM_NAME(InputEnvelopeMode));
 			FBoolReadRef bIsAnalogIn = Inputs.GetDataReadReferenceOrConstructWithVertexDefault<bool>(InputInterface, METASOUND_GET_PARAM_NAME(InputIsAnalog), InParams.OperatorSettings);
+			FBoolReadRef bIsUpwardsIn = Inputs.GetDataReadReferenceOrConstructWithVertexDefault<bool>(InputInterface, METASOUND_GET_PARAM_NAME(InputIsUpwards), InParams.OperatorSettings);
 			FFloatReadRef WetDryMixIn = Inputs.GetDataReadReferenceOrConstructWithVertexDefault<float>(InputInterface, METASOUND_GET_PARAM_NAME(InputWetDryMix), InParams.OperatorSettings);
 
 			bool bIsSidechainConnected = Inputs.ContainsDataReadReference<FAudioBuffer>(METASOUND_GET_PARAM_NAME(InputSidechain));
 
-			return MakeUnique<FCompressorOperator>(InParams.OperatorSettings, AudioIn, RatioIn, ThresholdDbIn, AttackTimeIn, ReleaseTimeIn, KneeIn, bIsSidechainConnected, SidechainIn, EnvelopeModeIn, bIsAnalogIn, WetDryMixIn);
+			return MakeUnique<FCompressorOperator>(InParams.OperatorSettings, AudioIn, RatioIn, ThresholdDbIn, AttackTimeIn, ReleaseTimeIn, LookaheadTimeIn, KneeIn, bIsSidechainConnected, SidechainIn, EnvelopeModeIn, bIsAnalogIn, bIsUpwardsIn, WetDryMixIn);
 		}
 
 		void Execute()
@@ -210,26 +232,42 @@ namespace Metasound
 			
 			// For a compressor, ratio values should be 1 or greater
 			Compressor.SetRatio(FMath::Max(*RatioInput, 1.0f));
-
-			if (FMath::IsNearlyEqual(*ThresholdDbInput, PrevThresholdDb) == false)
-			{
-				Compressor.SetThreshold(*ThresholdDbInput);
-			}
+			Compressor.SetThreshold(*ThresholdDbInput);
 
 			// Attack time cannot be negative
-			double CurrAttack = FMath::Max(FTime::ToMilliseconds(*AttackTimeInput), 0.0f);
+			float CurrAttack = FMath::Max(FTime::ToMilliseconds(*AttackTimeInput), 0.0f);
 			if (FMath::IsNearlyEqual(CurrAttack, PrevAttackTime) == false)
 			{
 				Compressor.SetAttackTime(CurrAttack);
+				PrevAttackTime = CurrAttack;
 			}
+
 			// Release time cannot be negative
-			double CurrRelease = FMath::Max(FTime::ToMilliseconds(*ReleaseTimeInput), 0.0f);
+			float CurrRelease = FMath::Max(FTime::ToMilliseconds(*ReleaseTimeInput), 0.0f);
 			if (FMath::IsNearlyEqual(CurrRelease, PrevReleaseTime) == false)
 			{
 				Compressor.SetReleaseTime(CurrRelease);
+				PrevReleaseTime = CurrRelease;
+			}
+
+			// Lookahead time cannot be negative
+			float CurrLookahead = FMath::Max(FTime::ToMilliseconds(*LookaheadTimeInput), 0.0f);
+			if (FMath::IsNearlyEqual(CurrLookahead, PrevLookaheadTime) == false)
+			{
+				Compressor.SetLookaheadMsec(CurrLookahead);
+				PrevLookaheadTime = CurrLookahead;
 			}
 
 			Compressor.SetKneeBandwidth(*KneeInput);
+
+			if (*bIsUpwardsInput)
+			{
+				Compressor.SetProcessingMode(Audio::EDynamicsProcessingMode::UpwardsCompressor);
+			}
+			else
+			{
+				Compressor.SetProcessingMode(Audio::EDynamicsProcessingMode::Compressor);
+			}
 				
 			switch (*EnvelopeModeInput)
 			{
@@ -272,10 +310,12 @@ namespace Metasound
 		FFloatReadRef ThresholdDbInput;
 		FTimeReadRef AttackTimeInput;
 		FTimeReadRef ReleaseTimeInput;
+		FTimeReadRef LookaheadTimeInput;
 		FFloatReadRef KneeInput;
 		FAudioBufferReadRef SidechainInput;
 		FEnvelopePeakModeReadRef EnvelopeModeInput;
 		FBoolReadRef bIsAnalogInput;
+		FBoolReadRef bIsUpwardsInput;
 		FFloatReadRef WetDryMixInput;
 
 		FAudioBufferWriteRef AudioOutput;
@@ -289,9 +329,9 @@ namespace Metasound
 		bool bUseSidechain;
 
 		// Cached variables
-		float PrevThresholdDb;
-		double PrevAttackTime;
-		double PrevReleaseTime;
+		float PrevAttackTime;
+		float PrevReleaseTime;
+		float PrevLookaheadTime;
 	};
 
 	// Node Class
diff --git a/Engine/Plugins/Runtime/Synthesis/Source/Synthesis/Classes/SourceEffects/SourceEffectDynamicsProcessor.h b/Engine/Plugins/Runtime/Synthesis/Source/Synthesis/Classes/SourceEffects/SourceEffectDynamicsProcessor.h
--- a/Engine/Plugins/Runtime/Synthesis/Source/Synthesis/Classes/SourceEffects/SourceEffectDynamicsProcessor.h
+++ b/Engine/Plugins/Runtime/Synthesis/Source/Synthesis/Classes/SourceEffects/SourceEffectDynamicsProcessor.h
@@ -13,6 +13,7 @@ enum class ESourceEffectDynamicsProcessorType : uint8
 	Limiter,
 	Expander,
 	Gate,
+	UpwardsCompressor,
 	Count UMETA(Hidden)
 };
 
diff --git a/Engine/Plugins/Runtime/Synthesis/Source/Synthesis/Private/SourceEffectDynamicsProcessor.cpp b/Engine/Plugins/Runtime/Synthesis/Source/Synthesis/Private/SourceEffectDynamicsProcessor.cpp
--- a/Engine/Plugins/Runtime/Synthesis/Source/Synthesis/Private/SourceEffectDynamicsProcessor.cpp
+++ b/Engine/Plugins/Runtime/Synthesis/Source/Synthesis/Private/SourceEffectDynamicsProcessor.cpp
@@ -35,6 +35,10 @@ void FSourceEffectDynamicsProcessor::OnPresetChanged()
 		case ESourceEffectDynamicsProcessorType::Gate:
 			DynamicsProcessor.SetProcessingMode(Audio::EDynamicsProcessingMode::Gate);
 			break;
+
+		case ESourceEffectDynamicsProcessorType::UpwardsCompressor:
+			DynamicsProcessor.SetProcessingMode(Audio::EDynamicsProcessingMode::UpwardsCompressor);
+			break;
 	}
 
 	switch (Settings.PeakMode)
diff --git a/Engine/Source/Runtime/AudioMixer/Classes/SubmixEffects/AudioMixerSubmixEffectDynamicsProcessor.h b/Engine/Source/Runtime/AudioMixer/Classes/SubmixEffects/AudioMixerSubmixEffectDynamicsProcessor.h
--- a/Engine/Source/Runtime/AudioMixer/Classes/SubmixEffects/AudioMixerSubmixEffectDynamicsProcessor.h
+++ b/Engine/Source/Runtime/AudioMixer/Classes/SubmixEffects/AudioMixerSubmixEffectDynamicsProcessor.h
@@ -29,6 +29,7 @@ enum class ESubmixEffectDynamicsProcessorType : uint8
 	Limiter,
 	Expander,
 	Gate,
+	UpwardsCompressor,
 	Count UMETA(Hidden)
 };
 
@@ -203,7 +204,7 @@ struct AUDIOMIXER_API FSubmixEffectDynamicsProcessorSettings
 
 	// The dynamics processor ratio used for compression/expansion
 	UPROPERTY(EditAnywhere, BlueprintReadWrite, Category = Dynamics, meta = (
-		EditCondition = "!bBypass && DynamicsProcessorType == ESubmixEffectDynamicsProcessorType::Compressor || DynamicsProcessorType == ESubmixEffectDynamicsProcessorType::Expander",
+		EditCondition = "!bBypass && DynamicsProcessorType == ESubmixEffectDynamicsProcessorType::Compressor || DynamicsProcessorType == ESubmixEffectDynamicsProcessorType::Expander ||  DynamicsProcessorType == ESubmixEffectDynamicsProcessorType::UpwardsCompressor",
 		ClampMin = "1.0", ClampMax = "20.0", UIMin = "1.0", UIMax = "20.0"))
 	float Ratio = 1.5f;
 
diff --git a/Engine/Source/Runtime/AudioMixer/Private/Effects/AudioMixerSubmixEffectDynamicsProcessor.cpp b/Engine/Source/Runtime/AudioMixer/Private/Effects/AudioMixerSubmixEffectDynamicsProcessor.cpp
--- a/Engine/Source/Runtime/AudioMixer/Private/Effects/AudioMixerSubmixEffectDynamicsProcessor.cpp
+++ b/Engine/Source/Runtime/AudioMixer/Private/Effects/AudioMixerSubmixEffectDynamicsProcessor.cpp
@@ -112,6 +112,10 @@ void FSubmixEffectDynamicsProcessor::OnPresetChanged()
 	case ESubmixEffectDynamicsProcessorType::Gate:
 		DynamicsProcessor.SetProcessingMode(Audio::EDynamicsProcessingMode::Gate);
 		break;
+
+	case ESubmixEffectDynamicsProcessorType::UpwardsCompressor:
+		DynamicsProcessor.SetProcessingMode(Audio::EDynamicsProcessingMode::UpwardsCompressor);
+		break;
 	}
 
 	switch (Settings.PeakMode)
diff --git a/Engine/Source/Runtime/SignalProcessing/Private/DynamicsProcessor.cpp b/Engine/Source/Runtime/SignalProcessing/Private/DynamicsProcessor.cpp
--- a/Engine/Source/Runtime/SignalProcessing/Private/DynamicsProcessor.cpp
+++ b/Engine/Source/Runtime/SignalProcessing/Private/DynamicsProcessor.cpp
@@ -439,6 +439,20 @@ namespace Audio
 		return true;
 	}
 
+	bool FDynamicsProcessor::IsInProcessingThreshold(const float InEnvFollowerDb) const
+	{
+		if (ProcessingMode == EDynamicsProcessingMode::UpwardsCompressor)
+		{
+			return HalfKneeBandwidthDb >= 0.0f
+				&& InEnvFollowerDb < (ThresholdDb - HalfKneeBandwidthDb)
+				&& InEnvFollowerDb > (ThresholdDb + HalfKneeBandwidthDb);
+		}
+
+		return HalfKneeBandwidthDb >= 0.0f
+			&& InEnvFollowerDb > (ThresholdDb - HalfKneeBandwidthDb)
+			&& InEnvFollowerDb < (ThresholdDb + HalfKneeBandwidthDb);
+	}
+
 	float FDynamicsProcessor::ComputeGain(const float InEnvFollowerDb)
 	{
 		float SlopeFactor = 0.0f;
@@ -450,6 +464,8 @@ namespace Audio
 
 			// Compressors smoothly reduce the gain as the gain gets louder
 			// CompressionRatio -> Inifinity is a limiter
+			// Upwards compression applies gain when below a threshold, but uses the same slope
+			case EDynamicsProcessingMode::UpwardsCompressor:
 			case EDynamicsProcessingMode::Compressor:
 				SlopeFactor = 1.0f - 1.0f / Ratio;
 				break;
@@ -472,7 +488,7 @@ namespace Audio
 		}
 
 		// If we are in the range of compression
-		if (HalfKneeBandwidthDb > 0.0f && InEnvFollowerDb > (ThresholdDb - HalfKneeBandwidthDb) && InEnvFollowerDb < ThresholdDb + HalfKneeBandwidthDb)
+		if (IsInProcessingThreshold(InEnvFollowerDb))
 		{
 			// Setup the knee for interpolation. Don't allow the top knee point to exceed 0.0
 			KneePoints[0].X = ThresholdDb - HalfKneeBandwidthDb;
@@ -486,7 +502,17 @@ namespace Audio
 		}
 
 		float OutputGainDb = SlopeFactor * (ThresholdDb - InEnvFollowerDb);
-		OutputGainDb = FMath::Min(0.0f, OutputGainDb);
+
+		if (ProcessingMode == EDynamicsProcessingMode::UpwardsCompressor)
+		{
+			// if left unchecked Upwards compression will try to apply infinite gain
+			OutputGainDb = FMath::Clamp(OutputGainDb, 0.f, UpwardsCompressionMaxGain);
+		}
+		else
+		{
+			OutputGainDb = FMath::Min(0.f, OutputGainDb);
+		}
+
 		return ConvertToLinear(OutputGainDb);
 	}
 }
diff --git a/Engine/Source/Runtime/SignalProcessing/Public/DSP/DynamicsProcessor.h b/Engine/Source/Runtime/SignalProcessing/Public/DSP/DynamicsProcessor.h
--- a/Engine/Source/Runtime/SignalProcessing/Public/DSP/DynamicsProcessor.h
+++ b/Engine/Source/Runtime/SignalProcessing/Public/DSP/DynamicsProcessor.h
@@ -18,6 +18,7 @@ namespace Audio
 			Limiter,
 			Expander,
 			Gate,
+			UpwardsCompressor,
 			Count
 		};
 	}
@@ -80,6 +81,8 @@ namespace Audio
 		// (Returns false in audition mode and writes straight to output).
 		bool ProcessKeyFrame(const float* InKeyFrame, float* OutFrame, bool bKeyIsInput);
 
+		bool IsInProcessingThreshold(const float InEnvFollowerDb) const;
+
 		// (Optional) Low-pass filter for input signal
 		FBiquadFilter InputLowshelfFilter;
 
@@ -150,5 +153,7 @@ namespace Audio
 
 		// Whether or not key low-pass filter is enabled
 		bool bKeyLowshelfEnabled;
+
+		static constexpr float UpwardsCompressionMaxGain = 36.0f;
 	};
 }
