AutoUpdate Clean-up
- Clarifies single point-of-execution
- Removes dead dirtying branch
- Add owning asset's path to transform for better logging when updating node/reference is missing
#rnx

diff --git a/Engine/Plugins/Runtime/Metasound/Source/MetasoundFrontend/Private/MetasoundAssetBase.cpp b/Engine/Plugins/Runtime/Metasound/Source/MetasoundFrontend/Private/MetasoundAssetBase.cpp
--- a/Engine/Plugins/Runtime/Metasound/Source/MetasoundFrontend/Private/MetasoundAssetBase.cpp
+++ b/Engine/Plugins/Runtime/Metasound/Source/MetasoundFrontend/Private/MetasoundAssetBase.cpp
@@ -111,13 +111,15 @@ void FMetasoundAssetBase::RegisterGraphWithFrontend(Metasound::Frontend::FMetaSo
 	// Auto update must be done after all referenced asset classes are registered
 	if (InRegistrationOptions.bAutoUpdate)
 	{
-		const bool bWasAutoUpdated = AutoUpdate();
-		if (bWasAutoUpdated)
-		{
+		FString OwningAssetName = GetOwningAssetName();
+		const bool bAutoUpdated = FAutoUpdateRootGraph(MoveTemp(OwningAssetName)).Transform(GetDocumentHandle());
 #if WITH_EDITORONLY_DATA
+		if (bAutoUpdated)
+		{
+ 
 			SetSynchronizationRequired();
-#endif // WITH_EDITORONLY_DATA
 		}
+#endif // WITH_EDITORONLY_DATA
 	}
 
 	// Registers node by copying document. Updates to document require re-registration.
@@ -329,20 +331,15 @@ void FMetasoundAssetBase::AddDefaultInterfaces()
 	FModifyRootGraphInterfaces({ }, InitInterfaces).Transform(DocumentHandle);
 }
 
-bool FMetasoundAssetBase::AutoUpdate(bool bInMarkDirty)
-{
-	using namespace Metasound::Frontend;
-	METASOUND_TRACE_CPUPROFILER_EVENT_SCOPE(MetaSoundAssetBase::AutoUpdate);
-
-	const bool bUpdated = FAutoUpdateRootGraph().Transform(GetDocumentHandle());
-	if (bUpdated && bInMarkDirty)
-	{
-		MarkMetasoundDocumentDirty();
-	}
-
-	return bUpdated;
-}
-
+ 
+ 
+ 
+ 
+ 
+ 
+ 
+ 
+ 
 bool FMetasoundAssetBase::VersionAsset()
 {
 	using namespace Metasound;
diff --git a/Engine/Plugins/Runtime/Metasound/Source/MetasoundFrontend/Private/MetasoundFrontendTransform.cpp b/Engine/Plugins/Runtime/Metasound/Source/MetasoundFrontend/Private/MetasoundFrontendTransform.cpp
--- a/Engine/Plugins/Runtime/Metasound/Source/MetasoundFrontend/Private/MetasoundFrontendTransform.cpp
+++ b/Engine/Plugins/Runtime/Metasound/Source/MetasoundFrontend/Private/MetasoundFrontendTransform.cpp
@@ -9,6 +9,7 @@
 #include "MetasoundFrontendRegistries.h"
 #include "MetasoundFrontendSearchEngine.h"
 #include "MetasoundLog.h"
+#include "MetasoundTrace.h"
 #include "Misc/App.h"
 
 namespace Metasound
@@ -449,6 +450,7 @@ namespace Metasound
 
 		bool FAutoUpdateRootGraph::Transform(FDocumentHandle InDocument) const
 		{
+			METASOUND_TRACE_CPUPROFILER_EVENT_SCOPE(FAutoUpdateRootGraph::Transform);
 			bool bDidEdit = false;
 
 			FMetasoundAssetBase* PresetReferencedMetaSoundAsset = nullptr;
@@ -468,7 +470,7 @@ namespace Metasound
 					PresetReferencedMetaSoundAsset = IMetaSoundAssetManager::GetChecked().TryLoadAssetFromKey(RegistryKey);
 					if (!PresetReferencedMetaSoundAsset)
 					{
-						UE_LOG(LogMetaSound, Error, TEXT("Failed to update Preset. Referenced asset class '%s' not found."), *ClassMetadata.GetDisplayName().ToString());
+						UE_LOG(LogMetaSound, Error, TEXT("Auto-Updating preset '%s' failed: Referenced class '%s' missing."), *AssetPath, *ClassMetadata.GetDisplayName().ToString());
 					}
 					return;
 				}
@@ -482,12 +484,12 @@ namespace Metasound
 				FMetasoundFrontendVersionNumber UpdateVersion = NodeHandle->FindHighestMinorVersionInRegistry();
 				if (UpdateVersion.IsValid() && UpdateVersion > ClassMetadata.GetVersion())
 				{
-					UE_LOG(LogMetaSound, Display, TEXT("Auto-Updating node class '%s': Newer version '%s' found."), *ClassMetadata.GetDisplayName().ToString(), *UpdateVersion.ToString());
+					UE_LOG(LogMetaSound, Display, TEXT("Auto-Updating '%s' node class '%s': Newer version '%s' found."), *AssetPath, *ClassMetadata.GetDisplayName().ToString(), *UpdateVersion.ToString());
 				}
 				else if (InterfaceUpdates.ContainsChanges())
 				{
 					UpdateVersion = ClassMetadata.GetVersion();
-					UE_LOG(LogMetaSound, Display, TEXT("Auto-Updating node with class '%s (%s)': Interface change detected."), *ClassMetadata.GetDisplayName().ToString(), *UpdateVersion.ToString());
+					UE_LOG(LogMetaSound, Display, TEXT("Auto-Updating '%s' node class '%s (%s)': Interface change detected."), *AssetPath, *ClassMetadata.GetDisplayName().ToString(), *UpdateVersion.ToString());
 				}
 
 				NodesToUpdate.Add(TPair<FNodeHandle, FMetasoundFrontendVersionNumber>(NodeHandle, UpdateVersion));
diff --git a/Engine/Plugins/Runtime/Metasound/Source/MetasoundFrontend/Public/MetasoundAssetBase.h b/Engine/Plugins/Runtime/Metasound/Source/MetasoundFrontend/Public/MetasoundAssetBase.h
--- a/Engine/Plugins/Runtime/Metasound/Source/MetasoundFrontend/Public/MetasoundAssetBase.h
+++ b/Engine/Plugins/Runtime/Metasound/Source/MetasoundFrontend/Public/MetasoundAssetBase.h
@@ -15,7 +15,6 @@
 #include "UObject/WeakObjectPtrTemplates.h"
 
 // Forward Declarations
-class FMetasoundAssetBase;
 class UEdGraph;
 
 namespace Metasound
@@ -123,7 +122,6 @@ public:
 	const FMetasoundFrontendDocument& GetDocumentChecked() const;
 
 	void AddDefaultInterfaces();
-	bool AutoUpdate(bool bInMarkDirty = false);
 	bool VersionAsset();
 
 #if WITH_EDITORONLY_DATA
diff --git a/Engine/Plugins/Runtime/Metasound/Source/MetasoundFrontend/Public/MetasoundFrontendTransform.h b/Engine/Plugins/Runtime/Metasound/Source/MetasoundFrontend/Public/MetasoundFrontendTransform.h
--- a/Engine/Plugins/Runtime/Metasound/Source/MetasoundFrontend/Public/MetasoundFrontendTransform.h
+++ b/Engine/Plugins/Runtime/Metasound/Source/MetasoundFrontend/Public/MetasoundFrontendTransform.h
@@ -120,12 +120,20 @@ namespace Metasound
 			FConstDocumentHandle ReferencedDocument;
 		};
 
+		/** Automatically updates all nodes and respective dependencies in graph where
+		  * newer versions exist in the loaded MetaSound Class Node Registry.
+		  */
 		class METASOUNDFRONTEND_API FAutoUpdateRootGraph : public IDocumentTransform
 		{
 		public:
-			FAutoUpdateRootGraph() = default;
+			FAutoUpdateRootGraph(FString&& InAssetPath)
+				: AssetPath(MoveTemp(InAssetPath))
+			{
+			}
 
 			bool Transform(FDocumentHandle InDocument) const override;
+
+			const FString AssetPath;
 		};
 
 		/** Synchronizes the document's root graph's display name with that of the asset. */
