More tooltip updates & removal of dead #ifdef 0 code
#rnx

diff --git a/Engine/Plugins/Runtime/Metasound/Source/MetasoundStandardNodes/Private/MetasoundMathNodes.cpp b/Engine/Plugins/Runtime/Metasound/Source/MetasoundStandardNodes/Private/MetasoundMathNodes.cpp
--- a/Engine/Plugins/Runtime/Metasound/Source/MetasoundStandardNodes/Private/MetasoundMathNodes.cpp
+++ b/Engine/Plugins/Runtime/Metasound/Source/MetasoundStandardNodes/Private/MetasoundMathNodes.cpp
@@ -20,7 +20,7 @@
 #include "MetasoundVertex.h"
 
 
-#define LOCTEXT_NAMESPACE "MetasoundMathOpNode"
+#define LOCTEXT_NAMESPACE "MetasoundStandardNodes"
 
 #define DEFINE_METASOUND_MATHOP(OpName, DataTypeName, DataClass, Description, Keywords) \
 	class F##OpName##DataTypeName##Node \
@@ -214,7 +214,7 @@ namespace Metasound
 					FNodeClassMetadata Info;
 					Info.ClassName = TDerivedClass::GetClassName();
 					Info.MajorVersion = 1;
-					Info.MinorVersion = 0;
+					Info.MinorVersion = 1;
 					Info.DisplayName = TDerivedClass::GetDisplayName();
 					Info.Description = TDerivedClass::GetDescription();
 					Info.Author = PluginAuthor;
@@ -309,11 +309,11 @@ namespace Metasound
 		{
 			static const FVertexInterface Interface(
 				FInputVertexInterface(
-					TInputDataVertexModel<TDataClass>(MathOpNames::PrimaryOperandName, LOCTEXT("MathOpAddendInitialTooltip", "Initial addend."), static_cast<TDataClass>(0)),
-					TInputDataVertexModel<TOperandDataClass>(MathOpNames::AdditionalOperandsName, LOCTEXT("MathOpAddendsTooltip", "Additional addend(s)."), static_cast<TOperandDataClass>(0))
+					TInputDataVertexModel<TDataClass>(MathOpNames::PrimaryOperandName, LOCTEXT("MathOpAddNode_InitialTooltip", "Initial addend."), static_cast<TDataClass>(0)),
+					TInputDataVertexModel<TOperandDataClass>(MathOpNames::AdditionalOperandsName, LOCTEXT("MathOpAddNode_AddendsTooltip", "Additional addend(s)."), static_cast<TOperandDataClass>(0))
 				),
 				FOutputVertexInterface(
-					TOutputDataVertexModel<TDataClass>(TEXT("Out"), LOCTEXT("MathOpOutTooltip", "Math operation result"))
+					TOutputDataVertexModel<TDataClass>(TEXT("Out"), LOCTEXT("MathOpAddNode_OutTooltip", "Math operation result"))
 				)
 			);
 
@@ -358,11 +358,11 @@ namespace Metasound
 		{
 			static const FVertexInterface Interface(
 				FInputVertexInterface(
-					TInputDataVertexModel<TDataClass>(MathOpNames::PrimaryOperandName, LOCTEXT("MathOpMinuendTooltip", "Minuend."), static_cast<TDataClass>(0)),
-					TInputDataVertexModel<TOperandDataClass>(MathOpNames::AdditionalOperandsName, LOCTEXT("MathOpSubtrahendsTooltip", "Subtrahend(s)."), static_cast<TOperandDataClass>(0))
+					TInputDataVertexModel<TDataClass>(MathOpNames::PrimaryOperandName, LOCTEXT("MathOpSubtractNode_MinuendTooltip", "Minuend."), static_cast<TDataClass>(0)),
+					TInputDataVertexModel<TOperandDataClass>(MathOpNames::AdditionalOperandsName, LOCTEXT("MathOpSubtractNode_SubtrahendsTooltip", "Subtrahend(s)."), static_cast<TOperandDataClass>(0))
 				),
 				FOutputVertexInterface(
-					TOutputDataVertexModel<TDataClass>(TEXT("Out"), LOCTEXT("MathOpOutTooltip", "Subtraction result"))
+					TOutputDataVertexModel<TDataClass>(TEXT("Out"), LOCTEXT("MathOpSubtractNode_OutTooltip", "Subtraction result"))
 				)
 			);
 
@@ -407,11 +407,11 @@ namespace Metasound
 		{
 			static const FVertexInterface Interface(
 				FInputVertexInterface(
-					TInputDataVertexModel<TDataClass>(MathOpNames::PrimaryOperandName, LOCTEXT("MathOpInitMultiplicandTooltip", "Initial multiplicand."), static_cast<TDataClass>(1)),
-					TInputDataVertexModel<TOperandDataClass>(MathOpNames::AdditionalOperandsName, LOCTEXT("MathOpMultiplicandsTooltip", "Additional multiplicand(s)."), static_cast<TOperandDataClass>(1))
+					TInputDataVertexModel<TDataClass>(MathOpNames::PrimaryOperandName, LOCTEXT("MathOpMultiplyNode_InitMultiplicandTooltip", "Initial multiplicand."), static_cast<TDataClass>(1)),
+					TInputDataVertexModel<TOperandDataClass>(MathOpNames::AdditionalOperandsName, LOCTEXT("MathOpMultiplyNode_MultiplicandsTooltip", "Additional multiplicand(s)."), static_cast<TOperandDataClass>(1))
 				),
 				FOutputVertexInterface(
-					TOutputDataVertexModel<TDataClass>(TEXT("Out"), LOCTEXT("MultiplicationResultTooltip", "Multiplication result"))
+					TOutputDataVertexModel<TDataClass>(TEXT("Out"), LOCTEXT("MathOpMultiplyNode_ResultTooltip", "Multiplication result"))
 				)
 			);
 
@@ -456,11 +456,11 @@ namespace Metasound
 		{
 			static const FVertexInterface Interface(
 				FInputVertexInterface(
-					TInputDataVertexModel<TDataClass>(MathOpNames::PrimaryOperandName, LOCTEXT("MathOpDividendTooltip", "Dividend."), static_cast<TDataClass>(1)),
-					TInputDataVertexModel<TOperandDataClass>(MathOpNames::AdditionalOperandsName, LOCTEXT("MathOpDivisorsTooltip", "Divisor(s)."), static_cast<TOperandDataClass>(1))
+					TInputDataVertexModel<TDataClass>(MathOpNames::PrimaryOperandName, LOCTEXT("MathOpDivideNode_DividendTooltip", "Dividend."), static_cast<TDataClass>(1)),
+					TInputDataVertexModel<TOperandDataClass>(MathOpNames::AdditionalOperandsName, LOCTEXT("MathOpDivideNode_DivisorsTooltip", "Divisor(s)."), static_cast<TOperandDataClass>(1))
 				),
 				FOutputVertexInterface(
-					TOutputDataVertexModel<TDataClass>(TEXT("Out"), LOCTEXT("MathOpOutTooltip", "Math operation result"))
+					TOutputDataVertexModel<TDataClass>(TEXT("Out"), LOCTEXT("MathOpDivideNode_OutTooltip", "Division result"))
 				)
 			);
 
@@ -512,11 +512,11 @@ namespace Metasound
 		{
 			static const FVertexInterface Interface(
 				FInputVertexInterface(
-					TInputDataVertexModel<TDataClass>(MathOpNames::PrimaryOperandName, LOCTEXT("MathOpModuloDividendTooltip", "Dividend."), static_cast<TDataClass>(1)),
-					TInputDataVertexModel<TOperandDataClass>(MathOpNames::AdditionalOperandsName, LOCTEXT("MathOpModuloDivisorsTooltip", "Divisor(s)."), static_cast<TOperandDataClass>(1))
+					TInputDataVertexModel<TDataClass>(MathOpNames::PrimaryOperandName, LOCTEXT("MathOpModuloNode_DividendTooltip", "Dividend."), static_cast<TDataClass>(1)),
+					TInputDataVertexModel<TOperandDataClass>(MathOpNames::AdditionalOperandsName, LOCTEXT("MathOpModuloNode_DivisorsTooltip", "Divisor(s)."), static_cast<TOperandDataClass>(1))
 				),
 				FOutputVertexInterface(
-					TOutputDataVertexModel<TDataClass>(TEXT("Out"), LOCTEXT("MathOpOutTooltip", "Resulting value"))
+					TOutputDataVertexModel<TDataClass>(TEXT("Out"), LOCTEXT("MathOpModuloNode_OutTooltip", "Modulo result"))
 				)
 			);
 
@@ -570,11 +570,11 @@ namespace Metasound
 		{
 			static const FVertexInterface DefaultInterface(
 				FInputVertexInterface(
-					TInputDataVertexModel<float>(MathOpNames::PrimaryOperandName, LOCTEXT("Base", "The base of the power")),
-					TInputDataVertexModel<float>(MathOpNames::AdditionalOperandsName, LOCTEXT("Exponent", "The exponent to take the base to the power of"))
+					TInputDataVertexModel<float>(MathOpNames::PrimaryOperandName, LOCTEXT("MathOpPowerNode_Base", "The base of the power")),
+					TInputDataVertexModel<float>(MathOpNames::AdditionalOperandsName, LOCTEXT("MathOpPowerNode_Exponent", "The exponent to take the base to the power of"))
 				),
 				FOutputVertexInterface(
-					TOutputDataVertexModel<float>(TEXT("Out"), LOCTEXT("Result", "Returns Base to the Exponent power"))
+					TOutputDataVertexModel<float>(TEXT("Out"), LOCTEXT("MathOpPowerNode_Result", "Returns Base to the Exponent power"))
 				)
 			);
 
@@ -622,11 +622,11 @@ namespace Metasound
 		{
 			static const FVertexInterface DefaultInterface(
 				FInputVertexInterface(
-					TInputDataVertexModel<float>(MathOpNames::PrimaryOperandName, LOCTEXT("Base", "The base of the logarithm")),
-					TInputDataVertexModel<float>(MathOpNames::AdditionalOperandsName, LOCTEXT("Value", "The value to find the logarithm of"))
+					TInputDataVertexModel<float>(MathOpNames::PrimaryOperandName, LOCTEXT("MathOpLogarithmNode_Base", "The base of the logarithm")),
+					TInputDataVertexModel<float>(MathOpNames::AdditionalOperandsName, LOCTEXT("MathOpLogarithmNode_Value", "The value to find the logarithm of"))
 				),
 				FOutputVertexInterface(
-					TOutputDataVertexModel<float>(TEXT("Out"), LOCTEXT("Result", "The logarithm of the inputted value"))
+					TOutputDataVertexModel<float>(TEXT("Out"), LOCTEXT("MathOpLogarithmNode_Result", "The logarithm of the inputted value"))
 				)
 			);
 
@@ -832,55 +832,6 @@ namespace Metasound
 		}
 	};
 
-#if 0
-	// int64 has to be explicitly implemented as they are PODs that are not implemented as a literal base type.
-	template <>
-	class TMathOpAdd<int64>
-	{
-	public:
-		static const FVertexInterface& GetVertexInterface()
-		{
-			static const FVertexInterface Interface(
-				FInputVertexInterface(
-					TInputDataVertexModel<int64>(MathOpNames::PrimaryOperandName, LOCTEXT("MathOpInt64InitialAddendMinuendTooltip", "Initial Addend."), 0),
-					TInputDataVertexModel<int64>(MathOpNames::AdditionalOperandsName, LOCTEXT("MathOpInt64AddendsTooltip", "Additional Addend(s)."), 0)
-				),
-				FOutputVertexInterface(
-					TOutputDataVertexModel<int64>(TEXT("Out"), LOCTEXT("MathOpOutTooltip", "Addition result"))
-				)
-			);
-
-			return Interface;
-		}
-
-		static int32 GetDefault(const FOperatorSettings& InSettings, const FLiteral& InVertexDefault)
-		{
-			return static_cast<int64>(InVertexDefault.Value.Get<int32>());
-		}
-
-		static int32 GetDefaultOp(const FOperatorSettings& InSettings, const FLiteral& InVertexDefault)
-		{
-			return static_cast<int64>(InVertexDefault.Value.Get<int32>());
-		}
-
-		static void Calculate(TMathOpAdd& InInstanceData, const FInt64ReadRef& InPrimaryOperand, const TArray<FInt64ReadRef>& InAdditionalOperands, FInt64WriteRef& OutResult)
-		{
-			*OutResult = *InPrimaryOperand;
-
-			if (!ensure(!InAdditionalOperands.IsEmpty()))
-			{
-				return;
-			}
-
-			for (int32 i = 0; i < InAdditionalOperands.Num(); ++i)
-			{
-				const int64 OperandValue = *InAdditionalOperands[i];
-				*OutResult += OperandValue;
-			}
-		}
-	};
-#endif
-
 	template <>
 	class TMathOpSubtract<FAudioBuffer>
 	{
@@ -964,11 +915,11 @@ namespace Metasound
 		{
 			static const FVertexInterface Interface(
 				FInputVertexInterface(
-					TInputDataVertexModel<FTime>(MathOpNames::PrimaryOperandName, LOCTEXT("MathOpTimeMinuendTooltip", "Time minuend."), 0.0f),
-					TInputDataVertexModel<FTime>(MathOpNames::AdditionalOperandsName, LOCTEXT("MathOpTimeSubtrahendsTooltip", "Time subtrahends."), 0.0f)
+					TInputDataVertexModel<FTime>(MathOpNames::PrimaryOperandName, LOCTEXT("MathOpTimeNode_MinuendTooltip", "Time minuend."), 0.0f),
+					TInputDataVertexModel<FTime>(MathOpNames::AdditionalOperandsName, LOCTEXT("MathOpTimeNode_SubtrahendsTooltip", "Time subtrahends."), 0.0f)
 				),
 				FOutputVertexInterface(
-					TOutputDataVertexModel<FTime>(TEXT("Out"), LOCTEXT("MathOpTimeSubtractOutTooltip", "Resulting time value"))
+					TOutputDataVertexModel<FTime>(TEXT("Out"), LOCTEXT("MathOpTimeNode_SubtractOutTooltip", "Resulting time value"))
 				)
 			);
 
@@ -1000,55 +951,6 @@ namespace Metasound
 		}
 	};
 
-#if 0
-	// int64 has to be explicitly implemented as they are PODs that are not implemented as a literal base type.
-	template <>
-	class TMathOpSubtract<int64>
-	{
-	public:
-		static const FVertexInterface& GetVertexInterface()
-		{
-			static const FVertexInterface Interface(
-				FInputVertexInterface(
-					TInputDataVertexModel<int64>(MathOpNames::PrimaryOperandName, LOCTEXT("MathOpMinuendTooltip", "Minuend."), 0),
-					TInputDataVertexModel<int64>(MathOpNames::AdditionalOperandsName, LOCTEXT("MathOpSubtrahendsTooltip", "Subtrahend(s)."), 0)
-				),
-				FOutputVertexInterface(
-					TOutputDataVertexModel<int64>(TEXT("Out"), LOCTEXT("MathOpOutTooltip", "Subtraction result"))
-				)
-			);
-
-			return Interface;
-		}
-
-		static int32 GetDefault(const FOperatorSettings& InSettings, const FLiteral& InVertexDefault)
-		{
-			return static_cast<int64>(InVertexDefault.Value.Get<int32>());
-		}
-
-		static int32 GetDefaultOp(const FOperatorSettings& InSettings, const FLiteral& InVertexDefault)
-		{
-			return static_cast<int64>(InVertexDefault.Value.Get<int32>());
-		}
-
-		static void Calculate(TMathOpSubtract& InInstanceData, const FInt64ReadRef& InPrimaryOperand, const TArray<FInt64ReadRef>& InAdditionalOperands, FInt64WriteRef& OutResult)
-		{
-			*OutResult = *InPrimaryOperand;
-
-			if (!ensure(!InAdditionalOperands.IsEmpty()))
-			{
-				return;
-			}
-
-			for (int32 i = 0; i < InAdditionalOperands.Num(); ++i)
-			{
-				const int64 OperandValue = *InAdditionalOperands[i];
-				*OutResult -= OperandValue;
-			}
-		}
-	};
-#endif
-
 	template <>
 	class TMathOpMultiply<FAudioBuffer>
 	{
@@ -1057,11 +959,11 @@ namespace Metasound
 		{
 			static const FVertexInterface Interface(
 				FInputVertexInterface(
-					TInputDataVertexModel<FAudioBuffer>(MathOpNames::PrimaryOperandName, LOCTEXT("MathOpAudioInitMultiplicandTooltip", "Initial audio to multiply.")),
-					TInputDataVertexModel<FAudioBuffer>(MathOpNames::AdditionalOperandsName, LOCTEXT("MathOpMultiplyAudioSubtrahendsTooltip", "Additional audio to multiply sample-by-sample."))
+					TInputDataVertexModel<FAudioBuffer>(MathOpNames::PrimaryOperandName, LOCTEXT("MathOpAudioMultiplyNode_InitMultiplicandTooltip", "Initial audio to multiply.")),
+					TInputDataVertexModel<FAudioBuffer>(MathOpNames::AdditionalOperandsName, LOCTEXT("MathOpAudioMultiplyNode_MultiplicandsTooltip", "Additional audio to multiply sample-by-sample."))
 				),
 				FOutputVertexInterface(
-					TOutputDataVertexModel<FAudioBuffer>(TEXT("Out"), LOCTEXT("MathOpMultiplyAudioOutTooltip", "Resulting buffer"))
+					TOutputDataVertexModel<FAudioBuffer>(TEXT("Out"), LOCTEXT("MathOpAudioMultiplyNode_OutTooltip", "Resulting buffer"))
 				)
 			);
 
@@ -1135,11 +1037,11 @@ namespace Metasound
 		{
 			static const FVertexInterface Interface(
 				FInputVertexInterface(
-					TInputDataVertexModel<FAudioBuffer>(MathOpNames::PrimaryOperandName, LOCTEXT("MathOpAudioMultiplyFloatTooltip", "Audio multiplicand.")),
-					TInputDataVertexModel<float>(MathOpNames::AdditionalOperandsName, LOCTEXT("MathOpAddendAdditionalTooltip", "Float multiplicand to apply sample-by-sample to audio. Interpolates over buffer size on value change."), 1.0f)
+					TInputDataVertexModel<FAudioBuffer>(MathOpNames::PrimaryOperandName, LOCTEXT("MathOpAudioFloatMultiplyNode_FloatTooltip", "Audio multiplicand.")),
+					TInputDataVertexModel<float>(MathOpNames::AdditionalOperandsName, LOCTEXT("MathOpAudioFloatMultiplyNode_MultiplicandTooltip", "Float multiplicand to apply sample-by-sample to audio. Interpolates over buffer size on value change."), 1.0f)
 				),
 				FOutputVertexInterface(
-					TOutputDataVertexModel<FAudioBuffer>(TEXT("Out"), LOCTEXT("MathOpAudioFloatMultiplyOutTooltip", "Resulting buffer"))
+					TOutputDataVertexModel<FAudioBuffer>(TEXT("Out"), LOCTEXT("MathOpAudioFloatMultiplyNode_OutTooltip", "Resulting buffer"))
 				)
 			);
 
@@ -1199,11 +1101,11 @@ namespace Metasound
 		{
 			static const FVertexInterface Interface(
 				FInputVertexInterface(
-					TInputDataVertexModel<FTime>(MathOpNames::PrimaryOperandName, LOCTEXT("MathOpTimeMultiplyFloatTooltip", "Time multiplicand."), 1.0f),
-					TInputDataVertexModel<float>(MathOpNames::AdditionalOperandsName, LOCTEXT("MathOpAddendAdditionalTooltip", "Float multiplicand(s)."), 1.0f)
+					TInputDataVertexModel<FTime>(MathOpNames::PrimaryOperandName, LOCTEXT("MathOpTimeFloatMultiplyNode_MultiplyFloatTooltip", "Time multiplicand."), 1.0f),
+					TInputDataVertexModel<float>(MathOpNames::AdditionalOperandsName, LOCTEXT("MathOpTimeFloatMultiplyNode_MultiplicandsTooltip", "Float multiplicand(s)."), 1.0f)
 				),
 				FOutputVertexInterface(
-					TOutputDataVertexModel<FTime>(TEXT("Out"), LOCTEXT("MathOpTimeMultiplyOutTooltip", "Resulting buffer"))
+					TOutputDataVertexModel<FTime>(TEXT("Out"), LOCTEXT("MathOpTimeFloatMultiplyNode_OutTooltip", "Time float multiplication output"))
 				)
 			);
 
@@ -1237,55 +1139,6 @@ namespace Metasound
 		}
 	};
 
-#if 0
-	// int64 has to be explicitly implemented as they are PODs that are not implemented as a literal base type.
-	template <>
-	class TMathOpMultiply<int64>
-	{
-	public:
-		static const FVertexInterface& GetVertexInterface()
-		{
-			static const FVertexInterface Interface(
-				FInputVertexInterface(
-					TInputDataVertexModel<int64>(MathOpNames::PrimaryOperandName, LOCTEXT("MathOpInitMultiplicandTooltip", "Initial multiplicand."), 1),
-					TInputDataVertexModel<int64>(MathOpNames::AdditionalOperandsName, LOCTEXT("MathOpMultiplicandsTooltip", "Additional multiplicand(s)."), 1)
-				),
-				FOutputVertexInterface(
-					TOutputDataVertexModel<int64>(TEXT("Out"), LOCTEXT("MultiplicationResultTooltip", "Multiplication result"))
-				)
-			);
-
-			return Interface;
-		}
-
-		static int32 GetDefault(const FOperatorSettings& InSettings, const FLiteral& InVertexDefault)
-		{
-			return static_cast<int64>(InVertexDefault.Value.Get<int32>());
-		}
-
-		static int32 GetDefaultOp(const FOperatorSettings& InSettings, const FLiteral& InVertexDefault)
-		{
-			return static_cast<int64>(InVertexDefault.Value.Get<int32>());
-		}
-
-		static void Calculate(TMathOpMultiply& InInstanceData, const FInt64ReadRef& InPrimaryOperand, const TArray<FInt64ReadRef>& InAdditionalOperands, FInt64WriteRef& OutResult)
-		{
-			*OutResult = *InPrimaryOperand;
-
-			if (!ensure(!InAdditionalOperands.IsEmpty()))
-			{
-				return;
-			}
-
-			for (int32 i = 0; i < InAdditionalOperands.Num(); ++i)
-			{
-				const int64 OperandValue = *InAdditionalOperands[i];
-				*OutResult *= OperandValue;
-			}
-		}
-	};
-#endif
-
 	template <>
 	class TMathOpDivide<FTime, float>
 	{
@@ -1294,11 +1147,11 @@ namespace Metasound
 		{
 			static const FVertexInterface Interface(
 				FInputVertexInterface(
-					TInputDataVertexModel<FTime>(MathOpNames::PrimaryOperandName, LOCTEXT("MathOpTimeMultiplyFloatTooltip", "Time multiplicand."), 1.0f),
-					TInputDataVertexModel<float>(MathOpNames::AdditionalOperandsName, LOCTEXT("MathOpAddendAdditionalTooltip", "Float multiplicand(s)."), 1.0f)
+					TInputDataVertexModel<FTime>(MathOpNames::PrimaryOperandName, LOCTEXT("MathOpTimeFloatDivideNode_DividendFloatTooltip", "Time dividend."), 1.0f),
+					TInputDataVertexModel<float>(MathOpNames::AdditionalOperandsName, LOCTEXT("MathOpTimeFloatDivideNode_DivisorsTooltip", "Float divisor(s)."), 1.0f)
 				),
 				FOutputVertexInterface(
-					TOutputDataVertexModel<FTime>(TEXT("Out"), LOCTEXT("MathOpTimeMultiplyOutTooltip", "Resulting buffer"))
+					TOutputDataVertexModel<FTime>(TEXT("Out"), LOCTEXT("MathOpTimeFloatDivideNode_OutTooltip", "Time divide-by-float output"))
 				)
 			);
 
@@ -1338,163 +1191,45 @@ namespace Metasound
 		}
 	};
 
-#if 0
-	// int64 has to be explicitly implemented as they are PODs that are not implemented as a literal base type.
-	template <>
-	class TMathOpDivide<int64>
-	{
-	public:
-		static const FVertexInterface& GetVertexInterface()
-		{
-			static const FVertexInterface Interface(
-				FInputVertexInterface(
-					TInputDataVertexModel<int64>(MathOpNames::PrimaryOperandName, LOCTEXT("MathOpDivideDividendTooltip", "Dividend of operation."), 1),
-					TInputDataVertexModel<int64>(MathOpNames::AdditionalOperandsName, LOCTEXT("MathOpDivideDivisorTooltip", "Divisor of operation."), 1)
-				),
-				FOutputVertexInterface(
-					TOutputDataVertexModel<int64>(TEXT("Out"), LOCTEXT("MathOpDivideOutTooltip", "Resulting value"))
-				)
-			);
-
-			return Interface;
-		}
-
-		static int64 GetDefault(const FOperatorSettings& InSettings, const FLiteral& InVertexDefault)
-		{
-			// TODO: we are truncating to int32, why do we need int64?
-			return static_cast<int64>(InVertexDefault.Value.Get<int32>());
-		}
-
-		static int64 GetDefaultOp(const FOperatorSettings& InSettings, const FLiteral& InVertexDefault)
-		{
-			// TODO: we are truncating to int32, why do we need in64?
-			return static_cast<int64>(InVertexDefault.Value.Get<int32>());
-		}
-
-		static void Calculate(TMathOpDivide& InInstanceData, const FInt64ReadRef& InPrimaryOperand, const TArray<FInt64ReadRef>& InAdditionalOperands, FInt64WriteRef& OutResult)
-		{
-			*OutResult = *InPrimaryOperand;
-
-			if (!ensure(!InAdditionalOperands.IsEmpty()))
-			{
-				return;
-			}
-
-			for (int32 i = 0; i < InAdditionalOperands.Num(); ++i)
-			{
-				const int64 OperandValue = *InAdditionalOperands[i];
-				if (OperandValue == 0)
-				{
-					// TODO: Error here
-					continue;
-				}
-
-				*OutResult /= OperandValue;
-			}
-		}
-	};
-
-
-	// int64 has to be explicitly implemented as they are PODs that are not implemented as a literal base type.
-	template <>
-	class TMathOpModulo<int64>
-	{
-	public:
-		static const FVertexInterface& GetVertexInterface()
-		{
-			static const FVertexInterface Interface(
-				FInputVertexInterface(
-					TInputDataVertexModel<int64>(MathOpNames::PrimaryOperandName, LOCTEXT("MathOpModuloDividendTooltip", "Dividend"), 1),
-					TInputDataVertexModel<int64>(MathOpNames::AdditionalOperandsName, LOCTEXT("MathOpModuluDivisorTooltip", "Divisor(s)"), 1)
-				),
-				FOutputVertexInterface(
-					TOutputDataVertexModel<int64>(TEXT("Out"), LOCTEXT("MathOpModuloOutTooltip", "Resulting value"))
-				)
-			);
-
-			return Interface;
-		}
-
-		static int64 GetDefault(const FOperatorSettings& InSettings, const FLiteral& InVertexDefault)
-		{
-			return static_cast<int64>(InVertexDefault.Value.Get<int32>());
-		}
-
-		static int64 GetDefaultOp(const FOperatorSettings& InSettings, const FLiteral& InVertexDefault)
-		{
-			return static_cast<int64>(InVertexDefault.Value.Get<int32>());
-		}
-
-		static void Calculate(TMathOpModulo& InInstanceData, const FInt64ReadRef& InPrimaryOperand, const TArray<FInt64ReadRef>& InAdditionalOperands, FInt64WriteRef& OutResult)
-		{
-			*OutResult = *InPrimaryOperand;
-
-			if (!ensure(!InAdditionalOperands.IsEmpty()))
-			{
-				return;
-			}
-
-			for (int32 i = 0; i < InAdditionalOperands.Num(); ++i)
-			{
-				const int64 OperandValue = *InAdditionalOperands[i];
-				if (OperandValue == 0)
-				{
-					// TODO: Error here
-					continue;
-				}
-
-				*OutResult %= OperandValue;
-			}
-		}
-	};
-#endif
-
 	// Definitions
-	DEFINE_METASOUND_MATHOP(Add, Float, float, LOCTEXT("Metasound_MathAddFloatNodeDescription", "Adds floats."), MathOpNames::AddKeywords)
-	DEFINE_METASOUND_MATHOP(Add, Int32, int32, LOCTEXT("Metasound_MathAddInt32NodeDescription", "Adds int32s."), MathOpNames::AddKeywords)
-	//DEFINE_METASOUND_MATHOP(Add, Int64, int64, LOCTEXT("Metasound_MathAddInt64NodeDescription", "Adds int64s."), MathOpNames::AddKeywords)
-	DEFINE_METASOUND_MATHOP(Add, Audio, FAudioBuffer, LOCTEXT("Metasound_MathAddBufferNodeDescription", "Adds buffers together by sample."), MathOpNames::AddKeywords)
-	DEFINE_METASOUND_MATHOP(Add, Time, FTime, LOCTEXT("Metasound_MathAddTimeNodeDescription", "Adds time values."), MathOpNames::AddKeywords)
+	DEFINE_METASOUND_MATHOP(Add, Float, float, LOCTEXT("MathAddFloatNode_Description", "Adds floats."), MathOpNames::AddKeywords)
+	DEFINE_METASOUND_MATHOP(Add, Int32, int32, LOCTEXT("MathAddInt32Node_Description", "Adds int32s."), MathOpNames::AddKeywords)
+	DEFINE_METASOUND_MATHOP(Add, Audio, FAudioBuffer, LOCTEXT("MathAddBufferNode_Description", "Adds buffers together by sample."), MathOpNames::AddKeywords)
+	DEFINE_METASOUND_MATHOP(Add, Time, FTime, LOCTEXT("MathAddTimeNode_Description", "Adds time values."), MathOpNames::AddKeywords)
 
-	DEFINE_METASOUND_OPERAND_TYPED_MATHOP(Add, Audio, FAudioBuffer, Float, float, LOCTEXT("Metasound_MathAddAudioFloatNodeDescription", "Add floats to buffer sample-by-sample."), MathOpNames::AddKeywords)
+	DEFINE_METASOUND_OPERAND_TYPED_MATHOP(Add, Audio, FAudioBuffer, Float, float, LOCTEXT("MathAddAudioFloatNode_Description", "Add floats to buffer sample-by-sample."), MathOpNames::AddKeywords)
 
-	DEFINE_METASOUND_MATHOP(Subtract, Float, float, LOCTEXT("Metasound_MathSubractFloatNodeDescription", "Subtracts floats."), MathOpNames::SubtractKeywords)
-	DEFINE_METASOUND_MATHOP(Subtract, Int32, int32, LOCTEXT("Metasound_MathSubractInt32NodeDescription", "Subtracts int32s."), MathOpNames::SubtractKeywords)
-	//DEFINE_METASOUND_MATHOP(Subtract, Int64, int64, LOCTEXT("Metasound_MathSubractInt64NodeDescription", "Subtracts int64s."), MathOpNames::SubtractKeywords)
-	DEFINE_METASOUND_MATHOP(Subtract, Audio, FAudioBuffer, LOCTEXT("Metasound_MathSubtractBufferNodeDescription", "Subtracts buffers sample-by-sample."), MathOpNames::SubtractKeywords)
-	DEFINE_METASOUND_MATHOP(Subtract, Time, FTime, LOCTEXT("Metasound_MathSubractTimeNodeDescription", "Subtracts time values."), MathOpNames::SubtractKeywords)
+	DEFINE_METASOUND_MATHOP(Subtract, Float, float, LOCTEXT("MathSubractFloatNode_Description", "Subtracts floats."), MathOpNames::SubtractKeywords)
+	DEFINE_METASOUND_MATHOP(Subtract, Int32, int32, LOCTEXT("MathSubractInt32Node_Description", "Subtracts int32s."), MathOpNames::SubtractKeywords)
+	DEFINE_METASOUND_MATHOP(Subtract, Audio, FAudioBuffer, LOCTEXT("MathSubtractBufferNode_Description", "Subtracts buffers sample-by-sample."), MathOpNames::SubtractKeywords)
+	DEFINE_METASOUND_MATHOP(Subtract, Time, FTime, LOCTEXT("MathSubractTimeNode_Description", "Subtracts time values."), MathOpNames::SubtractKeywords)
 
-	DEFINE_METASOUND_MATHOP(Multiply, Float, float, LOCTEXT("Metasound_MathMultiplyFloatNodeDescription", "Multiplies floats."), MathOpNames::MultiplyKeywords)
-	DEFINE_METASOUND_MATHOP(Multiply, Int32, int32, LOCTEXT("Metasound_MathMultiplyInt32NodeDescription", "Multiplies int32s."), MathOpNames::MultiplyKeywords)
-	//DEFINE_METASOUND_MATHOP(Multiply, Int64, int64, LOCTEXT("Metasound_MathMultiplyInt64NodeDescription", "Multiplies int64s."), MathOpNames::MultiplyKeywords)
-	DEFINE_METASOUND_MATHOP(Multiply, Audio, FAudioBuffer, LOCTEXT("Metasound_MathMultiplyBufferNodeDescription", "Multiplies buffers together sample-by-sample."), MathOpNames::MultiplyKeywords)
+	DEFINE_METASOUND_MATHOP(Multiply, Float, float, LOCTEXT("MathMultiplyFloatNode_Description", "Multiplies floats."), MathOpNames::MultiplyKeywords)
+	DEFINE_METASOUND_MATHOP(Multiply, Int32, int32, LOCTEXT("MathMultiplyInt32Node_Description", "Multiplies int32s."), MathOpNames::MultiplyKeywords)
+	DEFINE_METASOUND_MATHOP(Multiply, Audio, FAudioBuffer, LOCTEXT("MathMultiplyBufferNode_Description", "Multiplies buffers together sample-by-sample."), MathOpNames::MultiplyKeywords)
 
-	DEFINE_METASOUND_OPERAND_TYPED_MATHOP(Multiply, Audio, FAudioBuffer, Float, float, LOCTEXT("Metasound_MathMultiplyAudioByFloatDescription", "Multiplies buffer by float scalars."), MathOpNames::MultiplyKeywords)
-	DEFINE_METASOUND_OPERAND_TYPED_MATHOP(Multiply, Time, FTime, Float, float, LOCTEXT("Metasound_MathMultiplyTimeNodeDescription", "Scales time by floats."), MathOpNames::MultiplyKeywords)
+	DEFINE_METASOUND_OPERAND_TYPED_MATHOP(Multiply, Audio, FAudioBuffer, Float, float, LOCTEXT("MathMultiplyAudioByFloatDescription", "Multiplies buffer by float scalars."), MathOpNames::MultiplyKeywords)
+	DEFINE_METASOUND_OPERAND_TYPED_MATHOP(Multiply, Time, FTime, Float, float, LOCTEXT("MathMultiplyTimeNode_Description", "Scales time by floats."), MathOpNames::MultiplyKeywords)
 
-	DEFINE_METASOUND_MATHOP(Divide, Float, float, LOCTEXT("Metasound_MathDivideFloatNodeDescription", "Divide float by another float."), MathOpNames::DivideKeywords)
-	DEFINE_METASOUND_MATHOP(Divide, Int32, int32, LOCTEXT("Metasound_MathDivideInt32NodeDescription", "Divide int32 by another int32."), MathOpNames::DivideKeywords)
-	//DEFINE_METASOUND_MATHOP(Divide, Int64, int64, LOCTEXT("Metasound_MathDivideInt64NodeDescription", "Divide int64 by another int64."), MathOpNames::DivideKeywords)
+	DEFINE_METASOUND_MATHOP(Divide, Float, float, LOCTEXT("MathDivideFloatNode_Description", "Divide float by another float."), MathOpNames::DivideKeywords)
+	DEFINE_METASOUND_MATHOP(Divide, Int32, int32, LOCTEXT("MathDivideInt32Node_Description", "Divide int32 by another int32."), MathOpNames::DivideKeywords)
 
-	DEFINE_METASOUND_OPERAND_TYPED_MATHOP(Divide, Time, FTime, Float, float, LOCTEXT("Metasound_MathDivideTimeNodeDescription", "Divides time by floats."), MathOpNames::DivideKeywords)
+	DEFINE_METASOUND_OPERAND_TYPED_MATHOP(Divide, Time, FTime, Float, float, LOCTEXT("MathDivideTimeNode_Description", "Divides time by floats."), MathOpNames::DivideKeywords)
 
-	DEFINE_METASOUND_MATHOP(Modulo, Int32, int32, LOCTEXT("Metasound_MathModulusInt32NodeDescription", "Modulo int32 by another int32."), MathOpNames::ModuloKeywords)
-	//DEFINE_METASOUND_MATHOP(Modulo, Int64, int64, LOCTEXT("Metasound_MathModulusInt64NodeDescription", "Modulo int64 by another int64."), MathOpNames::ModuloKeywords)
+	DEFINE_METASOUND_MATHOP(Modulo, Int32, int32, LOCTEXT("MathModulusInt32Node_Description", "Modulo int32 by another int32."), MathOpNames::ModuloKeywords)
 
-	DEFINE_METASOUND_MATHOP(Power, Float, float, LOCTEXT("Metasound_MathPowerFloatNodeDescription", "Raise float to the power of another float."), MathOpNames::PowerKeywords)
+	DEFINE_METASOUND_MATHOP(Power, Float, float, LOCTEXT("MathPowerFloatNode_Description", "Raise float to the power of another float."), MathOpNames::PowerKeywords)
 
-	DEFINE_METASOUND_MATHOP(Logarithm, Float, float, LOCTEXT("Metasound_MathLogarithmFloatNodeDescription", "Calculate float-base logarithm of another float."), TArray<FText>())
+	DEFINE_METASOUND_MATHOP(Logarithm, Float, float, LOCTEXT("MathLogarithmFloatNode_Description", "Calculate float-base logarithm of another float."), TArray<FText>())
 
 	METASOUND_REGISTER_NODE(FAddFloatNode)
 	METASOUND_REGISTER_NODE(FAddInt32Node)
-	//METASOUND_REGISTER_NODE(FAddInt64Node)
 	METASOUND_REGISTER_NODE(FAddTimeNode)
 	METASOUND_REGISTER_NODE(FAddAudioNode)
 	METASOUND_REGISTER_NODE(FAddAudioFloatNode)
 
 	METASOUND_REGISTER_NODE(FSubtractFloatNode)
 	METASOUND_REGISTER_NODE(FSubtractInt32Node)
-	//METASOUND_REGISTER_NODE(FSubtractInt64Node)
 	METASOUND_REGISTER_NODE(FSubtractTimeNode)
 	METASOUND_REGISTER_NODE(FSubtractAudioNode)
 
@@ -1502,16 +1237,13 @@ namespace Metasound
 	METASOUND_REGISTER_NODE(FMultiplyAudioFloatNode)
 	METASOUND_REGISTER_NODE(FMultiplyFloatNode)
 	METASOUND_REGISTER_NODE(FMultiplyInt32Node)
-	//METASOUND_REGISTER_NODE(FMultiplyInt64Node)
 	METASOUND_REGISTER_NODE(FMultiplyTimeFloatNode)
 
 	METASOUND_REGISTER_NODE(FDivideFloatNode)
 	METASOUND_REGISTER_NODE(FDivideInt32Node)
-	//METASOUND_REGISTER_NODE(FDivideInt64Node)
 	METASOUND_REGISTER_NODE(FDivideTimeFloatNode)
 
 	METASOUND_REGISTER_NODE(FModuloInt32Node)
-	//METASOUND_REGISTER_NODE(FModuloInt64Node)
 
 	METASOUND_REGISTER_NODE(FPowerFloatNode)
 
