Remove implicit constructor for FMetasoundFrontendClass
- Favor static factory style initializer to clarify when it is being generated as a registry description (and thus has no ID/explicitly set type)
#rnx

diff --git a/Engine/Plugins/Runtime/Metasound/Source/MetasoundFrontend/Private/MetasoundFrontend.cpp b/Engine/Plugins/Runtime/Metasound/Source/MetasoundFrontend/Private/MetasoundFrontend.cpp
--- a/Engine/Plugins/Runtime/Metasound/Source/MetasoundFrontend/Private/MetasoundFrontend.cpp
+++ b/Engine/Plugins/Runtime/Metasound/Source/MetasoundFrontend/Private/MetasoundFrontend.cpp
@@ -35,8 +35,7 @@ namespace Metasound
 		{
 			FMetasoundFrontendClass ClassDescription;
 
-			ClassDescription.Metadata = FMetasoundFrontendClassMetadata(InNodeMetadata);
-			ClassDescription.Metadata.SetType(ClassType);
+			ClassDescription.Metadata = FMetasoundFrontendClassMetadata::GenerateClassDescription(InNodeMetadata, ClassType);
 
 			FMetasoundFrontendClassStyleDisplay DisplayStyle(InNodeMetadata.DisplayStyle);
 			ClassDescription.Style = FMetasoundFrontendClassStyle
diff --git a/Engine/Plugins/Runtime/Metasound/Source/MetasoundFrontend/Private/MetasoundFrontendDocument.cpp b/Engine/Plugins/Runtime/Metasound/Source/MetasoundFrontend/Private/MetasoundFrontendDocument.cpp
--- a/Engine/Plugins/Runtime/Metasound/Source/MetasoundFrontend/Private/MetasoundFrontendDocument.cpp
+++ b/Engine/Plugins/Runtime/Metasound/Source/MetasoundFrontend/Private/MetasoundFrontendDocument.cpp
@@ -23,7 +23,63 @@ namespace Metasound
 			} // namespace NodeLayout
 		} // namespace DisplayStyle
 	} // namespace Frontend
-}
+
+	namespace DocumentPrivate
+	{
+		/*
+		 * Sets an array to a given array and updates the change ID if the array changed.
+		 * @returns true if value changed, false if not.
+		 */
+		template <typename TElementType>
+		bool SetWithChangeID(const TElementType& InNewValue, TElementType& OutValue, FGuid& OutChangeID)
+		{
+			if (OutValue != InNewValue)
+			{
+				OutValue = InNewValue;
+				OutChangeID = FGuid::NewGuid();
+				return true;
+			}
+
+			return false;
+		}
+
+		/* Array Text specialization as FText does not implement == nor does it support IsBytewiseComparable */
+		template <>
+		bool SetWithChangeID<TArray<FText>>(const TArray<FText>& InNewArray, TArray<FText>& OutArray, FGuid& OutChangeID)
+		{
+			bool bIsEqual = OutArray.Num() == InNewArray.Num();
+			if (bIsEqual)
+			{
+				for (int32 i = 0; i < InNewArray.Num(); ++i)
+				{
+					bIsEqual &= InNewArray[i].IdenticalTo(OutArray[i]);
+				}
+			}
+
+			if (!bIsEqual)
+			{
+				OutArray = InNewArray;
+				OutChangeID = FGuid::NewGuid();
+			}
+
+			return !bIsEqual;
+		}
+
+		/* Text specialization as FText does not implement == nor does it support IsBytewiseComparable */
+		template <>
+		bool SetWithChangeID<FText>(const FText& InNewText, FText& OutText, FGuid& OutChangeID)
+		{
+			if (!InNewText.IdenticalTo(OutText))
+			{
+				OutText = InNewText;
+				OutChangeID = FGuid::NewGuid();
+				return true;
+			}
+
+			return false;
+		}
+	}
+} // namespace Metasound
 
 FMetasoundFrontendNodeInterface::FMetasoundFrontendNodeInterface(const FMetasoundFrontendClassInterface& InClassInterface)
 {
@@ -123,18 +179,88 @@ bool operator!=(const FMetasoundFrontendClassName& InLHS, const FMetasoundFronte
 	return !(InLHS == InRHS);
 }
 
-FMetasoundFrontendClassMetadata::FMetasoundFrontendClassMetadata(const Metasound::FNodeClassMetadata& InNodeClassMetadata)
-: ClassName(InNodeClassMetadata.ClassName)
-, Version{ InNodeClassMetadata.MajorVersion, InNodeClassMetadata.MinorVersion }
-, Type(EMetasoundFrontendClassType::External)
-, DisplayName(InNodeClassMetadata.DisplayName)
-, Description(InNodeClassMetadata.Description)
-, PromptIfMissing(InNodeClassMetadata.PromptIfMissing)
-, Author(InNodeClassMetadata.Author)
-, Keywords(InNodeClassMetadata.Keywords)
-, CategoryHierarchy(InNodeClassMetadata.CategoryHierarchy)
-, bIsDeprecated(InNodeClassMetadata.bDeprecated)
+void FMetasoundFrontendClassMetadata::SetAuthor(const FText& InAuthor)
+{
+	using namespace Metasound::DocumentPrivate;
+	SetWithChangeID(InAuthor, Author, ChangeID);
+}
+
+void FMetasoundFrontendClassMetadata::SetAutoUpdateManagesInterface(bool bInAutoUpdateManagesInterface)
+{
+	using namespace Metasound::DocumentPrivate;
+	SetWithChangeID(bInAutoUpdateManagesInterface, bAutoUpdateManagesInterface, ChangeID);
+}
+
+void FMetasoundFrontendClassMetadata::SetCategoryHierarchy(const TArray<FText>& InCategoryHierarchy)
+{
+	using namespace Metasound::DocumentPrivate;
+	SetWithChangeID(InCategoryHierarchy, CategoryHierarchy, ChangeID);
+}
+
+void FMetasoundFrontendClassMetadata::SetKeywords(const TArray<FText>& InKeywords)
+{
+	using namespace Metasound::DocumentPrivate;
+	SetWithChangeID(InKeywords, Keywords, ChangeID);
+}
+
+void FMetasoundFrontendClassMetadata::SetClassName(const FMetasoundFrontendClassName& InClassName)
+{
+	using namespace Metasound::DocumentPrivate;
+	SetWithChangeID(InClassName, ClassName, ChangeID);
+}
+
+void FMetasoundFrontendClassMetadata::SetDescription(const FText& InDescription)
+{
+	using namespace Metasound::DocumentPrivate;
+	SetWithChangeID(InDescription, Description, ChangeID);
+}
+
+void FMetasoundFrontendClassMetadata::SetDisplayName(const FText& InDisplayName)
+{
+	using namespace Metasound::DocumentPrivate;
+	SetWithChangeID(InDisplayName, DisplayName, ChangeID);
+}
+
+void FMetasoundFrontendClassMetadata::SetIsDeprecated(bool bInIsDeprecated)
+{
+	using namespace Metasound::DocumentPrivate;
+	SetWithChangeID(bInIsDeprecated, bIsDeprecated, ChangeID);
+}
+
+void FMetasoundFrontendClassMetadata::SetPromptIfMissing(const FText& InPromptIfMissing)
+{
+	using namespace Metasound::DocumentPrivate;
+	SetWithChangeID(InPromptIfMissing, PromptIfMissing, ChangeID);
+}
+
+void FMetasoundFrontendClassMetadata::SetVersion(const FMetasoundFrontendVersionNumber& InVersion)
+{
+	using namespace Metasound::DocumentPrivate;
+	SetWithChangeID(InVersion, Version, ChangeID);
+}
+
+FMetasoundFrontendClassMetadata FMetasoundFrontendClassMetadata::GenerateClassDescription(const Metasound::FNodeClassMetadata& InNodeClassMetadata, EMetasoundFrontendClassType InType)
 {
+	FMetasoundFrontendClassMetadata NewMetadata;
+	NewMetadata.Type = InType;
+
+	// TODO: This flag is only used by the graph class' metadata.
+	// Should probably be moved elsewhere (AssetBase?) as to not
+	// get confused with behavior encapsulated on registry class
+	// descriptions/individual node class dependencies.
+	NewMetadata.bAutoUpdateManagesInterface = false;
+
+	NewMetadata.ClassName = InNodeClassMetadata.ClassName;
+	NewMetadata.Version = { InNodeClassMetadata.MajorVersion, InNodeClassMetadata.MinorVersion };
+	NewMetadata.DisplayName = InNodeClassMetadata.DisplayName;
+	NewMetadata.Description = InNodeClassMetadata.Description;
+	NewMetadata.PromptIfMissing = InNodeClassMetadata.PromptIfMissing;
+	NewMetadata.Author = InNodeClassMetadata.Author;
+	NewMetadata.Keywords = InNodeClassMetadata.Keywords;
+	NewMetadata.CategoryHierarchy = InNodeClassMetadata.CategoryHierarchy;
+	NewMetadata.bIsDeprecated = InNodeClassMetadata.bDeprecated;
+
+	return NewMetadata;
 }
 
 FMetasoundFrontendClassInput::FMetasoundFrontendClassInput(const FMetasoundFrontendClassVertex& InOther)
diff --git a/Engine/Plugins/Runtime/Metasound/Source/MetasoundFrontend/Public/MetasoundFrontendDocument.h b/Engine/Plugins/Runtime/Metasound/Source/MetasoundFrontend/Public/MetasoundFrontendDocument.h
--- a/Engine/Plugins/Runtime/Metasound/Source/MetasoundFrontend/Public/MetasoundFrontendDocument.h
+++ b/Engine/Plugins/Runtime/Metasound/Source/MetasoundFrontend/Public/MetasoundFrontendDocument.h
@@ -888,209 +888,158 @@ struct METASOUNDFRONTEND_API FMetasoundFrontendClassMetadata
 {
 	GENERATED_BODY()
 
-	FMetasoundFrontendClassMetadata() = default;
-	FMetasoundFrontendClassMetadata(const Metasound::FNodeClassMetadata& InNodeClassMetadata);
+	// Generates class metadata intended to be used as a registry descriptor from FNodeClassMetadata. Does not initialize a change ID as it is not considered to be transactional.
+	static FMetasoundFrontendClassMetadata GenerateClassDescription(const Metasound::FNodeClassMetadata& InNodeClassMetadata, EMetasoundFrontendClassType InType);
 
-	private:
-		UPROPERTY(VisibleAnywhere, Category = Metasound)
-		FMetasoundFrontendClassName ClassName;
-
-		UPROPERTY(VisibleAnywhere, Category = Metasound)
-		FMetasoundFrontendVersionNumber Version;
-
-		UPROPERTY(VisibleAnywhere, Category = Metasound)
-		EMetasoundFrontendClassType Type = EMetasoundFrontendClassType::Invalid;
-
-		UPROPERTY(VisibleAnywhere, Category = Metasound)
-		FText DisplayName;
-
-		UPROPERTY(EditAnywhere, Category = Metasound)
-		FText Description;
-
-		UPROPERTY(VisibleAnywhere, Category = Metasound)
-		FText PromptIfMissing;
-
-		UPROPERTY(EditAnywhere, Category = Metasound)
-		FText Author;
-
-		UPROPERTY(VisibleAnywhere, Category = Metasound)
-		TArray<FText> Keywords;
-
-		UPROPERTY(EditAnywhere, Category = Metasound)
-		TArray<FText> CategoryHierarchy;
-
-		// If true, this node is deprecated and should not be used in new MetaSounds.
-		UPROPERTY(VisibleAnywhere, Category = Metasound)
-		bool bIsDeprecated = false;
+private:
+	UPROPERTY(VisibleAnywhere, Category = Metasound)
+	FMetasoundFrontendClassName ClassName;
 
-		// If true, auto-update will manage (add and remove)
-		// inputs/outputs associated with internally connected
-		// nodes when the interface of the given node is auto-updated.
-		UPROPERTY()
-		bool bAutoUpdateManagesInterface = false;
+	UPROPERTY(VisibleAnywhere, Category = Metasound)
+	FMetasoundFrontendVersionNumber Version;
 
-		// ID used to identify if any of the above have been modified,
-		// to determine if the parent class should be auto-updated.
-		UPROPERTY()
-		FGuid ChangeID;
+	UPROPERTY(VisibleAnywhere, Category = Metasound)
+	EMetasoundFrontendClassType Type = EMetasoundFrontendClassType::Invalid;
 
-	public:
-		static FName GetAuthorPropertyName()
-		{
-			return GET_MEMBER_NAME_CHECKED(FMetasoundFrontendClassMetadata, Author);
-		}
+	UPROPERTY(VisibleAnywhere, Category = Metasound)
+	FText DisplayName;
 
-		static FName GetCategoryHierarchyPropertyName()
-		{
-			return GET_MEMBER_NAME_CHECKED(FMetasoundFrontendClassMetadata, CategoryHierarchy);
-		}
+	UPROPERTY(EditAnywhere, Category = Metasound)
+	FText Description;
 
-		static FName GetClassNamePropertyName()
-		{
-			return GET_MEMBER_NAME_CHECKED(FMetasoundFrontendClassMetadata, ClassName);
-		}
+	UPROPERTY(VisibleAnywhere, Category = Metasound)
+	FText PromptIfMissing;
 
-		static FName GetDescriptionPropertyName()
-		{
-			return GET_MEMBER_NAME_CHECKED(FMetasoundFrontendClassMetadata, Description);
-		}
+	UPROPERTY(EditAnywhere, Category = Metasound)
+	FText Author;
 
-		static FName GetVersionPropertyName()
-		{
-			return GET_MEMBER_NAME_CHECKED(FMetasoundFrontendClassMetadata, Version);
-		}
+	UPROPERTY(VisibleAnywhere, Category = Metasound)
+	TArray<FText> Keywords;
 
-		const FMetasoundFrontendClassName& GetClassName() const
-		{
-			return ClassName;
-		}
+	UPROPERTY(EditAnywhere, Category = Metasound)
+	TArray<FText> CategoryHierarchy;
 
-		bool GetAutoUpdateManagesInterface() const
-		{
-			return bAutoUpdateManagesInterface;
-		}
+	// If true, this node is deprecated and should not be used in new MetaSounds.
+	UPROPERTY(VisibleAnywhere, Category = Metasound)
+	bool bIsDeprecated = false;
 
-		void SetClassName(const FMetasoundFrontendClassName& InClassName)
-		{
-			ClassName = InClassName;
-			ChangeID = FGuid::NewGuid();
-		}
+	// If true, auto-update will manage (add and remove)
+	// inputs/outputs associated with internally connected
+	// nodes when the interface of the given node is auto-updated.
+	UPROPERTY()
+	bool bAutoUpdateManagesInterface = false;
 
-		EMetasoundFrontendClassType GetType() const
-		{
-			return Type;
-		}
+	// ID used to identify if any of the above have been modified,
+	// to determine if the parent class should be auto-updated.
+	UPROPERTY()
+	FGuid ChangeID;
 
-		void SetAutoUpdateManagesInterface(bool bInAutoUpdateManagesInterface)
-		{
-			bAutoUpdateManagesInterface = bInAutoUpdateManagesInterface;
-			ChangeID = FGuid::NewGuid();
-		}
+public:
+	static FName GetAuthorPropertyName()
+	{
+		return GET_MEMBER_NAME_CHECKED(FMetasoundFrontendClassMetadata, Author);
+	}
 
-		void SetType(const EMetasoundFrontendClassType InType)
-		{
-			Type = InType;
-			// TODO: Type is modified while querying and swapped between
-			// to be external, so don't modify the ChangeID in this case.
-			// External/Internal should probably be a separate field.
-			// ChangeID = FGuid::NewGuid();
-		}
+	static FName GetCategoryHierarchyPropertyName()
+	{
+		return GET_MEMBER_NAME_CHECKED(FMetasoundFrontendClassMetadata, CategoryHierarchy);
+	}
 
-		const FMetasoundFrontendVersionNumber& GetVersion() const
-		{
-			return Version;
-		}
+	static FName GetClassNamePropertyName()
+	{
+		return GET_MEMBER_NAME_CHECKED(FMetasoundFrontendClassMetadata, ClassName);
+	}
 
-		void SetVersion(const FMetasoundFrontendVersionNumber& InVersion)
-		{
-			Version = InVersion;
-			ChangeID = FGuid::NewGuid();
-		}
+	static FName GetDescriptionPropertyName()
+	{
+		return GET_MEMBER_NAME_CHECKED(FMetasoundFrontendClassMetadata, Description);
+	}
 
-		const FText& GetDisplayName() const
-		{
-			return DisplayName;
-		}
+	static FName GetVersionPropertyName()
+	{
+		return GET_MEMBER_NAME_CHECKED(FMetasoundFrontendClassMetadata, Version);
+	}
 
-		void SetDisplayName(const FText& InDisplayName)
-		{
-			DisplayName = InDisplayName;
-			ChangeID = FGuid::NewGuid();
-		}
+	const FMetasoundFrontendClassName& GetClassName() const
+	{
+		return ClassName;
+	}
 
-		const FText& GetDescription() const
-		{
-			return Description;
-		}
+	bool GetAutoUpdateManagesInterface() const
+	{
+		return bAutoUpdateManagesInterface;
+	}
 
-		void SetDescription(const FText& InDescription)
-		{
-			Description = InDescription;
-			ChangeID = FGuid::NewGuid();
-		}
+	void SetClassName(const FMetasoundFrontendClassName& InClassName);
 
-		const FText& GetPromptIfMissing() const
-		{
-			return PromptIfMissing;
-		}
+	EMetasoundFrontendClassType GetType() const
+	{
+		return Type;
+	}
 
-		void SetPromptIfMissing(const FText& InPromptIfMissing)
-		{
-			PromptIfMissing = InPromptIfMissing;
-			ChangeID = FGuid::NewGuid();
-		}
+	const FMetasoundFrontendVersionNumber& GetVersion() const
+	{
+		return Version;
+	}
 
-		const FText& GetAuthor() const
-		{
-			return Author;
-		}
+	const FText& GetDisplayName() const
+	{
+		return DisplayName;
+	}
 
-		void SetAuthor(const FText& InAuthor)
-		{
-			Author = InAuthor;
-			ChangeID = FGuid::NewGuid();
-		}
+	const FText& GetDescription() const
+	{
+		return Description;
+	}
 
-		const TArray<FText>& GetKeywords() const
-		{
-			return Keywords;
-		}
+	const FText& GetPromptIfMissing() const
+	{
+		return PromptIfMissing;
+	}
 
-		void SetKeywords(const TArray<FText>& InKeywords)
-		{
-			Keywords = InKeywords;
-			ChangeID = FGuid::NewGuid();
-		}
+	const FText& GetAuthor() const
+	{
+		return Author;
+	}
 
-		const TArray<FText>& GetCategoryHierarchy() const
-		{
-			return CategoryHierarchy;
-		}
+	const TArray<FText>& GetKeywords() const
+	{
+		return Keywords;
+	}
 
-		void SetCategoryHierarchy(const TArray<FText>& InCategoryHierarchy)
-		{
-			CategoryHierarchy = InCategoryHierarchy;
-			ChangeID = FGuid::NewGuid();
-		}
+	const TArray<FText>& GetCategoryHierarchy() const
+	{
+		return CategoryHierarchy;
+	}
 
-		const FGuid& GetChangeID() const
-		{
-			return ChangeID;
-		}
+	const FGuid& GetChangeID() const
+	{
+		return ChangeID;
+	}
 
-		bool GetIsDeprecated() const
-		{
-			return bIsDeprecated;
-		}
+	bool GetIsDeprecated() const
+	{
+		return bIsDeprecated;
+	}
 
-		void SetIsDeprecated(bool bInIsDeprecated)
-		{
-			if (bInIsDeprecated != bIsDeprecated)
-			{
-				bIsDeprecated = bInIsDeprecated;
-				ChangeID = FGuid::NewGuid();
-			}
-		}
+	void SetAuthor(const FText& InAuthor);
+	void SetAutoUpdateManagesInterface(bool bInAutoUpdateManagesInterface);
+	void SetCategoryHierarchy(const TArray<FText>& InCategoryHierarchy);
+	void SetDescription(const FText& InDescription);
+	void SetDisplayName(const FText& InDisplayName);
+	void SetIsDeprecated(bool bInIsDeprecated);
+	void SetKeywords(const TArray<FText>& InKeywords);
+	void SetPromptIfMissing(const FText& InPromptIfMissing);
+	void SetVersion(const FMetasoundFrontendVersionNumber& InVersion);
+
+	void SetType(const EMetasoundFrontendClassType InType)
+	{
+		Type = InType;
+		// TODO: Type is modified while querying and swapped between
+		// to be external, so don't modify the ChangeID in this case.
+		// External/Internal should probably be a separate field.
+		// ChangeID = FGuid::NewGuid();
+	}
 };
 
 USTRUCT()
diff --git a/Engine/Plugins/Runtime/Metasound/Source/MetasoundFrontend/Public/MetasoundNodeRegistrationMacro.h b/Engine/Plugins/Runtime/Metasound/Source/MetasoundFrontend/Public/MetasoundNodeRegistrationMacro.h
--- a/Engine/Plugins/Runtime/Metasound/Source/MetasoundFrontend/Public/MetasoundNodeRegistrationMacro.h
+++ b/Engine/Plugins/Runtime/Metasound/Source/MetasoundFrontend/Public/MetasoundNodeRegistrationMacro.h
@@ -45,7 +45,7 @@ namespace Metasound
 		{
 		public:
 			FNodeRegistryEntry(const Metasound::FNodeClassMetadata& InMetadata)
-			: ClassInfo(InMetadata)
+			: ClassInfo(FMetasoundFrontendClassMetadata::GenerateClassDescription(InMetadata, EMetasoundFrontendClassType::External))
 			, FrontendClass(Metasound::Frontend::GenerateClassDescription(InMetadata)) 
 			{
 			}
