Fix bug with document dependency synchronization by refactoring Document FrontendClass Synchronization to be more clear by renaming as RemoveUnreferencedDependencies & adding option to FindOrAddClass to synchronize if class is found already with registry
#rnx

diff --git a/Engine/Plugins/Runtime/Metasound/Source/MetasoundEditor/Private/MetasoundEditorGraphBuilder.cpp b/Engine/Plugins/Runtime/Metasound/Source/MetasoundEditor/Private/MetasoundEditorGraphBuilder.cpp
--- a/Engine/Plugins/Runtime/Metasound/Source/MetasoundEditor/Private/MetasoundEditorGraphBuilder.cpp
+++ b/Engine/Plugins/Runtime/Metasound/Source/MetasoundEditor/Private/MetasoundEditorGraphBuilder.cpp
@@ -1187,7 +1187,7 @@ namespace Metasound
 
 						if (ensure(GraphHandle->RemoveNode(*NodeHandle)))
 						{
-							GraphHandle->GetOwningDocument()->SynchronizeDependencies();
+							GraphHandle->GetOwningDocument()->RemoveUnreferencedDependencies();
 						}
 					}
 					break;
diff --git a/Engine/Plugins/Runtime/Metasound/Source/MetasoundFrontend/Private/MetasoundFrontendDocumentController.cpp b/Engine/Plugins/Runtime/Metasound/Source/MetasoundFrontend/Private/MetasoundFrontendDocumentController.cpp
--- a/Engine/Plugins/Runtime/Metasound/Source/MetasoundFrontend/Private/MetasoundFrontendDocumentController.cpp
+++ b/Engine/Plugins/Runtime/Metasound/Source/MetasoundFrontend/Private/MetasoundFrontendDocumentController.cpp
@@ -232,20 +232,20 @@ namespace Metasound
 			return DocumentPtr.GetClassWithRegistryKey(InKey);
 		}
 
-		FConstClassAccessPtr FDocumentController::FindOrAddClass(const FNodeRegistryKey& InKey)
+		FConstClassAccessPtr FDocumentController::FindOrAddClass(const FNodeRegistryKey& InKey, bool bInRefreshFromRegistry)
 		{
 			if (FMetasoundFrontendDocument* Document = DocumentPtr.Get())
 			{
 				FConstClassAccessPtr ClassPtr = FindClass(InKey);
 
-				auto AddClass = [=](FMetasoundFrontendClass&& NewClassDescription)
+				auto AddClass = [=](FMetasoundFrontendClass&& NewClassDescription, const FGuid& NewClassID)
 				{
 					FConstClassAccessPtr NewClassPtr;
 
 					// Cannot add a subgraph using this method because dependencies
 					// of external graph are not added in this method.
 					check(EMetasoundFrontendClassType::Graph != NewClassDescription.Metadata.GetType());
-					NewClassDescription.ID = FGuid::NewGuid();
+					NewClassDescription.ID = NewClassID;
 
 					Document->Dependencies.Add(MoveTemp(NewClassDescription));
 					NewClassPtr = FindClass(InKey);
@@ -264,15 +264,22 @@ namespace Metasound
 						FMetasoundFrontendClass NewClass = GenerateClassDescription(InKey);
 						if (NewClass.Metadata.GetVersion().Major != MetasoundClass->Metadata.GetVersion().Major)
 						{
-							return AddClass(MoveTemp(NewClass));
+							return AddClass(MoveTemp(NewClass), FGuid::NewGuid());
 						}
 					}
 
+					if (bInRefreshFromRegistry)
+					{
+						FGuid ClassID = MetasoundClass->ID;
+						FMetasoundFrontendClass NewClass = GenerateClassDescription(InKey);
+						return AddClass(MoveTemp(NewClass), ClassID);
+					}
+
 					return ClassPtr;
 				}
 
 				FMetasoundFrontendClass NewClass = GenerateClassDescription(InKey);
-				return AddClass(MoveTemp(NewClass));
+				return AddClass(MoveTemp(NewClass), FGuid::NewGuid());
 			}
 
 			return FConstClassAccessPtr();
@@ -363,33 +370,7 @@ namespace Metasound
 			return ClassPtr;
 		}
 
-		const FMetasoundFrontendClass* FDocumentController::SynchronizeDependency(const FNodeRegistryKey& InKey)
-		{
-			FMetasoundFrontendDocument* Document = DocumentPtr.Get();
-			if (!ensure(Document))
-			{
-				return nullptr;
-			}
-
-			for (FMetasoundFrontendClass& Class : Document->Dependencies)
-			{
-				const FNodeRegistryKey DependencyRegistryKey = FMetasoundFrontendRegistryContainer::Get()->GetRegistryKey(Class.Metadata);
-				if (InKey == DependencyRegistryKey)
-				{
-					const FGuid ClassID = Class.ID;
-					FMetasoundFrontendClass RegisteredClass = GenerateClassDescription(InKey);
-					RegisteredClass.ID = ClassID;
-					Class = RegisteredClass;
-					return &Class;
-				}
-			}
-
-			FMetasoundFrontendClass NewClass = GenerateClassDescription(InKey);
-			NewClass.ID = FGuid::NewGuid();
-			return &Document->Dependencies.Add_GetRef(NewClass);
-		}
-
-		void FDocumentController::SynchronizeDependencies()
+		void FDocumentController::RemoveUnreferencedDependencies()
 		{
 			if (FMetasoundFrontendDocument* Document = DocumentPtr.Get())
 			{
@@ -422,9 +403,33 @@ namespace Metasound
 					};
 
 					NumDependenciesRemovedThisItr = Document->Dependencies.RemoveAllSwap(IsDependencyUnreferenced);
+				} while (NumDependenciesRemovedThisItr > 0);
+			}
+		}
+
+		TArray<FConstClassAccessPtr> FDocumentController::SynchronizeDependencyMetadata()
+		{
+			TArray<FConstClassAccessPtr> UpdatedClassPtrs;
+
+			if (FMetasoundFrontendDocument* Document = DocumentPtr.Get())
+			{
+				for (FMetasoundFrontendClass& Class : Document->Dependencies)
+				{
+					FNodeRegistryKey RegistryKey = NodeRegistryKey::CreateKey(Class.Metadata);
+
+					FMetasoundFrontendClass RegistryVersion;
+					if (FMetasoundFrontendRegistryContainer::Get()->FindFrontendClassFromRegistered(RegistryKey, RegistryVersion))
+					{
+						if (Class.Metadata.GetChangeID() != RegistryVersion.Metadata.GetChangeID())
+						{
+							Class.Metadata = MoveTemp(RegistryVersion.Metadata);
+							UpdatedClassPtrs.Add(FindClass(RegistryKey));
+						}
+					}
 				}
-				while (NumDependenciesRemovedThisItr > 0);
 			}
+
+			return UpdatedClassPtrs;
 		}
 
 		FGraphHandle FDocumentController::GetRootGraph()
diff --git a/Engine/Plugins/Runtime/Metasound/Source/MetasoundFrontend/Private/MetasoundFrontendDocumentController.h b/Engine/Plugins/Runtime/Metasound/Source/MetasoundFrontend/Private/MetasoundFrontendDocumentController.h
--- a/Engine/Plugins/Runtime/Metasound/Source/MetasoundFrontend/Private/MetasoundFrontendDocumentController.h
+++ b/Engine/Plugins/Runtime/Metasound/Source/MetasoundFrontend/Private/MetasoundFrontendDocumentController.h
@@ -46,7 +46,7 @@ namespace Metasound
 			FConstClassAccessPtr FindClass(const FNodeRegistryKey& InKey) const override;
 			FConstClassAccessPtr FindClass(const FMetasoundFrontendClassMetadata& InMetadata) const override;
 
-			FConstClassAccessPtr FindOrAddClass(const FNodeRegistryKey& InKey) override;
+			FConstClassAccessPtr FindOrAddClass(const FNodeRegistryKey& InKey, bool bInRefreshFromRegistry) override;
 			FConstClassAccessPtr FindOrAddClass(const FMetasoundFrontendClassMetadata& InMetadata) override;
 
 			FGraphHandle AddDuplicateSubgraph(const IGraphController& InGraph) override;
@@ -59,8 +59,8 @@ namespace Metasound
 			void SetMetadata(const FMetasoundFrontendDocumentMetadata& InMetadata) override;
 			const FMetasoundFrontendDocumentMetadata& GetMetadata() const override;
 
-			const FMetasoundFrontendClass* SynchronizeDependency(const FNodeRegistryKey& InKey) override;
-			void SynchronizeDependencies() override;
+			void RemoveUnreferencedDependencies() override;
+			TArray<FConstClassAccessPtr> SynchronizeDependencyMetadata() override;
 
 			FGraphHandle GetRootGraph() override;
 			FConstGraphHandle GetRootGraph() const override;
diff --git a/Engine/Plugins/Runtime/Metasound/Source/MetasoundFrontend/Private/MetasoundFrontendGraphController.cpp b/Engine/Plugins/Runtime/Metasound/Source/MetasoundFrontend/Private/MetasoundFrontendGraphController.cpp
--- a/Engine/Plugins/Runtime/Metasound/Source/MetasoundFrontend/Private/MetasoundFrontendGraphController.cpp
+++ b/Engine/Plugins/Runtime/Metasound/Source/MetasoundFrontend/Private/MetasoundFrontendGraphController.cpp
@@ -664,7 +664,7 @@ namespace Metasound
 				GraphClass->Interface.Inputs.Reset();
 				GraphClass->Interface.Outputs.Reset();
 				OwningDocument->ClearInterfaceVersions();
-				OwningDocument->SynchronizeDependencies();
+				OwningDocument->RemoveUnreferencedDependencies();
 			}
 		}
 
@@ -1398,7 +1398,7 @@ namespace Metasound
 
 				NumRemoved += GraphClass->Graph.Nodes.RemoveAll(IsNodeWithID);
 
-				OwningDocument->SynchronizeDependencies();
+				OwningDocument->RemoveUnreferencedDependencies();
 
 				return (NumRemoved > 0);
 			}
diff --git a/Engine/Plugins/Runtime/Metasound/Source/MetasoundFrontend/Private/MetasoundFrontendInvalidController.h b/Engine/Plugins/Runtime/Metasound/Source/MetasoundFrontend/Private/MetasoundFrontendInvalidController.h
--- a/Engine/Plugins/Runtime/Metasound/Source/MetasoundFrontend/Private/MetasoundFrontendInvalidController.h
+++ b/Engine/Plugins/Runtime/Metasound/Source/MetasoundFrontend/Private/MetasoundFrontendInvalidController.h
@@ -407,7 +407,7 @@ namespace Metasound
 				virtual FConstClassAccessPtr FindClassWithID(FGuid InClassID) const override { return FConstClassAccessPtr(); }
 
 				virtual FConstClassAccessPtr FindClass(const FNodeRegistryKey& InKey) const override { return FConstClassAccessPtr(); }
-				virtual FConstClassAccessPtr FindOrAddClass(const FNodeRegistryKey& InKey) override { return FConstClassAccessPtr(); }
+				virtual FConstClassAccessPtr FindOrAddClass(const FNodeRegistryKey& InKey, bool bInRefreshFromRegistry) override { return FConstClassAccessPtr(); }
 				virtual FConstClassAccessPtr FindClass(const FMetasoundFrontendClassMetadata& InMetadata) const override{ return FConstClassAccessPtr(); }
 				virtual FConstClassAccessPtr FindOrAddClass(const FMetasoundFrontendClassMetadata& InMetadata) override{ return FConstClassAccessPtr(); }
 				virtual FGraphHandle AddDuplicateSubgraph(const IGraphController& InGraph) override { return IGraphController::GetInvalidHandle(); }
@@ -420,8 +420,8 @@ namespace Metasound
 				virtual void SetMetadata(const FMetasoundFrontendDocumentMetadata& InMetadata) override { }
 				virtual const FMetasoundFrontendDocumentMetadata& GetMetadata() const override { return Invalid::GetInvalidDocumentMetadata(); }
 
-				virtual const FMetasoundFrontendClass* SynchronizeDependency(const FNodeRegistryKey& InKey) override { return nullptr; }
-				virtual void SynchronizeDependencies() override { }
+				virtual void RemoveUnreferencedDependencies() override { }
+				virtual TArray<FConstClassAccessPtr> SynchronizeDependencyMetadata() override { return { }; }
 
 				virtual TArray<FGraphHandle> GetSubgraphHandles() override { return TArray<FGraphHandle>(); }
 
diff --git a/Engine/Plugins/Runtime/Metasound/Source/MetasoundFrontend/Private/MetasoundFrontendNodeController.cpp b/Engine/Plugins/Runtime/Metasound/Source/MetasoundFrontend/Private/MetasoundFrontendNodeController.cpp
--- a/Engine/Plugins/Runtime/Metasound/Source/MetasoundFrontend/Private/MetasoundFrontendNodeController.cpp
+++ b/Engine/Plugins/Runtime/Metasound/Source/MetasoundFrontend/Private/MetasoundFrontendNodeController.cpp
@@ -795,7 +795,12 @@ namespace Metasound
 				});
 			});
 
+			// Maintain the NodeID to ensure the editor layer maintains node representation parity.
 			const FGuid ReplacedNodeGuid = GetID();
+
+			// Maintain class ID to ensure all nodes created & preexisting reference the same dependency.
+			const FGuid ClassID = GetClassID();
+
 			if (!ensureAlways(GetOwningGraph()->RemoveNode(*this)))
 			{
 				return this->AsShared();
@@ -804,9 +809,11 @@ namespace Metasound
 			// Make sure classes are up-to-date with registered versions of class.
 			// Note that this may break other nodes in the graph that have stale
 			// class API, but that's on the caller to fix-up or report invalid state.
-			const FNodeRegistryKey RegistryKey = FMetasoundFrontendRegistryContainer::Get()->GetRegistryKey(RegisteredClass->Metadata);
+			const FNodeRegistryKey RegistryKey = NodeRegistryKey::CreateKey(RegisteredClass->Metadata);
 			FDocumentHandle Document = GetOwningGraph()->GetOwningDocument();
-			ensureAlways(Document->SynchronizeDependency(RegistryKey) != nullptr);
+
+			constexpr bool bRefreshFromRegistry = true;
+			ensureAlways(Document->FindOrAddClass(RegistryKey, bRefreshFromRegistry).Get() != nullptr);
 
 			FNodeHandle ReplacementNode = GetOwningGraph()->AddNode(RegisteredClass->Metadata, ReplacedNodeGuid);
 			if (!ensureAlways(ReplacementNode->IsValid()))
diff --git a/Engine/Plugins/Runtime/Metasound/Source/MetasoundFrontend/Private/MetasoundFrontendTransform.cpp b/Engine/Plugins/Runtime/Metasound/Source/MetasoundFrontend/Private/MetasoundFrontendTransform.cpp
--- a/Engine/Plugins/Runtime/Metasound/Source/MetasoundFrontend/Private/MetasoundFrontendTransform.cpp
+++ b/Engine/Plugins/Runtime/Metasound/Source/MetasoundFrontend/Private/MetasoundFrontendTransform.cpp
@@ -511,7 +511,8 @@ namespace Metasound
 							ParentMetaSoundAsset->ConformObjectDataToInterfaces();
 						}
 
-						InDocument->SynchronizeDependencies();
+						InDocument->RemoveUnreferencedDependencies();
+						InDocument->SynchronizeDependencyMetadata();
 					}
 				}
 			}
@@ -528,7 +529,8 @@ namespace Metasound
 					NewNode->SetNodeStyle(Style);
 				}
 
-				InDocument->SynchronizeDependencies();
+				InDocument->RemoveUnreferencedDependencies();
+				InDocument->SynchronizeDependencyMetadata();
 			}
 
 			return bDidEdit;
diff --git a/Engine/Plugins/Runtime/Metasound/Source/MetasoundFrontend/Public/MetasoundFrontendController.h b/Engine/Plugins/Runtime/Metasound/Source/MetasoundFrontend/Public/MetasoundFrontendController.h
--- a/Engine/Plugins/Runtime/Metasound/Source/MetasoundFrontend/Public/MetasoundFrontendController.h
+++ b/Engine/Plugins/Runtime/Metasound/Source/MetasoundFrontend/Public/MetasoundFrontendController.h
@@ -955,11 +955,12 @@ namespace Metasound
 			/** Attempts to find an existing Metasound class description corresponding
 			 * to a dependency which matches the provided class information. If the
 			 * class is not found in the current dependencies, it is added to the 
-			 * dependencies.
+			 * dependencies. If found and bInRefreshFromRegistry is set (optional)
+			 * , copies the version found in the registry to the document.
 			 *
 			 * @return A pointer to the object, or nullptr on error.
 			 */
-			virtual FConstClassAccessPtr FindOrAddClass(const FNodeRegistryKey& InKey) = 0; 
+			virtual FConstClassAccessPtr FindOrAddClass(const FNodeRegistryKey& InKey, bool bInRefreshFromRegistry = false) = 0;
 
 			/** Returns an existing Metasound class description corresponding to 
 			 * a dependency which matches the provided class information.
@@ -991,15 +992,14 @@ namespace Metasound
 			virtual void RemoveInterfaceVersion(const FMetasoundFrontendVersion& InVersion) = 0;
 			virtual void ClearInterfaceVersions() = 0;
 
-			/** Updates dependency with given registry key if found with that in the active registry
-			  *
-			  * @return True if found (and subsequently updated), false if not found.
+			/** Removes all dependencies which are no longer referenced by any graphs within this document.
 			  */
-			virtual const FMetasoundFrontendClass* SynchronizeDependency(const FNodeRegistryKey& InKey) = 0;
+			virtual void RemoveUnreferencedDependencies() = 0;
 
-			/** Removes all dependencies which are no longer referenced by any graphs within this document
+			/** Synchronizes all dependency Metadata in document with that found in the registry. If not found
+			  * in the registry, no action taken.  Returns array of pointers to classes that were modified.
 			  */
-			virtual void SynchronizeDependencies() = 0;
+			virtual TArray<FConstClassAccessPtr> SynchronizeDependencyMetadata() = 0;
 
 			/** Returns an array of all subgraphs for this document. */
 			virtual TArray<FGraphHandle> GetSubgraphHandles() = 0;
