Prevent creating reference loops when dragging and dropping MetaSounds
- Also show appropriate messaging and cancel transaction if invalid

diff --git a/Engine/Plugins/Runtime/Metasound/Source/MetasoundEditor/Private/MetasoundEditorGraphSchema.cpp b/Engine/Plugins/Runtime/Metasound/Source/MetasoundEditor/Private/MetasoundEditorGraphSchema.cpp
--- a/Engine/Plugins/Runtime/Metasound/Source/MetasoundEditor/Private/MetasoundEditorGraphSchema.cpp
+++ b/Engine/Plugins/Runtime/Metasound/Source/MetasoundEditor/Private/MetasoundEditorGraphSchema.cpp
@@ -1361,6 +1361,7 @@ void UMetasoundEditorGraphSchema::BreakPinLinks(UEdGraphPin& TargetPin, bool bSe
 
 void UMetasoundEditorGraphSchema::GetAssetsGraphHoverMessage(const TArray<FAssetData>& Assets, const UEdGraph* HoverGraph, FString& OutTooltipText, bool& OutOkIcon) const
 {
+	using namespace Metasound;
 	using namespace Metasound::Editor;
 
 	OutOkIcon = true;
@@ -1382,6 +1383,36 @@ void UMetasoundEditorGraphSchema::GetAssetsGraphHoverMessage(const TArray<FAsset
 			OutTooltipText = TEXT("Asset(s) must all be MetaSounds.");
 			break;
 		}
+
+		const UMetasoundEditorGraph* MetaSoundGraph = CastChecked<UMetasoundEditorGraph>(HoverGraph);
+		const UObject& MetaSound = MetaSoundGraph->GetMetasoundChecked();
+
+		const FMetasoundAssetBase* MetaSoundAsset = IMetasoundUObjectRegistry::Get().GetObjectAsAssetBase(&MetaSound);
+		check(MetaSoundAsset);
+
+		if (UObject* DroppedObject = Data.GetAsset())
+		{
+			FMetasoundAssetBase* DroppedMetaSoundAsset = IMetasoundUObjectRegistry::Get().GetObjectAsAssetBase(DroppedObject);
+			if (!DroppedMetaSoundAsset)
+			{
+				OutOkIcon = false;
+				OutTooltipText = TEXT("Asset is not a valid MetaSound.");
+				break;
+			}
+
+			if (MetaSoundAsset->AddingReferenceCausesLoop(FSoftObjectPath(Data.ObjectPath)))
+			{
+				OutOkIcon = false;
+				OutTooltipText = TEXT("Cannot add an asset that would create a reference loop.");
+				break;
+			}
+		}
+		else
+		{
+			OutOkIcon = false;
+			OutTooltipText = TEXT("Asset not found.");
+			break;
+		}
 	}
 }
 
@@ -1390,12 +1421,12 @@ void UMetasoundEditorGraphSchema::DroppedAssetsOnGraph(const TArray<FAssetData>&
 	using namespace Metasound;
 	using namespace Metasound::Frontend;
 
-	const FScopedTransaction Transaction(LOCTEXT("DropMetaSoundOnGraph", "Drop MetaSound On Graph"));
+	FScopedTransaction Transaction(LOCTEXT("DropMetaSoundOnGraph", "Drop MetaSound On Graph"));
 
 	UMetasoundEditorGraph* MetaSoundGraph = CastChecked<UMetasoundEditorGraph>(Graph);
+	bool bTransactionSucceeded = false;
+	bool bModifiedObjects = false;
 	UObject& MetaSound = MetaSoundGraph->GetMetasoundChecked();
-	MetaSound.Modify();
-	Graph->Modify();
 
 	FMetasoundAssetBase* MetaSoundAsset = IMetasoundUObjectRegistry::Get().GetObjectAsAssetBase(&MetaSound);
 	check(MetaSoundAsset);
@@ -1415,6 +1446,13 @@ void UMetasoundEditorGraphSchema::DroppedAssetsOnGraph(const TArray<FAssetData>&
 				continue;
 			}
 
+			if (!bModifiedObjects)
+			{
+				MetaSound.Modify();
+				Graph->Modify();
+				bModifiedObjects = true;
+			}
+
 			// This may not be necessary as dropping an asset on the graph may load it, thus triggering the registration from the MetaSoundAssetManager.
 			const FNodeClassInfo ClassInfo = DroppedMetaSoundAsset->GetAssetClassInfo();
 			const FNodeRegistryKey RegistryKey = NodeRegistryKey::CreateKey(ClassInfo);
@@ -1431,9 +1469,22 @@ void UMetasoundEditorGraphSchema::DroppedAssetsOnGraph(const TArray<FAssetData>&
 			if (ensure(ISearchEngine::Get().FindClassWithHighestVersion(ClassName.ToNodeClassName(), Class)))
 			{
 				Metasound::Editor::FGraphBuilder::AddExternalNode(MetaSound, Class.Metadata, GraphPosition);
+				bTransactionSucceeded = true;
 			}
 		}
 	}
+
+	if (bTransactionSucceeded)
+	{
+		// Reregister this asset after adding the dropped asset to update references
+		FMetaSoundAssetRegistrationOptions RegOptions;
+		RegOptions.bForceReregister = true;
+		MetaSoundAsset->RegisterGraphWithFrontend(RegOptions);
+	}
+	else
+	{
+		Transaction.Cancel();
+	}
 }
 
 void UMetasoundEditorGraphSchema::DroppedAssetsOnNode(const TArray<FAssetData>& Assets, const FVector2D& GraphPosition, UEdGraphNode* Node) const
