Don't serialize cached dependency data on MetaSound graphs to avoid text desync or out-of-date data when autoupdate is disabled for a given asset.
TODO: Split out FMetasoundFrontendClass to class & dependency struct, where dependency is only the information required to find an item in the registry and potentially a couple non-localized breadcrumbs
#rnx

diff --git a/Engine/Plugins/Runtime/Metasound/Source/MetasoundEditor/Private/MetasoundEditor.cpp b/Engine/Plugins/Runtime/Metasound/Source/MetasoundEditor/Private/MetasoundEditor.cpp
--- a/Engine/Plugins/Runtime/Metasound/Source/MetasoundEditor/Private/MetasoundEditor.cpp
+++ b/Engine/Plugins/Runtime/Metasound/Source/MetasoundEditor/Private/MetasoundEditor.cpp
@@ -2257,6 +2257,8 @@ namespace Metasound
 
 			if (MetasoundAsset->GetSynchronizationRequired())
 			{
+				MetasoundAsset->CacheDependencyRegistryData();
+
 				// Capture before synchronizing as the flag is cleared therein.
 				const bool bShouldRefreshDetails = MetasoundAsset->GetSynchronizationUpdateDetails();
 				FGraphBuilder::SynchronizeGraph(*Metasound);
diff --git a/Engine/Plugins/Runtime/Metasound/Source/MetasoundEngine/Public/Metasound.h b/Engine/Plugins/Runtime/Metasound/Source/MetasoundEngine/Public/Metasound.h
--- a/Engine/Plugins/Runtime/Metasound/Source/MetasoundEngine/Public/Metasound.h
+++ b/Engine/Plugins/Runtime/Metasound/Source/MetasoundEngine/Public/Metasound.h
@@ -96,6 +96,25 @@ namespace Metasound
 		{
 			InMetaSound.VersionAsset();
 		}
+
+		// Clear dependency registry data both when saving and loading to clear
+		// out any fields (ex. text) that should not be serialized on the asset.
+		// Once a save is run on all pre-5.0 generated assets, this can be safely
+		// only run during *saving* persistent, non-transactional data.
+		if (InArchive.IsPersistent() && !InArchive.IsTransacting())
+		{
+			InMetaSound.ClearDependencyRegistryData();
+		}
+
+#if WITH_EDITOR
+		// When the editor is loaded, flag for synchronization required
+		// so that the registry data is reloaded post serialization for
+		// any open editors.
+		if (InArchive.IsSaving())
+		{
+			InMetaSound.SetSynchronizationRequired();
+		}
+#endif
 	}
 
 #if WITH_EDITORONLY_DATA
diff --git a/Engine/Plugins/Runtime/Metasound/Source/MetasoundFrontend/Private/MetasoundAssetBase.cpp b/Engine/Plugins/Runtime/Metasound/Source/MetasoundFrontend/Private/MetasoundAssetBase.cpp
--- a/Engine/Plugins/Runtime/Metasound/Source/MetasoundFrontend/Private/MetasoundAssetBase.cpp
+++ b/Engine/Plugins/Runtime/Metasound/Source/MetasoundFrontend/Private/MetasoundAssetBase.cpp
@@ -122,6 +122,11 @@ void FMetasoundAssetBase::RegisterGraphWithFrontend(Metasound::Frontend::FMetaSo
 #endif // WITH_EDITORONLY_DATA
 	}
 
+	// Caches commonly used class MetaData that isn't required for finding & building dependencies into local graph.
+	// Must be completed after auto-update to ensure all non-transient referenced dependency data is up-to-date (ex.
+	// class version), which is required for most accurately caching current registry metadata.
+	CacheDependencyRegistryData();
+
 	// Registers node by copying document. Updates to document require re-registration.
 	class FNodeRegistryEntry : public INodeRegistryEntry
 	{
@@ -331,15 +336,28 @@ void FMetasoundAssetBase::AddDefaultInterfaces()
 	FModifyRootGraphInterfaces({ }, InitInterfaces).Transform(DocumentHandle);
 }
 
- 
- 
- 
- 
- 
- 
- 
- 
- 
+void FMetasoundAssetBase::CacheDependencyRegistryData()
+{
+	if (FMetasoundFrontendDocument* Document = GetDocument().Get())
+	{
+		for (FMetasoundFrontendClass& Class : Document->Dependencies)
+		{
+			Class.CacheRegistryData();
+		}
+	}
+}
+
+void FMetasoundAssetBase::ClearDependencyRegistryData()
+{
+	if (FMetasoundFrontendDocument* Document = GetDocument().Get())
+	{
+		for (FMetasoundFrontendClass& Class : Document->Dependencies)
+		{
+			Class.ClearRegistryData();
+		}
+	}
+}
+
 bool FMetasoundAssetBase::VersionAsset()
 {
 	using namespace Metasound;
diff --git a/Engine/Plugins/Runtime/Metasound/Source/MetasoundFrontend/Private/MetasoundFrontendDocument.cpp b/Engine/Plugins/Runtime/Metasound/Source/MetasoundFrontend/Private/MetasoundFrontendDocument.cpp
--- a/Engine/Plugins/Runtime/Metasound/Source/MetasoundFrontend/Private/MetasoundFrontendDocument.cpp
+++ b/Engine/Plugins/Runtime/Metasound/Source/MetasoundFrontend/Private/MetasoundFrontendDocument.cpp
@@ -2,8 +2,10 @@
 
 #include "MetasoundFrontendDocument.h"
 
+#include "Algo/ForEach.h"
 #include "Algo/Transform.h"
 #include "IAudioParameterInterfaceRegistry.h"
+#include "Logging/LogMacros.h"
 #include "MetasoundFrontend.h"
 #include "MetasoundFrontendRegistries.h"
 #include "MetasoundLog.h"
@@ -239,6 +241,74 @@ void FMetasoundFrontendClassMetadata::SetVersion(const FMetasoundFrontendVersion
 	SetWithChangeID(InVersion, Version, ChangeID);
 }
 
+void FMetasoundFrontendClass::CacheRegistryData()
+{
+	using namespace Metasound::Frontend;
+
+	const FNodeRegistryKey Key = NodeRegistryKey::CreateKey(Metadata);
+	FMetasoundFrontendClass Class;
+
+	FMetasoundFrontendRegistryContainer* Registry = FMetasoundFrontendRegistryContainer::Get();
+	if (ensure(Registry))
+	{
+		if (Registry->FindFrontendClassFromRegistered(Key, Class))
+		{
+			Metadata = Class.Metadata;
+
+			using FNameTypeKey = TPair<FName, FName>;
+			TMap<FNameTypeKey, const FMetasoundFrontendVertexMetadata*> InterfaceMembers;
+
+			auto MakePairFromVertex = [](const FMetasoundFrontendClassVertex& InVertex)
+			{
+				const FNameTypeKey Key(InVertex.Name, InVertex.TypeName);
+				return TPair<FNameTypeKey, const FMetasoundFrontendVertexMetadata*> { Key, &InVertex.Metadata };
+			};
+
+			auto CacheRegistryVertexMetadata = [&](FMetasoundFrontendClassVertex& OutVertex)
+			{
+				const FNameTypeKey Key(OutVertex.Name, OutVertex.TypeName);
+				if (const FMetasoundFrontendVertexMetadata* RegVertex = InterfaceMembers.FindRef(Key))
+				{
+					OutVertex.Metadata = *RegVertex;
+				}
+			};
+
+			Algo::Transform(Class.Interface.Inputs, InterfaceMembers, [&](const FMetasoundFrontendClassInput& Input) { return MakePairFromVertex(Input); });
+			Algo::ForEach(Interface.Inputs, [&](FMetasoundFrontendClassInput& Input) { CacheRegistryVertexMetadata(Input); });
+
+			InterfaceMembers.Reset();
+
+			Algo::Transform(Class.Interface.Outputs, InterfaceMembers, [&](const FMetasoundFrontendClassOutput& Output) { return MakePairFromVertex(Output); });
+			Algo::ForEach(Interface.Outputs, [&](FMetasoundFrontendClassOutput& Output) { CacheRegistryVertexMetadata(Output); });
+		}
+		else
+		{
+			UE_LOG(LogMetaSound, Error, TEXT("Failed to load document dependency class metadata: Missing dependency with key '%s'"), *Key);
+		}
+
+		Interface.InputStyle = Class.Interface.InputStyle;
+		Interface.OutputStyle = Class.Interface.OutputStyle;
+		Style = Class.Style;
+	}
+}
+
+void FMetasoundFrontendClass::ClearRegistryData()
+{
+	Metadata.DisplayName = { };
+	Metadata.Description = { };
+	Metadata.PromptIfMissing = { };
+	Metadata.Author = { };
+	Metadata.Keywords.Reset();
+	Metadata.CategoryHierarchy.Reset();
+
+	Algo::ForEach(Interface.Inputs, [](FMetasoundFrontendClassInput& Input) { Input.Metadata = { }; });
+	Algo::ForEach(Interface.Outputs, [](FMetasoundFrontendClassOutput& Output) { Output.Metadata = { }; });
+
+	Interface.InputStyle = { };
+	Interface.OutputStyle = { };
+	Style = { };
+}
+
 FMetasoundFrontendClassMetadata FMetasoundFrontendClassMetadata::GenerateClassDescription(const Metasound::FNodeClassMetadata& InNodeClassMetadata, EMetasoundFrontendClassType InType)
 {
 	FMetasoundFrontendClassMetadata NewMetadata;
diff --git a/Engine/Plugins/Runtime/Metasound/Source/MetasoundFrontend/Private/MetasoundFrontendTransform.cpp b/Engine/Plugins/Runtime/Metasound/Source/MetasoundFrontend/Private/MetasoundFrontendTransform.cpp
--- a/Engine/Plugins/Runtime/Metasound/Source/MetasoundFrontend/Private/MetasoundFrontendTransform.cpp
+++ b/Engine/Plugins/Runtime/Metasound/Source/MetasoundFrontend/Private/MetasoundFrontendTransform.cpp
@@ -470,7 +470,7 @@ namespace Metasound
 					PresetReferencedMetaSoundAsset = IMetaSoundAssetManager::GetChecked().TryLoadAssetFromKey(RegistryKey);
 					if (!PresetReferencedMetaSoundAsset)
 					{
-						UE_LOG(LogMetaSound, Error, TEXT("Auto-Updating preset '%s' failed: Referenced class '%s' missing."), *AssetPath, *ClassMetadata.GetDisplayName().ToString());
+						UE_LOG(LogMetaSound, Error, TEXT("Auto-Updating preset '%s' failed: Referenced class '%s' missing."), *AssetPath, *ClassMetadata.GetClassName().ToString());
 					}
 					return;
 				}
@@ -484,12 +484,12 @@ namespace Metasound
 				FMetasoundFrontendVersionNumber UpdateVersion = NodeHandle->FindHighestMinorVersionInRegistry();
 				if (UpdateVersion.IsValid() && UpdateVersion > ClassMetadata.GetVersion())
 				{
-					UE_LOG(LogMetaSound, Display, TEXT("Auto-Updating '%s' node class '%s': Newer version '%s' found."), *AssetPath, *ClassMetadata.GetDisplayName().ToString(), *UpdateVersion.ToString());
+					UE_LOG(LogMetaSound, Display, TEXT("Auto-Updating '%s' node class '%s': Newer version '%s' found."), *AssetPath, *ClassMetadata.GetClassName().ToString(), *UpdateVersion.ToString());
 				}
 				else if (InterfaceUpdates.ContainsChanges())
 				{
 					UpdateVersion = ClassMetadata.GetVersion();
-					UE_LOG(LogMetaSound, Display, TEXT("Auto-Updating '%s' node class '%s (%s)': Interface change detected."), *AssetPath, *ClassMetadata.GetDisplayName().ToString(), *UpdateVersion.ToString());
+					UE_LOG(LogMetaSound, Display, TEXT("Auto-Updating '%s' node class '%s (%s)': Interface change detected."), *AssetPath, *ClassMetadata.GetClassName().ToString(), *UpdateVersion.ToString());
 				}
 
 				NodesToUpdate.Add(TPair<FNodeHandle, FMetasoundFrontendVersionNumber>(NodeHandle, UpdateVersion));
diff --git a/Engine/Plugins/Runtime/Metasound/Source/MetasoundFrontend/Public/MetasoundAssetBase.h b/Engine/Plugins/Runtime/Metasound/Source/MetasoundFrontend/Public/MetasoundAssetBase.h
--- a/Engine/Plugins/Runtime/Metasound/Source/MetasoundFrontend/Public/MetasoundAssetBase.h
+++ b/Engine/Plugins/Runtime/Metasound/Source/MetasoundFrontend/Public/MetasoundAssetBase.h
@@ -122,6 +122,8 @@ public:
 	const FMetasoundFrontendDocument& GetDocumentChecked() const;
 
 	void AddDefaultInterfaces();
+	void CacheDependencyRegistryData();
+	void ClearDependencyRegistryData();
 	bool VersionAsset();
 
 #if WITH_EDITORONLY_DATA
diff --git a/Engine/Plugins/Runtime/Metasound/Source/MetasoundFrontend/Public/MetasoundFrontendDocument.h b/Engine/Plugins/Runtime/Metasound/Source/MetasoundFrontend/Public/MetasoundFrontendDocument.h
--- a/Engine/Plugins/Runtime/Metasound/Source/MetasoundFrontend/Public/MetasoundFrontendDocument.h
+++ b/Engine/Plugins/Runtime/Metasound/Source/MetasoundFrontend/Public/MetasoundFrontendDocument.h
@@ -824,6 +824,9 @@ public:
 	{
 		ChangeID = FGuid::NewGuid();
 	}
+
+	// Required to allow caching registry data without modifying the ChangeID
+	friend struct FMetasoundFrontendClass;
 };
 
 
@@ -1040,6 +1043,9 @@ public:
 		// External/Internal should probably be a separate field.
 		// ChangeID = FGuid::NewGuid();
 	}
+
+	// Required to allow caching registry data without modifying the ChangeID
+	friend struct FMetasoundFrontendClass;
 };
 
 USTRUCT()
@@ -1051,18 +1057,6 @@ struct FMetasoundFrontendClassStyle
 	FMetasoundFrontendClassStyleDisplay Display;
 };
 
-USTRUCT()
-struct FMetasoundFrontendEditorData
-{
-	GENERATED_BODY()
-
-	UPROPERTY()
-	FMetasoundFrontendVersion Version;
-
-	UPROPERTY()
-	TArray<uint8> Data;
-};
-
 USTRUCT()
 struct METASOUNDFRONTEND_API FMetasoundFrontendClass
 {
@@ -1079,11 +1073,22 @@ struct METASOUNDFRONTEND_API FMetasoundFrontendClass
 	UPROPERTY(EditAnywhere, Category = CustomView)
 	FMetasoundFrontendClassInterface Interface;
 
-	UPROPERTY()
-	FMetasoundFrontendEditorData EditorData;
-
 	UPROPERTY()
 	FMetasoundFrontendClassStyle Style;
+
+	/*
+	 * Caches metadata found in the registry on the class
+	 * that is not necessary for serialization or core graph
+	 * generation.
+	 */
+	void CacheRegistryData();
+
+	/*
+	 * Clears cached metadata found in the registry on the class
+	 * that is not necessary for serialization or core graph
+	 * generation.
+	 */
+	void ClearRegistryData();
 };
 
 USTRUCT()
@@ -1121,9 +1126,6 @@ struct METASOUNDFRONTEND_API FMetasoundFrontendDocument
 	UPROPERTY(EditAnywhere, Category = Metadata)
 	FMetasoundFrontendDocumentMetadata Metadata;
 
-	UPROPERTY()
-	FMetasoundFrontendEditorData EditorData;
-
 public:
 	UPROPERTY(VisibleAnywhere, Category = CustomView)
 	TSet<FMetasoundFrontendVersion> Interfaces;
