Move error to warning & add more information regarding asset missing dependency
#rnx

diff --git a/Engine/Plugins/Runtime/Metasound/Source/MetasoundFrontend/Private/MetasoundAssetBase.cpp b/Engine/Plugins/Runtime/Metasound/Source/MetasoundFrontend/Private/MetasoundAssetBase.cpp
--- a/Engine/Plugins/Runtime/Metasound/Source/MetasoundFrontend/Private/MetasoundAssetBase.cpp
+++ b/Engine/Plugins/Runtime/Metasound/Source/MetasoundFrontend/Private/MetasoundAssetBase.cpp
@@ -342,7 +342,15 @@ void FMetasoundAssetBase::CacheDependencyRegistryData()
 	{
 		for (FMetasoundFrontendClass& Class : Document->Dependencies)
 		{
-			Class.CacheRegistryData();
+			if (!Class.CacheRegistryData())
+			{
+				UE_LOG(LogMetaSound, Warning,
+					TEXT("'%s' failed to cache dependency registry data: Registry missing class with key '%s'"),
+					*GetOwningAssetName(),
+					*Class.Metadata.GetClassName().ToString());
+				UE_LOG(LogMetaSound, Warning,
+					TEXT("Asset may fail to build runtime graph unless re-registered after dependency with given key is loaded."));
+			}
 		}
 	}
 }
diff --git a/Engine/Plugins/Runtime/Metasound/Source/MetasoundFrontend/Private/MetasoundFrontendDocument.cpp b/Engine/Plugins/Runtime/Metasound/Source/MetasoundFrontend/Private/MetasoundFrontendDocument.cpp
--- a/Engine/Plugins/Runtime/Metasound/Source/MetasoundFrontend/Private/MetasoundFrontendDocument.cpp
+++ b/Engine/Plugins/Runtime/Metasound/Source/MetasoundFrontend/Private/MetasoundFrontendDocument.cpp
@@ -241,7 +241,7 @@ void FMetasoundFrontendClassMetadata::SetVersion(const FMetasoundFrontendVersion
 	SetWithChangeID(InVersion, Version, ChangeID);
 }
 
-void FMetasoundFrontendClass::CacheRegistryData()
+bool FMetasoundFrontendClass::CacheRegistryData()
 {
 	using namespace Metasound::Frontend;
 
@@ -280,16 +280,16 @@ void FMetasoundFrontendClass::CacheRegistryData()
 
 			Algo::Transform(Class.Interface.Outputs, InterfaceMembers, [&](const FMetasoundFrontendClassOutput& Output) { return MakePairFromVertex(Output); });
 			Algo::ForEach(Interface.Outputs, [&](FMetasoundFrontendClassOutput& Output) { CacheRegistryVertexMetadata(Output); });
-		}
-		else
-		{
-			UE_LOG(LogMetaSound, Error, TEXT("Failed to load document dependency class metadata: Missing dependency with key '%s'"), *Key);
-		}
 
-		Interface.InputStyle = Class.Interface.InputStyle;
-		Interface.OutputStyle = Class.Interface.OutputStyle;
-		Style = Class.Style;
+			Interface.InputStyle = Class.Interface.InputStyle;
+			Interface.OutputStyle = Class.Interface.OutputStyle;
+			Style = Class.Style;
+
+			return true;
+		}
 	}
+
+	return false;
 }
 
 void FMetasoundFrontendClass::ClearRegistryData()
diff --git a/Engine/Plugins/Runtime/Metasound/Source/MetasoundFrontend/Public/MetasoundFrontendDocument.h b/Engine/Plugins/Runtime/Metasound/Source/MetasoundFrontend/Public/MetasoundFrontendDocument.h
--- a/Engine/Plugins/Runtime/Metasound/Source/MetasoundFrontend/Public/MetasoundFrontendDocument.h
+++ b/Engine/Plugins/Runtime/Metasound/Source/MetasoundFrontend/Public/MetasoundFrontendDocument.h
@@ -1080,8 +1080,10 @@ struct METASOUNDFRONTEND_API FMetasoundFrontendClass
 	 * Caches metadata found in the registry on the class
 	 * that is not necessary for serialization or core graph
 	 * generation.
+	 *
+	 * @return - Whether class was found in the registry & data was cached successfully.
 	 */
-	void CacheRegistryData();
+	bool CacheRegistryData();
 
 	/*
 	 * Clears cached metadata found in the registry on the class
