- Fix for not clearing MetaSound document dependency registry cache data when loading in loc builds
- Clean-up duplicate logic not needed as covered in PreSave template function
#rnx

diff --git a/Engine/Plugins/Runtime/Metasound/Source/MetasoundEngine/Public/Metasound.h b/Engine/Plugins/Runtime/Metasound/Source/MetasoundEngine/Public/Metasound.h
--- a/Engine/Plugins/Runtime/Metasound/Source/MetasoundEngine/Public/Metasound.h
+++ b/Engine/Plugins/Runtime/Metasound/Source/MetasoundEngine/Public/Metasound.h
@@ -57,9 +57,12 @@ namespace Metasound
 		// If not editor and cooking, no need to force re-registering as if its already registered,
 		// no edits have been made between last registration and this one if a parent asset has requested
 		// it to register.
-		FMetaSoundAssetRegistrationOptions NonEditorCookRegistrationOptions;
-		NonEditorCookRegistrationOptions.bForceReregister = false;
-		NonEditorCookRegistrationOptions.bRegisterDependencies = true;
+		FMetaSoundAssetRegistrationOptions NonEditorRegistrationOptions;
+		NonEditorRegistrationOptions.bForceReregister = false;
+		NonEditorRegistrationOptions.bRegisterDependencies = true;
+
+		NonEditorRegistrationOptions.bCacheDependencyMetaDataFromRegistry = false;
+		InMetaSound.ClearDependencyRegistryData();
 
 #if WITH_EDITORONLY_DATA
 		if (UMetasoundEditorGraphBase* MetaSoundGraph = Cast<UMetasoundEditorGraphBase>(InMetaSound.GetGraph()))
@@ -74,7 +77,7 @@ namespace Metasound
 				}
 				else
 				{
-					InMetaSound.RegisterGraphWithFrontend(NonEditorCookRegistrationOptions);
+					InMetaSound.RegisterGraphWithFrontend(NonEditorRegistrationOptions);
 				}
 				MetaSoundGraph->SetSynchronizationRequired(false /* bInClearUpdateNotes */);
 			}
@@ -85,7 +88,7 @@ namespace Metasound
 			}
 		}
 #else
-		InMetaSound.RegisterGraphWithFrontend(NonEditorCookRegistrationOptions);
+		InMetaSound.RegisterGraphWithFrontend(NonEditorRegistrationOptions);
 #endif // WITH_EDITORONLY_DATA
 	}
 
@@ -95,26 +98,13 @@ namespace Metasound
 		if (InArchive.IsLoading())
 		{
 			InMetaSound.VersionAsset();
-		}
 
-		// Clear dependency registry data both when saving and loading to clear
-		// out any fields (ex. text) that should not be serialized on the asset.
-		// Once a save is run on all pre-5.0 generated assets, this can be safely
-		// only run during *saving* persistent, non-transactional data.
-		if (InArchive.IsPersistent() && !InArchive.IsTransacting())
-		{
+			// Clear dependency registry data when loading to clear out any fields
+			// (ex. text) that should not have been serialized on the asset. Once
+			// a save is run on all pre-5.0 generated assets, this can be safely
+			// only run during pre-saving asset.
 			InMetaSound.ClearDependencyRegistryData();
 		}
-
-#if WITH_EDITOR
-		// When the editor is loaded, flag for synchronization required
-		// so that the registry data is reloaded post serialization for
-		// any open editors.
-		if (InArchive.IsSaving())
-		{
-			InMetaSound.SetSynchronizationRequired();
-		}
-#endif
 	}
 
 #if WITH_EDITORONLY_DATA
diff --git a/Engine/Plugins/Runtime/Metasound/Source/MetasoundFrontend/Private/MetasoundAssetBase.cpp b/Engine/Plugins/Runtime/Metasound/Source/MetasoundFrontend/Private/MetasoundAssetBase.cpp
--- a/Engine/Plugins/Runtime/Metasound/Source/MetasoundFrontend/Private/MetasoundAssetBase.cpp
+++ b/Engine/Plugins/Runtime/Metasound/Source/MetasoundFrontend/Private/MetasoundAssetBase.cpp
@@ -122,10 +122,12 @@ void FMetasoundAssetBase::RegisterGraphWithFrontend(Metasound::Frontend::FMetaSo
 #endif // WITH_EDITORONLY_DATA
 	}
 
-	// Caches commonly used class MetaData that isn't required for finding & building dependencies into local graph.
 	// Must be completed after auto-update to ensure all non-transient referenced dependency data is up-to-date (ex.
 	// class version), which is required for most accurately caching current registry metadata.
-	CacheDependencyRegistryData();
+	if (InRegistrationOptions.bCacheDependencyMetaDataFromRegistry)
+	{
+		CacheDependencyRegistryData();
+	}
 
 	// Registers node by copying document. Updates to document require re-registration.
 	class FNodeRegistryEntry : public INodeRegistryEntry
diff --git a/Engine/Plugins/Runtime/Metasound/Source/MetasoundFrontend/Public/MetasoundAssetBase.h b/Engine/Plugins/Runtime/Metasound/Source/MetasoundFrontend/Public/MetasoundAssetBase.h
--- a/Engine/Plugins/Runtime/Metasound/Source/MetasoundFrontend/Public/MetasoundAssetBase.h
+++ b/Engine/Plugins/Runtime/Metasound/Source/MetasoundFrontend/Public/MetasoundAssetBase.h
@@ -122,8 +122,13 @@ public:
 	const FMetasoundFrontendDocument& GetDocumentChecked() const;
 
 	void AddDefaultInterfaces();
+
+	// Caches commonly used class MetaData that isn't required for finding & building dependencies into runtime graph.
 	void CacheDependencyRegistryData();
+
+	// Clears commonly used class MetaData that isn't required for finding & building dependencies into runtime graph.
 	void ClearDependencyRegistryData();
+
 	bool VersionAsset();
 
 #if WITH_EDITORONLY_DATA
diff --git a/Engine/Plugins/Runtime/Metasound/Source/MetasoundFrontend/Public/MetasoundAssetManager.h b/Engine/Plugins/Runtime/Metasound/Source/MetasoundFrontend/Public/MetasoundAssetManager.h
--- a/Engine/Plugins/Runtime/Metasound/Source/MetasoundFrontend/Public/MetasoundAssetManager.h
+++ b/Engine/Plugins/Runtime/Metasound/Source/MetasoundFrontend/Public/MetasoundAssetManager.h
@@ -45,6 +45,9 @@ namespace Metasound
 
 			// Attempt to rebuild referenced class keys (only run if class not registered or set to force re-register)
 			bool bRebuildReferencedAssetClassKeys = true;
+
+			// If true, caches registry metadata not required by graph builder on the document's dependency class
+			bool bCacheDependencyMetaDataFromRegistry = true;
 		};
 
 		class METASOUNDFRONTEND_API IMetaSoundAssetManager
