Rebuild keys when searching for open MetaSounds to re-register (keys could be out-of-date and cause some MetaSounds to not get re-registered when expected. Can manifests as composed graphs being out-of-date in editor and not always updating).
#rnx

diff --git a/Engine/Plugins/Runtime/Metasound/Source/MetasoundEditor/Private/MetasoundEditorGraphBuilder.cpp b/Engine/Plugins/Runtime/Metasound/Source/MetasoundEditor/Private/MetasoundEditorGraphBuilder.cpp
--- a/Engine/Plugins/Runtime/Metasound/Source/MetasoundEditor/Private/MetasoundEditorGraphBuilder.cpp
+++ b/Engine/Plugins/Runtime/Metasound/Source/MetasoundEditor/Private/MetasoundEditorGraphBuilder.cpp
@@ -1278,6 +1278,7 @@ namespace Metasound
 					{
 						if (FMetasoundAssetBase* EditedMetaSound = IMetasoundUObjectRegistry::Get().GetObjectAsAssetBase(Asset))
 						{
+							EditedMetaSound->RebuildReferencedAssetClassKeys();
 							if (EditedMetaSound->IsReferencedAsset(*MetaSoundAsset))
 							{
 								EditedReferencingMetaSounds.Add(EditedMetaSound);
diff --git a/Engine/Plugins/Runtime/Metasound/Source/MetasoundFrontend/Private/MetasoundAssetBase.cpp b/Engine/Plugins/Runtime/Metasound/Source/MetasoundFrontend/Private/MetasoundAssetBase.cpp
--- a/Engine/Plugins/Runtime/Metasound/Source/MetasoundFrontend/Private/MetasoundAssetBase.cpp
+++ b/Engine/Plugins/Runtime/Metasound/Source/MetasoundFrontend/Private/MetasoundAssetBase.cpp
@@ -83,8 +83,7 @@ void FMetasoundAssetBase::RegisterGraphWithFrontend(Metasound::Frontend::FMetaSo
 
 	if (InRegistrationOptions.bRebuildReferencedAssetClassKeys)
 	{
-		TSet<FNodeRegistryKey> ReferencedKeys = IMetaSoundAssetManager::GetChecked().GetReferencedKeys(*this);
-		SetReferencedAssetClassKeys(MoveTemp(ReferencedKeys));
+		RebuildReferencedAssetClassKeys();
 	}
 
 	if (InRegistrationOptions.bRegisterDependencies)
@@ -351,7 +350,8 @@ void FMetasoundAssetBase::CacheDependencyRegistryData()
 					*GetOwningAssetName(),
 					*Class.Metadata.GetClassName().ToString());
 				UE_LOG(LogMetaSound, Warning,
-					TEXT("Asset may fail to build runtime graph unless re-registered after dependency with given key is loaded."));
+					TEXT("Asset '%s' may fail to build runtime graph unless re-registered after dependency with given key is loaded."),
+					*GetOwningAssetName());
 			}
 		}
 	}
@@ -781,6 +781,14 @@ TArray<const FMetasoundFrontendClassInput*> FMetasoundAssetBase::GetTransmittabl
 	return Inputs;
 }
 
+void FMetasoundAssetBase::RebuildReferencedAssetClassKeys()
+{
+	using namespace Metasound::Frontend;
+
+	TSet<FNodeRegistryKey> ReferencedKeys = IMetaSoundAssetManager::GetChecked().GetReferencedKeys(*this);
+	SetReferencedAssetClassKeys(MoveTemp(ReferencedKeys));
+}
+
 const FMetasoundAssetBase::FRuntimeData& FMetasoundAssetBase::GetRuntimeData() const
 {
 	
diff --git a/Engine/Plugins/Runtime/Metasound/Source/MetasoundFrontend/Public/MetasoundAssetBase.h b/Engine/Plugins/Runtime/Metasound/Source/MetasoundFrontend/Public/MetasoundAssetBase.h
--- a/Engine/Plugins/Runtime/Metasound/Source/MetasoundFrontend/Public/MetasoundAssetBase.h
+++ b/Engine/Plugins/Runtime/Metasound/Source/MetasoundFrontend/Public/MetasoundAssetBase.h
@@ -73,7 +73,8 @@ public:
 	// Sets/overwrites the root class metadata
 	virtual void SetMetadata(FMetasoundFrontendClassMetadata& InMetadata);
 
-	// Returns a default interface required when a given asset is created.
+	// Rebuild the asset class dependency key array.
+	void RebuildReferencedAssetClassKeys();
 
 	// Returns the interface entries declared by the given asset's document from the InterfaceRegistry.
 	bool GetDeclaredInterfaces(TArray<const Metasound::Frontend::IInterfaceRegistryEntry*>& OutInterfaces) const;
