Clean-up upgrade note logic (do not clear in validation step and clear style field when presaving via UMetasoundEditorGraph as doing so in validation marks graph as dirty post save when not desired)

#rnx

diff --git a/Engine/Plugins/Runtime/Metasound/Source/MetasoundEditor/Private/MetasoundEditorGraph.cpp b/Engine/Plugins/Runtime/Metasound/Source/MetasoundEditor/Private/MetasoundEditorGraph.cpp
--- a/Engine/Plugins/Runtime/Metasound/Source/MetasoundEditor/Private/MetasoundEditorGraph.cpp
+++ b/Engine/Plugins/Runtime/Metasound/Source/MetasoundEditor/Private/MetasoundEditorGraph.cpp
@@ -1200,6 +1200,23 @@ Metasound::Frontend::FConstGraphHandle UMetasoundEditorGraph::GetGraphHandle() c
 	return MetasoundAsset->GetRootGraphHandle();
 }
 
+void UMetasoundEditorGraph::PreSave(FObjectPreSaveContext InSaveContext)
+{
+	using namespace Metasound::Frontend;
+
+	TArray<UMetasoundEditorGraphNode*> MetaSoundNodes;
+	GetNodesOfClass<UMetasoundEditorGraphNode>(MetaSoundNodes);
+	for (UMetasoundEditorGraphNode* Node : MetaSoundNodes)
+	{
+		FNodeHandle NodeHandle = Node->GetNodeHandle();
+		FMetasoundFrontendNodeStyle Style = NodeHandle->GetNodeStyle();
+		Style.bMessageNodeUpdated = false;
+		NodeHandle->SetNodeStyle(Style);
+	}
+
+	Super::PreSave(InSaveContext);
+}
+
 UObject* UMetasoundEditorGraph::GetMetasound()
 {
 	return GetOuter();
@@ -1234,7 +1251,7 @@ void UMetasoundEditorGraph::RegisterGraphWithFrontend()
 	}
 }
 
-void UMetasoundEditorGraph::SetSynchronizationRequired(bool bInClearUpdateNotes)
+void UMetasoundEditorGraph::SetSynchronizationRequired()
 {
 	using namespace Metasound;
 
@@ -1242,15 +1259,7 @@ void UMetasoundEditorGraph::SetSynchronizationRequired(bool bInClearUpdateNotes)
 	{
 		FMetasoundAssetBase* MetasoundAsset = IMetasoundUObjectRegistry::Get().GetObjectAsAssetBase(ParentMetasound);
 		check(MetasoundAsset);
-
-		if (bInClearUpdateNotes)
-		{
-			MetasoundAsset->SetClearNodeNotesOnSynchronization();
-		}
-		else
-		{
-			MetasoundAsset->SetSynchronizationRequired();
-		}
+		MetasoundAsset->SetSynchronizationRequired();
 	}
 }
 
@@ -1542,7 +1551,7 @@ void UMetasoundEditorGraph::IterateOutputs(TUniqueFunction<void(UMetasoundEditor
 	}
 }
 
-bool UMetasoundEditorGraph::ValidateInternal(Metasound::Editor::FGraphValidationResults& OutResults, bool bClearUpgradeMessaging)
+bool UMetasoundEditorGraph::ValidateInternal(Metasound::Editor::FGraphValidationResults& OutResults)
 {
 	using namespace Metasound::Editor;
 	using namespace Metasound::Frontend;
@@ -1560,7 +1569,7 @@ bool UMetasoundEditorGraph::ValidateInternal(Metasound::Editor::FGraphValidation
 
 		if (UMetasoundEditorGraphExternalNode* ExternalNode = Cast<UMetasoundEditorGraphExternalNode>(Node))
 		{
-			bIsValid &= ExternalNode->Validate(NodeResult, bClearUpgradeMessaging);
+			bIsValid &= ExternalNode->Validate(NodeResult);
 			bMarkDirty |= NodeResult.bIsDirty;
 		}
 
diff --git a/Engine/Plugins/Runtime/Metasound/Source/MetasoundEditor/Private/MetasoundEditorGraphBuilder.cpp b/Engine/Plugins/Runtime/Metasound/Source/MetasoundEditor/Private/MetasoundEditorGraphBuilder.cpp
--- a/Engine/Plugins/Runtime/Metasound/Source/MetasoundEditor/Private/MetasoundEditorGraphBuilder.cpp
+++ b/Engine/Plugins/Runtime/Metasound/Source/MetasoundEditor/Private/MetasoundEditorGraphBuilder.cpp
@@ -227,6 +227,9 @@ namespace Metasound
 
 			SynchronizeNodeLocation(InLocation, InNodeHandle, *NewGraphNode);
 
+			// Adding external node may introduce referenced asset so rebuild referenced keys.
+			MetaSoundAsset->RebuildReferencedAssetClassKeys();
+
 			return NewGraphNode;
 		}
 
@@ -336,7 +339,7 @@ namespace Metasound
 
 			bool bMarkDirty = false;
 
-			Graph.ValidateInternal(Results, MetaSoundAsset->GetSynchronizationClearUpdateNotes());
+			Graph.ValidateInternal(Results);
 			for (const FGraphNodeValidationResult& Result : Results.GetResults())
 			{
 				bMarkDirty |= Result.bIsDirty;
diff --git a/Engine/Plugins/Runtime/Metasound/Source/MetasoundEditor/Private/MetasoundEditorGraphNode.cpp b/Engine/Plugins/Runtime/Metasound/Source/MetasoundEditor/Private/MetasoundEditorGraphNode.cpp
--- a/Engine/Plugins/Runtime/Metasound/Source/MetasoundEditor/Private/MetasoundEditorGraphNode.cpp
+++ b/Engine/Plugins/Runtime/Metasound/Source/MetasoundEditor/Private/MetasoundEditorGraphNode.cpp
@@ -691,7 +691,7 @@ UMetasoundEditorGraphExternalNode* UMetasoundEditorGraphExternalNode::UpdateToVe
 	return ReplacementEdNode;
 }
 
-bool UMetasoundEditorGraphExternalNode::Validate(Metasound::Editor::FGraphNodeValidationResult& OutResult, bool bClearUpgradeMessage)
+bool UMetasoundEditorGraphExternalNode::Validate(Metasound::Editor::FGraphNodeValidationResult& OutResult)
 {
 	using namespace Metasound::Editor;
 	using namespace Metasound::Frontend;
@@ -712,15 +712,6 @@ bool UMetasoundEditorGraphExternalNode::Validate(Metasound::Editor::FGraphNodeVa
 		OutResult.bIsDirty = true;
 	}
 
-	if (bClearUpgradeMessage)
-	{
-		if (!NodeUpgradeMessage.IsEmpty())
-		{
-			NodeUpgradeMessage = FText::GetEmpty();
-			OutResult.bIsDirty = true;
-		}
-	}
-
 	// 2. Check if node is invalid, version is missing and cache if interface changes exist between the document's records and the registry
 	FNodeHandle NodeHandle = GetNodeHandle();
 	const FMetasoundFrontendClassMetadata& Metadata = NodeHandle->GetClassMetadata();
@@ -728,7 +719,6 @@ bool UMetasoundEditorGraphExternalNode::Validate(Metasound::Editor::FGraphNodeVa
 	FClassInterfaceUpdates InterfaceUpdates;
 	if (!NodeHandle->DiffAgainstRegistryInterface(InterfaceUpdates, false /* bUseHighestMinorVersion */))
 	{
-		
 		if (NodeHandle->IsValid())
 		{
 			const FText* PromptIfMissing = nullptr;
@@ -780,10 +770,7 @@ bool UMetasoundEditorGraphExternalNode::Validate(Metasound::Editor::FGraphNodeVa
 	{
 		if (bNewIsClassNative)
 		{
-			GraphNodePrivate::SetGraphNodeMessage(*this, EMessageSeverity::Info, FString::Format(TEXT("Class '{0}' has been nativized."),
-			{
-				*Metadata.GetDisplayName().ToString()
-			}));
+			NodeUpgradeMessage = FText::Format(LOCTEXT("MetaSoundNode_NativizedMessage", "Class '{0}' has been nativized."), Metadata.GetDisplayName());
 		}
 
 		OutResult.bIsDirty = true;
@@ -794,20 +781,10 @@ bool UMetasoundEditorGraphExternalNode::Validate(Metasound::Editor::FGraphNodeVa
 	FMetasoundFrontendNodeStyle Style = NodeHandle->GetNodeStyle();
 	if (Style.bMessageNodeUpdated)
 	{
-		if (bClearUpgradeMessage)
-		{
-			Style.bMessageNodeUpdated = false;
-			NodeHandle->SetNodeStyle(Style);
-		}
-		else
-		{
-			GraphNodePrivate::SetGraphNodeMessage(*this, EMessageSeverity::Info, FString::Format(TEXT("Node class '{0}' updated to version {1}"),
-			{
-				*Metadata.GetDisplayName().ToString(),
-				*Metadata.GetVersion().ToString()
-			}));
-		}
-
+		NodeUpgradeMessage = FText::Format(LOCTEXT("MetaSoundNode_UpgradedMessage", "Node class '{0}' updated to version {1}"),
+			Metadata.GetDisplayName(),
+			FText::FromString(Metadata.GetVersion().ToString())
+		);
 		OutResult.bIsDirty = true;
 	}
 
diff --git a/Engine/Plugins/Runtime/Metasound/Source/MetasoundEditor/Public/MetasoundEditorGraph.h b/Engine/Plugins/Runtime/Metasound/Source/MetasoundEditor/Public/MetasoundEditorGraph.h
--- a/Engine/Plugins/Runtime/Metasound/Source/MetasoundEditor/Public/MetasoundEditorGraph.h
+++ b/Engine/Plugins/Runtime/Metasound/Source/MetasoundEditor/Public/MetasoundEditorGraph.h
@@ -342,7 +342,6 @@ private:
 	FVariableEditorNodes GetVariableNodes() const;
 	FVariableNodeLocations GetVariableNodeLocations() const;
 	void AddVariableNodes(UObject& InMetasound, Metasound::Frontend::FGraphHandle& InFrontendGraph, const FVariableNodeLocations& InNodeLocs);
-	void UpdateEditorLiteralType();
 };
 
 UCLASS()
@@ -358,6 +357,8 @@ public:
 	Metasound::Frontend::FGraphHandle GetGraphHandle();
 	Metasound::Frontend::FConstGraphHandle GetGraphHandle() const;
 
+	virtual void PreSave(FObjectPreSaveContext InSaveContext) override;
+
 	UObject* GetMetasound();
 	const UObject* GetMetasound() const;
 	UObject& GetMetasoundChecked();
@@ -375,14 +376,13 @@ public:
 
 	// UMetasoundEditorGraphBase Implementation
 	virtual void RegisterGraphWithFrontend() override;
-	virtual void SetSynchronizationRequired(bool bInClearUpdateNotes = false) override;
+	virtual void SetSynchronizationRequired() override;
 
 private:
-
 	bool RemoveFrontendInput(UMetasoundEditorGraphInput& Input);
 	bool RemoveFrontendOutput(UMetasoundEditorGraphOutput& Output);
 	bool RemoveFrontendVariable(UMetasoundEditorGraphVariable& Variable);
-	bool ValidateInternal(Metasound::Editor::FGraphValidationResults& OutResults, bool bClearUpgradeMessaging = true);
+	bool ValidateInternal(Metasound::Editor::FGraphValidationResults& OutResults);
 
 	// Preview ID is the Unique ID provided by the UObject that implements
 	// a sound's ParameterInterface when a sound begins playing.
diff --git a/Engine/Plugins/Runtime/Metasound/Source/MetasoundEditor/Public/MetasoundEditorGraphNode.h b/Engine/Plugins/Runtime/Metasound/Source/MetasoundEditor/Public/MetasoundEditorGraphNode.h
--- a/Engine/Plugins/Runtime/Metasound/Source/MetasoundEditor/Public/MetasoundEditorGraphNode.h
+++ b/Engine/Plugins/Runtime/Metasound/Source/MetasoundEditor/Public/MetasoundEditorGraphNode.h
@@ -192,7 +192,7 @@ public:
 	UMetasoundEditorGraphExternalNode* UpdateToVersion(const FMetasoundFrontendVersionNumber& InNewVersion, bool bInPropagateErrorMessages);
 
 	// Validates node and returns whether or not the node is valid.
-	bool Validate(Metasound::Editor::FGraphNodeValidationResult& OutResult, bool bClearUpgradeMessage = true);
+	bool Validate(Metasound::Editor::FGraphNodeValidationResult& OutResult);
 
 
 protected:
diff --git a/Engine/Plugins/Runtime/Metasound/Source/MetasoundEngine/Public/Metasound.h b/Engine/Plugins/Runtime/Metasound/Source/MetasoundEngine/Public/Metasound.h
--- a/Engine/Plugins/Runtime/Metasound/Source/MetasoundEngine/Public/Metasound.h
+++ b/Engine/Plugins/Runtime/Metasound/Source/MetasoundEngine/Public/Metasound.h
@@ -30,7 +30,7 @@ public:
 	virtual bool IsEditorOnly() const override { return true; }
 	virtual bool NeedsLoadForEditorGame() const override { return false; }
 
-	virtual void SetSynchronizationRequired(bool bInClearUpdateNotes = false) PURE_VIRTUAL(UMetasoundEditorGraphBase::Synchronize, )
+	virtual void SetSynchronizationRequired() PURE_VIRTUAL(UMetasoundEditorGraphBase::SetSynchronizationRequired, )
 	virtual void RegisterGraphWithFrontend() PURE_VIRTUAL(UMetasoundEditorGraphBase::RegisterGraphWithFrontend(), )
 };
 
@@ -79,12 +79,12 @@ namespace Metasound
 				{
 					InMetaSound.RegisterGraphWithFrontend(NonEditorRegistrationOptions);
 				}
-				MetaSoundGraph->SetSynchronizationRequired(false /* bInClearUpdateNotes */);
+				MetaSoundGraph->SetSynchronizationRequired();
 			}
 			else
 			{
 				MetaSoundGraph->RegisterGraphWithFrontend();
-				MetaSoundGraph->SetSynchronizationRequired(true /* bInClearUpdateNotes */);
+				MetaSoundGraph->SetSynchronizationRequired();
 			}
 		}
 #else
diff --git a/Engine/Plugins/Runtime/Metasound/Source/MetasoundFrontend/Private/MetasoundAssetBase.cpp b/Engine/Plugins/Runtime/Metasound/Source/MetasoundFrontend/Private/MetasoundAssetBase.cpp
--- a/Engine/Plugins/Runtime/Metasound/Source/MetasoundFrontend/Private/MetasoundAssetBase.cpp
+++ b/Engine/Plugins/Runtime/Metasound/Source/MetasoundFrontend/Private/MetasoundAssetBase.cpp
@@ -430,11 +430,6 @@ bool FMetasoundAssetBase::GetSynchronizationRequired() const
 	return bSynchronizationRequired;
 }
 
-bool FMetasoundAssetBase::GetSynchronizationClearUpdateNotes() const
-{
-	return bSynchronizationClearUpdateNotes;
-}
-
 bool FMetasoundAssetBase::GetSynchronizationUpdateDetails() const
 {
 	return bSynchronizationUpdateDetails;
@@ -445,12 +440,6 @@ void FMetasoundAssetBase::SetSynchronizationRequired()
 	bSynchronizationRequired = true;
 }
 
-void FMetasoundAssetBase::SetClearNodeNotesOnSynchronization()
-{
-	bSynchronizationClearUpdateNotes = true;
-	bSynchronizationRequired = true;
-}
-
 void FMetasoundAssetBase::SetUpdateDetailsOnSynchronization()
 {
 	bSynchronizationUpdateDetails = true;
@@ -459,7 +448,6 @@ void FMetasoundAssetBase::SetUpdateDetailsOnSynchronization()
 
 void FMetasoundAssetBase::ResetSynchronizationState()
 {
-	bSynchronizationClearUpdateNotes = false;
 	bSynchronizationUpdateDetails = false;
 	bSynchronizationRequired = false;
 }
diff --git a/Engine/Plugins/Runtime/Metasound/Source/MetasoundFrontend/Public/MetasoundAssetBase.h b/Engine/Plugins/Runtime/Metasound/Source/MetasoundFrontend/Public/MetasoundAssetBase.h
--- a/Engine/Plugins/Runtime/Metasound/Source/MetasoundFrontend/Public/MetasoundAssetBase.h
+++ b/Engine/Plugins/Runtime/Metasound/Source/MetasoundFrontend/Public/MetasoundAssetBase.h
@@ -137,10 +137,8 @@ public:
 	// Split synchronization requirement flag into synchronization required &
 	// object type refresh or checking frontend class guids when synchronizing.
 	bool GetSynchronizationRequired() const;
-	bool GetSynchronizationClearUpdateNotes() const;
 	bool GetSynchronizationUpdateDetails() const;
 	void ResetSynchronizationState();
-	void SetClearNodeNotesOnSynchronization();
 	void SetUpdateDetailsOnSynchronization();
 	void SetSynchronizationRequired();
 #endif // WITH_EDITORONLY_DATA
@@ -183,7 +181,6 @@ protected:
 
 #if WITH_EDITORONLY_DATA
 	bool bSynchronizationRequired = true;
-	bool bSynchronizationClearUpdateNotes = false;
 	bool bSynchronizationUpdateDetails = false;
 #endif // WITH_EDITORONLY_DATA
 
