Fix missed hard-coded data type names in a couple functions
#rnx

diff --git a/Engine/Plugins/Runtime/Metasound/Source/MetasoundEditor/Private/MetasoundEditor.cpp b/Engine/Plugins/Runtime/Metasound/Source/MetasoundEditor/Private/MetasoundEditor.cpp
--- a/Engine/Plugins/Runtime/Metasound/Source/MetasoundEditor/Private/MetasoundEditor.cpp
+++ b/Engine/Plugins/Runtime/Metasound/Source/MetasoundEditor/Private/MetasoundEditor.cpp
@@ -2393,8 +2393,8 @@ namespace Metasound
 					const FScopedTransaction Transaction(LOCTEXT("AddInputNode", "Add MetaSound Input"));
 					Metasound->Modify();
 
-					const FName DataTypeName = "Bool";
-					Frontend::FNodeHandle NodeHandle = FGraphBuilder::AddInputNodeHandle(*Metasound, DataTypeName, FText::GetEmpty());
+					const FName DataTypeName = GetMetasoundDataTypeName<bool>();
+					Frontend::FNodeHandle NodeHandle = FGraphBuilder::AddInputNodeHandle(*Metasound, DataTypeName);
 					if (ensure(NodeHandle->IsValid()))
 					{
 						NameToSelect = NodeHandle->GetNodeName();
@@ -2419,8 +2419,8 @@ namespace Metasound
 					const FScopedTransaction Transaction(TEXT(""), LOCTEXT("AddOutputNode", "Add MetaSound Output"), Metasound);
 					Metasound->Modify();
 
-					const FName DataTypeName = "Bool";
-					Frontend::FNodeHandle NodeHandle = FGraphBuilder::AddOutputNodeHandle(*Metasound, DataTypeName, FText::GetEmpty());
+					const FName DataTypeName = GetMetasoundDataTypeName<bool>();
+					Frontend::FNodeHandle NodeHandle = FGraphBuilder::AddOutputNodeHandle(*Metasound, DataTypeName);
 					if (ensure(NodeHandle->IsValid()))
 					{
 						NameToSelect = NodeHandle->GetNodeName();
