Fix bit from a CL not meant to go in with last fix-up
#rnx

diff --git a/Engine/Plugins/Runtime/Metasound/Source/MetasoundEditor/Private/MetasoundEditor.cpp b/Engine/Plugins/Runtime/Metasound/Source/MetasoundEditor/Private/MetasoundEditor.cpp
--- a/Engine/Plugins/Runtime/Metasound/Source/MetasoundEditor/Private/MetasoundEditor.cpp
+++ b/Engine/Plugins/Runtime/Metasound/Source/MetasoundEditor/Private/MetasoundEditor.cpp
@@ -2394,7 +2394,7 @@ namespace Metasound
 					Metasound->Modify();
 
 					const FName DataTypeName = GetMetasoundDataTypeName<bool>();
-					Frontend::FNodeHandle NodeHandle = FGraphBuilder::AddInputNodeHandle(*Metasound, DataTypeName);
+					Frontend::FNodeHandle NodeHandle = FGraphBuilder::AddInputNodeHandle(*Metasound, DataTypeName, FText::GetEmpty());
 					if (ensure(NodeHandle->IsValid()))
 					{
 						NameToSelect = NodeHandle->GetNodeName();
@@ -2420,7 +2420,7 @@ namespace Metasound
 					Metasound->Modify();
 
 					const FName DataTypeName = GetMetasoundDataTypeName<bool>();
-					Frontend::FNodeHandle NodeHandle = FGraphBuilder::AddOutputNodeHandle(*Metasound, DataTypeName);
+					Frontend::FNodeHandle NodeHandle = FGraphBuilder::AddOutputNodeHandle(*Metasound, DataTypeName, FText::GetEmpty());
 					if (ensure(NodeHandle->IsValid()))
 					{
 						NameToSelect = NodeHandle->GetNodeName();
