Fix for crash when metasound wave player uses ForceInline-ed audio on some platforms (increased decode size).

diff --git a/Engine/Plugins/Runtime/Metasound/Source/MetasoundEngine/Private/MetasoundWave.cpp b/Engine/Plugins/Runtime/Metasound/Source/MetasoundEngine/Private/MetasoundWave.cpp
--- a/Engine/Plugins/Runtime/Metasound/Source/MetasoundEngine/Private/MetasoundWave.cpp
+++ b/Engine/Plugins/Runtime/Metasound/Source/MetasoundEngine/Private/MetasoundWave.cpp
@@ -78,7 +78,7 @@ namespace Audio
 	: OutputSampleRate(InInitParams.OutputSampleRate)
 	, MaxPitchShiftCents(InInitParams.MaxPitchShiftMagnitudeAllowedInOctaves * 1200.0f)
 	, MaxPitchShiftRatio(FMath::Pow(2.0f, InInitParams.MaxPitchShiftMagnitudeAllowedInOctaves))
-	, DecodeBlockSizeInFrames(64)
+	, DecodeBlockSizeInFrames(128) // 128 chosen due to some limitations on some platform decoders (must be an integer multiple of 128)
 	, OutputBlockSizeInFrames(InInitParams.OutputBlockSizeInFrames)
 	{
 		check(OutputBlockSizeInFrames > 0);
diff --git a/Engine/Plugins/Runtime/Metasound/Source/MetasoundEngine/Public/MetasoundWave.h b/Engine/Plugins/Runtime/Metasound/Source/MetasoundEngine/Public/MetasoundWave.h
--- a/Engine/Plugins/Runtime/Metasound/Source/MetasoundEngine/Public/MetasoundWave.h
+++ b/Engine/Plugins/Runtime/Metasound/Source/MetasoundEngine/Public/MetasoundWave.h
@@ -126,8 +126,8 @@ namespace Audio
 		float MaxPitchShiftRatio{ 2.0f };
 
 		uint32 NumChannels{ 0 };
-		uint32 DecodeBlockSizeInFrames{ 64 };
-		uint32 DecodeBlockSizeInSamples{ 64 };
+		uint32 DecodeBlockSizeInFrames;
+		uint32 DecodeBlockSizeInSamples;
 		uint32 OutputBlockSizeInFrames { 512 };
 		
 		bool bDecoderIsDone{ true };
