Fixes for MetaSound rename and delete editor commands
- Change delete command to only use delete key instead of backspace
- Allow rename action on interface menu items
- Fix crash caused by invalid rename request handle

diff --git a/Engine/Plugins/Runtime/Metasound/Source/MetasoundEditor/Private/MetasoundEditor.cpp b/Engine/Plugins/Runtime/Metasound/Source/MetasoundEditor/Private/MetasoundEditor.cpp
--- a/Engine/Plugins/Runtime/Metasound/Source/MetasoundEditor/Private/MetasoundEditor.cpp
+++ b/Engine/Plugins/Runtime/Metasound/Source/MetasoundEditor/Private/MetasoundEditor.cpp
@@ -816,8 +816,12 @@ namespace Metasound
 				FExecuteAction::CreateSP(this, &FEditor::ConvertFromPreset));
 
 			ToolkitCommands->MapAction(
-				FGenericCommands::Get().Delete,
+				Commands.Delete,
 				FExecuteAction::CreateSP(this, &FEditor::DeleteSelected));
+
+			ToolkitCommands->MapAction(
+				FGenericCommands::Get().Rename,
+				FExecuteAction::CreateSP(this, &FEditor::RenameSelected));
 		}
 
 		void FEditor::Import()
@@ -1238,8 +1242,8 @@ namespace Metasound
 					FCanExecuteAction::CreateLambda([this]() { return CanDuplicateNodes(); }));
 
 				GraphEditorCommands->MapAction(FGenericCommands::Get().Rename,
-					FExecuteAction::CreateLambda([this] { RenameNode(); }),
-					FCanExecuteAction::CreateLambda([this]() { return CanRenameNodes(); }));
+					FExecuteAction::CreateLambda([this] { RenameSelected(); }),
+					FCanExecuteAction::CreateLambda([this]() { return CanRenameSelectedNodes(); }));
 
 				// Alignment Commands
 				GraphEditorCommands->MapAction(FGraphEditorCommands::Get().AlignNodesTop,
@@ -1840,7 +1844,7 @@ namespace Metasound
 			MetasoundGraphEditor->NotifyGraphChanged();
 		}
 
-		bool FEditor::CanRenameNodes() const
+		bool FEditor::CanRenameSelectedNodes() const
 		{
 			if (!IsGraphEditable())
 			{
@@ -1862,27 +1866,65 @@ namespace Metasound
 				{
 					if (const UMetasoundEditorGraphMember* Member = MemberNode->GetMember()) 
 					{
-						// Variable
-						if (const UMetasoundEditorGraphVariable* Variable = Cast<UMetasoundEditorGraphVariable>(Member))
-						{
-							return true;
-						}
+						return Member->CanRename();
+					}
+				}
+			}
+			return false;
+		}
+		
+		void FEditor::RenameSelected()
+		{
+			// Renaming selected nodes
+			if (CanRenameSelectedNodes())
+			{
+				RenameSelectedNode();
+				return;
+			}
+
+			// Renaming selected interface member
+			if (MetasoundInterfaceMenu.IsValid())
+			{
+				TArray<TSharedPtr<FEdGraphSchemaAction>> Actions;
+				MetasoundInterfaceMenu->GetSelectedActions(Actions);
 
-						// Input or output that is not an interface member (does not have a valid interface version)
-						const FMetasoundFrontendVersion* InterfaceVersion = nullptr;
-						if (const UMetasoundEditorGraphVertex* Vertex = Cast<UMetasoundEditorGraphVertex>(Member))
+				if (!Actions.IsEmpty())
+				{
+					for (const TSharedPtr<FEdGraphSchemaAction>& Action : Actions)
+					{
+						TSharedPtr<FMetasoundGraphMemberSchemaAction> MetasoundAction = StaticCastSharedPtr<FMetasoundGraphMemberSchemaAction>(Action);
+						if (MetasoundAction.IsValid())
 						{
-							InterfaceVersion = &Vertex->GetInterfaceVersion();
+							if (const UMetasoundEditorGraphMember* GraphMember = MetasoundAction->GetGraphMember())
+							{
+								if (GraphMember->CanRename())
+								{
+									RenameInterfaceItem(MetasoundAction);
+									return;
+								}
+							}
 						}
-
-						return (!InterfaceVersion || !InterfaceVersion->IsValid());
 					}
 				}
 			}
-			return false;
 		}
 
-		void FEditor::RenameNode() 
+		void FEditor::RenameInterfaceItem(TSharedPtr<FMetasoundGraphMemberSchemaAction> ActionToRename)
+		{
+			using namespace Metasound::Frontend;
+			check(Metasound);
+
+			UMetasoundEditorGraphMember* GraphMember = ActionToRename->GetGraphMember();
+			if (ensure(GraphMember))
+			{
+				if (GraphMember->OnRenameRequested.IsBound())
+				{
+					GraphMember->OnRenameRequested.Broadcast();
+				} 
+			}
+		}
+
+		void FEditor::RenameSelectedNode() 
 		{
 			const FGraphPanelSelectionSet& SelectedNodes = MetasoundGraphEditor->GetSelectedNodes();
 			for (FGraphPanelSelectionSet::TConstIterator SelectedIter(SelectedNodes); SelectedIter; ++SelectedIter)
@@ -1903,6 +1945,8 @@ namespace Metasound
 				{
 					if (const UMetasoundEditorGraphMember* Member = MemberNode->GetMember()) 
 					{
+						MetasoundInterfaceMenu->SelectItemByName(Member->GetMemberName(), ESelectInfo::Direct, static_cast<int32>(Member->GetSectionID()));
+
 						if (Member->OnRenameRequested.IsBound())
 						{
 							Member->OnRenameRequested.Broadcast();
diff --git a/Engine/Plugins/Runtime/Metasound/Source/MetasoundEditor/Private/MetasoundEditor.h b/Engine/Plugins/Runtime/Metasound/Source/MetasoundEditor/Private/MetasoundEditor.h
--- a/Engine/Plugins/Runtime/Metasound/Source/MetasoundEditor/Private/MetasoundEditor.h
+++ b/Engine/Plugins/Runtime/Metasound/Source/MetasoundEditor/Private/MetasoundEditor.h
@@ -130,11 +130,8 @@ namespace Metasound
 			void PasteNodes(const FVector2D* InLocation = nullptr);
 			void PasteNodes(const FVector2D* InLocation, const FText& InTransactionText);
 
-			/** Rename node (currently applies to comments and member nodes). */
-			void RenameNode();
-
-			/** Whether at least one of the currently selected node(s) can be renamed. */
-			bool CanRenameNodes() const;
+			/** Rename selected items (nodes or interface items). */
+			void RenameSelected();
 
 			/** Forces all UX pertaining to the root graph's interface to be refreshed. */
 			void RefreshInterface();
@@ -217,6 +214,15 @@ namespace Metasound
 			/** Whether the currently selected node(s) can be deleted */
 			bool CanDeleteNodes() const;
 
+			/** Whether at least one of the currently selected node(s) can be renamed. */
+			bool CanRenameSelectedNodes() const;
+
+			/** Rename selected node (currently applies to comments and member nodes). */
+			void RenameSelectedNode();
+
+			/** Rename selected interface item. */
+			void RenameInterfaceItem(TSharedPtr<FMetasoundGraphMemberSchemaAction> ActionToDelete);
+
 			/** Called to undo the last action */
 			void UndoGraphAction();
 
diff --git a/Engine/Plugins/Runtime/Metasound/Source/MetasoundEditor/Private/MetasoundEditorCommands.cpp b/Engine/Plugins/Runtime/Metasound/Source/MetasoundEditor/Private/MetasoundEditorCommands.cpp
--- a/Engine/Plugins/Runtime/Metasound/Source/MetasoundEditor/Private/MetasoundEditorCommands.cpp
+++ b/Engine/Plugins/Runtime/Metasound/Source/MetasoundEditor/Private/MetasoundEditorCommands.cpp
@@ -26,6 +26,7 @@ namespace Metasound
 			UI_COMMAND(UpdateNodes, "Update Nodes", "Update selected node(s) that have an available update.", EUserInterfaceActionType::Button, FInputChord());
 
 			UI_COMMAND(ConvertFromPreset, "Convert From Preset", "Converts this preset to a fully accessible MetaSound.", EUserInterfaceActionType::Button, FInputChord());
+			UI_COMMAND(Delete, "Delete Selected", "Delete selected items.", EUserInterfaceActionType::None, FInputChord(EKeys::Delete));
 		}
 	} // namespace Editor
 } // namespace Metasound
diff --git a/Engine/Plugins/Runtime/Metasound/Source/MetasoundEditor/Private/MetasoundEditorGraph.cpp b/Engine/Plugins/Runtime/Metasound/Source/MetasoundEditor/Private/MetasoundEditorGraph.cpp
--- a/Engine/Plugins/Runtime/Metasound/Source/MetasoundEditor/Private/MetasoundEditorGraph.cpp
+++ b/Engine/Plugins/Runtime/Metasound/Source/MetasoundEditor/Private/MetasoundEditorGraph.cpp
@@ -407,6 +407,11 @@ bool UMetasoundEditorGraphVertex::IsInterfaceMember() const
 	return GetConstNodeHandle()->IsInterfaceMember();
 }
 
+bool UMetasoundEditorGraphVertex::CanRename() const
+{
+	return !IsInterfaceMember();
+}
+
 bool UMetasoundEditorGraphVertex::CanRename(const FText& InNewName, FText& OutError) const
 {
 	using namespace Metasound::Frontend;
@@ -801,6 +806,11 @@ void UMetasoundEditorGraphVariable::SetDescription(const FText& InDescription, b
 	}
 }
 
+bool UMetasoundEditorGraphVariable::CanRename() const
+{
+	return true;
+}
+
 bool UMetasoundEditorGraphVariable::CanRename(const FText& InNewText, FText& OutError) const
 {
 	using namespace Metasound::Frontend;
diff --git a/Engine/Plugins/Runtime/Metasound/Source/MetasoundEditor/Private/MetasoundNodeDetailCustomization.h b/Engine/Plugins/Runtime/Metasound/Source/MetasoundEditor/Private/MetasoundNodeDetailCustomization.h
--- a/Engine/Plugins/Runtime/Metasound/Source/MetasoundEditor/Private/MetasoundNodeDetailCustomization.h
+++ b/Engine/Plugins/Runtime/Metasound/Source/MetasoundEditor/Private/MetasoundNodeDetailCustomization.h
@@ -230,19 +230,15 @@ namespace Metasound
 			{
 				if (UMetasoundEditorGraphMember* Member = InMemberDefaultLiteral.GetParentMember())
 				{
-					if (const UMetasoundEditorGraph* OwningGraph = GraphMember->GetOwningGraph())
+					if (Member->CanRename())
 					{
-						TSharedPtr<FEditor> MetasoundEditor = FGraphBuilder::GetEditorForGraph(*OwningGraph);
-						if (MetasoundEditor->CanRenameNodes())
+						if (!RenameRequestedHandle.IsValid())
 						{
-							if (!RenameRequestedHandle.IsValid())
+							Member->OnRenameRequested.Clear();
+							RenameRequestedHandle = Member->OnRenameRequested.AddLambda([this]()
 							{
-								Member->OnRenameRequested.Clear();
-								RenameRequestedHandle = Member->OnRenameRequested.AddLambda([this]()
-								{
-									FSlateApplication::Get().SetKeyboardFocus(NameEditableTextBox.ToSharedRef(), EFocusCause::SetDirectly);
-								});
-							}
+								FSlateApplication::Get().SetKeyboardFocus(NameEditableTextBox.ToSharedRef(), EFocusCause::SetDirectly);
+							});
 						}
 					}
 				}
diff --git a/Engine/Plugins/Runtime/Metasound/Source/MetasoundEditor/Public/MetasoundEditorCommands.h b/Engine/Plugins/Runtime/Metasound/Source/MetasoundEditor/Public/MetasoundEditorCommands.h
--- a/Engine/Plugins/Runtime/Metasound/Source/MetasoundEditor/Public/MetasoundEditorCommands.h
+++ b/Engine/Plugins/Runtime/Metasound/Source/MetasoundEditor/Public/MetasoundEditorCommands.h
@@ -58,6 +58,9 @@ namespace Metasound
 			// Converts a MetaSound from a restricted, preset edit state to a fully accessible MetaSound.
 			TSharedPtr<FUICommandInfo> ConvertFromPreset;
 
+			// Delete selected items. 
+			TSharedPtr<FUICommandInfo> Delete;
+
 			/** Initialize commands */
 			virtual void RegisterCommands() override;
 		};
diff --git a/Engine/Plugins/Runtime/Metasound/Source/MetasoundEditor/Public/MetasoundEditorGraph.h b/Engine/Plugins/Runtime/Metasound/Source/MetasoundEditor/Public/MetasoundEditorGraph.h
--- a/Engine/Plugins/Runtime/Metasound/Source/MetasoundEditor/Public/MetasoundEditorGraph.h
+++ b/Engine/Plugins/Runtime/Metasound/Source/MetasoundEditor/Public/MetasoundEditorGraph.h
@@ -147,6 +147,9 @@ public:
 	/** Marks all member's nodes as requiring a display refresh. */
 	void MarkNodesForRefresh();
 
+	/* Whether this member can be renamed. */
+	virtual bool CanRename() const PURE_VIRTUAL(UMetasoundEditorGraphMember::CanRename, return false;);
+
 #if WITH_EDITOR
 	virtual void PostEditUndo() override;
 #endif // WITH_EDITOR
@@ -223,6 +226,8 @@ public:
 
 	virtual UMetasoundEditorGraphMemberDefaultLiteral* GetLiteral() const PURE_VIRTUAL(UMetasoundEditorGraphMember::GetLiteral, return nullptr; );
 
+	virtual bool CanRename() const override;
+
 protected:
 	virtual void SetLiteral(UMetasoundEditorGraphMemberDefaultLiteral* InLiteral) PURE_VIRTUAL (UMetasoundEditorGraphMember::SetLiteral, );
 };
@@ -320,6 +325,8 @@ public:
 
 	virtual UMetasoundEditorGraphMemberDefaultLiteral* GetLiteral() const { return Literal; }
 
+	virtual bool CanRename() const override;
+
 protected:
 	virtual void SetLiteral(UMetasoundEditorGraphMemberDefaultLiteral* InLiteral) { Literal = InLiteral; }
 
