Move any places where FText is copied to a MetaSound document to transient UProperties that are populated by external registries (i.e. the interface registry & the class registry)
#rnx

diff --git a/Engine/Plugins/Runtime/Metasound/Source/MetasoundEditor/Private/MetasoundEditor.cpp b/Engine/Plugins/Runtime/Metasound/Source/MetasoundEditor/Private/MetasoundEditor.cpp
--- a/Engine/Plugins/Runtime/Metasound/Source/MetasoundEditor/Private/MetasoundEditor.cpp
+++ b/Engine/Plugins/Runtime/Metasound/Source/MetasoundEditor/Private/MetasoundEditor.cpp
@@ -2301,7 +2301,7 @@ namespace Metasound
 
 			if (MetasoundAsset->GetSynchronizationRequired())
 			{
-				MetasoundAsset->CacheDependencyRegistryData();
+				MetasoundAsset->CacheRegistryMetadata();
 
 				// Capture before synchronizing as the flag is cleared therein.
 				const bool bShouldRefreshDetails = MetasoundAsset->GetSynchronizationUpdateDetails();
@@ -2438,7 +2438,7 @@ namespace Metasound
 					Metasound->Modify();
 
 					const FName DataTypeName = GetMetasoundDataTypeName<bool>();
-					Frontend::FNodeHandle NodeHandle = FGraphBuilder::AddInputNodeHandle(*Metasound, DataTypeName, FText::GetEmpty());
+					Frontend::FNodeHandle NodeHandle = FGraphBuilder::AddInputNodeHandle(*Metasound, DataTypeName);
 					if (ensure(NodeHandle->IsValid()))
 					{
 						NameToSelect = NodeHandle->GetNodeName();
@@ -2464,7 +2464,7 @@ namespace Metasound
 					Metasound->Modify();
 
 					const FName DataTypeName = GetMetasoundDataTypeName<bool>();
-					Frontend::FNodeHandle NodeHandle = FGraphBuilder::AddOutputNodeHandle(*Metasound, DataTypeName, FText::GetEmpty());
+					Frontend::FNodeHandle NodeHandle = FGraphBuilder::AddOutputNodeHandle(*Metasound, DataTypeName);
 					if (ensure(NodeHandle->IsValid()))
 					{
 						NameToSelect = NodeHandle->GetNodeName();
diff --git a/Engine/Plugins/Runtime/Metasound/Source/MetasoundEditor/Private/MetasoundEditorGraph.cpp b/Engine/Plugins/Runtime/Metasound/Source/MetasoundEditor/Private/MetasoundEditorGraph.cpp
--- a/Engine/Plugins/Runtime/Metasound/Source/MetasoundEditor/Private/MetasoundEditorGraph.cpp
+++ b/Engine/Plugins/Runtime/Metasound/Source/MetasoundEditor/Private/MetasoundEditorGraph.cpp
@@ -502,7 +502,7 @@ Metasound::Frontend::FNodeHandle UMetasoundEditorGraphInput::AddNodeHandle(const
 	}
 
 	UObject& Metasound = Graph->GetMetasoundChecked();
-	Metasound::Frontend::FNodeHandle NewNodeHandle = FGraphBuilder::AddInputNodeHandle(Metasound, InDataType, FText::GetEmpty(), nullptr, &InName);
+	Metasound::Frontend::FNodeHandle NewNodeHandle = FGraphBuilder::AddInputNodeHandle(Metasound, InDataType, nullptr /* DefaultValue */, &InName);
 	return NewNodeHandle;
 }
 
@@ -621,7 +621,7 @@ Metasound::Frontend::FNodeHandle UMetasoundEditorGraphOutput::AddNodeHandle(cons
 	}
 
 	UObject& Metasound = Graph->GetMetasoundChecked();
-	FNodeHandle NewNodeHandle = FGraphBuilder::AddOutputNodeHandle(Metasound, InDataType, FText::GetEmpty(), &InName);
+	FNodeHandle NewNodeHandle = FGraphBuilder::AddOutputNodeHandle(Metasound, InDataType, &InName);
 	return NewNodeHandle;
 }
 
diff --git a/Engine/Plugins/Runtime/Metasound/Source/MetasoundEditor/Private/MetasoundEditorGraphBuilder.cpp b/Engine/Plugins/Runtime/Metasound/Source/MetasoundEditor/Private/MetasoundEditorGraphBuilder.cpp
--- a/Engine/Plugins/Runtime/Metasound/Source/MetasoundEditor/Private/MetasoundEditorGraphBuilder.cpp
+++ b/Engine/Plugins/Runtime/Metasound/Source/MetasoundEditor/Private/MetasoundEditorGraphBuilder.cpp
@@ -745,7 +745,8 @@ namespace Metasound
 					if (ensure(Input))
 					{
 						const FName PinName = Pin->GetFName();
-						NodeHandle = AddInputNodeHandle(InMetaSound, Input->GetDataType(), InGraphNode.GetTooltipText(), nullptr, &PinName);
+						NodeHandle = AddInputNodeHandle(InMetaSound, Input->GetDataType(), nullptr, &PinName);
+						NodeHandle->SetDescription(InGraphNode.GetTooltipText());
 					}
 				}
 			}
@@ -760,7 +761,8 @@ namespace Metasound
 					if (ensure(Output))
 					{
 						const FName PinName = Pin->GetFName();
-						NodeHandle = FGraphBuilder::AddOutputNodeHandle(InMetaSound, Output->GetDataType(), InGraphNode.GetTooltipText(), &PinName);
+						NodeHandle = FGraphBuilder::AddOutputNodeHandle(InMetaSound, Output->GetDataType(), &PinName);
+						NodeHandle->SetDescription(InGraphNode.GetTooltipText());
 					}
 				}
 			}
@@ -794,22 +796,22 @@ namespace Metasound
 			return NodeHandle;
 		}
 
-		Frontend::FNodeHandle FGraphBuilder::AddInputNodeHandle(UObject& InMetaSound, const FName InTypeName, const FText& InToolTip, const FMetasoundFrontendLiteral* InDefaultValue, const FName* InNameBase)
+		Frontend::FNodeHandle FGraphBuilder::AddInputNodeHandle(UObject& InMetaSound, const FName InTypeName, const FMetasoundFrontendLiteral* InDefaultValue, const FName* InNameBase)
 		{
 			FMetasoundAssetBase* MetaSoundAsset = IMetasoundUObjectRegistry::Get().GetObjectAsAssetBase(&InMetaSound);
 			check(MetaSoundAsset);
 
 			const FName NewName = GenerateUniqueNameByClassType(InMetaSound, EMetasoundFrontendClassType::Input, InNameBase ? InNameBase->ToString() : TEXT("Input"));
-			return MetaSoundAsset->GetRootGraphHandle()->AddInputVertex(NewName, InTypeName, InToolTip, InDefaultValue);
+			return MetaSoundAsset->GetRootGraphHandle()->AddInputVertex(NewName, InTypeName, InDefaultValue);
 		}
 
-		Frontend::FNodeHandle FGraphBuilder::AddOutputNodeHandle(UObject& InMetaSound, const FName InTypeName, const FText& InToolTip, const FName* InNameBase)
+		Frontend::FNodeHandle FGraphBuilder::AddOutputNodeHandle(UObject& InMetaSound, const FName InTypeName, const FName* InNameBase)
 		{
 			FMetasoundAssetBase* MetaSoundAsset = IMetasoundUObjectRegistry::Get().GetObjectAsAssetBase(&InMetaSound);
 			check(MetaSoundAsset);
 
 			const FName NewName = GenerateUniqueNameByClassType(InMetaSound, EMetasoundFrontendClassType::Output, InNameBase ? InNameBase->ToString() : TEXT("Output"));
-			return MetaSoundAsset->GetRootGraphHandle()->AddOutputVertex(NewName, InTypeName, InToolTip);
+			return MetaSoundAsset->GetRootGraphHandle()->AddOutputVertex(NewName, InTypeName);
 		}
 
 		FName FGraphBuilder::GenerateUniqueVariableName(const Frontend::FConstGraphHandle& InFrontendGraph, const FString& InBaseName)
@@ -1246,7 +1248,7 @@ namespace Metasound
 
 		void FGraphBuilder::RefreshPinMetadata(UEdGraphPin& InPin, const FMetasoundFrontendVertexMetadata& InMetadata)
 		{
-			InPin.PinToolTip = InMetadata.Description.ToString();
+			InPin.PinToolTip = InMetadata.GetDescription().ToString();
 			InPin.bAdvancedView = InMetadata.bIsAdvancedDisplay;
 			if (InPin.bAdvancedView)
 			{
diff --git a/Engine/Plugins/Runtime/Metasound/Source/MetasoundEditor/Private/MetasoundEditorGraphBuilder.h b/Engine/Plugins/Runtime/Metasound/Source/MetasoundEditor/Private/MetasoundEditorGraphBuilder.h
--- a/Engine/Plugins/Runtime/Metasound/Source/MetasoundEditor/Private/MetasoundEditorGraphBuilder.h
+++ b/Engine/Plugins/Runtime/Metasound/Source/MetasoundEditor/Private/MetasoundEditorGraphBuilder.h
@@ -110,7 +110,6 @@ namespace Metasound
 			static Frontend::FNodeHandle AddInputNodeHandle(
 				UObject& InMetaSound,
 				const FName InTypeName,
-				const FText& InToolTip,
 				const FMetasoundFrontendLiteral* InDefaultValue = nullptr,
 				const FName* InNameBase = nullptr);
 
@@ -131,7 +130,7 @@ namespace Metasound
 			static UMetasoundEditorGraphOutputNode* AddOutputNode(UObject& InMetaSound, Frontend::FNodeHandle& InNodeHandle, FVector2D InLocation, bool bInSelectNewNode = true);
 
 			// Generates analogous FNodeHandle for the given internal node data. Does not bind nor create EdGraph representation of given node.
-			static Frontend::FNodeHandle AddOutputNodeHandle(UObject& InMetaSound, const FName InTypeName, const FText& InToolTip, const FName* InNameBase = nullptr);
+			static Frontend::FNodeHandle AddOutputNodeHandle(UObject& InMetaSound, const FName InTypeName, const FName* InNameBase = nullptr);
 
 			// Create a unique name for the variable.
 			static FName GenerateUniqueVariableName(const Frontend::FConstGraphHandle& InFrontendGraph, const FString& InBaseName);
diff --git a/Engine/Plugins/Runtime/Metasound/Source/MetasoundEditor/Private/MetasoundEditorGraphSchema.cpp b/Engine/Plugins/Runtime/Metasound/Source/MetasoundEditor/Private/MetasoundEditorGraphSchema.cpp
--- a/Engine/Plugins/Runtime/Metasound/Source/MetasoundEditor/Private/MetasoundEditorGraphSchema.cpp
+++ b/Engine/Plugins/Runtime/Metasound/Source/MetasoundEditor/Private/MetasoundEditorGraphSchema.cpp
@@ -474,7 +474,7 @@ UEdGraphNode* FMetasoundGraphSchemaAction_PromoteToInput::PerformAction(UEdGraph
 	FGraphBuilder::GetPinLiteral(*FromPin, DefaultValue);
 
 	const FName InputName = InputHandle->GetName();
-	FNodeHandle NodeHandle = FGraphBuilder::AddInputNodeHandle(ParentMetasound, InputHandle->GetDataType(), FText::GetEmpty(), &DefaultValue, &InputName);
+	FNodeHandle NodeHandle = FGraphBuilder::AddInputNodeHandle(ParentMetasound, InputHandle->GetDataType(), &DefaultValue, &InputName);
 	if (ensure(NodeHandle->IsValid()))
 	{
 		UMetasoundEditorGraphInput* Input = MetasoundGraph->FindOrAddInput(NodeHandle);
@@ -686,7 +686,7 @@ UEdGraphNode* FMetasoundGraphSchemaAction_PromoteToOutput::PerformAction(UEdGrap
 
 	const FString OutputName = OutputHandle->GetName().ToString();
 	const FVertexName NewNodeName = FGraphBuilder::GenerateUniqueNameByClassType(ParentMetasound, EMetasoundFrontendClassType::Output, OutputName);
-	FNodeHandle NodeHandle = FGraphBuilder::AddOutputNodeHandle(ParentMetasound, OutputHandle->GetDataType(), FText::GetEmpty(), &NewNodeName);
+	FNodeHandle NodeHandle = FGraphBuilder::AddOutputNodeHandle(ParentMetasound, OutputHandle->GetDataType(), &NewNodeName);
 	if (ensure(NodeHandle->IsValid()))
 	{
 		UMetasoundEditorGraphOutput* Output = MetasoundGraph->FindOrAddOutput(NodeHandle);
diff --git a/Engine/Plugins/Runtime/Metasound/Source/MetasoundEngine/Private/MetasoundEngineArchetypes.cpp b/Engine/Plugins/Runtime/Metasound/Source/MetasoundEngine/Private/MetasoundEngineArchetypes.cpp
--- a/Engine/Plugins/Runtime/Metasound/Source/MetasoundEngine/Private/MetasoundEngineArchetypes.cpp
+++ b/Engine/Plugins/Runtime/Metasound/Source/MetasoundEngine/Private/MetasoundEngineArchetypes.cpp
@@ -101,8 +101,8 @@ namespace Metasound
 				OnPlayTrigger.Name = ArchetypesPrivate::GetOnPlayInputName();
 
 				OnPlayTrigger.TypeName = GetMetasoundDataTypeName<FTrigger>();
-				OnPlayTrigger.Metadata.DisplayName = LOCTEXT("OnPlay", "On Play");
-				OnPlayTrigger.Metadata.Description = LOCTEXT("OnPlayTriggerToolTip", "Trigger executed when this source is played.");
+				OnPlayTrigger.Metadata.SetDisplayName(LOCTEXT("OnPlay", "On Play"));
+				OnPlayTrigger.Metadata.SetDescription(LOCTEXT("OnPlayTriggerToolTip", "Trigger executed when this source is played."));
 				OnPlayTrigger.VertexID = FGuid::NewGuid();
 
 				Interface.Inputs.Add(OnPlayTrigger);
@@ -112,8 +112,8 @@ namespace Metasound
 				OnFinished.Name = ArchetypesPrivate::GetIsFinishedOutputName();
 
 				OnFinished.TypeName = GetMetasoundDataTypeName<FTrigger>();
-				OnFinished.Metadata.DisplayName = LOCTEXT("OnFinished", "On Finished");
-				OnFinished.Metadata.Description = LOCTEXT("OnFinishedToolTip", "Trigger executed to initiate stopping the source.");
+				OnFinished.Metadata.SetDisplayName(LOCTEXT("OnFinished", "On Finished"));
+				OnFinished.Metadata.SetDescription(LOCTEXT("OnFinishedToolTip", "Trigger executed to initiate stopping the source."));
 				OnFinished.VertexID = FGuid::NewGuid();
 
 				Interface.Outputs.Add(OnFinished);
@@ -121,8 +121,8 @@ namespace Metasound
 				FMetasoundFrontendClassVertex GeneratedAudio;
 				GeneratedAudio.Name = GetAudioOutputName();
 				GeneratedAudio.TypeName = GetMetasoundDataTypeName<FMonoAudioFormat>();
-				GeneratedAudio.Metadata.DisplayName = LOCTEXT("GeneratedMono", "Audio");
-				GeneratedAudio.Metadata.Description = LOCTEXT("GeneratedAudioToolTip", "The resulting output audio from this source.");
+				GeneratedAudio.Metadata.SetDisplayName(LOCTEXT("GeneratedMono", "Audio"));
+				GeneratedAudio.Metadata.SetDescription(LOCTEXT("GeneratedAudioToolTip", "The resulting output audio from this source."));
 				GeneratedAudio.VertexID = FGuid::NewGuid();
 
 				Interface.Outputs.Add(GeneratedAudio);
@@ -130,7 +130,6 @@ namespace Metasound
 				// Environment 
 				FMetasoundFrontendClassEnvironmentVariable AudioDeviceID;
 				AudioDeviceID.Name = ArchetypesPrivate::GetAudioDeviceIDVariableName();
-				AudioDeviceID.Metadata.Description = LOCTEXT("AudioDeviceIDToolTip", "Audio Device ID");
 
 				Interface.Environment.Add(AudioDeviceID);
 
@@ -162,8 +161,8 @@ namespace Metasound
 				OnPlayTrigger.Name = ArchetypesPrivate::GetOnPlayInputName();
 
 				OnPlayTrigger.TypeName = GetMetasoundDataTypeName<FTrigger>();
-				OnPlayTrigger.Metadata.DisplayName = LOCTEXT("OnPlay", "On Play");
-				OnPlayTrigger.Metadata.Description = LOCTEXT("OnPlayTriggerToolTip", "Trigger executed when this source is played.");
+				OnPlayTrigger.Metadata.SetDisplayName(LOCTEXT("OnPlay", "On Play"));
+				OnPlayTrigger.Metadata.SetDescription(LOCTEXT("OnPlayTriggerToolTip", "Trigger executed when this source is played."));
 				OnPlayTrigger.VertexID = FGuid::NewGuid();
 
 				Interface.Inputs.Add(OnPlayTrigger);
@@ -173,8 +172,8 @@ namespace Metasound
 				OnFinished.Name = ArchetypesPrivate::GetIsFinishedOutputName();
 
 				OnFinished.TypeName = GetMetasoundDataTypeName<FTrigger>();
-				OnFinished.Metadata.DisplayName = LOCTEXT("OnFinished", "On Finished");
-				OnFinished.Metadata.Description = LOCTEXT("OnFinishedToolTip", "Trigger executed to initiate stopping the source.");
+				OnFinished.Metadata.SetDisplayName(LOCTEXT("OnFinished", "On Finished"));
+				OnFinished.Metadata.SetDescription(LOCTEXT("OnFinishedToolTip", "Trigger executed to initiate stopping the source."));
 				OnFinished.VertexID = FGuid::NewGuid();
 
 				Interface.Outputs.Add(OnFinished);
@@ -182,8 +181,8 @@ namespace Metasound
 				FMetasoundFrontendClassVertex GeneratedAudio;
 				GeneratedAudio.Name = GetAudioOutputName();
 				GeneratedAudio.TypeName = GetMetasoundDataTypeName<FStereoAudioFormat>();
-				GeneratedAudio.Metadata.DisplayName = LOCTEXT("GeneratedStereo", "Audio");
-				GeneratedAudio.Metadata.Description = LOCTEXT("GeneratedAudioToolTip", "The resulting output audio from this source.");
+				GeneratedAudio.Metadata.SetDisplayName(LOCTEXT("GeneratedStereo", "Audio"));
+				GeneratedAudio.Metadata.SetDescription(LOCTEXT("GeneratedAudioToolTip", "The resulting output audio from this source."));
 				GeneratedAudio.VertexID = FGuid::NewGuid();
 
 				Interface.Outputs.Add(GeneratedAudio);
@@ -191,7 +190,6 @@ namespace Metasound
 				// Environment
 				FMetasoundFrontendClassEnvironmentVariable AudioDeviceID;
 				AudioDeviceID.Name = ArchetypesPrivate::GetAudioDeviceIDVariableName();
-				AudioDeviceID.Metadata.Description = LOCTEXT("AudioDeviceIDToolTip", "Audio Device ID");
 
 				Interface.Environment.Add(AudioDeviceID);
 
@@ -218,8 +216,8 @@ namespace Metasound
 				FMetasoundFrontendClassVertex GeneratedAudio;
 				GeneratedAudio.Name = GetAudioOutputName();
 				GeneratedAudio.TypeName = GetMetasoundDataTypeName<FAudioBuffer>();
-				GeneratedAudio.Metadata.DisplayName = LOCTEXT("GeneratedMono", "Audio");
-				GeneratedAudio.Metadata.Description = LOCTEXT("GeneratedAudioToolTip", "The resulting output audio from this source.");
+				GeneratedAudio.Metadata.SetDisplayName(LOCTEXT("GeneratedMono", "Audio"));
+				GeneratedAudio.Metadata.SetDescription(LOCTEXT("GeneratedAudioToolTip", "The resulting output audio from this source."));
 				GeneratedAudio.VertexID = FGuid::NewGuid();
 
 				return GeneratedAudio;
@@ -235,8 +233,8 @@ namespace Metasound
 				OnPlayTrigger.Name = ArchetypesPrivate::GetOnPlayInputName();
 
 				OnPlayTrigger.TypeName = GetMetasoundDataTypeName<FTrigger>();
-				OnPlayTrigger.Metadata.DisplayName = LOCTEXT("OnPlay", "On Play");
-				OnPlayTrigger.Metadata.Description = LOCTEXT("OnPlayTriggerToolTip", "Trigger executed when this source is played.");
+				OnPlayTrigger.Metadata.SetDisplayName(LOCTEXT("OnPlay", "On Play"));
+				OnPlayTrigger.Metadata.SetDescription(LOCTEXT("OnPlayTriggerToolTip", "Trigger executed when this source is played."));
 				OnPlayTrigger.VertexID = FGuid::NewGuid();
 
 				Interface.Inputs.Add(OnPlayTrigger);
@@ -246,8 +244,8 @@ namespace Metasound
 				OnFinished.Name = ArchetypesPrivate::GetIsFinishedOutputName();
 
 				OnFinished.TypeName = GetMetasoundDataTypeName<FTrigger>();
-				OnFinished.Metadata.DisplayName = LOCTEXT("OnFinished", "On Finished");
-				OnFinished.Metadata.Description = LOCTEXT("OnFinishedToolTip", "Trigger executed to initiate stopping the source.");
+				OnFinished.Metadata.SetDisplayName(LOCTEXT("OnFinished", "On Finished"));
+				OnFinished.Metadata.SetDescription(LOCTEXT("OnFinishedToolTip", "Trigger executed to initiate stopping the source."));
 				OnFinished.VertexID = FGuid::NewGuid();
 
 				Interface.Outputs.Add(OnFinished);
@@ -258,7 +256,6 @@ namespace Metasound
 				// Environment
 				FMetasoundFrontendClassEnvironmentVariable AudioDeviceID;
 				AudioDeviceID.Name = ArchetypesPrivate::GetAudioDeviceIDVariableName();
-				AudioDeviceID.Metadata.Description = LOCTEXT("AudioDeviceIDToolTip", "Audio Device ID");
 
 				Interface.Environment.Add(AudioDeviceID);
 
@@ -355,8 +352,8 @@ namespace Metasound
 				FMetasoundFrontendClassVertex GeneratedLeftAudio;
 				GeneratedLeftAudio.Name = GetLeftAudioOutputName();
 				GeneratedLeftAudio.TypeName = GetMetasoundDataTypeName<FAudioBuffer>();
-				GeneratedLeftAudio.Metadata.DisplayName = LOCTEXT("GeneratedStereoLeft", "Left Audio");
-				GeneratedLeftAudio.Metadata.Description = LOCTEXT("GeneratedLeftAudioToolTip", "The resulting output audio from this source.");
+				GeneratedLeftAudio.Metadata.SetDisplayName(LOCTEXT("GeneratedStereoLeft", "Left Audio"));
+				GeneratedLeftAudio.Metadata.SetDescription(LOCTEXT("GeneratedLeftAudioToolTip", "The resulting output audio from this source."));
 				GeneratedLeftAudio.VertexID = FGuid::NewGuid();
 
 				return GeneratedLeftAudio;
@@ -367,8 +364,8 @@ namespace Metasound
 				FMetasoundFrontendClassVertex GeneratedRightAudio;
 				GeneratedRightAudio.Name = GetRightAudioOutputName();
 				GeneratedRightAudio.TypeName = GetMetasoundDataTypeName<FAudioBuffer>();
-				GeneratedRightAudio.Metadata.DisplayName = LOCTEXT("GeneratedStereoRight", "Right Audio");
-				GeneratedRightAudio.Metadata.Description = LOCTEXT("GeneratedRightAudioToolTip", "The resulting output audio from this source.");
+				GeneratedRightAudio.Metadata.SetDisplayName(LOCTEXT("GeneratedStereoRight", "Right Audio"));
+				GeneratedRightAudio.Metadata.SetDescription(LOCTEXT("GeneratedRightAudioToolTip", "The resulting output audio from this source."));
 				GeneratedRightAudio.VertexID = FGuid::NewGuid();
 
 				return GeneratedRightAudio;
@@ -384,8 +381,8 @@ namespace Metasound
 				OnPlayTrigger.Name = ArchetypesPrivate::GetOnPlayInputName();
 
 				OnPlayTrigger.TypeName = GetMetasoundDataTypeName<FTrigger>();
-				OnPlayTrigger.Metadata.DisplayName = LOCTEXT("OnPlay", "On Play");
-				OnPlayTrigger.Metadata.Description = LOCTEXT("OnPlayTriggerToolTip", "Trigger executed when this source is played.");
+				OnPlayTrigger.Metadata.SetDisplayName(LOCTEXT("OnPlay", "On Play"));
+				OnPlayTrigger.Metadata.SetDescription(LOCTEXT("OnPlayTriggerToolTip", "Trigger executed when this source is played."));
 				OnPlayTrigger.VertexID = FGuid::NewGuid();
 
 				Interface.Inputs.Add(OnPlayTrigger);
@@ -395,8 +392,8 @@ namespace Metasound
 				OnFinished.Name = ArchetypesPrivate::GetIsFinishedOutputName();
 
 				OnFinished.TypeName = GetMetasoundDataTypeName<FTrigger>();
-				OnFinished.Metadata.DisplayName = LOCTEXT("OnFinished", "On Finished");
-				OnFinished.Metadata.Description = LOCTEXT("OnFinishedToolTip", "Trigger executed to initiate stopping the source.");
+				OnFinished.Metadata.SetDisplayName(LOCTEXT("OnFinished", "On Finished"));
+				OnFinished.Metadata.SetDescription(LOCTEXT("OnFinishedToolTip", "Trigger executed to initiate stopping the source."));
 				OnFinished.VertexID = FGuid::NewGuid();
 
 				Interface.Outputs.Add(OnFinished);
@@ -410,7 +407,6 @@ namespace Metasound
 				// Environment
 				FMetasoundFrontendClassEnvironmentVariable AudioDeviceID;
 				AudioDeviceID.Name = ArchetypesPrivate::GetAudioDeviceIDVariableName();
-				AudioDeviceID.Metadata.Description = LOCTEXT("AudioDeviceIDToolTip", "Audio Device ID");
 				Interface.Environment.Add(AudioDeviceID);
 
 				return Interface;
diff --git a/Engine/Plugins/Runtime/Metasound/Source/MetasoundEngine/Private/MetasoundInterface.cpp b/Engine/Plugins/Runtime/Metasound/Source/MetasoundEngine/Private/MetasoundInterface.cpp
--- a/Engine/Plugins/Runtime/Metasound/Source/MetasoundEngine/Private/MetasoundInterface.cpp
+++ b/Engine/Plugins/Runtime/Metasound/Source/MetasoundEngine/Private/MetasoundInterface.cpp
@@ -138,8 +138,13 @@ namespace Metasound
 				ClassInput.Name = Input.InitValue.ParamName;
 				ClassInput.DefaultLiteral = FMetasoundFrontendLiteral(Input.InitValue);
 				ClassInput.TypeName = ResolveMemberDataType(Input.DataType, Input.InitValue.ParamType);
-				ClassInput.Metadata.DisplayName = Input.DisplayName;
-				ClassInput.Metadata.Description = Input.Description;
+				ClassInput.Metadata.SetDisplayName(Input.DisplayName);
+				ClassInput.Metadata.SetDescription(Input.Description);
+
+				// Interfaces should never serialize text to avoid desync between
+				// copied versions serialized in assets and those defined in code.
+				ClassInput.Metadata.SetSerializeText(false);
+
 				ClassInput.VertexID = FGuid::NewGuid();
 
 				return ClassInput;
@@ -150,10 +155,14 @@ namespace Metasound
 				FMetasoundFrontendClassOutput ClassOutput;
 				ClassOutput.Name = Output.ParamName;
 				ClassOutput.TypeName = ResolveMemberDataType(Output.DataType, Output.ParamType);
-				ClassOutput.Metadata.DisplayName = Output.DisplayName;
-				ClassOutput.Metadata.Description = Output.Description;
+				ClassOutput.Metadata.SetDisplayName(Output.DisplayName);
+				ClassOutput.Metadata.SetDescription(Output.Description);
 				ClassOutput.VertexID = FGuid::NewGuid();
 
+				// Interfaces should never serialize text to avoid desync between
+				// copied versions serialized in assets and those defined in code.
+				ClassOutput.Metadata.SetSerializeText(false);
+
 				return ClassOutput;
 			});
 
@@ -166,9 +175,6 @@ namespace Metasound
 				// TODO: Either remove type info for environment variables all together or enforce type.
 				// EnvironmentVariable.TypeName = ResolveMemberDataType(Environment.DataType, Environment.ParamType);
 
-				EnvironmentVariable.Metadata.DisplayName = Environment.DisplayName;
-				EnvironmentVariable.Metadata.Description = Environment.Description;
-
 				return EnvironmentVariable;
 			});
 
diff --git a/Engine/Plugins/Runtime/Metasound/Source/MetasoundEngine/Public/Metasound.h b/Engine/Plugins/Runtime/Metasound/Source/MetasoundEngine/Public/Metasound.h
--- a/Engine/Plugins/Runtime/Metasound/Source/MetasoundEngine/Public/Metasound.h
+++ b/Engine/Plugins/Runtime/Metasound/Source/MetasoundEngine/Public/Metasound.h
@@ -61,9 +61,6 @@ namespace Metasound
 		NonEditorRegistrationOptions.bForceReregister = false;
 		NonEditorRegistrationOptions.bRegisterDependencies = true;
 
-		NonEditorRegistrationOptions.bCacheDependencyMetaDataFromRegistry = false;
-		InMetaSound.ClearDependencyRegistryData();
-
 #if WITH_EDITORONLY_DATA
 		if (UMetasoundEditorGraphBase* MetaSoundGraph = Cast<UMetasoundEditorGraphBase>(InMetaSound.GetGraph()))
 		{
@@ -98,12 +95,6 @@ namespace Metasound
 		if (InArchive.IsLoading())
 		{
 			InMetaSound.VersionAsset();
-
-			// Clear dependency registry data when loading to clear out any fields
-			// (ex. text) that should not have been serialized on the asset. Once
-			// a save is run on all pre-5.0 generated assets, this can be safely
-			// only run during pre-saving asset.
-			InMetaSound.ClearDependencyRegistryData();
 		}
 	}
 
diff --git a/Engine/Plugins/Runtime/Metasound/Source/MetasoundFrontend/Private/MetasoundAssetBase.cpp b/Engine/Plugins/Runtime/Metasound/Source/MetasoundFrontend/Private/MetasoundAssetBase.cpp
--- a/Engine/Plugins/Runtime/Metasound/Source/MetasoundFrontend/Private/MetasoundAssetBase.cpp
+++ b/Engine/Plugins/Runtime/Metasound/Source/MetasoundFrontend/Private/MetasoundAssetBase.cpp
@@ -121,12 +121,11 @@ void FMetasoundAssetBase::RegisterGraphWithFrontend(Metasound::Frontend::FMetaSo
 #endif // WITH_EDITORONLY_DATA
 	}
 
+#if WITH_EDITORONLY_DATA
 	// Must be completed after auto-update to ensure all non-transient referenced dependency data is up-to-date (ex.
 	// class version), which is required for most accurately caching current registry metadata.
-	if (InRegistrationOptions.bCacheDependencyMetaDataFromRegistry)
-	{
-		CacheDependencyRegistryData();
-	}
+	CacheRegistryMetadata();
+#endif // WITH_EDITORONLY_DATA
 
 	// Registers node by copying document. Updates to document require re-registration.
 	class FNodeRegistryEntry : public INodeRegistryEntry
@@ -337,37 +336,6 @@ void FMetasoundAssetBase::AddDefaultInterfaces()
 	FModifyRootGraphInterfaces({ }, InitInterfaces).Transform(DocumentHandle);
 }
 
-void FMetasoundAssetBase::CacheDependencyRegistryData()
-{
-	if (FMetasoundFrontendDocument* Document = GetDocument().Get())
-	{
-		for (FMetasoundFrontendClass& Class : Document->Dependencies)
-		{
-			if (!Class.CacheRegistryData())
-			{
-				UE_LOG(LogMetaSound, Warning,
-					TEXT("'%s' failed to cache dependency registry data: Registry missing class with key '%s'"),
-					*GetOwningAssetName(),
-					*Class.Metadata.GetClassName().ToString());
-				UE_LOG(LogMetaSound, Warning,
-					TEXT("Asset '%s' may fail to build runtime graph unless re-registered after dependency with given key is loaded."),
-					*GetOwningAssetName());
-			}
-		}
-	}
-}
-
-void FMetasoundAssetBase::ClearDependencyRegistryData()
-{
-	if (FMetasoundFrontendDocument* Document = GetDocument().Get())
-	{
-		for (FMetasoundFrontendClass& Class : Document->Dependencies)
-		{
-			Class.ClearRegistryData();
-		}
-	}
-}
-
 bool FMetasoundAssetBase::VersionAsset()
 {
 	using namespace Metasound;
@@ -425,6 +393,94 @@ bool FMetasoundAssetBase::VersionAsset()
 }
 
 #if WITH_EDITORONLY_DATA
+void FMetasoundAssetBase::CacheRegistryMetadata()
+{
+	using namespace Metasound::Frontend;
+
+	FMetasoundFrontendDocument* Document = GetDocument().Get();
+	if (!ensure(Document))
+	{
+		return;
+	}
+
+	using FNameDataTypePair = TPair<FName, FName>;
+	const TSet<FMetasoundFrontendVersion>& InterfaceVersions = Document->Interfaces;
+	FMetasoundFrontendClassInterface& RootGraphClassInterface = Document->RootGraph.Interface;
+
+	// 1. Gather inputs/outputs managed by interfaces
+	TMap<FNameDataTypePair, FMetasoundFrontendClassInput*> Inputs;
+	for (FMetasoundFrontendClassInput& Input : RootGraphClassInterface.Inputs)
+	{
+		FNameDataTypePair NameDataTypePair = FNameDataTypePair(Input.Name, Input.TypeName);
+		Inputs.Add(MoveTemp(NameDataTypePair), &Input);
+	}
+
+	TMap<FNameDataTypePair, FMetasoundFrontendClassOutput*> Outputs;
+	for (FMetasoundFrontendClassOutput& Output : RootGraphClassInterface.Outputs)
+	{
+		FNameDataTypePair NameDataTypePair = FNameDataTypePair(Output.Name, Output.TypeName);
+		Outputs.Add(MoveTemp(NameDataTypePair), &Output);
+	}
+
+	// 2. Copy metadata for inputs/outputs managed by interfaces, removing them from maps generated
+	for (const FMetasoundFrontendVersion& Version : InterfaceVersions)
+	{
+		const FInterfaceRegistryKey InterfaceKey = GetInterfaceRegistryKey(Version);
+		const IInterfaceRegistryEntry* Entry = IInterfaceRegistry::Get().FindInterfaceRegistryEntry(InterfaceKey);
+		if (ensure(Entry))
+		{
+			for (const FMetasoundFrontendClassInput& InterfaceInput : Entry->GetInterface().Inputs)
+			{
+				const FNameDataTypePair NameDataTypePair = FNameDataTypePair(InterfaceInput.Name, InterfaceInput.TypeName);
+				if (FMetasoundFrontendClassInput* Input = Inputs.FindRef(NameDataTypePair))
+				{
+					// Interface members are set to not serialize text and thus is handled by copy, so explicit call omitted.
+					Input->Metadata = InterfaceInput.Metadata;
+					Inputs.Remove(NameDataTypePair);
+				}
+			}
+
+			for (const FMetasoundFrontendClassOutput& InterfaceOutput : Entry->GetInterface().Outputs)
+			{
+				const FNameDataTypePair NameDataTypePair = FNameDataTypePair(InterfaceOutput.Name, InterfaceOutput.TypeName);
+				if (FMetasoundFrontendClassOutput* Output = Outputs.FindRef(NameDataTypePair))
+				{
+					// Interface members are set to not serialize text and thus is handled by copy, so explicit call omitted.
+					Output->Metadata = InterfaceOutput.Metadata;
+					Outputs.Remove(NameDataTypePair);
+				}
+			}
+		}
+	}
+
+	// 3. Iterate remaining inputs/outputs not managed by interfaces and set to serialize text
+	// (in case they were orphaned by an interface no longer being implemented).
+	for (TPair<FNameDataTypePair, FMetasoundFrontendClassInput*>& Pair : Inputs)
+	{
+		Pair.Value->Metadata.SetSerializeText(true);
+	}
+
+	for (TPair<FNameDataTypePair, FMetasoundFrontendClassOutput*>& Pair : Outputs)
+	{
+		Pair.Value->Metadata.SetSerializeText(true);
+	}
+
+	// 4. Cache registry data on document dependencies
+	for (FMetasoundFrontendClass& Class : Document->Dependencies)
+	{
+		if (!Class.CacheRegistryMetadata())
+		{
+			UE_LOG(LogMetaSound, Warning,
+				TEXT("'%s' failed to cache dependency registry data: Registry missing class with key '%s'"),
+				*GetOwningAssetName(),
+				*Class.Metadata.GetClassName().ToString());
+			UE_LOG(LogMetaSound, Warning,
+				TEXT("Asset '%s' may fail to build runtime graph unless re-registered after dependency with given key is loaded."),
+				*GetOwningAssetName());
+		}
+	}
+}
+
 bool FMetasoundAssetBase::GetSynchronizationRequired() const
 {
 	return bSynchronizationRequired;
@@ -601,7 +657,7 @@ Metasound::Frontend::FNodeHandle FMetasoundAssetBase::AddInputPinForSendAddress(
 
 	Description.Name = InSendInfo.Address.GetChannelName();
 	Description.TypeName = Metasound::GetMetasoundDataTypeName<Metasound::FSendAddress>();
-	Description.Metadata.Description = FText::GetEmpty();
+	Description.Metadata.SetDescription(FText::GetEmpty());
 	Description.VertexID = VertexID;
 	Description.DefaultLiteral.Set(InSendInfo.Address.GetChannelName().ToString());
 
diff --git a/Engine/Plugins/Runtime/Metasound/Source/MetasoundFrontend/Private/MetasoundFrontend.cpp b/Engine/Plugins/Runtime/Metasound/Source/MetasoundFrontend/Private/MetasoundFrontend.cpp
--- a/Engine/Plugins/Runtime/Metasound/Source/MetasoundFrontend/Private/MetasoundFrontend.cpp
+++ b/Engine/Plugins/Runtime/Metasound/Source/MetasoundFrontend/Private/MetasoundFrontend.cpp
@@ -56,8 +56,11 @@ namespace Metasound
 				ClassInput.TypeName = InputVertex.GetDataTypeName();
 				ClassInput.VertexID = FGuid::NewGuid();
 
+
 				const FDataVertexMetadata& VertexMetadata = InputVertex.GetMetadata();
-				ClassInput.Metadata.Description = VertexMetadata.Description;
+
+				ClassInput.Metadata.SetDescription(VertexMetadata.Description);
+
 				ClassInput.Metadata.bIsAdvancedDisplay = VertexMetadata.bIsAdvancedDisplay;
 
 				FLiteral DefaultLiteral = InputVertex.GetDefaultLiteral();
@@ -89,7 +92,7 @@ namespace Metasound
 				ClassOutput.VertexID = FGuid::NewGuid();
 
 				const FDataVertexMetadata& VertexMetadata = OutputTuple.Value.GetMetadata();
-				ClassOutput.Metadata.Description = VertexMetadata.Description;
+				ClassOutput.Metadata.SetDescription(VertexMetadata.Description);
 				ClassOutput.Metadata.bIsAdvancedDisplay = VertexMetadata.bIsAdvancedDisplay;
 
 				// Advanced display items are pushed to bottom below non-advanced
@@ -109,7 +112,6 @@ namespace Metasound
 				FMetasoundFrontendClassEnvironmentVariable EnvVar;
 
 				EnvVar.Name = EnvTuple.Value.GetVertexName();
-				EnvVar.Metadata.Description = EnvTuple.Value.GetDescription();
 				EnvVar.bIsRequired = true;
 
 				ClassInterface.Environment.Add(EnvVar);
diff --git a/Engine/Plugins/Runtime/Metasound/Source/MetasoundFrontend/Private/MetasoundFrontendDocument.cpp b/Engine/Plugins/Runtime/Metasound/Source/MetasoundFrontend/Private/MetasoundFrontendDocument.cpp
--- a/Engine/Plugins/Runtime/Metasound/Source/MetasoundFrontend/Private/MetasoundFrontendDocument.cpp
+++ b/Engine/Plugins/Runtime/Metasound/Source/MetasoundFrontend/Private/MetasoundFrontendDocument.cpp
@@ -184,7 +184,9 @@ bool operator!=(const FMetasoundFrontendClassName& InLHS, const FMetasoundFronte
 void FMetasoundFrontendClassMetadata::SetAuthor(const FText& InAuthor)
 {
 	using namespace Metasound::DocumentPrivate;
-	SetWithChangeID(InAuthor, Author, ChangeID);
+
+	FText& TextToSet = bSerializeText ? Author : AuthorTransient;
+	SetWithChangeID(InAuthor, TextToSet, ChangeID);
 }
 
 void FMetasoundFrontendClassMetadata::SetAutoUpdateManagesInterface(bool bInAutoUpdateManagesInterface)
@@ -196,7 +198,9 @@ void FMetasoundFrontendClassMetadata::SetAutoUpdateManagesInterface(bool bInAuto
 void FMetasoundFrontendClassMetadata::SetCategoryHierarchy(const TArray<FText>& InCategoryHierarchy)
 {
 	using namespace Metasound::DocumentPrivate;
-	SetWithChangeID(InCategoryHierarchy, CategoryHierarchy, ChangeID);
+
+	TArray<FText>& TextToSet = bSerializeText ? CategoryHierarchy : CategoryHierarchyTransient;
+	SetWithChangeID(InCategoryHierarchy, TextToSet, ChangeID);
 }
 
 void FMetasoundFrontendClassMetadata::SetKeywords(const TArray<FText>& InKeywords)
@@ -214,13 +218,17 @@ void FMetasoundFrontendClassMetadata::SetClassName(const FMetasoundFrontendClass
 void FMetasoundFrontendClassMetadata::SetDescription(const FText& InDescription)
 {
 	using namespace Metasound::DocumentPrivate;
-	SetWithChangeID(InDescription, Description, ChangeID);
+
+	FText& TextToSet = bSerializeText ? Description : DescriptionTransient;
+	SetWithChangeID(InDescription, TextToSet, ChangeID);
 }
 
 void FMetasoundFrontendClassMetadata::SetDisplayName(const FText& InDisplayName)
 {
 	using namespace Metasound::DocumentPrivate;
-	SetWithChangeID(InDisplayName, DisplayName, ChangeID);
+
+	FText& TextToSet = bSerializeText ? DisplayName : DisplayNameTransient;
+	SetWithChangeID(InDisplayName, TextToSet, ChangeID);
 }
 
 void FMetasoundFrontendClassMetadata::SetIsDeprecated(bool bInIsDeprecated)
@@ -232,7 +240,45 @@ void FMetasoundFrontendClassMetadata::SetIsDeprecated(bool bInIsDeprecated)
 void FMetasoundFrontendClassMetadata::SetPromptIfMissing(const FText& InPromptIfMissing)
 {
 	using namespace Metasound::DocumentPrivate;
-	SetWithChangeID(InPromptIfMissing, PromptIfMissing, ChangeID);
+	SetWithChangeID(InPromptIfMissing, PromptIfMissingTransient, ChangeID);
+}
+
+void FMetasoundFrontendClassMetadata::SetSerializeText(bool bInSerializeText)
+{
+	if (bSerializeText)
+	{
+		if (!bInSerializeText)
+		{
+			AuthorTransient = Author;
+			DescriptionTransient = Description;
+			DisplayNameTransient = DisplayName;
+
+			Author = { };
+			Description = { };
+			DisplayName = { };
+
+			KeywordsTransient = MoveTemp(Keywords);
+			CategoryHierarchyTransient = MoveTemp(CategoryHierarchy);
+		}
+	}
+	else
+	{
+		if (bInSerializeText)
+		{
+			Author = AuthorTransient;
+			Description = DescriptionTransient;
+			DisplayName = DisplayNameTransient;
+
+			AuthorTransient = { };
+			DescriptionTransient = { };
+			DisplayNameTransient = { };
+
+			Keywords = MoveTemp(KeywordsTransient);
+			CategoryHierarchy = MoveTemp(CategoryHierarchyTransient);
+		}
+	}
+
+	bSerializeText = bInSerializeText;
 }
 
 void FMetasoundFrontendClassMetadata::SetVersion(const FMetasoundFrontendVersionNumber& InVersion)
@@ -241,7 +287,8 @@ void FMetasoundFrontendClassMetadata::SetVersion(const FMetasoundFrontendVersion
 	SetWithChangeID(InVersion, Version, ChangeID);
 }
 
-bool FMetasoundFrontendClass::CacheRegistryData()
+#if WITH_EDITORONLY_DATA
+bool FMetasoundFrontendClass::CacheRegistryMetadata()
 {
 	using namespace Metasound::Frontend;
 
@@ -291,27 +338,14 @@ bool FMetasoundFrontendClass::CacheRegistryData()
 
 	return false;
 }
-
-void FMetasoundFrontendClass::ClearRegistryData()
-{
-	Metadata.DisplayName = { };
-	Metadata.Description = { };
-	Metadata.PromptIfMissing = { };
-	Metadata.Author = { };
-	Metadata.Keywords.Reset();
-	Metadata.CategoryHierarchy.Reset();
-
-	Algo::ForEach(Interface.Inputs, [](FMetasoundFrontendClassInput& Input) { Input.Metadata = { }; });
-	Algo::ForEach(Interface.Outputs, [](FMetasoundFrontendClassOutput& Output) { Output.Metadata = { }; });
-
-	Interface.InputStyle = { };
-	Interface.OutputStyle = { };
-	Style = { };
-}
+#endif // WITH_EDITORONLY_DATA
 
 FMetasoundFrontendClassMetadata FMetasoundFrontendClassMetadata::GenerateClassDescription(const Metasound::FNodeClassMetadata& InNodeClassMetadata, EMetasoundFrontendClassType InType)
 {
 	FMetasoundFrontendClassMetadata NewMetadata;
+
+	NewMetadata.SetSerializeText(false);
+
 	NewMetadata.Type = InType;
 
 	// TODO: This flag is only used by the graph class' metadata.
@@ -322,12 +356,14 @@ FMetasoundFrontendClassMetadata FMetasoundFrontendClassMetadata::GenerateClassDe
 
 	NewMetadata.ClassName = InNodeClassMetadata.ClassName;
 	NewMetadata.Version = { InNodeClassMetadata.MajorVersion, InNodeClassMetadata.MinorVersion };
-	NewMetadata.DisplayName = InNodeClassMetadata.DisplayName;
-	NewMetadata.Description = InNodeClassMetadata.Description;
-	NewMetadata.PromptIfMissing = InNodeClassMetadata.PromptIfMissing;
-	NewMetadata.Author = InNodeClassMetadata.Author;
-	NewMetadata.Keywords = InNodeClassMetadata.Keywords;
-	NewMetadata.CategoryHierarchy = InNodeClassMetadata.CategoryHierarchy;
+
+	NewMetadata.SetDisplayName(InNodeClassMetadata.DisplayName);
+	NewMetadata.SetDescription(InNodeClassMetadata.Description);
+	NewMetadata.SetPromptIfMissing(InNodeClassMetadata.PromptIfMissing);
+	NewMetadata.SetAuthor(InNodeClassMetadata.Author);
+	NewMetadata.SetKeywords(InNodeClassMetadata.Keywords);
+	NewMetadata.SetCategoryHierarchy(InNodeClassMetadata.CategoryHierarchy);
+
 	NewMetadata.bIsDeprecated = InNodeClassMetadata.bDeprecated;
 
 	return NewMetadata;
diff --git a/Engine/Plugins/Runtime/Metasound/Source/MetasoundFrontend/Private/MetasoundFrontendDocumentController.cpp b/Engine/Plugins/Runtime/Metasound/Source/MetasoundFrontend/Private/MetasoundFrontendDocumentController.cpp
--- a/Engine/Plugins/Runtime/Metasound/Source/MetasoundFrontend/Private/MetasoundFrontendDocumentController.cpp
+++ b/Engine/Plugins/Runtime/Metasound/Source/MetasoundFrontend/Private/MetasoundFrontendDocumentController.cpp
@@ -38,6 +38,22 @@ namespace Metasound
 			return Invalid::GetInvalidClassArray();
 		}
 
+		void FDocumentController::IterateDependencies(TFunctionRef<void(FMetasoundFrontendClass&)> InFunction)
+		{
+			if (FMetasoundFrontendDocument* Document = DocumentPtr.Get())
+			{
+				Algo::ForEach(Document->Dependencies, InFunction);
+			}
+		}
+
+		void FDocumentController::IterateDependencies(TFunctionRef<void(const FMetasoundFrontendClass&)> InFunction) const
+		{
+			if (const FMetasoundFrontendDocument* Document = DocumentPtr.Get())
+			{
+				Algo::ForEach(Document->Dependencies, InFunction);
+			}
+		}
+
 		const TArray<FMetasoundFrontendGraphClass>& FDocumentController::GetSubgraphs() const
 		{
 			if (const FMetasoundFrontendDocument* Document = DocumentPtr.Get())
@@ -56,6 +72,14 @@ namespace Metasound
 			return Invalid::GetInvalidGraphClass();
 		}
 
+		void FDocumentController::SetRootGraphClass(FMetasoundFrontendGraphClass&& InClass)
+		{
+			if (FMetasoundFrontendDocument* Doc = DocumentPtr.Get())
+			{
+				Doc->RootGraph = MoveTemp(InClass);
+			}
+		}
+
 		bool FDocumentController::AddDuplicateSubgraph(const FMetasoundFrontendGraphClass& InGraphToCopy, const FMetasoundFrontendDocument& InOtherDocument)
 		{
 			if (FMetasoundFrontendDocument* Document = DocumentPtr.Get())
diff --git a/Engine/Plugins/Runtime/Metasound/Source/MetasoundFrontend/Private/MetasoundFrontendDocumentController.h b/Engine/Plugins/Runtime/Metasound/Source/MetasoundFrontend/Private/MetasoundFrontendDocumentController.h
--- a/Engine/Plugins/Runtime/Metasound/Source/MetasoundFrontend/Private/MetasoundFrontendDocumentController.h
+++ b/Engine/Plugins/Runtime/Metasound/Source/MetasoundFrontend/Private/MetasoundFrontendDocumentController.h
@@ -36,8 +36,11 @@ namespace Metasound
 			bool IsValid() const override;
 
 			const TArray<FMetasoundFrontendClass>& GetDependencies() const override;
+			void IterateDependencies(TFunctionRef<void(FMetasoundFrontendClass&)> InFunction) override;
+			void IterateDependencies(TFunctionRef<void(const FMetasoundFrontendClass&)> InFunction) const override;
 			const TArray<FMetasoundFrontendGraphClass>& GetSubgraphs() const override;
 			const FMetasoundFrontendGraphClass& GetRootGraphClass() const override;
+			void SetRootGraphClass(FMetasoundFrontendGraphClass&& InClass) override;
 
 			FConstClassAccessPtr FindDependencyWithID(FGuid InClassID) const override;
 			FConstGraphClassAccessPtr FindSubgraphWithID(FGuid InClassID) const override;
diff --git a/Engine/Plugins/Runtime/Metasound/Source/MetasoundFrontend/Private/MetasoundFrontendGraphController.cpp b/Engine/Plugins/Runtime/Metasound/Source/MetasoundFrontend/Private/MetasoundFrontendGraphController.cpp
--- a/Engine/Plugins/Runtime/Metasound/Source/MetasoundFrontend/Private/MetasoundFrontendGraphController.cpp
+++ b/Engine/Plugins/Runtime/Metasound/Source/MetasoundFrontend/Private/MetasoundFrontendGraphController.cpp
@@ -893,7 +893,7 @@ namespace Metasound
 			return INodeController::GetInvalidHandle();
 		}
 
-		FNodeHandle FGraphController::AddInputVertex(const FVertexName& InName, const FName InTypeName, const FText& InToolTip, const FMetasoundFrontendLiteral* InDefaultValue)
+		FNodeHandle FGraphController::AddInputVertex(const FVertexName& InName, const FName InTypeName, const FMetasoundFrontendLiteral* InDefaultValue)
 		{
 			const FGuid VertexID = FGuid::NewGuid();
 
@@ -901,7 +901,6 @@ namespace Metasound
 
 			Description.Name = InName;
 			Description.TypeName = InTypeName;
-			Description.Metadata.Description = InToolTip;
 			Description.VertexID = VertexID;
 
 			if (InDefaultValue)
@@ -936,7 +935,7 @@ namespace Metasound
 			return false;
 		}
 
-		FNodeHandle FGraphController::AddOutputVertex(const FVertexName& InName, const FName InTypeName, const FText& InToolTip)
+		FNodeHandle FGraphController::AddOutputVertex(const FVertexName& InName, const FName InTypeName)
 		{
 			const FGuid VertexID = FGuid::NewGuid();
 
@@ -944,7 +943,6 @@ namespace Metasound
 
 			Description.Name = InName;
 			Description.TypeName = InTypeName;
-			Description.Metadata.Description = InToolTip;
 			Description.VertexID = VertexID;
 
 			return AddOutputVertex(Description);
@@ -1106,7 +1104,7 @@ namespace Metasound
 		{
 			if (const FMetasoundFrontendClassInput* Desc = FindInputDescriptionWithName(InName))
 			{
-				return Desc->Metadata.Description;
+				return Desc->Metadata.GetDescription();
 			}
 
 			return FText::GetEmpty();
@@ -1116,7 +1114,7 @@ namespace Metasound
 		{
 			if (const FMetasoundFrontendClassOutput* Desc = FindOutputDescriptionWithName(InName))
 			{
-				return Desc->Metadata.Description;
+				return Desc->Metadata.GetDescription();
 			}
 
 			return FText::GetEmpty();
@@ -1126,7 +1124,7 @@ namespace Metasound
 		{
 			if (FMetasoundFrontendClassInput* Desc = FindInputDescriptionWithName(InName))
 			{
-				Desc->Metadata.DisplayName = InDisplayName;
+				Desc->Metadata.SetDisplayName(InDisplayName);
 			}
 		}
 
@@ -1134,7 +1132,7 @@ namespace Metasound
 		{
 			if (FMetasoundFrontendClassOutput* Desc = FindOutputDescriptionWithName(InName))
 			{
-				Desc->Metadata.DisplayName = InDisplayName;
+				Desc->Metadata.SetDisplayName(InDisplayName);
 			}
 		}
 
@@ -1142,7 +1140,7 @@ namespace Metasound
 		{
 			if (FMetasoundFrontendClassInput* Desc = FindInputDescriptionWithName(InName))
 			{
-				Desc->Metadata.Description = InDescription;
+				Desc->Metadata.SetDescription(InDescription);
 			}
 		}
 
@@ -1150,7 +1148,7 @@ namespace Metasound
 		{
 			if (FMetasoundFrontendClassOutput* Desc = FindOutputDescriptionWithName(InName))
 			{
-				Desc->Metadata.Description = InDescription;
+				Desc->Metadata.SetDescription(InDescription);
 			}
 		}
 
diff --git a/Engine/Plugins/Runtime/Metasound/Source/MetasoundFrontend/Private/MetasoundFrontendGraphController.h b/Engine/Plugins/Runtime/Metasound/Source/MetasoundFrontend/Private/MetasoundFrontendGraphController.h
--- a/Engine/Plugins/Runtime/Metasound/Source/MetasoundFrontend/Private/MetasoundFrontendGraphController.h
+++ b/Engine/Plugins/Runtime/Metasound/Source/MetasoundFrontend/Private/MetasoundFrontendGraphController.h
@@ -104,11 +104,11 @@ namespace Metasound
 			FNodeHandle GetOutputNodeWithName(const FVertexName& InName) override;
 
 			FNodeHandle AddInputVertex(const FMetasoundFrontendClassInput& InDescription) override;
-			FNodeHandle AddInputVertex(const FVertexName& InName, const FName InTypeName, const FText& InToolTip, const FMetasoundFrontendLiteral* InDefaultValue) override;
+			FNodeHandle AddInputVertex(const FVertexName& InName, const FName InTypeName, const FMetasoundFrontendLiteral* InDefaultValue) override;
 			bool RemoveInputVertex(const FVertexName& InName) override;
 
 			FNodeHandle AddOutputVertex(const FMetasoundFrontendClassOutput& InDescription) override;
-			FNodeHandle AddOutputVertex(const FVertexName& InName, const FName InTypeName, const FText& InToolTip) override;
+			FNodeHandle AddOutputVertex(const FVertexName& InName, const FName InTypeName) override;
 			bool RemoveOutputVertex(const FVertexName& InName) override;
 
 			void UpdateInterfaceChangeID() override;
diff --git a/Engine/Plugins/Runtime/Metasound/Source/MetasoundFrontend/Private/MetasoundFrontendInputController.cpp b/Engine/Plugins/Runtime/Metasound/Source/MetasoundFrontend/Private/MetasoundFrontendInputController.cpp
--- a/Engine/Plugins/Runtime/Metasound/Source/MetasoundFrontend/Private/MetasoundFrontendInputController.cpp
+++ b/Engine/Plugins/Runtime/Metasound/Source/MetasoundFrontend/Private/MetasoundFrontendInputController.cpp
@@ -61,7 +61,7 @@ namespace Metasound
 		{
 			if (const FMetasoundFrontendClassInput* ClassInput = ClassInputPtr.Get())
 			{
-				return ClassInput->Metadata.DisplayName;
+				return ClassInput->Metadata.GetDisplayName();
 			}
 			
 			return Invalid::GetInvalidText();
@@ -108,7 +108,7 @@ namespace Metasound
 		{
 			if (const FMetasoundFrontendClassInput* ClassInput = ClassInputPtr.Get())
 			{
-				return ClassInput->Metadata.Description;
+				return ClassInput->Metadata.GetDescription();
 			}
 			
 			return Invalid::GetInvalidText();
@@ -407,12 +407,12 @@ namespace Metasound
 				{
 					// If there the ClassInput exists, combine the variable name and class input name.
 					// of the variable should be added to the names of the vertices.
-					CachedDisplayName = FText::Format(LOCTEXT("OutputNodeInputControllerFormat", "{1} {0}"), OwningOutput->Metadata.DisplayName, ClassInput->Metadata.DisplayName);
+					CachedDisplayName = FText::Format(LOCTEXT("OutputNodeInputControllerFormat", "{1} {0}"), OwningOutput->Metadata.GetDisplayName(), ClassInput->Metadata.GetDisplayName());
 				}
 				else
 				{
 					// If there is not ClassInput, then use the variable name.
-					CachedDisplayName = OwningOutput->Metadata.DisplayName;
+					CachedDisplayName = OwningOutput->Metadata.GetDisplayName();
 				}
 			}
 
@@ -423,7 +423,7 @@ namespace Metasound
 		{
 			if (const FMetasoundFrontendClassOutput* OwningOutput = OwningGraphClassOutputPtr.Get())
 			{
-				return OwningOutput->Metadata.Description;
+				return OwningOutput->Metadata.GetDescription();
 			}
 			
 			return Invalid::GetInvalidText();
@@ -484,7 +484,7 @@ namespace Metasound
 		{
 			if (const FMetasoundFrontendClassInput* OwningInput = OwningGraphClassInputPtr.Get())
 			{
-				return OwningInput->Metadata.DisplayName;
+				return OwningInput->Metadata.GetDisplayName();
 			}
 
 			return Invalid::GetInvalidText();
@@ -494,7 +494,7 @@ namespace Metasound
 		{
 			if (const FMetasoundFrontendClassInput* OwningInput = OwningGraphClassInputPtr.Get())
 			{
-				return OwningInput->Metadata.Description;
+				return OwningInput->Metadata.GetDescription();
 			}
 			
 			return Invalid::GetInvalidText();
diff --git a/Engine/Plugins/Runtime/Metasound/Source/MetasoundFrontend/Private/MetasoundFrontendInvalidController.h b/Engine/Plugins/Runtime/Metasound/Source/MetasoundFrontend/Private/MetasoundFrontendInvalidController.h
--- a/Engine/Plugins/Runtime/Metasound/Source/MetasoundFrontend/Private/MetasoundFrontendInvalidController.h
+++ b/Engine/Plugins/Runtime/Metasound/Source/MetasoundFrontend/Private/MetasoundFrontendInvalidController.h
@@ -321,11 +321,11 @@ namespace Metasound
 			virtual FConstClassOutputAccessPtr FindClassOutputWithName(const FVertexName& InName) const override { return FConstClassOutputAccessPtr(); }
 
 			virtual FNodeHandle AddInputVertex(const FMetasoundFrontendClassInput& InDescription) override { return INodeController::GetInvalidHandle(); }
-			virtual TSharedRef<INodeController> AddInputVertex(const FVertexName& InName, const FName InTypeName, const FText& InToolTip, const FMetasoundFrontendLiteral* InDefaultValue) override { return INodeController::GetInvalidHandle(); }
+			virtual TSharedRef<INodeController> AddInputVertex(const FVertexName& InName, const FName InTypeName, const FMetasoundFrontendLiteral* InDefaultValue) override { return INodeController::GetInvalidHandle(); }
 			virtual bool RemoveInputVertex(const FVertexName& InputName) override { return false; }
 
 			virtual FNodeHandle AddOutputVertex(const FMetasoundFrontendClassOutput& InDescription) override { return INodeController::GetInvalidHandle(); }
-			virtual TSharedRef<INodeController> AddOutputVertex(const FVertexName& InName, const FName InTypeName, const FText& InToolTip) override { return INodeController::GetInvalidHandle(); }
+			virtual TSharedRef<INodeController> AddOutputVertex(const FVertexName& InName, const FName InTypeName) override { return INodeController::GetInvalidHandle(); }
 			virtual bool RemoveOutputVertex(const FVertexName& OutputName) override { return false; }
 
 			// This can be used to determine what kind of property editor we should use for the data type of a given input.
@@ -399,8 +399,11 @@ namespace Metasound
 				virtual bool IsValid() const override { return false; }
 
 				virtual const TArray<FMetasoundFrontendClass>& GetDependencies() const override { return Invalid::GetInvalidClassArray(); }
+				virtual void IterateDependencies(TFunctionRef<void(FMetasoundFrontendClass&)> InFunction) override { }
+				virtual void IterateDependencies(TFunctionRef<void(const FMetasoundFrontendClass&)> InFunction) const override { }
 				virtual const TArray<FMetasoundFrontendGraphClass>& GetSubgraphs() const override { return Invalid::GetInvalidGraphClassArray(); }
 				virtual const FMetasoundFrontendGraphClass& GetRootGraphClass() const override { return Invalid::GetInvalidGraphClass(); }
+				virtual void SetRootGraphClass(FMetasoundFrontendGraphClass&& InClass) override { }
 
 				virtual FConstClassAccessPtr FindDependencyWithID(FGuid InClassID) const override { return FConstClassAccessPtr(); }
 				virtual FConstGraphClassAccessPtr FindSubgraphWithID(FGuid InClassID) const override { return FConstGraphClassAccessPtr(); }
diff --git a/Engine/Plugins/Runtime/Metasound/Source/MetasoundFrontend/Private/MetasoundFrontendNodeController.cpp b/Engine/Plugins/Runtime/Metasound/Source/MetasoundFrontend/Private/MetasoundFrontendNodeController.cpp
--- a/Engine/Plugins/Runtime/Metasound/Source/MetasoundFrontend/Private/MetasoundFrontendNodeController.cpp
+++ b/Engine/Plugins/Runtime/Metasound/Source/MetasoundFrontend/Private/MetasoundFrontendNodeController.cpp
@@ -1152,7 +1152,7 @@ namespace Metasound
 		{
 			if (const FMetasoundFrontendClassOutput* OwningOutput = OwningGraphClassOutputPtr.Get())
 			{
-				return OwningOutput->Metadata.Description;
+				return OwningOutput->Metadata.GetDescription();
 			}
 
 			return Invalid::GetInvalidText();
@@ -1162,7 +1162,7 @@ namespace Metasound
 		{
 			if (const FMetasoundFrontendClassOutput* OwningOutput = OwningGraphClassOutputPtr.Get())
 			{
-				return OwningOutput->Metadata.DisplayName;
+				return OwningOutput->Metadata.GetDisplayName();
 			}
 
 			return Invalid::GetInvalidText();
@@ -1173,7 +1173,7 @@ namespace Metasound
 			// TODO: can we remove the const cast by constructing output nodes with a non-const access to class outputs?
 			if (FMetasoundFrontendClassOutput* ClassOutput = ConstCastAccessPtr<FClassOutputAccessPtr>(OwningGraphClassOutputPtr).Get())
 			{
-				ClassOutput->Metadata.Description = InDescription;
+				ClassOutput->Metadata.SetDescription(InDescription);
 				OwningGraph->UpdateInterfaceChangeID();
 			}
 		}
@@ -1208,7 +1208,7 @@ namespace Metasound
 			// TODO: can we remove the const cast by constructing output nodes with a non-const access to class outputs?
 			if (FMetasoundFrontendClassOutput* ClassOutput = ConstCastAccessPtr<FClassOutputAccessPtr>(OwningGraphClassOutputPtr).Get())
 			{
-				ClassOutput->Metadata.DisplayName = InDisplayName;
+				ClassOutput->Metadata.SetDisplayName(InDisplayName);
 				OwningGraph->UpdateInterfaceChangeID();
 			}
 		}
@@ -1253,40 +1253,30 @@ namespace Metasound
 
 			// Test if this node exists on the document's root graph.
 			const bool bIsNodeOnRootGraph = OwningGraph->IsValid() && (RootGraph->GetClassID() == OwningGraph->GetClassID());
-
 			if (bIsNodeOnRootGraph)
 			{
-				// If the node is on the root graph, test if it is in the interfaces' required inputs or outputs.
-				FMetasoundFrontendInterface Interface;
-
-				for (const FMetasoundFrontendVersion& InterfaceVersion : OwningDocument->GetInterfaceVersions())
+				if (const FMetasoundFrontendNode* Node = NodePtr.Get())
 				{
-					FInterfaceRegistryKey InterfaceKey = GetInterfaceRegistryKey(InterfaceVersion);
-					bool bFoundInterface = IInterfaceRegistry::Get().FindInterface(InterfaceKey, Interface);
-					if (bFoundInterface)
+					if (ensure(Node->Interface.Outputs.Num() == 1))
 					{
-						if (const FMetasoundFrontendNode* Node = NodePtr.Get())
+						const FMetasoundFrontendVertex& Output = Node->Interface.Outputs.Last();
+						for (const FMetasoundFrontendVersion& InterfaceVersion : OwningDocument->GetInterfaceVersions())
 						{
-							if (ensure(Node->Interface.Outputs.Num() == 1))
+							FInterfaceRegistryKey InterfaceKey = GetInterfaceRegistryKey(InterfaceVersion);
+							if (const IInterfaceRegistryEntry* Entry = IInterfaceRegistry::Get().FindInterfaceRegistryEntry(InterfaceKey))
 							{
-								const FVertexName& Name = Node->Name;
-								const FName& DataType = Node->Interface.Outputs[0].TypeName;
-								auto IsOutputWithSameNameAndType = [&Name, &DataType](const FMetasoundFrontendClassOutput& InOutput)
+								auto IsOutput = [&Output](const FMetasoundFrontendClassOutput& InterfaceOutput)
 								{
-									return InOutput.Name == Name && InOutput.TypeName == DataType;
+									return FMetasoundFrontendVertex::IsFunctionalEquivalent(Output, InterfaceOutput);
 								};
 
-								if (Interface.Outputs.ContainsByPredicate(IsOutputWithSameNameAndType))
+								if (Entry->GetInterface().Outputs.ContainsByPredicate(IsOutput))
 								{
 									return InterfaceVersion;
 								}
 							}
 						}
 					}
-					else
-					{
-						UE_LOG(LogMetaSound, Warning, TEXT("Document using unregistered interface [InterfaceVersion:%s]"), *InterfaceVersion.ToString());
-					}
 				}
 			}
 
@@ -1413,7 +1403,7 @@ namespace Metasound
 		{
 			if (const FMetasoundFrontendClassInput* OwningInput = OwningGraphClassInputPtr.Get())
 			{
-				return OwningInput->Metadata.Description;
+				return OwningInput->Metadata.GetDescription();
 			}
 
 			return Invalid::GetInvalidText();
@@ -1423,7 +1413,7 @@ namespace Metasound
 		{
 			if (const FMetasoundFrontendClassInput* OwningInput = OwningGraphClassInputPtr.Get())
 			{
-				return OwningInput->Metadata.DisplayName;
+				return OwningInput->Metadata.GetDisplayName();
 			}
 
 			return Invalid::GetInvalidText();
@@ -1444,36 +1434,31 @@ namespace Metasound
 			const bool bIsNodeOnRootGraph = OwningGraph->IsValid() && (RootGraph->GetClassID() == OwningGraph->GetClassID());
 			if (bIsNodeOnRootGraph)
 			{
-				for (const FMetasoundFrontendVersion& InterfaceVersion : OwningDocument->GetInterfaceVersions())
+				if (const FMetasoundFrontendNode* Node = NodePtr.Get())
 				{
-					// If the node is on the root graph, test if it is in the interfaces required inputs.
-					FMetasoundFrontendInterface Interface;
-					FInterfaceRegistryKey InterfaceKey = GetInterfaceRegistryKey(InterfaceVersion);
-					bool bFoundInterface = IInterfaceRegistry::Get().FindInterface(InterfaceKey, Interface);
-					if (bFoundInterface)
+					if (ensure(Node->Interface.Inputs.Num() == 1))
 					{
-						if (const FMetasoundFrontendNode* Node = NodePtr.Get())
+						const TSet<FMetasoundFrontendVersion>& InterfaceVersions = OwningDocument->GetInterfaceVersions();
+						const FMetasoundFrontendVertex& Input = Node->Interface.Inputs.Last();
+						for (const FMetasoundFrontendVersion& InterfaceVersion : InterfaceVersions)
 						{
-							const FVertexName& Name = Node->Name;
-							if (ensure(Node->Interface.Inputs.Num() == 1))
+							// If the node is on the root graph, test if it is in the interfaces required inputs.
+							FMetasoundFrontendInterface Interface;
+							FInterfaceRegistryKey InterfaceKey = GetInterfaceRegistryKey(InterfaceVersion);
+							if (const IInterfaceRegistryEntry* Entry = IInterfaceRegistry::Get().FindInterfaceRegistryEntry(InterfaceKey))
 							{
-								const FName& DataType = Node->Interface.Inputs[0].TypeName;
-								auto IsInputWithSameNameAndType = [&Name, &DataType](const FMetasoundFrontendClassInput& InInput)
+								auto IsInput = [&Input](const FMetasoundFrontendClassInput& InterfaceInput)
 								{
-									return InInput.Name == Name && InInput.TypeName == DataType;
+									return FMetasoundFrontendVertex::IsFunctionalEquivalent(Input, InterfaceInput);
 								};
 
-								if (Interface.Inputs.ContainsByPredicate(IsInputWithSameNameAndType))
+								if (Entry->GetInterface().Inputs.ContainsByPredicate(IsInput))
 								{
 									return InterfaceVersion;
 								}
 							}
 						}
 					}
-					else
-					{
-						UE_LOG(LogMetaSound, Warning, TEXT("Document using unregistered interface [InterfaceVersion:%s]"), *InterfaceVersion.ToString());
-					}
 				}
 			}
 
@@ -1485,7 +1470,7 @@ namespace Metasound
 			// TODO: can we remove these const casts by constructing FINputNodeController with non-const access to the class input?
 			if (FMetasoundFrontendClassInput* ClassInput = ConstCastAccessPtr<FClassInputAccessPtr>(OwningGraphClassInputPtr).Get())
 			{
-				ClassInput->Metadata.Description = InDescription;
+				ClassInput->Metadata.SetDescription(InDescription);
 				OwningGraph->UpdateInterfaceChangeID();
 			}
 		}
@@ -1518,7 +1503,7 @@ namespace Metasound
 			// TODO: can we remove these const casts by constructing FINputNodeController with non-const access to the class input?
 			if (FMetasoundFrontendClassInput* ClassInput = ConstCastAccessPtr<FClassInputAccessPtr>(OwningGraphClassInputPtr).Get())
 			{
-				ClassInput->Metadata.DisplayName = InDisplayName;
+				ClassInput->Metadata.SetDisplayName(InDisplayName);
 				OwningGraph->UpdateInterfaceChangeID();
 			}
 		}
diff --git a/Engine/Plugins/Runtime/Metasound/Source/MetasoundFrontend/Private/MetasoundFrontendOutputController.cpp b/Engine/Plugins/Runtime/Metasound/Source/MetasoundFrontend/Private/MetasoundFrontendOutputController.cpp
--- a/Engine/Plugins/Runtime/Metasound/Source/MetasoundFrontend/Private/MetasoundFrontendOutputController.cpp
+++ b/Engine/Plugins/Runtime/Metasound/Source/MetasoundFrontend/Private/MetasoundFrontendOutputController.cpp
@@ -171,7 +171,7 @@ namespace Metasound
 		{
 			if (const FMetasoundFrontendClassOutput* Output = ClassOutputPtr.Get())
 			{
-				return Output->Metadata.DisplayName;
+				return Output->Metadata.GetDisplayName();
 			}
 
 			return Invalid::GetInvalidText();
@@ -181,7 +181,7 @@ namespace Metasound
 		{
 			if (const FMetasoundFrontendClassOutput* Output = ClassOutputPtr.Get())
 			{
-				return Output->Metadata.Description;
+				return Output->Metadata.GetDescription();
 			}
 
 			return Invalid::GetInvalidText();
@@ -241,12 +241,12 @@ namespace Metasound
 				if (const FMetasoundFrontendClassOutput* ClassOutput = ClassOutputPtr.Get())
 				{
 					// If there is a valid ClassOutput, combine the names.
-					CachedDisplayName = FText::Format(LOCTEXT("InputNodeOutputControllerFormat", "{1} {0}"), OwningInput->Metadata.DisplayName, ClassOutput->Metadata.DisplayName);
+					CachedDisplayName = FText::Format(LOCTEXT("InputNodeOutputControllerFormat", "{1} {0}"), OwningInput->Metadata.GetDisplayName(), ClassOutput->Metadata.GetDisplayName());
 				}
 				else
 				{
 					// If there is no valid ClassOutput, use the owning value display name.
-					CachedDisplayName = OwningInput->Metadata.DisplayName;
+					CachedDisplayName = OwningInput->Metadata.GetDisplayName();
 				}
 			}
 
@@ -257,7 +257,7 @@ namespace Metasound
 		{
 			if (const FMetasoundFrontendClassInput* Input = OwningGraphClassInputPtr.Get())
 			{
-				return Input->Metadata.Description;
+				return Input->Metadata.GetDescription();
 			}
 
 			return Invalid::GetInvalidText();
@@ -317,7 +317,7 @@ namespace Metasound
 		{
 			if (const FMetasoundFrontendClassOutput* OwningOutput = OwningGraphClassOutputPtr.Get())
 			{
-				return OwningOutput->Metadata.DisplayName;
+				return OwningOutput->Metadata.GetDisplayName();
 			}
 
 			return Invalid::GetInvalidText();
@@ -327,7 +327,7 @@ namespace Metasound
 		{
 			if (const FMetasoundFrontendClassOutput* OwningOutput = OwningGraphClassOutputPtr.Get())
 			{
-				return OwningOutput->Metadata.Description;
+				return OwningOutput->Metadata.GetDescription();
 			}
 
 			return Invalid::GetInvalidText();
diff --git a/Engine/Plugins/Runtime/Metasound/Source/MetasoundFrontend/Private/MetasoundFrontendTransform.cpp b/Engine/Plugins/Runtime/Metasound/Source/MetasoundFrontend/Private/MetasoundFrontendTransform.cpp
--- a/Engine/Plugins/Runtime/Metasound/Source/MetasoundFrontend/Private/MetasoundFrontendTransform.cpp
+++ b/Engine/Plugins/Runtime/Metasound/Source/MetasoundFrontend/Private/MetasoundFrontendTransform.cpp
@@ -337,6 +337,11 @@ namespace Metasound
 					if (FMetasoundFrontendVertex::IsFunctionalEquivalent(*ClassOutput, OriginalVertex))
 					{
 						NewVertex.NodeID = ClassOutput->NodeID;
+
+						// Interface members do not serialize text to avoid localization
+						// mismatches between assets and interfaces defined in code.
+						NewVertex.Metadata.SetSerializeText(false);
+
 						FNodeHandle OriginalOutputNode = GraphHandle->GetOutputNodeWithName(OriginalVertex.Name);
 						Locations = OriginalOutputNode->GetNodeStyle().Display.Locations;
 						FInputHandle Input = OriginalOutputNode->GetInputWithVertexName(OriginalVertex.Name);
@@ -672,8 +677,10 @@ namespace Metasound
 
 					ClassInput.Name = NodeName;
 					ClassInput.TypeName = Input->GetDataType();
-					ClassInput.Metadata.Description = InputNode->GetDescription();
-					ClassInput.Metadata.DisplayName = Input->GetMetadata().DisplayName;
+
+					ClassInput.Metadata.SetDescription(InputNode->GetDescription());
+					ClassInput.Metadata.SetDisplayName(Input->GetMetadata().GetDisplayName());
+
 					ClassInput.VertexID = FGuid::NewGuid();
 
 					if (const FMetasoundFrontendClassInput* ExistingClassInput = InParentGraph->FindClassInputWithName(NodeName).Get())
@@ -721,8 +728,10 @@ namespace Metasound
 
 					ClassOutput.Name = NodeName;
 					ClassOutput.TypeName = Output->GetDataType();
-					ClassOutput.Metadata.Description = OutputNode->GetDescription();
-					ClassOutput.Metadata.DisplayName = Output->GetMetadata().DisplayName;
+
+					ClassOutput.Metadata.SetDescription(OutputNode->GetDescription());
+					ClassOutput.Metadata.SetDisplayName(Output->GetMetadata().GetDisplayName());
+
 					ClassOutput.VertexID = FGuid::NewGuid();
 
 					if (const FMetasoundFrontendClassOutput* ExistingClassOutput = InParentGraph->FindClassOutputWithName(NodeName).Get())
@@ -1033,6 +1042,81 @@ namespace Metasound
 			}
 		};
 
+		/** Versions document from 1.7 to 1.8. */
+		class FVersionDocument_1_8 : public FVersionDocumentTransform
+		{
+		public:
+			FVersionDocument_1_8() = default;
+
+			FMetasoundFrontendVersionNumber GetTargetVersion() const override
+			{
+				return { 1, 8 };
+			}
+
+			void TransformInternal(FDocumentHandle InDocument) const override
+			{
+				// Do not serialize MetaData text for dependencies as
+				// CacheRegistryData dynamically provides this.
+				InDocument->IterateDependencies([](FMetasoundFrontendClass& Dependency)
+				{
+					constexpr bool bSerializeText = false;
+					Dependency.Metadata.SetSerializeText(bSerializeText);
+
+					for (FMetasoundFrontendClassInput& Input : Dependency.Interface.Inputs)
+					{
+						Input.Metadata.SetSerializeText(false);
+					}
+
+					for (FMetasoundFrontendClassOutput& Output : Dependency.Interface.Outputs)
+					{
+						Output.Metadata.SetSerializeText(false);
+					}
+				});
+
+				const TSet<FMetasoundFrontendVersion>& InterfaceVersions = InDocument->GetInterfaceVersions();
+
+				using FNameDataTypePair = TPair<FName, FName>;
+				TSet<FNameDataTypePair> InterfaceInputs;
+				TSet<FNameDataTypePair> InterfaceOutputs;
+
+				for (const FMetasoundFrontendVersion& Version : InterfaceVersions)
+				{
+					FInterfaceRegistryKey RegistryKey = GetInterfaceRegistryKey(Version);
+					const IInterfaceRegistryEntry* Entry = IInterfaceRegistry::Get().FindInterfaceRegistryEntry(RegistryKey);
+					if (ensure(Entry))
+					{
+						const FMetasoundFrontendInterface& Interface = Entry->GetInterface();
+						Algo::Transform(Interface.Inputs, InterfaceInputs, [](const FMetasoundFrontendClassInput& Input)
+						{
+							return FNameDataTypePair(Input.Name, Input.TypeName);
+						});
+
+						Algo::Transform(Interface.Outputs, InterfaceOutputs, [](const FMetasoundFrontendClassOutput& Output)
+						{
+							return FNameDataTypePair(Output.Name, Output.TypeName);
+						});
+					}
+				}
+
+				// Only serialize MetaData text for inputs owned by the graph (not by interfaces)
+				FMetasoundFrontendGraphClass RootGraphClass = InDocument->GetRootGraphClass();
+				for (FMetasoundFrontendClassInput& Input : RootGraphClass.Interface.Inputs)
+				{
+					const bool bSerializeText = !InterfaceInputs.Contains(FNameDataTypePair(Input.Name, Input.TypeName));
+					Input.Metadata.SetSerializeText(bSerializeText);
+				}
+
+				// Only serialize MetaData text for outputs owned by the graph (not by interfaces)
+				for (FMetasoundFrontendClassOutput& Output : RootGraphClass.Interface.Outputs)
+				{
+					const bool bSerializeText = !InterfaceOutputs.Contains(FNameDataTypePair(Output.Name, Output.TypeName));
+					Output.Metadata.SetSerializeText(bSerializeText);
+				}
+
+				InDocument->SetRootGraphClass(MoveTemp(RootGraphClass));
+			}
+		};
+
 		FVersionDocument::FVersionDocument(FName InName, const FString& InPath)
 			: Name(InName)
 			, Path(InPath)
@@ -1058,6 +1142,7 @@ namespace Metasound
 			bWasUpdated |= FVersionDocument_1_5(Name).Transform(InDocument);
 			bWasUpdated |= FVersionDocument_1_6().Transform(InDocument);
 			bWasUpdated |= FVersionDocument_1_7().Transform(InDocument);
+			bWasUpdated |= FVersionDocument_1_8().Transform(InDocument);
 
 			if (bWasUpdated)
 			{
diff --git a/Engine/Plugins/Runtime/Metasound/Source/MetasoundFrontend/Public/MetasoundAssetBase.h b/Engine/Plugins/Runtime/Metasound/Source/MetasoundFrontend/Public/MetasoundAssetBase.h
--- a/Engine/Plugins/Runtime/Metasound/Source/MetasoundFrontend/Public/MetasoundAssetBase.h
+++ b/Engine/Plugins/Runtime/Metasound/Source/MetasoundFrontend/Public/MetasoundAssetBase.h
@@ -124,15 +124,18 @@ public:
 
 	void AddDefaultInterfaces();
 
-	// Caches commonly used class MetaData that isn't required for finding & building dependencies into runtime graph.
-	void CacheDependencyRegistryData();
-
-	// Clears commonly used class MetaData that isn't required for finding & building dependencies into runtime graph.
-	void ClearDependencyRegistryData();
 
 	bool VersionAsset();
 
 #if WITH_EDITORONLY_DATA
+	/*
+	 * Caches transient metadata (class & vertex) found in the registry
+	 * that is not necessary for serialization or core graph generation.
+	 *
+	 * @return - Whether class was found in the registry & data was cached successfully.
+	 */
+	void CacheRegistryMetadata();
+
 	// TODO: These flags & associated functions are highly UE editor-specific.
 	// Split synchronization requirement flag into synchronization required &
 	// object type refresh or checking frontend class guids when synchronizing.
diff --git a/Engine/Plugins/Runtime/Metasound/Source/MetasoundFrontend/Public/MetasoundAssetManager.h b/Engine/Plugins/Runtime/Metasound/Source/MetasoundFrontend/Public/MetasoundAssetManager.h
--- a/Engine/Plugins/Runtime/Metasound/Source/MetasoundFrontend/Public/MetasoundAssetManager.h
+++ b/Engine/Plugins/Runtime/Metasound/Source/MetasoundFrontend/Public/MetasoundAssetManager.h
@@ -45,9 +45,6 @@ namespace Metasound
 
 			// Attempt to rebuild referenced class keys (only run if class not registered or set to force re-register)
 			bool bRebuildReferencedAssetClassKeys = true;
-
-			// If true, caches registry metadata not required by graph builder on the document's dependency class
-			bool bCacheDependencyMetaDataFromRegistry = true;
 		};
 
 		class METASOUNDFRONTEND_API IMetaSoundAssetManager
diff --git a/Engine/Plugins/Runtime/Metasound/Source/MetasoundFrontend/Public/MetasoundFrontendArchetypeRegistry.h b/Engine/Plugins/Runtime/Metasound/Source/MetasoundFrontend/Public/MetasoundFrontendArchetypeRegistry.h
--- a/Engine/Plugins/Runtime/Metasound/Source/MetasoundFrontend/Public/MetasoundFrontendArchetypeRegistry.h
+++ b/Engine/Plugins/Runtime/Metasound/Source/MetasoundFrontend/Public/MetasoundFrontendArchetypeRegistry.h
@@ -77,9 +77,16 @@ namespace Metasound
 
 			virtual ~IInterfaceRegistry() = default;
 
+			// Register an interface
 			virtual void RegisterInterface(TUniquePtr<IInterfaceRegistryEntry>&& InEntry) = 0;
+
+			// Find an interface entry with the given key. Returns null if entry not found with given key.
 			virtual const IInterfaceRegistryEntry* FindInterfaceRegistryEntry(const FInterfaceRegistryKey& InKey) const = 0;
+
+			// Find an interface with the given key. Returns true if interface is found, false if not.
 			virtual bool FindInterface(const FInterfaceRegistryKey& InKey, FMetasoundFrontendInterface& OutInterface) const = 0;
+
+			// Executes a function against all transactions since a provided ID and provides the current transaction ID.
 			virtual void ForEachRegistryTransactionSince(FRegistryTransactionID InSince, FRegistryTransactionID* OutCurrentRegistryTransactionID, TFunctionRef<void(const FInterfaceRegistryTransaction&)> InFunc) const = 0;
 		};
 	}
diff --git a/Engine/Plugins/Runtime/Metasound/Source/MetasoundFrontend/Public/MetasoundFrontendController.h b/Engine/Plugins/Runtime/Metasound/Source/MetasoundFrontend/Public/MetasoundFrontendController.h
--- a/Engine/Plugins/Runtime/Metasound/Source/MetasoundFrontend/Public/MetasoundFrontendController.h
+++ b/Engine/Plugins/Runtime/Metasound/Source/MetasoundFrontend/Public/MetasoundFrontendController.h
@@ -765,7 +765,7 @@ namespace Metasound
 			 * @return On success, a valid input node handle. On failure, an invalid node handle.
 			 */
 			virtual FNodeHandle AddInputVertex(const FMetasoundFrontendClassInput& InDescription) = 0;
-			virtual FNodeHandle AddInputVertex(const FVertexName& InName, const FName InTypeName, const FText& InToolTip, const FMetasoundFrontendLiteral* InDefaultValue) = 0;
+			virtual FNodeHandle AddInputVertex(const FVertexName& InName, const FName InTypeName, const FMetasoundFrontendLiteral* InDefaultValue) = 0;
 
 			/** Remove the input with the given name. Returns true if successfully removed, false otherwise. */
 			virtual bool RemoveInputVertex(const FVertexName& InputName) = 0;
@@ -776,7 +776,7 @@ namespace Metasound
 			 * @return On success, a valid output node handle. On failure, an invalid node handle.
 			 */
 			virtual FNodeHandle AddOutputVertex(const FMetasoundFrontendClassOutput& InDescription) = 0;
-			virtual FNodeHandle AddOutputVertex(const FVertexName& InName, const FName InTypeName, const FText& InToolTip) = 0;
+			virtual FNodeHandle AddOutputVertex(const FVertexName& InName, const FName InTypeName) = 0;
 
 			/** Remove the output with the given name. Returns true if successfully removed, false otherwise. */
 			virtual bool RemoveOutputVertex(const FVertexName& OutputName) = 0;
@@ -929,14 +929,19 @@ namespace Metasound
 			/** Returns true if the controller is in a valid state. */
 			virtual bool IsValid() const = 0;
 
+			// Iterates all dependent class definitions that exist on the document (do not necessarily correspond
+			// to registered dependencies)
+			virtual void IterateDependencies(TFunctionRef<void(FMetasoundFrontendClass&)> InFunction) = 0;
+			virtual void IterateDependencies(TFunctionRef<void(const FMetasoundFrontendClass&)> InFunction) const = 0;
+
 			// TODO: add info on environment variables.
 			// TODO: consider find/add subgraph
 			// TODO: perhaps functions returning access pointers could be removed from main interface and only exist in FDocumentController.
-			
 			/** Returns an array of all class dependencies for this document. */
 			virtual const TArray<FMetasoundFrontendClass>& GetDependencies() const = 0;
 			virtual const TArray<FMetasoundFrontendGraphClass>& GetSubgraphs() const = 0;
 			virtual const FMetasoundFrontendGraphClass& GetRootGraphClass() const = 0;
+			virtual void SetRootGraphClass(FMetasoundFrontendGraphClass&& InClass) = 0;
 
 			virtual FConstClassAccessPtr FindDependencyWithID(FGuid InClassID) const = 0;
 			virtual FConstGraphClassAccessPtr FindSubgraphWithID(FGuid InClassID) const = 0;
diff --git a/Engine/Plugins/Runtime/Metasound/Source/MetasoundFrontend/Public/MetasoundFrontendDocument.h b/Engine/Plugins/Runtime/Metasound/Source/MetasoundFrontend/Public/MetasoundFrontendDocument.h
--- a/Engine/Plugins/Runtime/Metasound/Source/MetasoundFrontend/Public/MetasoundFrontendDocument.h
+++ b/Engine/Plugins/Runtime/Metasound/Source/MetasoundFrontend/Public/MetasoundFrontendDocument.h
@@ -539,14 +539,26 @@ struct FMetasoundFrontendVertexMetadata
 {
 	GENERATED_BODY()
 
+private:
 	// Display name for a vertex
 	UPROPERTY(EditAnywhere, Category = Parameters, meta = (DisplayName = "Name"))
 	FText DisplayName;
 
+	// Display name for a vertex if vertex is natively defined
+	// (must be transient to avoid localization desync on load)
+	UPROPERTY(Transient)
+	FText DisplayNameTransient;
+
 	// Description of the vertex.
 	UPROPERTY(EditAnywhere, Category = Parameters)
 	FText Description;
 
+	// Description of the vertex if vertex is natively defined
+	// (must be transient to avoid localization desync on load)
+	UPROPERTY(Transient)
+	FText DescriptionTransient;
+
+public:
 	// Keywords associated with the vertex
 	UPROPERTY()
 	TArray<FString> Keywords;
@@ -558,22 +570,80 @@ struct FMetasoundFrontendVertexMetadata
 	// If true, vertex is shown for advanced display.
 	UPROPERTY()
 	bool bIsAdvancedDisplay = false;
-};
 
-USTRUCT()
-struct FMetasoundFrontendClassEnvironmentVariableMetadata
-{
-	GENERATED_BODY()
-
-	// Display name for a environment variable
+private:
+	// Whether or not the given metadata text should be serialized
+	// or is procedurally maintained via auto-update & the referenced
+	// registry class (to avoid localization text desync).  Should be
+	// false for classes serialized as externally-defined dependencies
+	// or interfaces.
 	UPROPERTY()
-	FText DisplayName;
+	bool bSerializeText = true;
 
-	// Description of the environment variable
-	UPROPERTY()
-	FText Description;
-};
+	FText& GetDescription()
+	{
+		return bSerializeText ? Description : DescriptionTransient;
+	}
+
+	FText& GetDisplayName()
+	{
+		return bSerializeText ? DisplayName : DisplayNameTransient;
+	}
 
+public:
+	const FText& GetDescription() const
+	{
+		return bSerializeText ? Description : DescriptionTransient;
+	}
+
+	const FText& GetDisplayName() const
+	{
+		return bSerializeText ? DisplayName : DisplayNameTransient;
+	}
+
+	bool GetSerializeText() const
+	{
+		return bSerializeText;
+	}
+
+	void SetDescription(const FText& InText)
+	{
+		GetDescription() = InText;
+	}
+
+	void SetDisplayName(const FText& InText)
+	{
+		GetDisplayName() = InText;
+	}
+
+	void SetSerializeText(bool bInSerializeText)
+	{
+		if (bSerializeText)
+		{
+			if (!bInSerializeText)
+			{
+				DisplayNameTransient = DisplayName;
+				DescriptionTransient = Description;
+
+				DisplayName = { };
+				Description  = { };
+			}
+		}
+		else
+		{
+			if (bInSerializeText)
+			{
+				DisplayName = DisplayNameTransient;
+				Description = DescriptionTransient;
+
+				DisplayNameTransient = { };
+				DescriptionTransient = { };
+			}
+		}
+		
+		bSerializeText = bInSerializeText;
+	}
+};
 
 
 USTRUCT() 
@@ -690,10 +760,6 @@ struct METASOUNDFRONTEND_API FMetasoundFrontendClassEnvironmentVariable
 	UPROPERTY()
 	FName TypeName;
 
-	// Metadata of environment variable.
-	UPROPERTY()
-	FMetasoundFrontendClassEnvironmentVariableMetadata Metadata;
-
 	// True if the environment variable is needed in order to instantiate a node instance of the class.
 	// TODO: Should be deprecated?
 	UPROPERTY()
@@ -907,21 +973,40 @@ private:
 	UPROPERTY(VisibleAnywhere, Category = Metasound)
 	FText DisplayName;
 
+	UPROPERTY(Transient)
+	FText DisplayNameTransient;
+
 	UPROPERTY(EditAnywhere, Category = Metasound)
 	FText Description;
 
-	UPROPERTY(VisibleAnywhere, Category = Metasound)
-	FText PromptIfMissing;
+	UPROPERTY(Transient)
+	FText DescriptionTransient;
+
+	// TODO: Move to using a non-localized hint path.  Due to localization,
+	// loading & the fact that class registration happens on demand (post serialization),
+	// copying an FText to the referencing document can result in localization ids
+	// mismatched to different text when attempting to gather text.
+	UPROPERTY(Transient)
+	FText PromptIfMissingTransient;
 
 	UPROPERTY(EditAnywhere, Category = Metasound)
 	FText Author;
 
+	UPROPERTY(Transient)
+	FText AuthorTransient;
+
 	UPROPERTY(VisibleAnywhere, Category = Metasound)
 	TArray<FText> Keywords;
 
+	UPROPERTY(Transient)
+	TArray<FText> KeywordsTransient;
+
 	UPROPERTY(EditAnywhere, Category = Metasound)
 	TArray<FText> CategoryHierarchy;
 
+	UPROPERTY(Transient)
+	TArray<FText> CategoryHierarchyTransient;
+
 	// If true, this node is deprecated and should not be used in new MetaSounds.
 	UPROPERTY(VisibleAnywhere, Category = Metasound)
 	bool bIsDeprecated = false;
@@ -932,6 +1017,14 @@ private:
 	UPROPERTY()
 	bool bAutoUpdateManagesInterface = false;
 
+	// Whether or not the given metadata text should be serialized
+	// or is procedurally maintained via auto-update & the referenced
+	// registry class (to avoid localization text desync).  Should be
+	// false for classes serialized as externally-defined dependencies
+	// or interfaces.
+	UPROPERTY()
+	bool bSerializeText = true;
+
 	// ID used to identify if any of the above have been modified,
 	// to determine if the parent class should be auto-updated.
 	UPROPERTY()
@@ -948,14 +1041,14 @@ public:
 		return GET_MEMBER_NAME_CHECKED(FMetasoundFrontendClassMetadata, CategoryHierarchy);
 	}
 
-	static FName GetClassNamePropertyName()
+	static FName GetDescriptionPropertyName()
 	{
-		return GET_MEMBER_NAME_CHECKED(FMetasoundFrontendClassMetadata, ClassName);
+		return GET_MEMBER_NAME_CHECKED(FMetasoundFrontendClassMetadata, Description);
 	}
 
-	static FName GetDescriptionPropertyName()
+	static FName GetClassNamePropertyName()
 	{
-		return GET_MEMBER_NAME_CHECKED(FMetasoundFrontendClassMetadata, Description);
+		return GET_MEMBER_NAME_CHECKED(FMetasoundFrontendClassMetadata, ClassName);
 	}
 
 	static FName GetVersionPropertyName()
@@ -987,32 +1080,32 @@ public:
 
 	const FText& GetDisplayName() const
 	{
-		return DisplayName;
+		return bSerializeText ? DisplayName : DisplayNameTransient;
 	}
 
 	const FText& GetDescription() const
 	{
-		return Description;
+		return bSerializeText ? Description : DescriptionTransient;
 	}
 
 	const FText& GetPromptIfMissing() const
 	{
-		return PromptIfMissing;
+		return PromptIfMissingTransient;
 	}
 
 	const FText& GetAuthor() const
 	{
-		return Author;
+		return bSerializeText ? Author : AuthorTransient;
 	}
 
 	const TArray<FText>& GetKeywords() const
 	{
-		return Keywords;
+		return bSerializeText ? Keywords : KeywordsTransient;
 	}
 
 	const TArray<FText>& GetCategoryHierarchy() const
 	{
-		return CategoryHierarchy;
+		return bSerializeText ? CategoryHierarchy : CategoryHierarchyTransient;
 	}
 
 	const FGuid& GetChangeID() const
@@ -1035,6 +1128,8 @@ public:
 	void SetPromptIfMissing(const FText& InPromptIfMissing);
 	void SetVersion(const FMetasoundFrontendVersionNumber& InVersion);
 
+	void SetSerializeText(bool bInSerializeText);
+
 	void SetType(const EMetasoundFrontendClassType InType)
 	{
 		Type = InType;
@@ -1043,9 +1138,6 @@ public:
 		// External/Internal should probably be a separate field.
 		// ChangeID = FGuid::NewGuid();
 	}
-
-	// Required to allow caching registry data without modifying the ChangeID
-	friend struct FMetasoundFrontendClass;
 };
 
 USTRUCT()
@@ -1076,21 +1168,15 @@ struct METASOUNDFRONTEND_API FMetasoundFrontendClass
 	UPROPERTY()
 	FMetasoundFrontendClassStyle Style;
 
+#if WITH_EDITORONLY_DATA
 	/*
-	 * Caches metadata found in the registry on the class
-	 * that is not necessary for serialization or core graph
-	 * generation.
+	 * Caches transient metadata (class & vertex) found in the registry
+	 * that is not necessary for serialization or core graph generation.
 	 *
 	 * @return - Whether class was found in the registry & data was cached successfully.
 	 */
-	bool CacheRegistryData();
-
-	/*
-	 * Clears cached metadata found in the registry on the class
-	 * that is not necessary for serialization or core graph
-	 * generation.
-	 */
-	void ClearRegistryData();
+	bool CacheRegistryMetadata();
+#endif // WITH_EDITORONLY_DATA
 };
 
 USTRUCT()
diff --git a/Engine/Plugins/Runtime/Metasound/Source/MetasoundFrontend/Public/MetasoundFrontendTransform.h b/Engine/Plugins/Runtime/Metasound/Source/MetasoundFrontend/Public/MetasoundFrontendTransform.h
--- a/Engine/Plugins/Runtime/Metasound/Source/MetasoundFrontend/Public/MetasoundFrontendTransform.h
+++ b/Engine/Plugins/Runtime/Metasound/Source/MetasoundFrontend/Public/MetasoundFrontendTransform.h
@@ -171,7 +171,7 @@ namespace Metasound
 		public:
 			static FMetasoundFrontendVersionNumber GetMaxVersion()
 			{
-				return FMetasoundFrontendVersionNumber { 1, 7 };
+				return FMetasoundFrontendVersionNumber { 1, 8 };
 			}
 
 			FVersionDocument(FName InName, const FString& InPath);
