Rework Metasound Wave Player

diff --git a/Engine/Plugins/Runtime/Metasound/Source/MetasoundEngine/Private/MetasoundWave.cpp b/Engine/Plugins/Runtime/Metasound/Source/MetasoundEngine/Private/MetasoundWave.cpp
--- a/Engine/Plugins/Runtime/Metasound/Source/MetasoundEngine/Private/MetasoundWave.cpp
+++ b/Engine/Plugins/Runtime/Metasound/Source/MetasoundEngine/Private/MetasoundWave.cpp
@@ -10,6 +10,7 @@
 #include "IAudioCodec.h"
 #include "IAudioCodecRegistry.h"
 #include "MetasoundPrimitives.h"
+#include "MetasoundTrace.h"
 #include "Sound/SoundWave.h"
 
 
@@ -23,7 +24,6 @@ FAutoConsoleVariableRef CVarDisableMetasoundWaveAssetCachePriming(
 
 namespace Metasound
 {
-
 	FWaveAsset::FWaveAsset(const TUniquePtr<Audio::IProxyData>& InInitData)
 	{
 		if (InInitData.IsValid())
@@ -71,234 +71,3 @@ namespace Metasound
 		return SoundWaveProxy.IsValid();
 	}
 }
-
-namespace Audio
-{
-	FSimpleDecoderWrapper::FSimpleDecoderWrapper(const InitParams& InInitParams)
-	: OutputSampleRate(InInitParams.OutputSampleRate)
-	, MaxPitchShiftCents(InInitParams.MaxPitchShiftMagnitudeAllowedInOctaves * 1200.0f)
-	, MaxPitchShiftRatio(FMath::Pow(2.0f, InInitParams.MaxPitchShiftMagnitudeAllowedInOctaves))
-	, DecodeBlockSizeInFrames(128) // 128 chosen due to some limitations on some platform decoders (must be an integer multiple of 128)
-	, OutputBlockSizeInFrames(InInitParams.OutputBlockSizeInFrames)
-	{
-		check(OutputBlockSizeInFrames > 0);
-	}
-
-	bool FSimpleDecoderWrapper::SetWave(const FSoundWaveProxyPtr& InWave, float InStartTimeSeconds, float InInitialPitchShiftSemitones)
-	{
-		bool bSuccessful = true;
-
-		if (InWave.IsValid())
-		{
-			// Determine which values differ so that correct things can be updated.
-			const bool bIsDifferentWave = (InWave.Get() != Wave.Get());
-			const bool bUpdateNumChannels = (InWave->GetNumChannels() != NumChannels);
-			const bool bUpdateInputSampleRate = (InWave->GetSampleRate() != InputSampleRate);
-
-			const bool bUpdateAudioFormat = bUpdateNumChannels || bUpdateInputSampleRate || !bIsInitialized;
-			const bool bReinitDecoder = bIsDifferentWave || !bIsInitialized;
-
-			Wave = InWave;
-
-			if (bReinitDecoder)
-			{
-				// try to initialize decoders
-				bSuccessful = InitializeDecodersInternal(Wave, InStartTimeSeconds);
-			}
-			else
-			{
-				// If the wave has not changed, only need to seek.
-				bSuccessful = SeekToTime(InStartTimeSeconds);
-			}
-			bDecoderIsDone = !bSuccessful;
-
-			if (bUpdateAudioFormat)
-			{
-				// initialize input/output data
-				InputSampleRate = Wave->GetSampleRate();
-				FsOutToInRatio = (OutputSampleRate / InputSampleRate);
-
-				NumChannels = Wave->GetNumChannels();
-				DecodeBlockSizeInSamples = DecodeBlockSizeInFrames * NumChannels;
-
-				// set Circular Buffer capacity
-				int32 Capacity = FMath::Max(1, static_cast<int32>(OutputBlockSizeInFrames * NumChannels * (1.0f + FsOutToInRatio * MaxPitchShiftRatio) * 2));
-				OutputCircularBuffer.Reserve(Capacity, true /* bRetainExistingSamples */);
-
-				Resampler.Init(Audio::EResamplingMethod::Linear, FsOutToInRatio, NumChannels);
-				bIsInitialized = true;
-			}
-
-			if (bUpdateNumChannels)
-			{
-				
-				PitchShifter.Reset(Wave->GetNumChannels(), InInitialPitchShiftSemitones);
-				// Have to discard previous samples if the channel count changes. 
-				OutputCircularBuffer.SetNum(0);
-			}
-			else
-			{
-				PitchShifter.UpdatePitchShift(InInitialPitchShiftSemitones);
-			}
-		}
-
-		return bSuccessful;
-	}
-
-	bool FSimpleDecoderWrapper::SeekToTime(const float InSeconds)
-	{
-		if (Wave.IsValid())
-		{
-			// try to initialize decoders
-			bool bSuccessful = InitializeDecodersInternal(Wave, InSeconds);
-			bDecoderIsDone = !bSuccessful;
-
-			return bSuccessful;
-		}
-
-		return false;
-	}
-	
-	bool FSimpleDecoderWrapper::CanGenerateAudio() const
-	{
-		// If there is a valid decoder, than this object can generate audio.
-		const bool bCanDecoderGenerateAudio = !bDecoderIsDone && Input.IsValid() && Output.IsValid() && Decoder.IsValid() && (NumChannels > 0);
-
-		// If there is audio remaining in the output circular buffer, then this can 
-		// generate some audio.
-		const bool bIsAudioInOutputBuffer = (0 != OutputCircularBuffer.Num());
-
-		return bCanDecoderGenerateAudio || bIsAudioInOutputBuffer;
-	}
-
-	uint32 FSimpleDecoderWrapper::GenerateAudio(float* OutputDest, int32 NumOutputFrames, int32& OutNumFramesConsumed, float PitchShiftInCents, bool bIsLooping)
-	{
-		if (0 == NumChannels)
-		{
-			return 0;
-		}
-
-		const uint32 NumOutputSamples = NumOutputFrames * NumChannels;
-
-		OutNumFramesConsumed = 0;
-
-		if (OutputCircularBuffer.Num() < NumOutputSamples)
-		{
-
-			// (multiply by two just to be sure we can handle SRC output size)
-			const int32 MaxNumResamplerOutputFramesPerBlock = FMath::CeilToInt(FsOutToInRatio * DecodeBlockSizeInFrames) * 2;
-			const int32 MaxNumResamplerOutputSamplesPerBlock = MaxNumResamplerOutputFramesPerBlock * NumChannels;
-
-			PreSrcBuffer.Reset(DecodeBlockSizeInSamples);
-			PreSrcBuffer.AddUninitialized(DecodeBlockSizeInSamples);
-
-			PostSrcBuffer.Reset(MaxNumResamplerOutputSamplesPerBlock);
-			PostSrcBuffer.AddUninitialized(MaxNumResamplerOutputSamplesPerBlock);
-
-			Resampler.SetSampleRateRatio(FsOutToInRatio);
-			PitchShifter.UpdatePitchShift(FMath::Clamp(PitchShiftInCents, -MaxPitchShiftCents, MaxPitchShiftCents) / 100.0f);
-		
-
-			// perform SRC and push to circular buffer until we have enough frames for the output
-			while (Decoder && !(bDecoderIsDone || bDecoderHasLooped) && (OutputCircularBuffer.Num() < NumOutputSamples))
-			{
-				// get more audio from the decoder
-				Audio::IDecoderOutput::FPushedAudioDetails Details;
-				const Audio::IDecoder::EDecodeResult  DecodeResult = Decoder->Decode(bIsLooping);
-				const int32 NumFramesDecoded = Output->PopAudio(PreSrcBuffer, Details) / NumChannels;
-
-				bDecoderIsDone = DecodeResult == Audio::IDecoder::EDecodeResult::Finished;
-				bDecoderHasLooped = DecodeResult == Audio::IDecoder::EDecodeResult::Looped;
-
-				OutNumFramesConsumed += NumFramesDecoded;
-				int32 NumResamplerOutputFrames = 0;
-				int32 Error = Resampler.ProcessAudio(PreSrcBuffer.GetData(), NumFramesDecoded, bDecoderIsDone, PostSrcBuffer.GetData(), MaxNumResamplerOutputFramesPerBlock, NumResamplerOutputFrames);
-				ensure(Error == 0);
-
-				if (!PostSrcBuffer.Num() || !NumResamplerOutputFrames)
-				{
-					continue;
-				}
-				int32 OutputFrames = FMath::Min((int32)PostSrcBuffer.Num(), (int32)(NumResamplerOutputFrames * NumChannels));
-
-				// perform linear pitch shift into OutputCircularBuffer
-				if (NumResamplerOutputFrames > 0)
-				{
-					const TArrayView<float> BufferToPitchShift(PostSrcBuffer.GetData(), NumResamplerOutputFrames * NumChannels);
-					PitchShifter.ProcessAudio(BufferToPitchShift, OutputCircularBuffer);
-				}
-			}
-		}
-
-		if (OutputCircularBuffer.Num() >= NumOutputSamples)
-		{
-			OutputCircularBuffer.Pop(OutputDest, NumOutputSamples);
-		}
-		else if (ensure(bDecoderHasLooped || bDecoderIsDone))
-		{
-			bDecoderHasLooped = false;
-
-			const int32 NumSamplesToPop = OutputCircularBuffer.Num();
-			const int32 NumSamplesRemaining = NumOutputSamples - NumSamplesToPop;
-			OutputCircularBuffer.Pop(OutputDest, OutputCircularBuffer.Num());
-			FMemory::Memzero(&OutputDest[NumSamplesToPop], sizeof(float) * NumSamplesRemaining);
-			return NumSamplesToPop;
-		}
-		else
-		{
-			ensureMsgf(false, TEXT("Something went wrong with decoding."));
-			bDecoderIsDone = true;
-			return 0;
-		}
-
-		return NumOutputSamples; // update once we are aware of partial decode on last buffer
-	}
-
-	uint32 FSimpleDecoderWrapper::GetNumChannels() const
-	{
-		return NumChannels;
-	}
-
-	bool FSimpleDecoderWrapper::InitializeDecodersInternal(const FSoundWaveProxyPtr& InWave, float InStartTimeSeconds)
-	{
-		if (!ensure(InWave.IsValid()))
-		{
-			return false;
-		}
-
-		// Input:
-		FName OldFormat = InWave->GetRuntimeFormat();
-		// TODO: Why is this shared? Doesn't look like it needs to be shared. 
-		Input = MakeShareable(Audio::CreateBackCompatDecoderInput(OldFormat, InWave).Release());
-		Input->SeekToTime(InStartTimeSeconds);
-
-		if (!Input)
-		{
-			return false;
-		}
-
-		// acquire codec:
-		ICodecRegistry::FCodecPtr Codec = ICodecRegistry::Get().FindCodecByParsingInput(Input.Get());
-		if (!Codec)
-		{
-			return false;
-		}
-
-		// specify requirements
-		IDecoderOutput::FRequirements Reqs
-		{
-			Float32_Interleaved,
-			static_cast<int32>(DecodeBlockSizeInFrames),
-			static_cast<int32>(OutputSampleRate)
-		};
-
-		// Output:
-		Output = IDecoderOutput::Create(Reqs);
-
-		// Decoder:
-		Decoder = Codec->CreateDecoder(Input.Get(), Output.Get());
-
-		// return true if all the components were successfully create
-		return Input.IsValid() && Output.IsValid() && Decoder.IsValid();
-	}
-} // namespace Audio
diff --git a/Engine/Plugins/Runtime/Metasound/Source/MetasoundEngine/Private/MetasoundWavePlayerNode.cpp b/Engine/Plugins/Runtime/Metasound/Source/MetasoundEngine/Private/MetasoundWavePlayerNode.cpp
--- a/Engine/Plugins/Runtime/Metasound/Source/MetasoundEngine/Private/MetasoundWavePlayerNode.cpp
+++ b/Engine/Plugins/Runtime/Metasound/Source/MetasoundEngine/Private/MetasoundWavePlayerNode.cpp
@@ -1,21 +1,25 @@
 // Copyright Epic Games, Inc. All Rights Reserved.
 
-#include "AudioResampler.h"
-#include "CoreMinimal.h"
+#include "DecoderInputFactory.h"
 #include "DSP/BufferVectorOperations.h"
+#include "DSP/ConvertDeinterleave.h"
+#include "DSP/MultichannelBuffer.h"
+#include "DSP/MultichannelLinearResampler.h"
+#include "IAudioCodec.h"
 #include "MetasoundBuildError.h"
+#include "MetasoundEngineNodesNames.h"
 #include "MetasoundExecutableOperator.h"
 #include "MetasoundLog.h"
 #include "MetasoundNodeRegistrationMacro.h"
 #include "MetasoundPrimitives.h"
-#include "MetasoundWave.h"
+#include "MetasoundTrace.h"
 #include "MetasoundTrigger.h"
-#include "MetasoundEngineNodesNames.h"
 #include "MetasoundVertex.h"
+#include "MetasoundWave.h"
+#include "MetasoundWaveProxyReader.h"
 
 #define LOCTEXT_NAMESPACE "MetasoundWaveNode"
 
-
 namespace Metasound
 {
 	namespace WavePlayerVertexNames
@@ -28,7 +32,6 @@ namespace Metasound
 		static const TCHAR* InputLoopName = TEXT("Loop");
 		static const TCHAR* InputLoopStartName = TEXT("Loop Start");
 		static const TCHAR* InputLoopDurationName = TEXT("Loop Duration");
-		//static const TCHAR* InputStartTimeAsPercentName = TEXT("Start/Loop Times As Percent");
 
 		static const TCHAR* OutputTriggerOnPlayName = TEXT("On Play");
 		static const TCHAR* OutputTriggerOnDoneName = TEXT("On Finished");
@@ -48,9 +51,8 @@ namespace Metasound
 		static FText InputStartTimeTT = LOCTEXT("StartTimeTT", "Time into the wave asset to start (seek) the wave asset. For real-time decoding, the wave asset must be set to seekable!)");
 		static FText InputPitchShiftTT = LOCTEXT("PitchShiftTT", "The pitch shift to use for the wave asset in semitones.");
 		static FText InputLoopTT = LOCTEXT("LoopTT", "Whether or not to loop between the start and specified end times.");
-		static FText InputLoopStartTT = LOCTEXT("LoopStartTT", "When to start the loop. Will be a percentage if \"Start Time As Percent\" is true.");
-		static FText InputLoopDurationTT = LOCTEXT("LoopDurationTT", "The duration of the loop when wave player is enabled for looping. A value of -1.0 will loop the whole wave asset.");
-		//static FText InputStartTimeAsPercentTT = LOCTEXT("StartTimeAsPercentTT", "Whether to treat the start time and loop start times as a percentage of total wave duration vs seconds (default).");
+		static FText InputLoopStartTT = LOCTEXT("LoopStartTT", "When to start the loop.");
+		static FText InputLoopDurationTT = LOCTEXT("LoopDurationTT", "The duration of the loop when wave player is enabled for looping. A negative value will loop the whole wave asset.");
 
 		static FText OutputTriggerOnPlayTT = LOCTEXT("OnPlayTT", "Triggers when Play is triggered.");
 		static FText OutputTriggerOnDoneTT = LOCTEXT("OnDoneTT", "Triggers when the wave played has finished playing.");
@@ -63,7 +65,6 @@ namespace Metasound
 		static FText OutputPlaybackLocationTT = LOCTEXT("PlaybackLocationTT", "Returns the absolute position of the wave playback as a precentage of wave duration.");
 		static FText OutputAudioLeftNameTT = LOCTEXT("AudioLeftTT", "The left channel audio output. Mono wave assets will be upmixed to dual stereo.");
 		static FText OutputAudioRightNameTT = LOCTEXT("AudioRightTT", "The right channel audio output. Mono wave assets will be upmixed to dual stereo.");
-
 	}
 
 	class FWavePlayerNode : public FNode
@@ -91,6 +92,243 @@ namespace Metasound
 		FVertexInterface Interface;
 	};
 
+	namespace WavePlayerNodePrivate
+	{
+		int32 GetCuePointFrame(const FSoundWaveCuePoint& InPoint)
+		{
+			return InPoint.FramePosition;
+		}
+
+
+		/** FSourceBufferState tracks the current frame and loop indices held in 
+		 * a circular buffer. It describes how the content of a circular buffer
+		 * relates to the frame indices of an FWaveProxyReader 
+		 *
+		 * FSourceBufferState is tied to the implementation of the FWaveProxyReader
+		 * and TCircularAudioBuffer<>, and thus does not serve much purpose outside
+		 * of this wave player node.
+		 *
+		 * However, it does provide a convenient place to perform frame counting
+		 * arithmetic that would otherwise make code more difficult to read.
+		 */
+		struct FSourceBufferState
+		{
+			FSourceBufferState() = default;
+
+			/** Construct a FSourceBufferState
+			 *
+			 * @param InStartFrameIndex - The frame index in the wave corresponding to the first frame in the circular buffer.
+			 * @param InNumFrames - The number of frames in the circular buffer.
+			 * @param bIsLooping - True if the wave player is looping, false if not.
+			 * @param InLoopStartFrameIndexInWave - Frame index in the wave corresponding to a loop start.
+			 * @param InLoopEndFrameIndexInWave - Frame index in the wave corresponding to a loop end.
+			 * @param InEOFFrameIndexInWave - Frame index in the wave corresponding to the end of the file.
+			 */
+			FSourceBufferState(int32 InStartFrameIndex, int32 InNumFrames, bool bIsLooping, int32 InLoopStartFrameIndexInWave, int32 InLoopEndFrameIndexInWave, int32 InEOFFrameIndexInWave)
+			{
+				check(InStartFrameIndex >= 0);
+				check(InNumFrames >= 0);
+				check(InLoopStartFrameIndexInWave >= 0);
+				check(InLoopEndFrameIndexInWave >= 0);
+				check(InEOFFrameIndexInWave >= 0);
+
+				StartFrameIndex = InStartFrameIndex;
+				EndFrameIndex = InStartFrameIndex; // Initialize to starting frame index. Will be adjusted during call to Append()
+				EOFFrameIndexInBuffer = InEOFFrameIndexInWave - InStartFrameIndex;
+				LoopEndFrameIndexInBuffer = InLoopEndFrameIndexInWave - InStartFrameIndex;
+				LoopStartFrameIndexInWave = InLoopStartFrameIndexInWave;
+				LoopEndFrameIndexInWave = InLoopEndFrameIndexInWave;
+				EOFFrameIndexInWave = InEOFFrameIndexInWave;
+
+				Append(InNumFrames, bIsLooping);
+			}
+
+			/** Construct an FSourceBufferState
+			 *
+			 * @param ProxyReader - The wave proxy reader producing the audio.
+			 * @parma InSourceBuffer - The audio buffer holding a range of samples popped from the reader.
+			 */
+			FSourceBufferState(const FWaveProxyReader& ProxyReader, const Audio::FMultichannelCircularBuffer& InSourceBuffer)
+			: FSourceBufferState(ProxyReader.GetFrameIndex(), Audio::GetMultichannelBufferNumFrames(InSourceBuffer), ProxyReader.IsLooping(), ProxyReader.GetLoopStartFrameIndex(), ProxyReader.GetLoopEndFrameIndex(), ProxyReader.GetNumFramesInWave())
+			{
+			}
+
+			/** Track frames removed from the circular buffer. This generally coincides
+			 * with a Pop(...) call to the circular buffer. */
+			void Advance(int32 InNumFrames, bool bIsLooping)
+			{
+				check(InNumFrames >= 0);
+
+				StartFrameIndex += InNumFrames;
+				if (bIsLooping)
+				{
+					StartFrameIndex = WrapLoop(StartFrameIndex);
+				}
+
+				EOFFrameIndexInBuffer = EOFFrameIndexInWave - StartFrameIndex;
+				LoopEndFrameIndexInBuffer = LoopEndFrameIndexInWave - StartFrameIndex;
+			}
+
+
+			/** Track frames appended to the source buffer. This generally coincides
+			 * with a Push(...) call to the circular buffer. */
+			void Append(int32 InNumFrames, bool bIsLooping)
+			{
+				check(InNumFrames >= 0);
+				EndFrameIndex += InNumFrames;
+
+				if (bIsLooping)
+				{
+					EndFrameIndex = WrapLoop(EndFrameIndex);
+				}
+			}
+
+			/** Update loop frame indices. */
+			void SetLoopFrameIndices(int32 InLoopStartFrameIndexInWave, int32 InLoopEndFrameIndexInWave)
+			{
+				LoopStartFrameIndexInWave = InLoopStartFrameIndexInWave;
+				LoopEndFrameIndexInWave = InLoopEndFrameIndexInWave;
+				LoopEndFrameIndexInBuffer = LoopEndFrameIndexInWave - StartFrameIndex;
+			}
+
+			/** Update loop frame indices. */
+			void SetLoopFrameIndices(const FWaveProxyReader& InProxyReader)
+			{
+				SetLoopFrameIndices(InProxyReader.GetLoopStartFrameIndex(), InProxyReader.GetLoopEndFrameIndex());
+			}
+
+			/** Returns the corresponding frame index in the wave which corresponds
+			 * to the first frame in the circular buffer. 
+			 */
+			FORCEINLINE int32 GetStartFrameIndex() const
+			{
+				return StartFrameIndex;
+			}
+
+			/** Returns the corresponding frame index in the wave which corresponds
+			 * to the end frame in the circular buffer (non-inclusive).
+			 */
+			FORCEINLINE int32 GetEndFrameIndex() const
+			{
+				return EndFrameIndex;
+			}
+
+			/** Returns the frame index in the wave where the loop starts */
+			FORCEINLINE int32 GetLoopStartFrameIndexInWave() const
+			{
+				return LoopStartFrameIndexInWave;
+			}
+
+			/** Returns the frame index in the wave where the loop end*/
+			FORCEINLINE int32 GetLoopEndFrameIndexInWave() const
+			{
+				return LoopEndFrameIndexInWave;
+			}
+
+			/** Returns the end-of-file frame index in the wave. */
+			FORCEINLINE int32 GetEOFFrameIndexInWave() const
+			{
+				return EOFFrameIndexInWave;
+			}
+
+			/** Returns the frame index in the circular buffer which represents
+			 * the end of file in the wave. */
+			FORCEINLINE int32 GetEOFFrameIndexInBuffer() const
+			{
+				return EOFFrameIndexInBuffer;
+			}
+
+			/** Returns the frame index in the circular buffer which represents
+			 * the ending loop frame index in the wave. */
+			FORCEINLINE int32 GetLoopEndFrameIndexInBuffer() const
+			{
+				return LoopEndFrameIndexInBuffer;
+			}
+
+			/** Returns the ratio of the current frame index divided by the total
+			 * number of frames in the wave. */
+			FORCEINLINE float GetPlaybackFraction() const
+			{
+				const float PlaybackFraction = static_cast<float>(StartFrameIndex) / FMath::Max(static_cast<float>(EOFFrameIndexInWave), 1.f);
+				return FMath::Max(0.f, PlaybackFraction);
+			}
+
+			/** Returns the ratio of the relative position of the current frame 
+			 * index to the start loop frame index, divided by the number of frames
+			 * in the loop.
+			 * This value can be negative if the current frame index is less
+			 * than the first loop frame index. 
+			 */
+			FORCEINLINE float GetLoopFraction() const
+			{
+				const float LoopNumFrames = static_cast<float>(FMath::Max(1, LoopEndFrameIndexInWave - LoopStartFrameIndexInWave));
+				const float LoopRelativeLocation = static_cast<float>(StartFrameIndex - LoopStartFrameIndexInWave);
+
+				return LoopRelativeLocation / LoopNumFrames;
+			}
+
+			/** Map a index representing a frame in a wave file to an index representing
+			 * a frame in the associated circular buffer. */
+			FORCEINLINE int32 MapFrameInWaveToFrameInBuffer(int32 InFrameIndexInWave, bool bIsLooping) const
+			{
+				if (!bIsLooping || (InFrameIndexInWave >= StartFrameIndex))
+				{
+					return InFrameIndexInWave - StartFrameIndex;
+				}
+				else
+				{
+					const int32 NumFramesFromStartToLoopEnd = LoopEndFrameIndexInWave - StartFrameIndex;
+					const int32 NumFramesFromLoopStartToFrameIndex = InFrameIndexInWave - LoopStartFrameIndexInWave;
+					return NumFramesFromStartToLoopEnd + NumFramesFromLoopStartToFrameIndex;
+				}
+			}
+
+		private:
+
+			int32 WrapLoop(int32 InSourceFrameIndex)
+			{
+				int32 Overshot = InSourceFrameIndex - LoopEndFrameIndexInWave;
+				if (Overshot > 0)
+				{
+					InSourceFrameIndex = LoopStartFrameIndexInWave + Overshot;
+				}
+				return InSourceFrameIndex;
+			}
+
+			int32 StartFrameIndex = INDEX_NONE;
+			int32 EndFrameIndex = INDEX_NONE;
+			int32 EOFFrameIndexInBuffer = INDEX_NONE;
+			int32 LoopEndFrameIndexInBuffer = INDEX_NONE;
+			int32 EOFFrameIndexInWave = INDEX_NONE;
+			int32 LoopStartFrameIndexInWave = INDEX_NONE;
+			int32 LoopEndFrameIndexInWave = INDEX_NONE;
+		};
+
+		/** FSourceEvents contains the frame indices of wave events. 
+		 * Indices are INDEX_NONE if they are unset. 
+		 */
+		struct FSourceEvents
+		{
+			/** Frame index of a loop end. */
+			int32 OnLoopFrameIndex = INDEX_NONE;
+			/** Frame index of an end of file. */
+			int32 OnEOFFrameIndex = INDEX_NONE;
+			/** Frame index of a cue points. */
+			int32 OnCuePointFrameIndex = INDEX_NONE;
+			/** Cue point associated with OnCuePointFrameIndex. */
+			const FSoundWaveCuePoint* CuePoint = nullptr;
+
+			/** Clear all frame indices and associated data. */
+			void Reset()
+			{
+				OnLoopFrameIndex = INDEX_NONE;
+				OnEOFFrameIndex = INDEX_NONE;
+				OnCuePointFrameIndex = INDEX_NONE;
+				CuePoint = nullptr;
+			}
+		};
+	}
+
 	struct FWavePlayerOpArgs
 	{
 		FOperatorSettings Settings;
@@ -101,14 +339,22 @@ namespace Metasound
 		FFloatReadRef PitchShift;
 		FBoolReadRef bLoop;
 		FTimeReadRef LoopStartTime;
-		FTimeReadRef LoopDurationSeconds;
-		//FBoolReadRef bStartTimeAsPercentage;
+		FTimeReadRef LoopDuration;
 	};
 
+	/** MetaSound operator for the wave player node. */
 	class FWavePlayerOperator : public TExecutableOperator<FWavePlayerOperator>
 	{	
 	public:
-		static constexpr float MaxPitchShiftOctaves = 6.0f;
+
+		// Maximum absolute pitch shift in octaves. 
+		static constexpr float MaxAbsPitchShiftInOctaves = 6.0f;
+		// Maximum decode size in frames. 
+		static constexpr int32 MaxDecodeSizeInFrames = 8192;
+		// Number of output audio channels. 
+		static constexpr int32 NumOutputChannels = 2;
+		// Block size for deinterleaving audio. 
+		static constexpr int32 DeinterleaveBlockSizeInFrames = 512;
 
 		FWavePlayerOperator(const FWavePlayerOpArgs& InArgs)
 			: OperatorSettings(InArgs.Settings)
@@ -116,12 +362,10 @@ namespace Metasound
 			, StopTrigger(InArgs.StopTrigger)
 			, WaveAsset(InArgs.WaveAsset)
 			, StartTime(InArgs.StartTime)
-			//, bStartTimeAsPercentage(InArgs.bStartTimeAsPercentage)
 			, PitchShift(InArgs.PitchShift)
 			, bLoop(InArgs.bLoop)
 			, LoopStartTime(InArgs.LoopStartTime)
-			, LoopDurationSeconds(InArgs.LoopDurationSeconds)
-			, TriggerOnPlay(FTriggerWriteRef::CreateNew(InArgs.Settings))
+			, LoopDuration(InArgs.LoopDuration)
 			, TriggerOnDone(FTriggerWriteRef::CreateNew(InArgs.Settings))
 			, TriggerOnNearlyDone(FTriggerWriteRef::CreateNew(InArgs.Settings))
 			, TriggerOnLooped(FTriggerWriteRef::CreateNew(InArgs.Settings))
@@ -132,12 +376,13 @@ namespace Metasound
 			, PlaybackLocation(FFloatWriteRef::CreateNew(0.0f))
 			, AudioBufferL(FAudioBufferWriteRef::CreateNew(InArgs.Settings))
 			, AudioBufferR(FAudioBufferWriteRef::CreateNew(InArgs.Settings))
-			, Decoder(Audio::FSimpleDecoderWrapper::InitParams{static_cast<uint32>(InArgs.Settings.GetNumFramesPerBlock()), InArgs.Settings.GetSampleRate(), 6.f /*MaxPitchShiftMagnitudeAllowedInOctages */})
-			, OutputSampleRate(InArgs.Settings.GetSampleRate())
-			, OutputBlockSizeInFrames(InArgs.Settings.GetNumFramesPerBlock())
 		{
-			check(OutputSampleRate);
-			check(AudioBufferL->Num() == OutputBlockSizeInFrames && AudioBufferR->Num() == OutputBlockSizeInFrames);
+			// Hold on to a view of the output audio. Audio buffers are only writable
+			// by this object and will not be reallocated. 
+			OutputAudioView.Emplace(AudioBufferL->GetData(), AudioBufferL->Num());
+			OutputAudioView.Emplace(AudioBufferR->GetData(), AudioBufferR->Num());
+
+			check(OutputAudioView.Num() == NumOutputChannels);
 		}
 
 		virtual FDataReferenceCollection GetInputs() const override
@@ -152,8 +397,7 @@ namespace Metasound
 			InputDataReferences.AddDataReadReference(InputPitchShiftName, PitchShift);
 			InputDataReferences.AddDataReadReference(InputLoopName, bLoop);
 			InputDataReferences.AddDataReadReference(InputLoopStartName, LoopStartTime);
-			InputDataReferences.AddDataReadReference(InputLoopDurationName, LoopDurationSeconds);
-			//InputDataReferences.AddDataReadReference(InputStartTimeAsPercentName, bStartTimeAsPercentage);
+			InputDataReferences.AddDataReadReference(InputLoopDurationName, LoopDuration);
 			return InputDataReferences;
 		}
 
@@ -162,7 +406,7 @@ namespace Metasound
 			using namespace WavePlayerVertexNames;
 
 			FDataReferenceCollection OutputDataReferences;
-			OutputDataReferences.AddDataReadReference(OutputTriggerOnPlayName, TriggerOnPlay);
+			OutputDataReferences.AddDataReadReference(OutputTriggerOnPlayName, PlayTrigger);
 			OutputDataReferences.AddDataReadReference(OutputTriggerOnDoneName, TriggerOnDone);
 			OutputDataReferences.AddDataReadReference(OutputTriggerOnNearlyDoneName, TriggerOnNearlyDone);
 			OutputDataReferences.AddDataReadReference(OutputTriggerOnLoopedName, TriggerOnLooped);
@@ -176,43 +420,47 @@ namespace Metasound
 			return OutputDataReferences;
 		}
 
-		float UpdateAndGetLoopStartTime()
-		{
-			float LoopSeekTime = 0.0f;
-// 			if (*bStartTimeAsPercentage)
-// 			{
-// 				float LoopStartTimeClamped = FMath::Clamp((float)LoopStartTime->GetSeconds(), 0.0f, 1.0f);
-// 				LoopSeekTime = LoopStartTimeClamped * SoundAssetDurationSeconds;
-// 			}
-// 			else
-			{
-				LoopSeekTime = FMath::Clamp((float)LoopStartTime->GetSeconds(), 0.0f, SoundAssetDurationSeconds);
-			}
-
-			LoopStartFrame = SoundAssetSampleRate * LoopSeekTime;
-			FramesToConsumePlayBeforeLooping = LoopStartFrame;
-
-			if (!bLooped && PlaybackStartFrame > LoopStartFrame)
-			{
-				FramesToConsumePlayBeforeLooping += (SoundAssetNumFrames - PlaybackStartFrame);
-			}
-
-			return LoopSeekTime;
-		}
 
 		void Execute()
 		{
-			TriggerOnPlay->AdvanceBlock();
+			METASOUND_TRACE_CPUPROFILER_EVENT_SCOPE(Metasound::FWavePlayerNode::Execute);
+
+			// Advance all triggers owned by this operator. 
 			TriggerOnDone->AdvanceBlock();
 			TriggerOnNearlyDone->AdvanceBlock();
 			TriggerOnCuePoint->AdvanceBlock();
 			TriggerOnLooped->AdvanceBlock();
 
+			// Update wave proxy reader with any new looping bounds. 
+			if (WaveProxyReader.IsValid())
+			{
+				WaveProxyReader->SetIsLooping(*bLoop);
+				WaveProxyReader->SetLoopStartTime(LoopStartTime->GetSeconds());
+				WaveProxyReader->SetLoopDuration(LoopDuration->GetSeconds());
+				SourceState.SetLoopFrameIndices(*WaveProxyReader);
+			}
+
+			// Update resampler with new frame ratio. 
+			if (Resampler.IsValid())
+			{
+				Resampler->SetFrameRatio(GetFrameRatio(), OperatorSettings.GetNumFramesPerBlock());
+			}
 
 			// zero output buffers
-			FMemory::Memzero(AudioBufferL->GetData(), OutputBlockSizeInFrames * sizeof(float));
-			FMemory::Memzero(AudioBufferR->GetData(), OutputBlockSizeInFrames * sizeof(float));
+			FMemory::Memzero(AudioBufferL->GetData(), OperatorSettings.GetNumFramesPerBlock() * sizeof(float));
+			FMemory::Memzero(AudioBufferR->GetData(), OperatorSettings.GetNumFramesPerBlock() * sizeof(float));
 
+			// Performs execution per sub block based on triggers.
+			ExecuteSubblocks();
+
+			// Updates output playhead information
+			UpdatePlaybackLocation();
+		}
+
+	private:
+
+		void ExecuteSubblocks()
+		{
 			// Parse triggers and render audio
 			int32 PlayTrigIndex = 0;
 			int32 NextPlayFrame = 0;
@@ -224,11 +472,11 @@ namespace Metasound
 
 			int32 CurrAudioFrame = 0;
 			int32 NextAudioFrame = 0;
+			const int32 LastAudioFrame = OperatorSettings.GetNumFramesPerBlock() - 1;
+			const int32 NoTrigger = OperatorSettings.GetNumFramesPerBlock() << 1;
 
-			while (NextAudioFrame < (OutputBlockSizeInFrames -1))
+			while (NextAudioFrame < LastAudioFrame)
 			{
-				const int32 NoTrigger = (OutputBlockSizeInFrames << 1);
-
 				// get the next Play and Stop indices
 				// (play)
 				if (PlayTrigIndex < NumPlayTrigs)
@@ -256,7 +504,7 @@ namespace Metasound
 				// no more triggers, rendering to the end of the block
 				if (NextAudioFrame == NoTrigger)
 				{
-					NextAudioFrame = OutputBlockSizeInFrames;
+					NextAudioFrame = OperatorSettings.GetNumFramesPerBlock();
 				}
 
 				// render audio (while loop handles looping audio)
@@ -264,23 +512,14 @@ namespace Metasound
 				{
 					if (bIsPlaying)
 					{
-						CurrAudioFrame += ExecuteInternal(CurrAudioFrame, NextAudioFrame);
-					}
-					else
-					{
-						CurrAudioFrame = NextAudioFrame;
+						RenderFrameRange(CurrAudioFrame, NextAudioFrame);
 					}
+					CurrAudioFrame = NextAudioFrame;
 				}
 
 				// execute the next trigger
 				if (CurrAudioFrame == NextPlayFrame)
 				{
-					if (*StartTime > 0.0f)
-					{
-						ExecuteSeekRequest();
-					}
-					TriggerOnPlay->TriggerFrame(CurrAudioFrame);
-
 					StartPlaying();
 					++PlayTrigIndex;
 				}
@@ -294,346 +533,423 @@ namespace Metasound
 			}
 		}
 
-		void StartPlaying()
+		void RenderFrameRange(int32 StartFrame, int32 EndFrame)
 		{
-			// Copy the wave asset off on init in case the user changes it while we're playing it.
-			// We'll only check for new wave assets when the current one finishes for sample accurate concantenation
-			CurrentWaveAsset = *WaveAsset;
-
-			if (IsWaveValid())
-			{
-				UE_LOG(LogMetaSound, Verbose, TEXT("Starting Sound: '%s'"), *CurrentWaveAsset->GetFName().ToString());
-				SoundAssetSampleRate = CurrentWaveAsset->GetSampleRate();
-				SoundAssetDurationSeconds = CurrentWaveAsset->GetDuration();
-				SoundAssetNumFrames = CurrentWaveAsset->GetNumFrames();
-				CurrentSoundWaveName = CurrentWaveAsset->GetFName();
-
-				PlaybackStartFrame = FMath::Clamp((float)StartTime->GetSeconds(), 0.0f, SoundAssetDurationSeconds) / SoundAssetSampleRate;
-
-				float StartTimeSeconds = 0.0f;
-				const FSoundWaveProxyPtr& WaveProxy = CurrentWaveAsset.GetSoundWaveProxy();
-// 				if (*bStartTimeAsPercentage)
-// 				{
-// 
-// 					float Percentage = FMath::Clamp((float)StartTime->GetSeconds(), 0.0f, 1.0f);
-// 					StartTimeSeconds = Percentage * WaveProxy->GetDuration();
-// 				}
-// 				else
-				{
-					StartTimeSeconds = FMath::Clamp((float)StartTime->GetSeconds(), 0.0f, WaveProxy->GetDuration());
-				}
-
-				PlaybackStartFrame = (uint32)(StartTimeSeconds * SoundAssetSampleRate);
-				UpdateAndGetLoopStartTime();
+			using namespace Audio;
 
+			METASOUND_TRACE_CPUPROFILER_EVENT_SCOPE(Metasound::FWavePlayerNode::RenderFrameRange);
 
+			// Assume this is set to true and checked by outside callers
+			check(bIsPlaying);
 
-				// Update our current loop frame
-				CurrentConsumedFrameCount = LoopStartFrame;
-
-				// Reset our decoded frame counting for sub loop
-				NumTotalDecodedFramesInLoop = 0;
-
-				*LoopPercent = 0.0f;
+			const int32 NumFramesToGenerate = EndFrame - StartFrame;
+			if (NumFramesToGenerate > 0)
+			{
+				// Trigger any events that occur within this frame range
+				TriggerUpcomingEvents(StartFrame, NumFramesToGenerate, SourceState);
 
-				StartDecoder(StartTimeSeconds, true /* bLogFailures */);
+				// Generate audio
+				FMultichannelBufferView BufferToGenerate = SliceMultichannelBufferView(OutputAudioView, StartFrame, NumFramesToGenerate);
+				GeneratePitchedAudio(BufferToGenerate);
 
-				if (!bIsPlaying)
+				// Check if the source is empty.
+				if (!(*bLoop))
 				{
-					bIsPlaying = Decoder.CanGenerateAudio();
+					bIsPlaying = (SourceState.GetStartFrameIndex() <= SourceState.GetEOFFrameIndexInWave());
 				}
 			}
-
 		}
 
-		void ExecuteSeekRequest()
+		void UpdatePlaybackLocation()
 		{
-			// TODO: Don't do full decoder reset (as performed below) and instead seek decoder.
-			// ex: Decoder.SeekToTime(FMath::Max(0.f, *SeekTimeSeconds));
-			if (!bIsPlaying)
+			*PlaybackLocation = SourceState.GetPlaybackFraction();
+
+			if (*bLoop)
+			{
+				*LoopPercent = SourceState.GetLoopFraction();
+			}
+			else
 			{
-				StartPlaying();
+				*LoopPercent = 0.f;
 			}
 		}
 
-		bool IsWaveValidInternal(const FWaveAsset& InWaveAsset, bool bReportToLog = false) const
+		float GetPitchShiftClamped() const
 		{
-			if (!InWaveAsset.IsSoundWaveValid())
-			{
-				if (bReportToLog)
-				{
-					UE_LOG(LogMetaSound, Error, TEXT("Failed to initialize SoundWave decoder. WavePlayerNode references invalid or missing Wave."));
-				}
-				return false;
-			}
+			return FMath::Clamp(*PitchShift, -12.0f * MaxAbsPitchShiftInOctaves, 12.0f * MaxAbsPitchShiftInOctaves);
+		}
 
-			const FSoundWaveProxyPtr& WaveProxy = InWaveAsset.GetSoundWaveProxy();
+		float GetPitchShiftFrameRatio() const
+		{
+			return FMath::Pow(2.0f, GetPitchShiftClamped() / 12.0f);
+		}
+
+		// Updates the sample rate frame ratio. Used when a new wave proxy reader
+		// is created. 
+		void UpdateSampleRateFrameRatio() 
+		{
+			SampleRateFrameRatio = 1.f;
 
-			// WavePlayerNode currently only supports mono or stereo inputs
-			const int32 NumChannels = WaveProxy->GetNumChannels();
-			if (NumChannels > 2)
+			if (WaveProxyReader.IsValid())
 			{
-				if (bReportToLog)
+				float SourceSampleRate = WaveProxyReader->GetSampleRate();
+				if (SourceSampleRate > 0.f)
 				{
-					UE_LOG(LogMetaSound, Error, TEXT("Failed to initialize SoundWave decoder. WavePlayerNode only supports 2 channels max [Package %s Sound %s: %d Channels])"), *WaveProxy->GetPackageName().ToString(), *WaveProxy->GetFName().ToString(), NumChannels);
+					float TargetSampleRate = OperatorSettings.GetSampleRate();
+					if (TargetSampleRate > 0.f)
+					{
+						SampleRateFrameRatio = SourceSampleRate / OperatorSettings.GetSampleRate();
+					}
 				}
-				return false;
 			}
-
-			return true;
 		}
 
-		bool IsCurrentWaveValid(bool bReportToLog = false) const
+		float GetSampleRateFrameRatio() const
 		{
-			return IsWaveValidInternal(CurrentWaveAsset, bReportToLog);
+			return SampleRateFrameRatio;
 		}
 
-		bool IsWaveValid(bool bReportToLog = false) const
+		float GetFrameRatio() const
 		{
-			return IsWaveValidInternal(*WaveAsset, bReportToLog);
+			return GetSampleRateFrameRatio() * GetPitchShiftFrameRatio();
 		}
 
-		float GetPlaybackRateFromPitchShift(float InPitchShift)
+		float GetMaxPitchShiftFrameRatio() const
 		{
-			return FMath::Pow(2.0f, InPitchShift / 12.0f);
+			return FMath::Pow(2.0f, MaxAbsPitchShiftInOctaves);
 		}
 
-		bool StartDecoder(float SeekTimeSeconds, bool bLogFailures = false)
+		float GetMaxFrameRatio() const
 		{
-			if (IsCurrentWaveValid(bLogFailures))
-			{
-				const FSoundWaveProxyPtr& WaveProxy = CurrentWaveAsset.GetSoundWaveProxy();
-				
-				CuePoints = WaveProxy->GetCuePoints();
-
-				float PitchShiftClamped = FMath::Clamp(*PitchShift, -12.0f * MaxPitchShiftOctaves, 12.0f * MaxPitchShiftOctaves);
-
-				return Decoder.SetWave(WaveProxy, SeekTimeSeconds, PitchShiftClamped);
-			}
-
-			return false;
+			return GetSampleRateFrameRatio() * GetMaxPitchShiftFrameRatio();
 		}
 
-		bool SeekDecoder(float SeekTimeSeconds, bool bLogFailures = false)
+		// Start playing the current wave by creating a wave proxy reader and
+		// recreating the DSP stack.
+		void StartPlaying()
 		{
-			return Decoder.SeekToTime(SeekTimeSeconds);
-		}
+			using namespace WavePlayerNodePrivate;
+			METASOUND_TRACE_CPUPROFILER_EVENT_SCOPE(Metasound::FWavePlayerNode::StartPlaying);
+
+			// MetasoundWavePlayerNode DSP Stack
+			//
+			// Legend:
+			// 	[ObjectName] - An Object which generates or process audio.
+			// 	(BufferName) - A buffer which holds audio.
+			//
+			// [WaveProxyReader]->(InterleavedBuffer)->[ConvertDeinterleave]->(DeinterleavedBuffer)->(SourceCircularBuffer)->[Resampler]->(AudioOutputView)
+			//
 
+			// Copy the wave asset off on init in case the user changes it while we're playing it.
+			// We'll only check for new wave assets when the current one finishes for sample accurate concantenation
+			CurrentWaveAsset = *WaveAsset;
+			FSoundWaveProxyPtr WaveProxy = CurrentWaveAsset.GetSoundWaveProxy();
 
-		int32 ExecuteInternal(int32 StartFrame, int32 EndFrame)
-		{
-			// note: output is hard-coded to stereo (dual-mono)
-			float* FinalOutputLeft = AudioBufferL->GetData() + StartFrame;
-			float* FinalOutputRight = AudioBufferR->GetData() + StartFrame;
-			int32 NumOutputFramesToGenerate = (EndFrame - StartFrame);
-			int32 NumFramesGenerated = 0;
+			bOnNearlyDoneTriggeredForWave = false;
+			bIsPlaying = false;
 
-			const bool bCanDecoderGenerateAudio = Decoder.CanGenerateAudio();
-			if (bCanDecoderGenerateAudio)
-			{
-				ensure(Decoder.CanGenerateAudio());
+			// Reset dsp stack.
+			ResetSourceBufferAndState();
+			WaveProxyReader.Reset();
+			ConvertDeinterleave.Reset();
+			Resampler.Reset();
+			SortedCuePoints.Reset();
 
-				const int32 NumInputChannels = Decoder.GetNumChannels();
-				const bool bNeedsUpmix = (NumInputChannels == 1);
-				const bool bNeedsDeinterleave = !bNeedsUpmix;
+			if (WaveProxy.IsValid())
+			{
+				UE_LOG(LogMetaSound, Verbose, TEXT("Starting Sound: '%s'"), *CurrentWaveAsset->GetFName().ToString());
 
-				const int32 NumSamplesToGenerate = NumOutputFramesToGenerate * NumInputChannels;
+				// Create local sorted copy of cue points.
+				SortedCuePoints = WaveProxy->GetCuePoints();
+				Algo::SortBy(SortedCuePoints, WavePlayerNodePrivate::GetCuePointFrame);
 				
-				float PitchShiftClamped = FMath::Clamp(*PitchShift, -12.0f * MaxPitchShiftOctaves, 12.0f * MaxPitchShiftOctaves);
-
-				PostSrcBuffer.Reset(NumSamplesToGenerate);
-				PostSrcBuffer.AddZeroed(NumSamplesToGenerate);
-				float* PostSrcBufferPtr = PostSrcBuffer.GetData();
+				// Create the wave proxy reader.
+				FWaveProxyReader::FSettings WaveReaderSettings;
+				WaveReaderSettings.MaxDecodeSizeInFrames = MaxDecodeSizeInFrames;
+				WaveReaderSettings.StartTimeInSeconds = StartTime->GetSeconds();
+				WaveReaderSettings.LoopStartTimeInSeconds = LoopStartTime->GetSeconds();
+				WaveReaderSettings.LoopDurationInSeconds = LoopDuration->GetSeconds(); 
+				WaveReaderSettings.bIsLooping = *bLoop;
 
-				// Update looping state which may have changed
-				bool bIsLooping = *bLoop;
+				WaveProxyReader = FWaveProxyReader::Create(WaveProxy.ToSharedRef(), WaveReaderSettings);
 
-				// If we're looping we got some extra logic we need to do before we generate audio
-				if (*bLoop)
+				if (WaveProxyReader.IsValid())
 				{
-					float LoopDuration = SoundAssetDurationSeconds;
-					if (LoopDurationSeconds->GetSeconds() > 0.0f)
-					{
-						LoopDuration = FMath::Clamp((float)LoopDurationSeconds->GetSeconds(), 0.0f, SoundAssetDurationSeconds);
-					}
-					
-					// Give the loop frame a minimum width
-					LoopNumFrames = (uint32)FMath::Max(LoopDuration * SoundAssetSampleRate, 10.0f);
+					UpdateSampleRateFrameRatio();
+					int32 WaveProxyNumChannels = WaveProxyReader->GetNumChannels();
 
-					// We need to look for the case where we are about to loop and we don't want to "overshoot" the loop
-					// So we need to check the number output frames to generate against how much is left in our loop
-					if (NumTotalDecodedFramesInLoop < LoopNumFrames)
+					if (WaveProxyNumChannels > 0)
 					{
-						// The number of frames left in the loop we need to consume from the source file
-						int32 NumFramesLeftToConsumeInLoop = (int32)LoopNumFrames - NumTotalDecodedFramesInLoop;
+						// Create buffer for interleaved audio
+						int32 InterleavedBufferNumSamples = WaveProxyNumChannels * DeinterleaveBlockSizeInFrames;
+						InterleavedBuffer.Reset(InterleavedBufferNumSamples);
+						InterleavedBuffer.AddUninitialized(InterleavedBufferNumSamples);
 
-						// Translate that to a generated frames count taking into account the pitch scale and the playback rate.
-						// Note this is going to be slightly inaccurate due to pitch interpolation... but it should be super close
-						// and w/ loop cross fading implemented, won't matter too much.
-						float PlaybackRate = GetPlaybackRateFromPitchShift(PitchShiftClamped);
+						NumDeinterleaveChannels = NumOutputChannels;
 
-						float SampleRateRatioWithPitchScale = PlaybackRate * SoundAssetSampleRate / OutputSampleRate;
-						int32 NumFramesLeftToGenerateForLoop = SampleRateRatioWithPitchScale * NumFramesLeftToConsumeInLoop;
+						// Create algorithm for channel conversion and deinterleave 
+						ConvertDeinterleave = Audio::IConvertDeinterleave::Create(WaveProxyReader->GetNumChannels(), NumDeinterleaveChannels);
+						Audio::SetMultichannelBufferSize(NumDeinterleaveChannels, DeinterleaveBlockSizeInFrames, DeinterleavedBuffer);
 
-						// Generate the min of the requested frames to generate and the amount we need to finish the loop
-						NumOutputFramesToGenerate = FMath::Min(NumOutputFramesToGenerate, NumFramesLeftToGenerateForLoop);
-					}
-					else
-					{
-						NumOutputFramesToGenerate = 1;
+						// Initialize source buffer
+						int32 FrameCapacity = DeinterleaveBlockSizeInFrames + FMath::CeilToInt(GetMaxFrameRatio() * OperatorSettings.GetNumFramesPerBlock());
+						Audio::SetMultichannelCircularBufferCapacity(NumOutputChannels, FrameCapacity, SourceCircularBuffer);
+						SourceState = FSourceBufferState(*WaveProxyReader, SourceCircularBuffer);
+
+						// Create a resampler.
+						Resampler = MakeUnique<Audio::FMultichannelLinearResampler>(NumDeinterleaveChannels);
+						Resampler->SetFrameRatio(GetFrameRatio(), 0 /* NumFramesToInperolate */);
+
+						// Need to add upmixing if this is not true
+						check(NumDeinterleaveChannels == NumOutputChannels);
 					}
 				}
-				else
-				{
-					LoopNumFrames = 0;
-					NumTotalDecodedFramesInLoop = 0;
-				}
-
-				int32 NumFramesConsumed = 0;
-				if (NumOutputFramesToGenerate > 0)
-				{
-					NumFramesGenerated = Decoder.GenerateAudio(PostSrcBufferPtr, NumOutputFramesToGenerate, NumFramesConsumed, 100.0f * PitchShiftClamped, bIsLooping) / NumInputChannels;
+			}
 
-					// Update loop and playback state logic based on the number of frames consumed
-					FramesConsumedSinceStart += NumFramesConsumed;
+			// If everything was created successfully, start playing.
+			bIsPlaying = WaveProxyReader.IsValid() && ConvertDeinterleave.IsValid() && Resampler.IsValid();
+		}
 
-					int32 PrevConsumedFrameCount = CurrentConsumedFrameCount;
-					CurrentConsumedFrameCount = (CurrentConsumedFrameCount + NumFramesConsumed) % SoundAssetNumFrames;
+		/** Removes all samples from the source buffer and resets SourceState. */
+		void ResetSourceBufferAndState()
+		{
+			using namespace WavePlayerNodePrivate;
+			using namespace Audio;
 
-					*PlaybackLocation = SoundAssetDurationSeconds * ((float)CurrentConsumedFrameCount / SoundAssetNumFrames);
+			SourceState = FSourceBufferState();
+			for (TCircularAudioBuffer<float>& ChannelCircularBuffer : SourceCircularBuffer)
+			{
+				ChannelCircularBuffer.SetNum(0);
+			}
+		}
 
-					// Check for any cue trigger events 
-					FSoundWaveCuePoint OutCuePoint;
-					int32 OutFrameOffset = INDEX_NONE;
-					if (GetCuePointForBlock(PrevConsumedFrameCount, CurrentConsumedFrameCount, OutCuePoint, OutFrameOffset))
-					{
-						// Write the outputs
-						*CuePointID = OutCuePoint.CuePointID;
-						*CuePointLabel = OutCuePoint.Label;
+		/** Generates audio from the wave proxy reader.
+		 *
+		 * @param OutBuffer - Buffer to place generated audio.
+		 * @param OutSourceState - Source state for tracking state of OutBuffer.
+		 */
+		void GenerateSourceAudio(Audio::FMultichannelCircularBuffer& OutBuffer, WavePlayerNodePrivate::FSourceBufferState& OutSourceState)
+		{
+			using namespace WavePlayerNodePrivate;
 
-						// Do the trigger
-						TriggerOnCuePoint->TriggerFrame(StartFrame + OutFrameOffset);
-					}
+			if (bIsPlaying)
+			{
+				const int32 NumExistingFrames = Audio::GetMultichannelBufferNumFrames(OutBuffer);
+				const int32 NumSamplesToGenerate = DeinterleaveBlockSizeInFrames * WaveProxyReader->GetNumChannels();
+ 				check(NumSamplesToGenerate == InterleavedBuffer.Num())
 
-					// Check for nearly done trigger logic. If we're not looping and the next render block will likely finish the file.
-					if (!bIsLooping && (CurrentConsumedFrameCount + OutputBlockSizeInFrames) >= SoundAssetNumFrames)
-					{
-						TriggerOnNearlyDone->TriggerFrame(StartFrame + NumFramesGenerated);
-					}
+				WaveProxyReader->PopAudio(InterleavedBuffer);
+				ConvertDeinterleave->ProcessAudio(InterleavedBuffer, DeinterleavedBuffer);
 
+				for (int32 ChannelIndex = 0; ChannelIndex < NumDeinterleaveChannels; ChannelIndex++)
+				{
+					OutBuffer[ChannelIndex].Push(DeinterleavedBuffer[ChannelIndex]);
 				}
+				OutSourceState.Append(DeinterleaveBlockSizeInFrames, *bLoop);
+			}
+			else
+			{
+				OutSourceState = FSourceBufferState();
+			}
+		}
 
-				// If we're looping and we've made it past the loop start point and we've not yet looped, check the looping condition
-				// Note: we need to handle the case where the "start time" of the player is past the loop start point. In that case
-				// we allow the audio file to wrap around and loop from the beginning before checking the loop logic
-				if (bIsLooping && (bLooped || FramesConsumedSinceStart >= FramesToConsumePlayBeforeLooping))
-				{					
-					NumTotalDecodedFramesInLoop += NumFramesConsumed;
 
-					if (NumTotalDecodedFramesInLoop >= LoopNumFrames || !NumOutputFramesToGenerate)
-					{
-						// We've looped now so we will continue to loop between the loop start and end points
-						bLooped = true;
+		/** Updates frame indices of events if the event occurs in the source within
+		 * the frame range. The frame range begins with the start frame in InSourceState
+		 * and continues for InNumSourceFrames in the source buffer. 
+		 *
+		 * @param InSourceState - Description of the current state of the source buffer.
+		 * @param InNumSourceFrames - Number of frames to inspect.
+		 * @param OutEvents - Event structure to fill out. 
+		 */
+		void MapSourceEventsIfInRange(const WavePlayerNodePrivate::FSourceBufferState& InSourceState, int32 InNumSourceFrames, WavePlayerNodePrivate::FSourceEvents& OutEvents)
+		{
+			OutEvents.Reset();
 
-						// Trigger that we looped
-						TriggerOnLooped->TriggerFrame(StartFrame + NumFramesGenerated);
-					
-						// Update the loop start frame based on any recent inputs
-						float LoopSeekTime = UpdateAndGetLoopStartTime();
+			// Loop end
+			if (*bLoop && FMath::IsWithin(SourceState.GetLoopEndFrameIndexInBuffer(), 0, InNumSourceFrames))
+			{
+				OutEvents.OnLoopFrameIndex = FMath::RoundToInt(Resampler->MapInputFrameToOutputFrame(SourceState.GetLoopEndFrameIndexInBuffer()));
+			}
+
+			// End of file
+			if (FMath::IsWithin(SourceState.GetEOFFrameIndexInBuffer(), 0, InNumSourceFrames))
+			{
+				OutEvents.OnEOFFrameIndex = FMath::RoundToInt(Resampler->MapInputFrameToOutputFrame(SourceState.GetEOFFrameIndexInBuffer()));
+			}
 
-						// Update our current loop frame
-						CurrentConsumedFrameCount = LoopStartFrame;
+			// Map Cue point. Since only one can be mapped, map the first one found.
+			// The first cue point found has the best chance of being rendered.
+			int32 SearchStartFrameIndexInWave = InSourceState.GetStartFrameIndex();
+			int32 SearchEndFrameIndexInWave = SearchStartFrameIndexInWave + InNumSourceFrames;
 
-						// Reset our decoded frame counting for sub loop
-						NumTotalDecodedFramesInLoop = 0;
+			const bool bFramesCrossLoopBoundary = *bLoop && (OutEvents.OnLoopFrameIndex != INDEX_NONE);
+			if (bFramesCrossLoopBoundary)
+			{
+				SearchEndFrameIndexInWave = SearchStartFrameIndexInWave + SourceState.GetLoopEndFrameIndexInBuffer();
+			}
 
-						*LoopPercent = 0.0f;
+			OutEvents.CuePoint = FindCuePoint(SearchStartFrameIndexInWave, SearchEndFrameIndexInWave);
 
-						// Now seek to the loop start frame if we are looping a custom amount. Otherwise, we let the decoder do the looping.
-						if (LoopDurationSeconds->GetSeconds() > 0.0f)
-						{
-							SeekDecoder(LoopSeekTime);
-						}
+			if (bFramesCrossLoopBoundary)
+			{
+				SearchStartFrameIndexInWave = SourceState.GetLoopStartFrameIndexInWave();
+				int32 RemainingFrames = InNumSourceFrames - SourceState.GetLoopEndFrameIndexInBuffer();
+				SearchEndFrameIndexInWave = RemainingFrames;
+
+				// Only override OutEvents.CuePoint if one exists in this subsection
+				// of the buffer.
+				if (const FSoundWaveCuePoint* CuePoint = FindCuePoint(SearchStartFrameIndexInWave, SearchEndFrameIndexInWave))
+				{
+					if (nullptr == OutEvents.CuePoint)
+					{
+						OutEvents.CuePoint = CuePoint;
 					}
-					else
+					else 
 					{
-						*LoopPercent = (float)NumTotalDecodedFramesInLoop / LoopNumFrames;
+						UE_LOG(LogMetaSound, Verbose, TEXT("Skipping cue point \"%s\" at frame %d due to multiple cue points in same render block"), *CuePoint->Label, CuePoint->FramePosition);
 					}
 				}
-				else
+			}
+
+			if (nullptr != OutEvents.CuePoint)
+			{
+				int32 CuePointFrameInBuffer = SourceState.MapFrameInWaveToFrameInBuffer(OutEvents.CuePoint->FramePosition, *bLoop);
+				OutEvents.OnCuePointFrameIndex = FMath::RoundToInt(Resampler->MapInputFrameToOutputFrame(CuePointFrameInBuffer));
+			}
+		}
+
+		// Search for cue points in frame range. Return the first cue point in the frame range.
+		const FSoundWaveCuePoint* FindCuePoint(int32 InStartFrameInWave, int32 InEndFrameInWave)
+		{
+			int32 LowerBoundIndex = Algo::LowerBoundBy(SortedCuePoints, InStartFrameInWave, WavePlayerNodePrivate::GetCuePointFrame);
+			int32 UpperBoundIndex = Algo::LowerBoundBy(SortedCuePoints, InEndFrameInWave, WavePlayerNodePrivate::GetCuePointFrame);
+
+			if (LowerBoundIndex < UpperBoundIndex)
+			{
+				// Inform about skipped cue points. 
+				for (int32 i = LowerBoundIndex + 1; i < UpperBoundIndex; i++)
 				{
-					*LoopPercent = 0.0f;
+					const FSoundWaveCuePoint& CuePoint = SortedCuePoints[i];
+
+					UE_LOG(LogMetaSound, Verbose, TEXT("Skipping cue point \"%s\" at frame %d due to multiple cue points in same render block"), *CuePoint.Label, CuePoint.FramePosition);
 				}
-				
-				// handle decoder having completed during it's decode
-				const bool bCannotGenerateAudio = !Decoder.CanGenerateAudio();
-				const bool bDidNotGenerateEnough = NumFramesGenerated < NumOutputFramesToGenerate;
-				const bool bReachedEOF = !Decoder.CanGenerateAudio() || (NumFramesGenerated < NumOutputFramesToGenerate);
-				if (bReachedEOF && !bIsLooping)
-				{
-					// Check the wave file input to see if it has changed... if it has, we want to restart everything up!
-					bool bConcatenatedPlayback = false;
-					if (WaveAsset->IsSoundWaveValid())
-					{
-						const FName SoundWaveName = (*WaveAsset)->GetFName();
-						if (SoundWaveName != CurrentSoundWaveName)
-						{
-							StartPlaying();
-							bConcatenatedPlayback = true;
-						}
-					}
+				return &SortedCuePoints[LowerBoundIndex];
+			}
 
-					if (!bConcatenatedPlayback)
-					{
-						bIsPlaying = false;
-						TriggerOnDone->TriggerFrame(StartFrame + NumFramesGenerated);
-					}
+			return nullptr;
+		}
+
+		// Check the expected output positions for various sample accurate events
+		// before resampling. 
+		//
+		// Note: The resampler can only accurately map samples *before* processing 
+		// audio because processing audio modifies the internal state of the resampler.
+		void TriggerUpcomingEvents(int32 InOperatorStartFrame, int32 InNumFrames, const WavePlayerNodePrivate::FSourceBufferState& InState)
+		{
+			WavePlayerNodePrivate::FSourceEvents Events;
+
+			// Check extra frames to hit the 
+			const int32 NumOutputFramesToCheck = (2 * OperatorSettings.GetNumFramesPerBlock() + 1) - InOperatorStartFrame;
+			const int32 NumSourceFramesToCheck = FMath::CeilToInt(Resampler->MapOutputFrameToInputFrame(NumOutputFramesToCheck));
+
+			// Selectively map events in the source buffer to frame indices in the
+			// resampled output buffer. 
+			MapSourceEventsIfInRange(SourceState, NumSourceFramesToCheck, Events);
+
+			// Check whether to trigger loops based on actual number of output frames 
+			if (*bLoop)
+			{
+				if (FMath::IsWithin(Events.OnLoopFrameIndex, 0, InNumFrames))
+				{
+					TriggerOnLooped->TriggerFrame(InOperatorStartFrame + Events.OnLoopFrameIndex);
 				}
+			}
+			else
+			{
+				const int32 IsNearlyDoneStartFrameIndex = OperatorSettings.GetNumFramesPerBlock() - InOperatorStartFrame;
+				const int32 IsNearlyDoneEndFrameIndex = IsNearlyDoneStartFrameIndex + OperatorSettings.GetNumFramesPerBlock();
 
-				if (bNeedsUpmix)
+				if (FMath::IsWithin(Events.OnEOFFrameIndex, 0, InNumFrames))
 				{
-					// Dual mono output
-					FMemory::Memcpy(FinalOutputLeft, PostSrcBufferPtr, sizeof(float)* NumFramesGenerated);
-					FMemory::Memcpy(FinalOutputRight, PostSrcBufferPtr, sizeof(float) * NumFramesGenerated);
+					TriggerOnDone->TriggerFrame(InOperatorStartFrame + Events.OnEOFFrameIndex);
 				}
-				else if (bNeedsDeinterleave)
+				else if (FMath::IsWithin(Events.OnEOFFrameIndex, IsNearlyDoneStartFrameIndex, IsNearlyDoneEndFrameIndex))
 				{
-					for (int32 i = 0; i < NumFramesGenerated; ++i)
+					// Protect against triggering OnNearlyDone multiple times
+					// in the scenario where significant pitch shift changes drastically
+					// alter the predicted OnDone frame between render blocks.
+					if (!bOnNearlyDoneTriggeredForWave)
 					{
-						// deinterleave each stereo frame into output buffers
-						FinalOutputLeft[i] = PostSrcBufferPtr[(i << 1)];
-						FinalOutputRight[i] = PostSrcBufferPtr[(i << 1) + 1];
+						TriggerOnNearlyDone->TriggerFrame(InOperatorStartFrame + Events.OnEOFFrameIndex);
+						bOnNearlyDoneTriggeredForWave = true;
 					}
 				}
 			}
-			else
-			{
-				FMemory::Memzero(FinalOutputLeft, sizeof(float) * NumOutputFramesToGenerate);
-				FMemory::Memzero(FinalOutputRight, sizeof(float) * NumOutputFramesToGenerate);
 
-				NumFramesGenerated = NumOutputFramesToGenerate;
+			if (nullptr != Events.CuePoint)
+			{
+				if (FMath::IsWithin(Events.OnCuePointFrameIndex, 0, InNumFrames))
+				{
+					if (!TriggerOnCuePoint->IsTriggeredInBlock())
+					{
+						*CuePointID = Events.CuePoint->CuePointID;
+						*CuePointLabel = Events.CuePoint->Label;
+						TriggerOnCuePoint->TriggerFrame(InOperatorStartFrame + Events.OnCuePointFrameIndex);
+					}
+					else
+					{
+						UE_LOG(LogMetaSound, Verbose, TEXT("Skipping cue point \"%s\" at frame %d due to multiple cue points in same render block"), *Events.CuePoint->Label, Events.CuePoint->FramePosition);
+					}
+				}
 			}
-
-			return NumFramesGenerated;
 		}
 
-	private:
 
-		// Called to get the cue point in the given block. Will return the first cue point if multiple cues are in the same block.
-		// Note: we can't currently support sub-block cue triggers due to cue ID and cue label being block-rate params. 
-		bool GetCuePointForBlock(int32 InStartFrame, int32 InEndFrame, FSoundWaveCuePoint& OutCuePoint, int32& OutFrameOffset)
+		void GeneratePitchedAudio(Audio::FMultichannelBufferView& OutBuffer)
 		{
-			for (FSoundWaveCuePoint& CuePoint : CuePoints)
+			using namespace Audio;
+
+			// Outside callers should ensure that bIsPlaying is true if calling this function.
+			check(bIsPlaying);
+
+			int32 NumFramesRequested = GetMultichannelBufferNumFrames(OutBuffer);
+			int32 NumSourceFramesAvailable = GetMultichannelBufferNumFrames(SourceCircularBuffer);
+
+			while (NumFramesRequested > 0)
 			{
-				if (CuePoint.FramePosition >= InStartFrame && CuePoint.FramePosition < InEndFrame)
+				// Determine how many frames are needed to produce the output.
+				int32 NumSourceFramesNeeded = Resampler->GetNumInputFramesNeededToProduceOutputFrames(NumFramesRequested + 1);
+				if (NumSourceFramesNeeded > NumSourceFramesAvailable)
 				{
-					OutCuePoint = CuePoint;
-					OutFrameOffset = CuePoint.FramePosition - InStartFrame;
-					return true;
+					// Generate more source audio, but still may not be enough to produce all requested frames.
+					GenerateSourceAudio(SourceCircularBuffer, SourceState);
+				}
+				NumSourceFramesAvailable = GetMultichannelBufferNumFrames(SourceCircularBuffer);
+
+				// Resample frames. 
+				int32 NumFramesProduced = Resampler->ProcessAndConsumeAudio(SourceCircularBuffer, OutBuffer);
+				if (NumFramesProduced < 1)
+				{
+					UE_LOG(LogMetaSound, Error, TEXT("Aborting currently playing metasound wave %s. Failed to produce any resampled audio frames with %d input frames and a frame ratio of %f."), *CurrentWaveAsset->GetFName().ToString(), NumSourceFramesAvailable, GetFrameRatio());
+					bIsPlaying = false;
+					break;
+				}
+
+				// Update sample counters
+				int32 NewNumSourceFramesAvailable = GetMultichannelBufferNumFrames(SourceCircularBuffer);
+				int32 NumSourceFramesConsumed = NumSourceFramesAvailable - NewNumSourceFramesAvailable;
+				NumSourceFramesAvailable = NewNumSourceFramesAvailable;
+				NumFramesRequested -= NumFramesProduced;
+
+				SourceState.Advance(NumSourceFramesConsumed, *bLoop);
+
+				// Shift buffer if there are more samples to create
+				if (NumFramesRequested > 0)
+				{
+					ShiftMultichannelBufferView(NumFramesProduced, OutBuffer);
 				}
 			}
-			// No cue points in this block
-			return false;
 		}
 
 		const FOperatorSettings OperatorSettings;
@@ -643,13 +959,11 @@ namespace Metasound
 		FTriggerReadRef StopTrigger;
 		FWaveAssetReadRef WaveAsset;
 		FTimeReadRef StartTime;
-		//FBoolReadRef bStartTimeAsPercentage;
 		FFloatReadRef PitchShift;
 		FBoolReadRef bLoop;
 		FTimeReadRef LoopStartTime;
-		FTimeReadRef LoopDurationSeconds;
+		FTimeReadRef LoopDuration;
 
-		FTriggerWriteRef TriggerOnPlay;
 		FTriggerWriteRef TriggerOnDone;
 		FTriggerWriteRef TriggerOnNearlyDone;
 		FTriggerWriteRef TriggerOnLooped;
@@ -661,34 +975,23 @@ namespace Metasound
 		FAudioBufferWriteRef AudioBufferL;
 		FAudioBufferWriteRef AudioBufferR;
 
-		// source decode
-		Audio::FAlignedFloatBuffer PostSrcBuffer;
-		Audio::FSimpleDecoderWrapper Decoder;
-
-		FName CurrentSoundWaveName;
-			
-		float SoundAssetSampleRate = 0.0f;
-		float SoundAssetDurationSeconds = 0.0f;
-		uint32 SoundAssetNumFrames = 0;
-
-		const float OutputSampleRate = 0.f;
-		const int32 OutputBlockSizeInFrames = 0;
-		
-		TArray<FSoundWaveCuePoint> CuePoints;
-		TArray<FSoundWaveCuePoint> CuePointsInCurrentBock;
-
-		// Data to track decoder
-		uint32 PlaybackStartFrame = 0;
-		uint32 CurrentConsumedFrameCount = 0;
-		uint32 FramesConsumedSinceStart = 0;
-		uint32 NumTotalDecodedFramesInLoop = 0;
-		uint32 LoopStartFrame = 0;
-		uint32 FramesToConsumePlayBeforeLooping = 0;
-		uint32 LoopNumFrames = 0;
-		bool bLooped = false;
-		bool bIsPlaying = false;
+		TUniquePtr<FWaveProxyReader> WaveProxyReader;
+		TUniquePtr<Audio::IConvertDeinterleave> ConvertDeinterleave;
+		TUniquePtr<Audio::FMultichannelLinearResampler> Resampler;
 
 		FWaveAsset CurrentWaveAsset;
+		TArray<FSoundWaveCuePoint> SortedCuePoints;
+		Audio::FAlignedFloatBuffer InterleavedBuffer;
+		Audio::FMultichannelBuffer DeinterleavedBuffer;
+		Audio::FMultichannelCircularBuffer SourceCircularBuffer;
+		Audio::FMultichannelBufferView OutputAudioView;
+
+		WavePlayerNodePrivate::FSourceBufferState SourceState;
+		float SampleRateFrameRatio = 1.f;
+		int32 NumDeinterleaveChannels;
+		bool bOnNearlyDoneTriggeredForWave = false;
+		bool bIsPlaying = false;
+		
 	};
 
 	TUniquePtr<IOperator> FWavePlayerNode::FOperatorFactory::CreateOperator(
@@ -713,7 +1016,6 @@ namespace Metasound
 			Inputs.GetDataReadReferenceOrConstruct<bool>(InputLoopName),
 			Inputs.GetDataReadReferenceOrConstruct<FTime>(InputLoopStartName),
 			Inputs.GetDataReadReferenceOrConstruct<FTime>(InputLoopDurationName)
-			//Inputs.GetDataReadReferenceOrConstruct<bool>(InputStartTimeAsPercentName)
 		};
 
 		return MakeUnique<FWavePlayerOperator>(Args);
@@ -733,7 +1035,6 @@ namespace Metasound
 				TInputDataVertexModel<bool>(InputLoopName, InputLoopTT, false),
 				TInputDataVertexModel<FTime>(InputLoopStartName, InputLoopStartTT, 0.0f),
 				TInputDataVertexModel<FTime>(InputLoopDurationName, InputLoopDurationTT, -1.0f)
-				//TInputDataVertexModel<bool>(InputStartTimeAsPercentName, InputStartTimeAsPercentTT, false)
 				),
 			FOutputVertexInterface(
 				TOutputDataVertexModel<FTrigger>(OutputTriggerOnPlayName, OutputTriggerOnPlayTT),
@@ -806,6 +1107,8 @@ namespace Metasound
 	}
 
 	METASOUND_REGISTER_NODE(FWavePlayerNode)
-
 } // namespace Metasound
+
 #undef LOCTEXT_NAMESPACE // MetasoundWaveNode
+
+
diff --git a/Engine/Plugins/Runtime/Metasound/Source/MetasoundEngine/Private/MetasoundWaveProxyReader.cpp b/Engine/Plugins/Runtime/Metasound/Source/MetasoundEngine/Private/MetasoundWaveProxyReader.cpp
new file mode 100644
--- /dev/null
+++ b/Engine/Plugins/Runtime/Metasound/Source/MetasoundEngine/Private/MetasoundWaveProxyReader.cpp
@@ -0,0 +1,500 @@
+// Copyright Epic Games, Inc. All Rights Reserved.
+
+#include "MetasoundWaveProxyReader.h"
+
+#include "Containers/Array.h"
+#include "Containers/ArrayView.h"
+#include "DecoderInputFactory.h"
+#include "DSP/BufferVectorOperations.h"
+#include "DSP/Dsp.h"
+#include "HAL/IConsoleManager.h"
+#include "HAL/Platform.h"
+#include "HAL/UnrealMemory.h"
+#include "IAudioCodec.h"
+#include "IAudioCodecRegistry.h"
+#include "Math/UnrealMathUtility.h"
+#include "MetasoundLog.h"
+#include "MetasoundTrace.h"
+#include "Misc/AssertionMacros.h"
+#include "Sound/SoundWave.h"
+#include "Templates/SharedPointer.h"
+#include "Templates/UniquePtr.h"
+
+static int32 GMetaSoundWaveProxyReaderSimulateSeekOnNonSeekable = 0;
+FAutoConsoleVariableRef CVarClearMetaSoundFrontendSearchEngineCache(
+	TEXT("au.MetaSounds.WavePlayer.SimulateSeek"),
+	GMetaSoundWaveProxyReaderSimulateSeekOnNonSeekable ,
+	TEXT("If true, SoundWaves which are not of a seekable format will simulate seek calls by reading and discarding samples.\n")
+	TEXT("0: Do not simulate seek, !0: Simulate seek"),
+	ECVF_Default);
+
+namespace Metasound
+{
+	namespace MetasoundWaveProxyReaderPrivate
+	{
+		/** Construct a FDecoderOutput
+		 *
+		 * @param InNumFramesPerDecode - Maximum number of frames pushed when
+		 *                               audio is decoded. The buffer will be 
+		 *                               able to hold twice the number of frames. 
+		 */
+		FDecoderOutput::FDecoderOutput(uint32 InNumFramesPerDecode)
+		: NumFramesPerDecode(FMath::Max(InNumFramesPerDecode, MinNumFramesPerDecode))
+		{
+			Init();
+		}
+
+		/** Set the number of channels in the audio being decoded. */
+		void FDecoderOutput::SetNumChannels(uint32 InNumChannels)
+		{
+			NumChannels = FMath::Max(MinNumChannels, InNumChannels);
+			Init();
+		}
+
+		/** Removes all samples from the buffer. */
+		void FDecoderOutput::Reset()
+		{
+			Buffer.SetNum(0);
+		}
+
+		/** Returns requirements used by the audio codec system. */
+		Audio::IDecoderOutput::FRequirements FDecoderOutput::GetRequirements(const Audio::FDecodedFormatInfo& InFormat) const
+		{
+			FRequirements Requirements;
+
+			Requirements.DownstreamFormat = Audio::EBitRepresentation::Float32_Interleaved;
+			Requirements.NumSampleFramesWanted = NumFramesPerDecode;
+			Requirements.NumSampleFramesPerSecond = InFormat.NumFramesPerSec;
+			Requirements.NumChannels = InFormat.NumChannels;
+
+			return Requirements;
+		}
+
+		/** Adds samples to the buffer. 
+		 *
+		 * This is called by the Decoder and should not be called otherwise.
+		 */
+		int32 FDecoderOutput::PushAudio(const Audio::IDecoderOutput::FPushedAudioDetails& InDetails, TArrayView<const int16> In16BitInterleave)
+		{
+			if (InDetails.NumChannels != NumChannels)
+			{
+				SetNumChannels(InDetails.NumChannels);
+			}
+			return PushAudioInternal(InDetails, In16BitInterleave);
+		}
+
+		/** Adds samples to the buffer. 
+		 *
+		 * This is called by the Decoder and should not be called otherwise.
+		 */
+		int32 FDecoderOutput::PushAudio(const FPushedAudioDetails& InDetails, TArrayView<const float> InFloat32Interleave) 
+		{
+			if (InDetails.NumChannels != NumChannels)
+			{
+				SetNumChannels(InDetails.NumChannels);
+			}
+			return PushAudioInternal(InDetails, InFloat32Interleave);
+		}
+
+		/** This should not be called. It removes 16 bit PCM samples from the buffer
+		 * which is an unsupported operation of this class. 
+		 */
+		int32 FDecoderOutput::PopAudio(TArrayView<int16> InExternalInt16Buffer, FPushedAudioDetails& OutDetails)
+		{
+			// This buffer cannot produce 16 bit PCM audio.
+			checkNoEntry();
+			return 0;
+		}
+
+		/** Copy samples to OutBuffer and remove them from this objects internal
+		 * buffer. 
+		 *
+		 *
+		 * @param OutBuffer - A destination array to copy samples to.
+		 * @param OutDetails - Unused. 
+		 *
+		 * @return The actual number of samples copied. 
+		 */
+		int32 FDecoderOutput::PopAudio(TArrayView<float> OutBuffer, FPushedAudioDetails& OutDetails) 
+		{
+			return Buffer.Pop(OutBuffer.GetData(), OutBuffer.Num());
+		}
+
+		// Initialize buffer size.
+		void FDecoderOutput::Init()
+		{
+			// Allow it to hold two decode buffers max.
+			uint32 MinBufferCapacity = NumChannels * NumFramesPerDecode * 2;
+			constexpr bool bRetainExistingSamples = false;
+			Buffer.Reserve(MinBufferCapacity, bRetainExistingSamples);
+		}
+
+		int32 FDecoderOutput::PushAudioInternal(const FPushedAudioDetails& InDetails, TArrayView<const float> InBuffer) 
+		{
+			return Buffer.Push(InBuffer.GetData(), InBuffer.Num());
+		}
+
+		int32 FDecoderOutput::PushAudioInternal(const FPushedAudioDetails& InDetails, TArrayView<const int16> InBuffer) 
+		{
+			SampleConversionBuffer.SetNumUninitialized(InBuffer.Num(), false /* bAllowShrinking */);
+
+			// Convert 16 bit pcm to 32 bit float.
+			constexpr float Scalar = 1.f / 32768.f;
+			const int16* Src = InBuffer.GetData();
+			float* Dst = SampleConversionBuffer.GetData();
+			int32 Num = InBuffer.Num();
+
+			// Convert 1 sample at a time, slow.
+			for (int32 i = 0; i < Num; ++i)
+			{
+				*Dst++ = static_cast<float>(*Src++) * Scalar;
+			}
+
+			return PushAudioInternal(InDetails, SampleConversionBuffer);
+		}
+
+	}
+
+
+	
+
+	uint32 FWaveProxyReader::ConformDecodeSize(uint32 InMaxDesiredDecodeSizeInFrames)
+	{
+		static_assert(DefaultMinDecodeSizeInFrames >= DecodeSizeQuantizationInFrames, "Min decode size less than decode size quantization");
+		static_assert((DefaultMinDecodeSizeInFrames % DecodeSizeQuantizationInFrames) == 0, "Min decode size must be equally divisible by decode size quantization");
+
+		const uint32 QuantizedDecodeSizeInFrames = (InMaxDesiredDecodeSizeInFrames / DecodeSizeQuantizationInFrames) * DecodeSizeQuantizationInFrames;
+		const uint32 ConformedDecodeSizeInFrames = FMath::Max(QuantizedDecodeSizeInFrames, DefaultMinDecodeSizeInFrames);
+
+		check((ConformedDecodeSizeInFrames % QuantizedDecodeSizeInFrames) == 0);
+		check(ConformedDecodeSizeInFrames > 0);
+
+		return ConformedDecodeSizeInFrames;
+	}
+
+	/** Construct a wave proxy reader. 
+	 *
+	 * @param InWaveProxy - A TSharedRef of a FSoundWaveProxy which is to be played. 
+	 * @param InSettings - Reader settings. 
+	 */
+	FWaveProxyReader::FWaveProxyReader(FSoundWaveProxyRef InWaveProxy, const FSettings& InSettings)
+	: WaveProxy(InWaveProxy)
+	, DecoderOutput(ConformDecodeSize(InSettings.MaxDecodeSizeInFrames))
+	, Settings(InSettings)
+	{
+		// Get local copies of some values from the proxy. 
+		SampleRate = WaveProxy->GetSampleRate();
+		NumChannels = WaveProxy->GetNumChannels();
+		NumFramesInWave = WaveProxy->GetNumFrames();
+		DurationInSeconds = WaveProxy->GetDuration();
+		MaxLoopStartTimeInSeconds = FMath::Max(0, DurationInSeconds - MinLoopDurationInSeconds);
+
+		// Clamp times
+		Settings.StartTimeInSeconds = FMath::Clamp(Settings.StartTimeInSeconds, 0.f, DurationInSeconds);
+		Settings.LoopStartTimeInSeconds = ClampLoopStartTime(Settings.LoopStartTimeInSeconds);
+		Settings.LoopDurationInSeconds = ClampLoopDuration(Settings.LoopDurationInSeconds);
+		
+		// Setup frame indices
+		CurrentFrameIndex = 0;
+		UpdateLoopBoundaries();
+
+		// Determine max size of decode buffer
+		Settings.MaxDecodeSizeInFrames = FMath::Max(DefaultMinDecodeSizeInFrames, Settings.MaxDecodeSizeInFrames);
+		DecoderOutput.SetNumChannels(NumChannels);
+
+		// Prepare to read audio
+		bIsDecoderValid = InitializeDecoder(Settings.StartTimeInSeconds);
+	}
+
+	/** Create a wave proxy reader. 
+	 *
+	 * @param InWaveProxy - A TSharedRef of a FSoundWaveProxy which is to be played. 
+	 * @param InSettings - Reader settings. 
+	 */
+	TUniquePtr<FWaveProxyReader> FWaveProxyReader::Create(FSoundWaveProxyRef InWaveProxy, const FSettings& InSettings)
+	{
+		if (InWaveProxy->GetSampleRate() <= 0.f)
+		{
+			UE_LOG(LogMetaSound, Warning, TEXT("Cannot create FWaveProxy reader due to invalid sample rate (%f). Package: %s"), InWaveProxy->GetSampleRate(), *InWaveProxy->GetPackageName().ToString());
+			return TUniquePtr<FWaveProxyReader>();
+		}
+
+		if (InWaveProxy->GetNumChannels() <= 0)
+		{
+			UE_LOG(LogMetaSound, Warning, TEXT("Cannot create FWaveProxy reader due to invalid num channels (%d). Package: %s"), InWaveProxy->GetNumChannels(), *InWaveProxy->GetPackageName().ToString());
+			return TUniquePtr<FWaveProxyReader>();
+		}
+
+		if (InWaveProxy->GetNumFrames() <= 0)
+		{
+			UE_LOG(LogMetaSound, Warning, TEXT("Cannot create FWaveProxy reader due to invalid num frames (%d). Package: %s"), InWaveProxy->GetNumFrames(), *InWaveProxy->GetPackageName().ToString());
+			return TUniquePtr<FWaveProxyReader>();
+		}
+
+		return TUniquePtr<FWaveProxyReader>(new FWaveProxyReader(InWaveProxy, InSettings));
+	}
+
+
+	/** Set whether the reader should loop the audio or not. */
+	void FWaveProxyReader::SetIsLooping(bool bInIsLooping)
+	{
+		if (Settings.bIsLooping != bInIsLooping)
+		{
+			Settings.bIsLooping = bInIsLooping;
+			UpdateLoopBoundaries();
+		}
+	}
+
+	/** Sets the beginning position of the loop. */
+	void FWaveProxyReader::SetLoopStartTime(float InLoopStartTimeInSeconds)
+	{
+		InLoopStartTimeInSeconds = ClampLoopStartTime(InLoopStartTimeInSeconds);
+		if (!FMath::IsNearlyEqual(Settings.LoopStartTimeInSeconds, InLoopStartTimeInSeconds))
+		{
+			Settings.LoopStartTimeInSeconds = InLoopStartTimeInSeconds;
+			UpdateLoopBoundaries();
+		}
+	}
+
+	/** Sets the duration of the loop in seconds. If the value is negative, the
+	 * loop duration consists of the entire file. */
+	void FWaveProxyReader::SetLoopDuration(float InLoopDurationInSeconds)
+	{
+		InLoopDurationInSeconds = ClampLoopDuration(InLoopDurationInSeconds);
+		if (!FMath::IsNearlyEqual(Settings.LoopDurationInSeconds, InLoopDurationInSeconds))
+		{
+			Settings.LoopDurationInSeconds = InLoopDurationInSeconds;
+			UpdateLoopBoundaries();
+		}
+	}
+
+	bool FWaveProxyReader::SeekToTime(float InSeconds)
+	{
+		// Direct seeking is not supported. A new decoder must be created. 
+		bIsDecoderValid = InitializeDecoder(InSeconds);
+		return bIsDecoderValid;
+	}
+
+	/** Copies audio into OutBuffer. It returns the number of samples copied.
+	 * Samples not written to will be set to zero.
+	 */
+	int32 FWaveProxyReader::PopAudio(Audio::FAlignedFloatBuffer& OutBuffer)
+	{
+		using namespace Audio;
+
+		METASOUND_TRACE_CPUPROFILER_EVENT_SCOPE(Metasound::FWaveProxyReader::PopAudio);
+
+		if ((0 == NumChannels) || !bIsDecoderValid)
+		{
+			if (OutBuffer.Num() > 0)
+			{
+				FMemory::Memset(OutBuffer.GetData(), 0, sizeof(float) * OutBuffer.Num());
+			}
+			return 0;
+		}
+
+		checkf(0 == (OutBuffer.Num() % NumChannels), TEXT("Output buffer size must be evenly divisible by the number of channels"));
+
+		TArrayView<float> OutBufferView{OutBuffer};
+		int32 NumSamplesUnset = OutBuffer.Num();
+		int32 NumSamplesCopied = 0;
+		bool bCanProduceMoreAudio = true;
+
+		while ((NumSamplesUnset > 0) && bCanProduceMoreAudio)
+		{
+			int32 NumSamplesCopiedThisLoop = PopAudioFromDecoderOutput(OutBufferView);
+
+			// Update sample counters and output buffer view
+			NumSamplesCopied += NumSamplesCopiedThisLoop;
+			NumSamplesUnset -= NumSamplesCopiedThisLoop;
+			OutBufferView = OutBufferView.Slice(NumSamplesCopiedThisLoop, NumSamplesUnset);
+
+			// Determine if we can / should decode more data. 
+			if (bIsDecoderValid && (NumSamplesUnset > 0) && (DecoderOutput.Num() == 0))
+			{
+				// Looping is handle within the FWaveProxyReader instead of 
+				// within the decoder.
+				DecodeResult = Decoder->Decode(false /* bIsLooping */); 
+
+			}
+
+			const bool bDecoderOutputHasMoreData = DecoderOutput.Num() > 0;
+			const bool bDecoderCanDecodeMoreData = bIsDecoderValid && (EDecodeResult::MoreDataRemaining == DecodeResult);
+			bCanProduceMoreAudio = bDecoderOutputHasMoreData || bDecoderCanDecodeMoreData;
+		}
+
+		// Zero pad any unset samples. 
+		if (OutBufferView.Num() > 0)
+		{
+			check(!bCanProduceMoreAudio);
+
+			// Zero out audio that was not set.
+			FMemory::Memset(OutBufferView.GetData(), 0, sizeof(float) * OutBufferView.Num());
+			CurrentFrameIndex += OutBufferView.Num() / NumChannels;
+		}
+
+		return NumSamplesCopied;
+	}
+
+	int32 FWaveProxyReader::PopAudioFromDecoderOutput(TArrayView<float> OutBufferView)
+	{
+		using namespace MetasoundWaveProxyReaderPrivate;
+
+		check(NumChannels > 0);
+
+		int32 NumSamplesCopied = 0;
+		if (DecoderOutput.Num() > 0)
+		{
+			// Get samples from the decoder buffer.
+			FDecoderOutput::FPushedAudioDetails UnusedDetails;  
+			NumSamplesCopied = DecoderOutput.PopAudio(OutBufferView, UnusedDetails);
+
+			int32 NumFramesCopied = NumSamplesCopied / NumChannels;
+			CurrentFrameIndex += NumFramesCopied;
+
+			// Check whether the samples copied from the decoder extend 
+			// past the end of the loop.
+			if (Settings.bIsLooping)
+			{
+				const bool bDidOvershoot = CurrentFrameIndex >= LoopEndFrameIndex;
+				if (bDidOvershoot)
+				{
+					// Rewind sample counters if the loop boundary was overshot.
+					NumFramesCopied -= (CurrentFrameIndex - LoopEndFrameIndex);
+					// If Settings.bIsLooping info was altered, NumFramesCopied can end up 
+					// negative if the current frame index is past the loop end frame. 
+					NumFramesCopied = FMath::Max(0, NumFramesCopied); 
+					NumSamplesCopied = NumFramesCopied * NumChannels;
+
+					// Remove any remaining samples in the decoder because
+					DecoderOutput.Reset();
+				}
+
+				// Determine whether or not to seek to the loop start.
+				const bool bSeekToLoopStart = bDidOvershoot
+					|| ((DecodeResult == EDecodeResult::Finished) && (DecoderOutput.Num() == 0)) // Decoder cannot produce more audio
+					|| (CurrentFrameIndex >= NumFramesInWave); // Reached end of file.
+
+				if (bSeekToLoopStart)
+				{
+					// Start decoder at beginning of loop
+					SeekToTime(Settings.LoopStartTimeInSeconds);
+				}
+			}
+		}
+
+		return NumSamplesCopied;
+	}
+
+	bool FWaveProxyReader::InitializeDecoder(float InStartTimeInSeconds)
+	{
+		using namespace Audio;
+
+		METASOUND_TRACE_CPUPROFILER_EVENT_SCOPE(MetaSound::FWaveProxyReader::InitializeDecoder);
+
+		// Create decoder input
+		FName Format = WaveProxy->GetRuntimeFormat();
+		DecoderInput = Audio::CreateBackCompatDecoderInput(Format, WaveProxy);
+		if (!DecoderInput.IsValid())
+		{
+			UE_LOG(LogMetaSound, Error, TEXT("Failed to create decoder input (format:%s) for wave (package:%s)"), *Format.ToString(), *WaveProxy->GetPackageName().ToString());
+			return false;
+		}
+
+		// Seek input to start time.
+		DecoderInput->SeekToTime(InStartTimeInSeconds);
+		CurrentFrameIndex = InStartTimeInSeconds * GetSampleRate();
+
+		// Get codec ptr by reading the header info from the decoder input.
+		ICodecRegistry::FCodecPtr Codec = ICodecRegistry::Get().FindCodecByParsingInput(DecoderInput.Get());
+		if (nullptr == Codec)
+		{
+			UE_LOG(LogMetaSound, Error, TEXT("Failed to find codec (format:%s) for wave (package:%s)"), *Format.ToString(), *WaveProxy->GetPackageName().ToString());
+			return false;
+		}
+
+		// Create the decoder
+		Decoder = Codec->CreateDecoder(DecoderInput.Get(), &DecoderOutput);
+		if (!Decoder.IsValid())
+		{
+			UE_LOG(LogMetaSound, Error, TEXT("Failed to create decoder (format:%s) for wave (package:%s)"), *Format.ToString(), *WaveProxy->GetPackageName().ToString());
+			return false;
+		}
+
+		// For non-seekable streaming waves, use a fallback method to seek
+		// to the start time. 
+		const bool bUseFallbackSeekMethod = (GMetaSoundWaveProxyReaderSimulateSeekOnNonSeekable != 0) && (InStartTimeInSeconds != 0.f) && (!WaveProxy->IsSeekableStreaming());
+		if (bUseFallbackSeekMethod)
+		{
+			if (!bFallbackSeekMethodWarningLogged)
+			{
+				UE_LOG(LogMetaSound, Warning, TEXT("Simulating seeking in wave which is not seekable (package:%s). For better performance, set wave to a seekable format"), *WaveProxy->GetPackageName().ToString()); 
+				bFallbackSeekMethodWarningLogged = true;
+			}
+
+			int32 NumSamplesToDiscard = CurrentFrameIndex * NumChannels;
+			DiscardSamples(NumSamplesToDiscard);
+		}
+
+		// return true if all the components were successfully create
+		return true;
+	}
+
+	void FWaveProxyReader::DiscardSamples(int32 InNumSamplesToDiscard)
+	{
+		while (InNumSamplesToDiscard > 0)
+		{
+			DecodeResult = Decoder->Decode(false /* bIsLooping */); 
+
+			int32 NumSamplesToDiscardThisLoop = FMath::Min(DecoderOutput.Num(), InNumSamplesToDiscard);
+
+			if (NumSamplesToDiscardThisLoop < 1)
+			{
+				UE_LOG(LogMetaSound, Error, TEXT("Failed to decode samples (package: %s)"), *WaveProxy->GetPackageName().ToString());
+				break;
+			}
+
+			int32 ActualNumSamplesDiscarded = DecoderOutput.PopAudio(NumSamplesToDiscardThisLoop);
+			InNumSamplesToDiscard -= ActualNumSamplesDiscarded;
+
+			if ((InNumSamplesToDiscard > 0) && (DecodeResult != EDecodeResult::MoreDataRemaining))
+			{
+				UE_LOG(LogMetaSound, Error, TEXT("Failed to simulate seek (seek to frame: %d, package: %s)"), CurrentFrameIndex, *WaveProxy->GetPackageName().ToString());
+				break;
+			}
+		}
+	}
+
+	float FWaveProxyReader::ClampLoopStartTime(float InStartTimeInSeconds)
+	{
+		return FMath::Clamp(InStartTimeInSeconds, 0.f, MaxLoopStartTimeInSeconds);
+	}
+
+	float FWaveProxyReader::ClampLoopDuration(float InDurationInSeconds)
+	{
+		if (InDurationInSeconds <= 0.f)
+		{
+			return MaxLoopDurationInSeconds;
+		}
+		return FMath::Max(InDurationInSeconds, MinLoopDurationInSeconds);
+	}
+
+	void FWaveProxyReader::UpdateLoopBoundaries()
+	{
+		if (Settings.bIsLooping)
+		{
+			const int32 MinLoopDurationInFrames = FMath::CeilToInt(MinLoopDurationInSeconds * SampleRate);
+			const float LoopEndTime = Settings.LoopStartTimeInSeconds + Settings.LoopDurationInSeconds;
+			const int32 MinLoopEndFrameIndex = LoopStartFrameIndex + MinLoopDurationInFrames;
+			const int32 MaxLoopEndFrameIndex = NumFramesInWave;
+
+			LoopStartFrameIndex = FMath::FloorToInt(Settings.LoopStartTimeInSeconds * SampleRate);
+			LoopEndFrameIndex = FMath::Clamp(LoopEndTime * SampleRate, MinLoopEndFrameIndex, MaxLoopEndFrameIndex);
+		}
+		else
+		{
+			LoopStartFrameIndex = 0;
+			LoopEndFrameIndex = NumFramesInWave;
+		}
+	}
+}
diff --git a/Engine/Plugins/Runtime/Metasound/Source/MetasoundEngine/Private/MetasoundWaveProxyReader.h b/Engine/Plugins/Runtime/Metasound/Source/MetasoundEngine/Private/MetasoundWaveProxyReader.h
new file mode 100644
--- /dev/null
+++ b/Engine/Plugins/Runtime/Metasound/Source/MetasoundEngine/Private/MetasoundWaveProxyReader.h
@@ -0,0 +1,263 @@
+// Copyright Epic Games, Inc. All Rights Reserved.
+
+#pragma once
+
+#include "DSP/BufferVectorOperations.h"
+#include "HAL/Platform.h"
+#include "IAudioCodec.h"
+#include "Templates/SharedPointer.h"
+#include "Templates/UniquePtr.h"
+
+// Forward declare
+class FSoundWaveProxy;
+
+namespace Metasound
+{
+	namespace MetasoundWaveProxyReaderPrivate
+	{
+		/** FDecoderOutput maintains a circular buffer of audio which is produced
+		 * from the decoder. 
+		 */
+		class FDecoderOutput : public Audio::IDecoderOutput
+		{
+			static constexpr uint32 MinNumFramesPerDecode = 1;
+			static constexpr uint32 DefaultNumChannels = 1;
+			static constexpr uint32 MinNumChannels = 1;
+
+		public:
+			using FPushedAudioDetails = Audio::IDecoderOutput::FPushedAudioDetails;
+
+			/** Construct a FDecoderOutput
+			 *
+			 * @param InNumFramesPerDecode - Maximum number of frames pushed when
+			 *                               audio is decoded. The buffer will be 
+			 *                               able to hold twice the number of frames. 
+			 */
+			FDecoderOutput(uint32 InNumFramesPerDecode);
+
+			/** Set the number of channels in the audio being decoded. */
+			void SetNumChannels(uint32 InNumChannels);
+
+			/** Returns the number of samples in the buffer. */
+			FORCEINLINE int32 Num() const
+			{
+				return Buffer.Num();
+			}
+
+			/** Removes all samples from the buffer. */
+			void Reset();
+
+			/** Returns requirements used by the audio codec system. */
+			virtual Audio::IDecoderOutput::FRequirements GetRequirements(const Audio::FDecodedFormatInfo& InFormat) const override;
+
+			/** Adds samples to the buffer. 
+			 *
+			 * This is called by the Decoder and should not be called otherwise.
+			 */
+			virtual int32 PushAudio(const Audio::IDecoderOutput::FPushedAudioDetails& InDetails, TArrayView<const int16> In16BitInterleave) override;
+
+			/** Adds samples to the buffer. 
+			 *
+			 * This is called by the Decoder and should not be called otherwise.
+			 */
+			virtual int32 PushAudio(const FPushedAudioDetails& InDetails, TArrayView<const float> InFloat32Interleave) override;
+
+			/** This should not be called. It removes 16 bit PCM samples from the buffer
+			 * which is an unsupported operation of this class. 
+			 */
+			virtual int32 PopAudio(TArrayView<int16> InExternalInt16Buffer, FPushedAudioDetails& OutDetails) override;
+
+			/** Copy samples to OutBuffer and remove them from this objects internal
+			 * buffer. 
+			 *
+			 *
+			 * @param OutBuffer - A destination array to copy samples to.
+			 * @param OutDetails - Unused. 
+			 *
+			 * @return The actual number of samples copied. 
+			 */
+			virtual int32 PopAudio(TArrayView<float> OutBuffer, FPushedAudioDetails& OutDetails) override;
+
+			/** Remove samples from the internal buffer.
+			 *
+			 * @param InNumSamples - The desired number of samples to remove.
+			 *
+			 * @return The actual number of samples removed. 
+			 */
+			FORCEINLINE int32 PopAudio(int32 InNumSamples)
+			{
+				return Buffer.Pop(InNumSamples);
+			}
+
+		private:
+
+			// Initialize buffer size.
+			void Init();
+			int32 PushAudioInternal(const FPushedAudioDetails& InDetails, TArrayView<const float> InBuffer);
+			int32 PushAudioInternal(const FPushedAudioDetails& InDetails, TArrayView<const int16> InBuffer);
+
+			uint32 NumFramesPerDecode = MinNumFramesPerDecode;
+			uint32 NumChannels = DefaultNumChannels;
+			Audio::TCircularAudioBuffer<float> Buffer;
+			Audio::FAlignedFloatBuffer SampleConversionBuffer;
+		};
+	}
+	
+	/** FWaveProxyReader reads a FWaveProxy and outputs 32 bit interleaved audio.
+	 *
+	 * FWaveProxyReader provides controls for looping and relevant frame index
+	 * values.
+	 */
+	class FWaveProxyReader
+	{
+
+	public:
+		using FSoundWaveProxyRef = TSharedRef<FSoundWaveProxy, ESPMode::ThreadSafe>;
+		using EDecodeResult = Audio::IDecoder::EDecodeResult;
+
+		/** Minimum number of frames to decode per a call to the decoder.  */
+		static constexpr uint32 DefaultMinDecodeSizeInFrames = 128;
+
+		/** Some codecs have strict requirements on decode size. In order to be
+		 * functional with all supported codecs, the decode size must be a multiple
+		 * of 128.
+		 */
+		static constexpr uint32 DecodeSizeQuantizationInFrames = 128;
+
+		static uint32 ConformDecodeSize(uint32 InMaxDesiredDecodeSizeInFrames);
+
+		/** Minimum duration of a loop. */
+		static constexpr float MinLoopDurationInSeconds = 0.05f;
+
+		/** Maximum duration of a loop.
+		 *
+		 * 10 years. ridiculously high. This exists to prevent floating point 
+		 * undefined overflow behavior if performing calculations with 
+		 * TNumericLimits<float>::Max().
+		 */
+		static constexpr float MaxLoopDurationInSeconds = 60.f * 60.f * 365.f * 10.f; 
+
+		/** Settings for a FWaveProxyReader. */
+		struct FSettings
+		{
+			uint32 MaxDecodeSizeInFrames = 8192;
+			float StartTimeInSeconds = 0.f;
+			bool bIsLooping = false;
+			float LoopStartTimeInSeconds = 0.f;
+			float LoopDurationInSeconds = 0.f;
+		};
+
+	private:
+		/** Construct a wave proxy reader. 
+		 *
+		 * @param InWaveProxy - A TSharedRef of a FSoundWaveProxy which is to be played. 
+		 * @param InSettings - Reader settings. 
+		 */
+		FWaveProxyReader(FSoundWaveProxyRef InWaveProxy, const FSettings& InSettings);
+
+	public:
+		/** Create a wave proxy reader. 
+		 *
+		 * @param InWaveProxy - A TSharedRef of a FSoundWaveProxy which is to be played. 
+		 * @param InSettings - Reader settings. 
+		 */
+		static TUniquePtr<FWaveProxyReader> Create(FSoundWaveProxyRef InWaveProxy, const FSettings& InSettings);
+
+		/** Set whether the reader should loop the audio or not. */
+		void SetIsLooping(bool bInIsLooping);
+
+		/** Returns true if the audio will be looped, false otherwise. */
+		FORCEINLINE bool IsLooping() const
+		{
+			return Settings.bIsLooping;
+		}
+
+		/** Sets the beginning position of the loop. */
+		void SetLoopStartTime(float InLoopStartTimeInSeconds);
+
+		/** Sets the duration of the loop in seconds. 
+		 *
+		 * If the value is negative, the MaxLoopDurationInSeconds will be used
+		 * which will effectively loop at the end of the file. 
+		 */
+		void SetLoopDuration(float InLoopDurationInSeconds);
+
+		FORCEINLINE float GetSampleRate() const
+		{
+			return SampleRate;
+		}
+
+		FORCEINLINE int32 GetNumChannels() const
+		{
+			return NumChannels;
+		}
+
+		/** Returns the index of the playhead within the complete wave. */
+		FORCEINLINE int32 GetFrameIndex() const
+		{
+			return CurrentFrameIndex;
+		}
+
+		FORCEINLINE int32 GetNumFramesInWave() const
+		{
+			return NumFramesInWave;
+		}
+
+		FORCEINLINE int32 GetNumFramesInLoop() const
+		{
+			return LoopEndFrameIndex - LoopStartFrameIndex;
+		}
+
+		FORCEINLINE int32 GetLoopStartFrameIndex() const
+		{
+			return LoopStartFrameIndex;
+		}
+
+		FORCEINLINE int32 GetLoopEndFrameIndex() const
+		{
+			return LoopEndFrameIndex;
+		}
+
+		/** Seeks to position in wave. 
+		 * 
+		 * @param InSeconds - The location to seek the playhead
+		 *
+		 * @return true on success, false on failuer.
+		 */
+		bool SeekToTime(float InSeconds);
+
+		/** Copies audio into OutBuffer. It returns the number of samples copied.
+		 * Samples not written to will be set to zero.
+		 */
+		int32 PopAudio(Audio::FAlignedFloatBuffer& OutBuffer);
+
+	private:
+
+		int32 PopAudioFromDecoderOutput(TArrayView<float> OutBufferView);
+		bool InitializeDecoder(float InStartTimeInSeconds);
+		void DiscardSamples(int32 InNumSamplesToDiscard);
+		float ClampLoopStartTime(float InStartTimeInSeconds);
+		float ClampLoopDuration(float InDurationInSeconds);
+		void UpdateLoopBoundaries();
+
+		FSoundWaveProxyRef WaveProxy;
+		TUniquePtr<Audio::IDecoderInput> DecoderInput;
+		TUniquePtr<Audio::IDecoder> Decoder;
+		MetasoundWaveProxyReaderPrivate::FDecoderOutput DecoderOutput;
+
+		FSettings Settings;
+
+		float SampleRate = 0.f;
+		int32 NumChannels = 0;
+		EDecodeResult DecodeResult = EDecodeResult::MoreDataRemaining;
+		int32 NumFramesInWave = 0;
+		int32 CurrentFrameIndex = 0;
+		int32 LoopStartFrameIndex = 0;
+		int32 LoopEndFrameIndex = -1;
+		float DurationInSeconds = 0.f;
+		float MaxLoopStartTimeInSeconds = 0.f;
+
+		bool bIsDecoderValid = false;
+		bool bFallbackSeekMethodWarningLogged = false;
+	};
+}
diff --git a/Engine/Plugins/Runtime/Metasound/Source/MetasoundEngine/Public/MetasoundWave.h b/Engine/Plugins/Runtime/Metasound/Source/MetasoundEngine/Public/MetasoundWave.h
--- a/Engine/Plugins/Runtime/Metasound/Source/MetasoundEngine/Public/MetasoundWave.h
+++ b/Engine/Plugins/Runtime/Metasound/Source/MetasoundEngine/Public/MetasoundWave.h
@@ -51,88 +51,3 @@ namespace Metasound
 	DECLARE_METASOUND_DATA_REFERENCE_TYPES(FWaveAsset, METASOUNDENGINE_API, FWaveAssetTypeInfo, FWaveAssetReadRef, FWaveAssetWriteRef)
 }
  
-namespace Audio
-{
-	// Forward declares
-	class ICodecRegistry;
-	struct IDecoderInput;
-	struct IDecoderOutput;
-	struct IDecoder;
-
-	class FSimpleDecoderWrapper
-	{
-	public:
-		struct InitParams
-		{
-			uint32 OutputBlockSizeInFrames{ 512 };
-			float OutputSampleRate { 44100.0f };
-			float MaxPitchShiftMagnitudeAllowedInOctaves { 4.0f };
-		};
-
-		FSimpleDecoderWrapper(const InitParams& InInitparams);
-
-
-		/** Set the wave for decoding.
-		 *
-		 * @param InWave - The wave proxy to decode.
-		 * @param InStartTimeSeconds - Number of seconds into the wave to start decoding.
-		 * @param InInitialPitchShiftSemitones - Initial pitch shift of wave.
-		 *
-		 * @return True on success, false on failure. 
-		 */
-		bool SetWave(const FSoundWaveProxyPtr& InWave, float InStartTimeSeconds=0.f, float InInitialPitchShiftSemitones=0.f);
-
-		/** Seek in the wave.
-		 *
-		 * @param InSeconds - Number of seconds into the wave to start decoding.
-		 *
-		 * @return True on success, false on failure. 
-		 */
-		bool SeekToTime(const float InSeconds);
-
-		bool CanGenerateAudio() const;
-
-		// returns number of samples written.   
-		uint32 GenerateAudio(float* OutputDest, int32 NumOutputFrames, int32& OutNumFramesConsumed, float PitchShiftInCents = 0.0f, bool bIsLooping = false);
-
-		uint32 GetNumChannels() const;
-
-
-	private:
-		FSoundWaveProxyPtr Wave;
-
-		// actual decoder objects
-		TUniquePtr<Audio::IDecoder> Decoder;
-		TUniquePtr<Audio::IDecoderOutput> Output;
-		TSharedPtr<Audio::IDecoderInput, ESPMode::ThreadSafe> Input;
-
-		// init helper for decoders
-		bool InitializeDecodersInternal(const FSoundWaveProxyPtr& InWave, float InStartTimeSeconds);
-
-		// SRC objects
-		Audio::FResampler Resampler;
-		FLinearPitchShifter PitchShifter;
-
-		// buffers
-		TArray<float> PreSrcBuffer;
-		TArray<float> PostSrcBuffer;
-		Audio::TCircularAudioBuffer<float> OutputCircularBuffer;
-
-		// meta data:
-		float InputSampleRate{ -1.f };
-		float OutputSampleRate{ -1.f };
-		float FsOutToInRatio{ 1.f };
-		float MaxPitchShiftCents{ 12.0f };
-		float MaxPitchShiftRatio{ 2.0f };
-
-		uint32 NumChannels{ 0 };
-		uint32 DecodeBlockSizeInFrames;
-		uint32 DecodeBlockSizeInSamples;
-		uint32 OutputBlockSizeInFrames { 512 };
-		
-		bool bDecoderIsDone{ true };
-		bool bDecoderHasLooped{ false };
-		bool bIsInitialized{ false };
-
-	}; // class FSimpleDecoderWrapper
-} // namespace Audio
diff --git a/Engine/Source/Runtime/AudioCodecEngine/Private/DecoderBackCompat.cpp b/Engine/Source/Runtime/AudioCodecEngine/Private/DecoderBackCompat.cpp
--- a/Engine/Source/Runtime/AudioCodecEngine/Private/DecoderBackCompat.cpp
+++ b/Engine/Source/Runtime/AudioCodecEngine/Private/DecoderBackCompat.cpp
@@ -94,16 +94,18 @@ namespace Audio
 			{
 				bFinished = Info->ReadCompressedData(Buff, bLoop, BuffSizeInBytes);
 			}
+			int32 NumSamplesStreamed = NumBytesStreamed / sizeof(int16);
+			int32 NumFramesStreamed = NumSamplesStreamed / Desc.NumChannels;
 			PushedDetails.SampleFramesStartOffset = FrameOffset;
-			ResidualBuffer.SetNum(NumBytesStreamed / sizeof(int16));
+			//ResidualBuffer.SetNum(NumBytesStreamed / sizeof(int16));
 
 			Dst->PushAudio(
 				PushedDetails,
-				MakeArrayView(ResidualBuffer)
+				MakeArrayView(ResidualBuffer.GetData(), NumSamplesStreamed)
 			);
 
-			FrameOffset += BuffSizeInFrames;
-			NumFramesRemaining -= BuffSizeInFrames;
+			FrameOffset += NumFramesStreamed;
+			NumFramesRemaining -= NumFramesStreamed;
 		}
 
 		if (!bFinished)
