Fix CIS linker issue
#okforgithbu public
#rnx

diff --git a/Engine/Plugins/Runtime/Metasound/Source/MetasoundEngine/Private/MetasoundWaveProxyReader.cpp b/Engine/Plugins/Runtime/Metasound/Source/MetasoundEngine/Private/MetasoundWaveProxyReader.cpp
--- a/Engine/Plugins/Runtime/Metasound/Source/MetasoundEngine/Private/MetasoundWaveProxyReader.cpp
+++ b/Engine/Plugins/Runtime/Metasound/Source/MetasoundEngine/Private/MetasoundWaveProxyReader.cpp
@@ -21,7 +21,7 @@
 #include "Templates/UniquePtr.h"
 
 static int32 GMetaSoundWaveProxyReaderSimulateSeekOnNonSeekable = 0;
-FAutoConsoleVariableRef CVarClearMetaSoundFrontendSearchEngineCache(
+FAutoConsoleVariableRef CVarMetaSoundWaveProxyReaderSimulateSeekOnNonSeekable(
 	TEXT("au.MetaSounds.WavePlayer.SimulateSeek"),
 	GMetaSoundWaveProxyReaderSimulateSeekOnNonSeekable ,
 	TEXT("If true, SoundWaves which are not of a seekable format will simulate seek calls by reading and discarding samples.\n")
