Add notification when MetaSound document has been versioned & mark dirty on editor open if so
#rnx

diff --git a/Engine/Plugins/Runtime/Metasound/Source/MetasoundEditor/Private/MetasoundEditor.cpp b/Engine/Plugins/Runtime/Metasound/Source/MetasoundEditor/Private/MetasoundEditor.cpp
--- a/Engine/Plugins/Runtime/Metasound/Source/MetasoundEditor/Private/MetasoundEditor.cpp
+++ b/Engine/Plugins/Runtime/Metasound/Source/MetasoundEditor/Private/MetasoundEditor.cpp
@@ -502,6 +502,8 @@ namespace Metasound
 
 			ExtendToolbar();
 			RegenerateMenusAndToolbars();
+
+			NotifyDocumentVersioned();
 		}
 
 		UObject* FEditor::GetMetasoundObject() const
@@ -569,6 +571,32 @@ namespace Metasound
 			FSlateApplication::Get().DismissAllMenus();
 		}
 
+		void FEditor::NotifyDocumentVersioned()
+		{
+			if (MetasoundGraphEditor.IsValid())
+			{
+				UMetasoundEditorGraph& MetaSoundGraph = GetMetaSoundGraphChecked();
+				const bool bVersionedOnLoad = MetaSoundGraph.GetVersionedOnLoad();
+				if (bVersionedOnLoad)
+				{
+					MetaSoundGraph.ClearVersionedOnLoad();
+					const FMetasoundAssetBase* MetaSoundAsset = IMetasoundUObjectRegistry::Get().GetObjectAsAssetBase(Metasound);
+					check(MetaSoundAsset);
+
+					const FString VersionString = MetaSoundAsset->GetDocumentChecked().Metadata.Version.Number.ToString();
+					FText Msg = FText::Format(LOCTEXT("MetaSoundDocumentVersioned", "Document versioned to '{0}' on load."), FText::FromString(VersionString));
+					FNotificationInfo Info(Msg);
+					Info.bFireAndForget = true;
+					Info.bUseSuccessFailIcons = false;
+					Info.ExpireDuration = 5.0f;
+
+					MetasoundGraphEditor->AddNotification(Info, false /* bSuccess */);
+
+					MetaSoundAsset->MarkMetasoundDocumentDirty();
+				}
+			}
+		}
+
 		void FEditor::NotifyNodePasteFailure_ReferenceLoop()
 		{
 			FNotificationInfo Info(LOCTEXT("NodePasteFailed_ReferenceLoop", "Node(s) not pasted: Nodes would create asset reference cycle."));
diff --git a/Engine/Plugins/Runtime/Metasound/Source/MetasoundEditor/Private/MetasoundEditor.h b/Engine/Plugins/Runtime/Metasound/Source/MetasoundEditor/Private/MetasoundEditor.h
--- a/Engine/Plugins/Runtime/Metasound/Source/MetasoundEditor/Private/MetasoundEditor.h
+++ b/Engine/Plugins/Runtime/Metasound/Source/MetasoundEditor/Private/MetasoundEditor.h
@@ -301,6 +301,7 @@ namespace Metasound
 
 			void EditObjectSettings();
 
+			void NotifyDocumentVersioned();
 			void NotifyNodePasteFailure_ReferenceLoop();
 
 			bool IsPlaying() const;
diff --git a/Engine/Plugins/Runtime/Metasound/Source/MetasoundEditor/Private/MetasoundEditorGraph.cpp b/Engine/Plugins/Runtime/Metasound/Source/MetasoundEditor/Private/MetasoundEditorGraph.cpp
--- a/Engine/Plugins/Runtime/Metasound/Source/MetasoundEditor/Private/MetasoundEditorGraph.cpp
+++ b/Engine/Plugins/Runtime/Metasound/Source/MetasoundEditor/Private/MetasoundEditorGraph.cpp
@@ -1261,6 +1261,21 @@ void UMetasoundEditorGraph::RegisterGraphWithFrontend()
 	}
 }
 
+void UMetasoundEditorGraph::ClearVersionedOnLoad()
+{
+	bVersionedOnLoad = false;
+}
+
+bool UMetasoundEditorGraph::GetVersionedOnLoad() const
+{
+	return bVersionedOnLoad;
+}
+
+void UMetasoundEditorGraph::SetVersionedOnLoad()
+{
+	bVersionedOnLoad = true;
+}
+
 void UMetasoundEditorGraph::SetSynchronizationRequired()
 {
 	using namespace Metasound;
diff --git a/Engine/Plugins/Runtime/Metasound/Source/MetasoundEditor/Public/MetasoundEditorGraph.h b/Engine/Plugins/Runtime/Metasound/Source/MetasoundEditor/Public/MetasoundEditorGraph.h
--- a/Engine/Plugins/Runtime/Metasound/Source/MetasoundEditor/Public/MetasoundEditorGraph.h
+++ b/Engine/Plugins/Runtime/Metasound/Source/MetasoundEditor/Public/MetasoundEditorGraph.h
@@ -385,6 +385,10 @@ public:
 	virtual void RegisterGraphWithFrontend() override;
 	virtual void SetSynchronizationRequired() override;
 
+	virtual void ClearVersionedOnLoad() override;
+	virtual bool GetVersionedOnLoad() const override;
+	virtual void SetVersionedOnLoad() override;
+
 private:
 	bool RemoveFrontendInput(UMetasoundEditorGraphInput& Input);
 	bool RemoveFrontendOutput(UMetasoundEditorGraphOutput& Output);
@@ -395,6 +399,8 @@ private:
 	// a sound's ParameterInterface when a sound begins playing.
 	uint32 PreviewID = INDEX_NONE;
 
+	bool bVersionedOnLoad = false;
+
 	UPROPERTY()
 	TArray<TObjectPtr<UMetasoundEditorGraphInput>> Inputs;
 
diff --git a/Engine/Plugins/Runtime/Metasound/Source/MetasoundEngine/Public/Metasound.h b/Engine/Plugins/Runtime/Metasound/Source/MetasoundEngine/Public/Metasound.h
--- a/Engine/Plugins/Runtime/Metasound/Source/MetasoundEngine/Public/Metasound.h
+++ b/Engine/Plugins/Runtime/Metasound/Source/MetasoundEngine/Public/Metasound.h
@@ -32,6 +32,12 @@ public:
 
 	virtual void SetSynchronizationRequired() PURE_VIRTUAL(UMetasoundEditorGraphBase::SetSynchronizationRequired, )
 	virtual void RegisterGraphWithFrontend() PURE_VIRTUAL(UMetasoundEditorGraphBase::RegisterGraphWithFrontend(), )
+
+#if WITH_EDITORONLY_DATA
+	virtual void ClearVersionedOnLoad() PURE_VIRTUAL(UMetasoundEditorGraphBase::ClearVersionedOnLoad, )
+	virtual bool GetVersionedOnLoad() const PURE_VIRTUAL(UMetasoundEditorGraphBase::GetVersionedOnLoad, return false; )
+	virtual void SetVersionedOnLoad() PURE_VIRTUAL(UMetasoundEditorGraphBase::SetVersionedOnLoad, )
+#endif // WITH_EDITORONLY_DATA
 };
 
 
@@ -94,7 +100,15 @@ namespace Metasound
 	{
 		if (InArchive.IsLoading())
 		{
-			InMetaSound.VersionAsset();
+			if (InMetaSound.VersionAsset())
+			{
+#if WITH_EDITORONLY_DATA
+				if (UMetasoundEditorGraphBase* MetaSoundGraph = Cast<UMetasoundEditorGraphBase>(InMetaSound.GetGraph()))
+				{
+					MetaSoundGraph->SetVersionedOnLoad();
+				}
+#endif // WITH_EDITORONLY_DATA
+			}
 		}
 	}
 
