Metasounds: Updated GetOwningAssetName to return GetPathName, and renamed FAutoUpdateRootGraph's member to DebugAssetPath to make the intention/uses clear

diff --git a/Engine/Plugins/Runtime/Metasound/Source/MetasoundFrontend/Private/MetasoundAssetBase.cpp b/Engine/Plugins/Runtime/Metasound/Source/MetasoundFrontend/Private/MetasoundAssetBase.cpp
--- a/Engine/Plugins/Runtime/Metasound/Source/MetasoundFrontend/Private/MetasoundAssetBase.cpp
+++ b/Engine/Plugins/Runtime/Metasound/Source/MetasoundFrontend/Private/MetasoundAssetBase.cpp
@@ -767,7 +767,7 @@ FString FMetasoundAssetBase::GetOwningAssetName() const
 {
 	if (const UObject* OwningAsset = GetOwningAsset())
 	{
-		return OwningAsset->GetName();
+		return OwningAsset->GetPathName();
 	}
 	return FString();
 }
diff --git a/Engine/Plugins/Runtime/Metasound/Source/MetasoundFrontend/Private/MetasoundFrontendTransform.cpp b/Engine/Plugins/Runtime/Metasound/Source/MetasoundFrontend/Private/MetasoundFrontendTransform.cpp
--- a/Engine/Plugins/Runtime/Metasound/Source/MetasoundFrontend/Private/MetasoundFrontendTransform.cpp
+++ b/Engine/Plugins/Runtime/Metasound/Source/MetasoundFrontend/Private/MetasoundFrontendTransform.cpp
@@ -475,7 +475,7 @@ namespace Metasound
 					PresetReferencedMetaSoundAsset = IMetaSoundAssetManager::GetChecked().TryLoadAssetFromKey(RegistryKey);
 					if (!PresetReferencedMetaSoundAsset)
 					{
-						UE_LOG(LogMetaSound, Error, TEXT("Auto-Updating preset '%s' failed: Referenced class '%s' missing."), *AssetPath, *ClassMetadata.GetClassName().ToString());
+						UE_LOG(LogMetaSound, Error, TEXT("Auto-Updating preset '%s' failed: Referenced class '%s' missing."), *DebugAssetPath, *ClassMetadata.GetClassName().ToString());
 					}
 					return;
 				}
@@ -489,12 +489,12 @@ namespace Metasound
 				FMetasoundFrontendVersionNumber UpdateVersion = NodeHandle->FindHighestMinorVersionInRegistry();
 				if (UpdateVersion.IsValid() && UpdateVersion > ClassMetadata.GetVersion())
 				{
-					UE_LOG(LogMetaSound, Display, TEXT("Auto-Updating '%s' node class '%s': Newer version '%s' found."), *AssetPath, *ClassMetadata.GetClassName().ToString(), *UpdateVersion.ToString());
+					UE_LOG(LogMetaSound, Display, TEXT("Auto-Updating '%s' node class '%s': Newer version '%s' found."), *DebugAssetPath, *ClassMetadata.GetClassName().ToString(), *UpdateVersion.ToString());
 				}
 				else if (InterfaceUpdates.ContainsChanges())
 				{
 					UpdateVersion = ClassMetadata.GetVersion();
-					UE_LOG(LogMetaSound, Display, TEXT("Auto-Updating '%s' node class '%s (%s)': Interface change detected."), *AssetPath, *ClassMetadata.GetClassName().ToString(), *UpdateVersion.ToString());
+					UE_LOG(LogMetaSound, Display, TEXT("Auto-Updating '%s' node class '%s (%s)': Interface change detected."), *DebugAssetPath, *ClassMetadata.GetClassName().ToString(), *UpdateVersion.ToString());
 				}
 
 				NodesToUpdate.Add(TPair<FNodeHandle, FMetasoundFrontendVersionNumber>(NodeHandle, UpdateVersion));
diff --git a/Engine/Plugins/Runtime/Metasound/Source/MetasoundFrontend/Public/MetasoundFrontendTransform.h b/Engine/Plugins/Runtime/Metasound/Source/MetasoundFrontend/Public/MetasoundFrontendTransform.h
--- a/Engine/Plugins/Runtime/Metasound/Source/MetasoundFrontend/Public/MetasoundFrontendTransform.h
+++ b/Engine/Plugins/Runtime/Metasound/Source/MetasoundFrontend/Public/MetasoundFrontendTransform.h
@@ -126,14 +126,14 @@ namespace Metasound
 		class METASOUNDFRONTEND_API FAutoUpdateRootGraph : public IDocumentTransform
 		{
 		public:
-			FAutoUpdateRootGraph(FString&& InAssetPath)
-				: AssetPath(MoveTemp(InAssetPath))
+			FAutoUpdateRootGraph(FString&& InDebugAssetPath)
+				: DebugAssetPath(MoveTemp(InDebugAssetPath))
 			{
 			}
 
 			bool Transform(FDocumentHandle InDocument) const override;
 
-			const FString AssetPath;
+			const FString DebugAssetPath;
 		};
 
 		/** Synchronizes the document's root graph's display name with that of the asset. */
