- Fix interface vertex metadata still serializing text during presave.
- Fix node tooltips caching old strings & improperly displaying them if a node's tooltips are modified but the node is not versioned

Minor Clean-up:
- Move interface registry generation logic to static function to mirror metadata's analogous call
- Remove misleading, unused GetPinTooltip cruft from SMetasoundGraphPin

#rnx

diff --git a/Engine/Plugins/Runtime/Metasound/Source/MetasoundEditor/Private/MetasoundEditorGraphBuilder.cpp b/Engine/Plugins/Runtime/Metasound/Source/MetasoundEditor/Private/MetasoundEditorGraphBuilder.cpp
--- a/Engine/Plugins/Runtime/Metasound/Source/MetasoundEditor/Private/MetasoundEditorGraphBuilder.cpp
+++ b/Engine/Plugins/Runtime/Metasound/Source/MetasoundEditor/Private/MetasoundEditorGraphBuilder.cpp
@@ -1248,7 +1248,8 @@ namespace Metasound
 
 		void FGraphBuilder::RefreshPinMetadata(UEdGraphPin& InPin, const FMetasoundFrontendVertexMetadata& InMetadata)
 		{
-			InPin.PinToolTip = InMetadata.GetDescription().ToString();
+			// Pin ToolTips are no longer cached on pins, and are instead dynamically generated via UMetasoundEditorGraphNode::GetPinHoverText
+			InPin.PinToolTip = { };
 			InPin.bAdvancedView = InMetadata.bIsAdvancedDisplay;
 			if (InPin.bAdvancedView)
 			{
diff --git a/Engine/Plugins/Runtime/Metasound/Source/MetasoundEditor/Private/MetasoundEditorGraphNode.cpp b/Engine/Plugins/Runtime/Metasound/Source/MetasoundEditor/Private/MetasoundEditorGraphNode.cpp
--- a/Engine/Plugins/Runtime/Metasound/Source/MetasoundEditor/Private/MetasoundEditorGraphNode.cpp
+++ b/Engine/Plugins/Runtime/Metasound/Source/MetasoundEditor/Private/MetasoundEditorGraphNode.cpp
@@ -250,6 +250,23 @@ FText UMetasoundEditorGraphNode::GetNodeTitle(ENodeTitleType::Type TitleType) co
 	return NodeHandle->GetDisplayTitle();
 }
 
+void UMetasoundEditorGraphNode::GetPinHoverText(const UEdGraphPin& Pin, FString& OutHoverText) const
+{
+	using namespace Metasound::Editor;
+	using namespace Metasound::Frontend;
+
+	if (Pin.Direction == EGPD_Input)
+	{
+		FConstInputHandle InputHandle = FGraphBuilder::GetConstInputHandleFromPin(&Pin);
+		OutHoverText = InputHandle->GetTooltip().ToString();
+	}
+	else // Pin.Direction == EGPD_Output
+	{
+		FConstOutputHandle OutputHandle = FGraphBuilder::GetConstOutputHandleFromPin(&Pin);
+		OutHoverText = OutputHandle->GetTooltip().ToString();
+	}
+}
+
 void UMetasoundEditorGraphNode::PinDefaultValueChanged(UEdGraphPin* Pin)
 {
 	using namespace Metasound::Editor;
diff --git a/Engine/Plugins/Runtime/Metasound/Source/MetasoundEditor/Private/SMetasoundGraphNode.cpp b/Engine/Plugins/Runtime/Metasound/Source/MetasoundEditor/Private/SMetasoundGraphNode.cpp
--- a/Engine/Plugins/Runtime/Metasound/Source/MetasoundEditor/Private/SMetasoundGraphNode.cpp
+++ b/Engine/Plugins/Runtime/Metasound/Source/MetasoundEditor/Private/SMetasoundGraphNode.cpp
@@ -245,58 +245,49 @@ namespace Metasound
 				// Don't show default value field for container types
 				if (InPin->PinType.ContainerType != EPinContainerType::None)
 				{
-					PinWidget = SNew(SMetasoundGraphPin, InPin)
-						.ToolTipText(this, &SMetaSoundGraphNode::GetPinTooltip, InPin);
+					PinWidget = SNew(SMetasoundGraphPin, InPin);
 				}
 
 				else if (InPin->PinType.PinCategory == FGraphBuilder::PinCategoryAudio)
 				{
-					PinWidget = SNew(SMetasoundGraphPin, InPin)
-						.ToolTipText(this, &SMetaSoundGraphNode::GetPinTooltip, InPin);
+					PinWidget = SNew(SMetasoundGraphPin, InPin);
 				}
 
 				else if (InPin->PinType.PinCategory == FGraphBuilder::PinCategoryBoolean)
 				{
-					PinWidget = SNew(SMetasoundGraphPinBool, InPin)
-						.ToolTipText(this, &SMetaSoundGraphNode::GetPinTooltip, InPin);
+					PinWidget = SNew(SMetasoundGraphPinBool, InPin);
 				}
 
 				else if (InPin->PinType.PinCategory == FGraphBuilder::PinCategoryFloat)
 				{
-					PinWidget = SNew(SMetasoundGraphPinFloat, InPin)
-						.ToolTipText(this, &SMetaSoundGraphNode::GetPinTooltip, InPin);
+					PinWidget = SNew(SMetasoundGraphPinFloat, InPin);
 				}
 
 				else if (InPin->PinType.PinCategory == FGraphBuilder::PinCategoryInt32)
 				{
 					if (SMetasoundGraphEnumPin::FindEnumInterfaceFromPin(InPin))
 					{
-						PinWidget = SNew(SMetasoundGraphEnumPin, InPin)
-							.ToolTipText(this, &SMetaSoundGraphNode::GetPinTooltip, InPin);
+						PinWidget = SNew(SMetasoundGraphEnumPin, InPin);
 					}
 					else
 					{
-						PinWidget = SNew(SMetasoundGraphPinInteger, InPin)
-							.ToolTipText(this, &SMetaSoundGraphNode::GetPinTooltip, InPin);
+						PinWidget = SNew(SMetasoundGraphPinInteger, InPin);
 					}
 				}
 
 				else if (InPin->PinType.PinCategory == FGraphBuilder::PinCategoryObject)
 				{
-					PinWidget = SNew(SMetasoundGraphPinObject, InPin)
-						.ToolTipText(this, &SMetaSoundGraphNode::GetPinTooltip, InPin);
+					PinWidget = SNew(SMetasoundGraphPinObject, InPin);
 				}
 
 				else if (InPin->PinType.PinCategory == FGraphBuilder::PinCategoryString)
 				{
-					PinWidget = SNew(SMetasoundGraphPinString, InPin)
-						.ToolTipText(this, &SMetaSoundGraphNode::GetPinTooltip, InPin);
+					PinWidget = SNew(SMetasoundGraphPinString, InPin);
 				}
 
 				else if (InPin->PinType.PinCategory == FGraphBuilder::PinCategoryTrigger)
 				{
-					PinWidget = SNew(SMetasoundGraphPin, InPin)
-						.ToolTipText(this, &SMetaSoundGraphNode::GetPinTooltip, InPin);
+					PinWidget = SNew(SMetasoundGraphPin, InPin);
 
 					if (const ISlateStyle* MetasoundStyle = FSlateStyleRegistry::FindSlateStyle("MetaSoundStyle"))
 					{
@@ -309,8 +300,7 @@ namespace Metasound
 
 			if (!PinWidget.IsValid())
 			{
-				PinWidget = SNew(SMetasoundGraphPin, InPin)
-					.ToolTipText(this, &SMetaSoundGraphNode::GetPinTooltip, InPin);
+				PinWidget = SNew(SMetasoundGraphPin, InPin);
 			}
 
 			return PinWidget;
@@ -344,22 +334,6 @@ namespace Metasound
 			}
 		}
 
-		FText SMetaSoundGraphNode::GetPinTooltip(UEdGraphPin* InPin) const
-		{
-			using namespace Frontend;
-
-			if (InPin->Direction == EGPD_Input)
-			{
-				FConstInputHandle InputHandle = FGraphBuilder::GetConstInputHandleFromPin(InPin);
-				return InputHandle->GetTooltip();
-			}
-			else
-			{
-				FConstOutputHandle OutputHandle = FGraphBuilder::GetConstOutputHandleFromPin(InPin);
-				return OutputHandle->GetTooltip();
-			}
-		}
-
 		TSharedRef<SWidget> SMetaSoundGraphNode::CreateTitleWidget(TSharedPtr<SNodeTitle> NodeTitle)
 		{
 			Metasound::Frontend::FNodeHandle NodeHandle = GetMetaSoundNode().GetNodeHandle();
diff --git a/Engine/Plugins/Runtime/Metasound/Source/MetasoundEditor/Private/SMetasoundGraphNode.h b/Engine/Plugins/Runtime/Metasound/Source/MetasoundEditor/Private/SMetasoundGraphNode.h
--- a/Engine/Plugins/Runtime/Metasound/Source/MetasoundEditor/Private/SMetasoundGraphNode.h
+++ b/Engine/Plugins/Runtime/Metasound/Source/MetasoundEditor/Private/SMetasoundGraphNode.h
@@ -71,8 +71,6 @@ namespace Metasound
 
 			TAttribute<EVisibility> GetSimulationVisibilityAttribute() const;
 
-			FText GetPinTooltip(UEdGraphPin* InPin) const;
-
 			// Slider widget for float input
 			TSharedPtr<SAudioInputWidget> InputWidget;
 			// Handle for on value changed delegate for input slider 
diff --git a/Engine/Plugins/Runtime/Metasound/Source/MetasoundEditor/Public/MetasoundEditorGraphNode.h b/Engine/Plugins/Runtime/Metasound/Source/MetasoundEditor/Public/MetasoundEditorGraphNode.h
--- a/Engine/Plugins/Runtime/Metasound/Source/MetasoundEditor/Public/MetasoundEditorGraphNode.h
+++ b/Engine/Plugins/Runtime/Metasound/Source/MetasoundEditor/Public/MetasoundEditorGraphNode.h
@@ -55,6 +55,7 @@ public:
 	virtual bool CanCreateUnderSpecifiedSchema(const UEdGraphSchema* Schema) const override;
 	virtual bool CanUserDeleteNode() const override;
 	virtual FText GetNodeTitle(ENodeTitleType::Type TitleType) const override;
+	virtual void GetPinHoverText(const UEdGraphPin& Pin, FString& OutHoverText) const override;
 	virtual FString GetDocumentationExcerptName() const override;
 	virtual FString GetDocumentationLink() const override;
 	virtual void GetNodeContextMenuActions(class UToolMenu* Menu, class UGraphNodeContextMenuContext* Context) const override;
diff --git a/Engine/Plugins/Runtime/Metasound/Source/MetasoundFrontend/Private/MetasoundAssetBase.cpp b/Engine/Plugins/Runtime/Metasound/Source/MetasoundFrontend/Private/MetasoundAssetBase.cpp
--- a/Engine/Plugins/Runtime/Metasound/Source/MetasoundFrontend/Private/MetasoundAssetBase.cpp
+++ b/Engine/Plugins/Runtime/Metasound/Source/MetasoundFrontend/Private/MetasoundAssetBase.cpp
@@ -466,14 +466,14 @@ void FMetasoundAssetBase::CacheRegistryMetadata()
 	}
 
 	// 4. Cache registry data on document dependencies
-	for (FMetasoundFrontendClass& Class : Document->Dependencies)
+	for (FMetasoundFrontendClass& Dependency : Document->Dependencies)
 	{
-		if (!Class.CacheRegistryMetadata())
+		if (!FMetasoundFrontendClass::CacheGraphDependencyMetadataFromRegistry(Dependency))
 		{
 			UE_LOG(LogMetaSound, Warning,
 				TEXT("'%s' failed to cache dependency registry data: Registry missing class with key '%s'"),
 				*GetOwningAssetName(),
-				*Class.Metadata.GetClassName().ToString());
+				*Dependency.Metadata.GetClassName().ToString());
 			UE_LOG(LogMetaSound, Warning,
 				TEXT("Asset '%s' may fail to build runtime graph unless re-registered after dependency with given key is loaded."),
 				*GetOwningAssetName());
diff --git a/Engine/Plugins/Runtime/Metasound/Source/MetasoundFrontend/Private/MetasoundFrontend.cpp b/Engine/Plugins/Runtime/Metasound/Source/MetasoundFrontend/Private/MetasoundFrontend.cpp
--- a/Engine/Plugins/Runtime/Metasound/Source/MetasoundFrontend/Private/MetasoundFrontend.cpp
+++ b/Engine/Plugins/Runtime/Metasound/Source/MetasoundFrontend/Private/MetasoundFrontend.cpp
@@ -36,86 +36,7 @@ namespace Metasound
 			FMetasoundFrontendClass ClassDescription;
 
 			ClassDescription.Metadata = FMetasoundFrontendClassMetadata::GenerateClassDescription(InNodeMetadata, ClassType);
-
-			FMetasoundFrontendClassStyleDisplay DisplayStyle(InNodeMetadata.DisplayStyle);
-			ClassDescription.Style = FMetasoundFrontendClassStyle
-			{
-				DisplayStyle
-			};
-
-			FMetasoundFrontendClassInterface& ClassInterface = ClassDescription.Interface;
-
-			const FInputVertexInterface& InputInterface = InNodeMetadata.DefaultInterface.GetInputInterface();
-			FMetasoundFrontendInterfaceStyle InputStyle;
-			for (const TPair<FVertexName, FInputDataVertex>& InputTuple : InputInterface)
-			{
-				FMetasoundFrontendClassInput ClassInput;
-
-				const FInputDataVertex& InputVertex = InputTuple.Value;
-				ClassInput.Name = InputVertex.GetVertexName();
-				ClassInput.TypeName = InputVertex.GetDataTypeName();
-				ClassInput.VertexID = FGuid::NewGuid();
-
-
-				const FDataVertexMetadata& VertexMetadata = InputVertex.GetMetadata();
-
-				ClassInput.Metadata.SetDescription(VertexMetadata.Description);
-
-				ClassInput.Metadata.bIsAdvancedDisplay = VertexMetadata.bIsAdvancedDisplay;
-
-				FLiteral DefaultLiteral = InputVertex.GetDefaultLiteral();
-				if (DefaultLiteral.GetType() != ELiteralType::Invalid)
-				{
-					ClassInput.DefaultLiteral.SetFromLiteral(DefaultLiteral);
-				}
-
-				// Advanced display items are pushed to bottom of sort order
-				int32 OrderIndex = InputInterface.GetOrderIndex(InputTuple.Key);
-				if (ClassInput.Metadata.bIsAdvancedDisplay)
-				{
-					OrderIndex += InputInterface.Num();
-				}
-				InputStyle.DefaultSortOrder.Add(OrderIndex);
-
-				ClassInterface.Inputs.Add(ClassInput);
-			}
-			ClassInterface.SetInputStyle(InputStyle);
-
-			const FOutputVertexInterface& OutputInterface = InNodeMetadata.DefaultInterface.GetOutputInterface();
-			FMetasoundFrontendInterfaceStyle OutputStyle;
-			for (const TPair<FVertexName, FOutputDataVertex>& OutputTuple : OutputInterface)
-			{
-				FMetasoundFrontendClassOutput ClassOutput;
-
-				ClassOutput.Name = OutputTuple.Value.GetVertexName();
-				ClassOutput.TypeName = OutputTuple.Value.GetDataTypeName();
-				ClassOutput.VertexID = FGuid::NewGuid();
-
-				const FDataVertexMetadata& VertexMetadata = OutputTuple.Value.GetMetadata();
-				ClassOutput.Metadata.SetDescription(VertexMetadata.Description);
-				ClassOutput.Metadata.bIsAdvancedDisplay = VertexMetadata.bIsAdvancedDisplay;
-
-				// Advanced display items are pushed to bottom below non-advanced
-				int32 OrderIndex = OutputInterface.GetOrderIndex(OutputTuple.Key);
-				if (ClassOutput.Metadata.bIsAdvancedDisplay)
-				{
-					OrderIndex += OutputInterface.Num();
-				}
-				OutputStyle.DefaultSortOrder.Add(OrderIndex);
-
-				ClassInterface.Outputs.Add(ClassOutput);
-			}
-			ClassInterface.SetOutputStyle(OutputStyle);
-
-			for (auto& EnvTuple : InNodeMetadata.DefaultInterface.GetEnvironmentInterface())
-			{
-				FMetasoundFrontendClassEnvironmentVariable EnvVar;
-
-				EnvVar.Name = EnvTuple.Value.GetVertexName();
-				EnvVar.bIsRequired = true;
-
-				ClassInterface.Environment.Add(EnvVar);
-			}
+			ClassDescription.Interface = FMetasoundFrontendClassInterface::GenerateClassDescription(InNodeMetadata);
 
 			return ClassDescription;
 		}
diff --git a/Engine/Plugins/Runtime/Metasound/Source/MetasoundFrontend/Private/MetasoundFrontendDocument.cpp b/Engine/Plugins/Runtime/Metasound/Source/MetasoundFrontend/Private/MetasoundFrontendDocument.cpp
--- a/Engine/Plugins/Runtime/Metasound/Source/MetasoundFrontend/Private/MetasoundFrontendDocument.cpp
+++ b/Engine/Plugins/Runtime/Metasound/Source/MetasoundFrontend/Private/MetasoundFrontendDocument.cpp
@@ -9,6 +9,7 @@
 #include "MetasoundFrontend.h"
 #include "MetasoundFrontendRegistries.h"
 #include "MetasoundLog.h"
+#include "MetasoundVertex.h"
 
 namespace Metasound
 {
@@ -181,6 +182,105 @@ bool operator!=(const FMetasoundFrontendClassName& InLHS, const FMetasoundFronte
 	return !(InLHS == InRHS);
 }
 
+FMetasoundFrontendClassInterface FMetasoundFrontendClassInterface::GenerateClassDescription(const Metasound::FNodeClassMetadata& InNodeClassMetadata)
+{
+	using namespace Metasound;
+	using namespace Metasound::Frontend;
+
+	FMetasoundFrontendClassInterface ClassInterface;
+
+	// Copy over inputs
+	{
+		const FInputVertexInterface& InputInterface = InNodeClassMetadata.DefaultInterface.GetInputInterface();
+		FMetasoundFrontendInterfaceStyle InputStyle;
+		for (const TPair<FVertexName, FInputDataVertex>& InputTuple : InputInterface)
+		{
+			FMetasoundFrontendClassInput ClassInput;
+
+			const FInputDataVertex& InputVertex = InputTuple.Value;
+			ClassInput.Name = InputVertex.GetVertexName();
+			ClassInput.TypeName = InputVertex.GetDataTypeName();
+			ClassInput.VertexID = FGuid::NewGuid();
+
+			const FDataVertexMetadata& VertexMetadata = InputVertex.GetMetadata();
+
+			ClassInput.Metadata.SetSerializeText(false);
+			ClassInput.Metadata.SetDescription(VertexMetadata.Description);
+
+			ClassInput.Metadata.bIsAdvancedDisplay = VertexMetadata.bIsAdvancedDisplay;
+
+			FLiteral DefaultLiteral = InputVertex.GetDefaultLiteral();
+			if (DefaultLiteral.GetType() != ELiteralType::Invalid)
+			{
+				ClassInput.DefaultLiteral.SetFromLiteral(DefaultLiteral);
+			}
+
+			// Advanced display items are pushed to bottom of sort order
+			int32 OrderIndex = InputInterface.GetOrderIndex(InputTuple.Key);
+			if (ClassInput.Metadata.bIsAdvancedDisplay)
+			{
+				OrderIndex += InputInterface.Num();
+			}
+			InputStyle.DefaultSortOrder.Add(OrderIndex);
+
+			ClassInterface.Inputs.Add(MoveTemp(ClassInput));
+		}
+
+		// Must set via direct accessor to avoid updating the change GUID
+		// (All instances of this generation call should be done for code
+		// defined classes only, which do not currently create a persistent
+		// change hash between builds and leave the guid 0'ed).
+		ClassInterface.InputStyle = InputStyle;
+	}
+
+	// Copy over outputs
+	{
+		const FOutputVertexInterface& OutputInterface = InNodeClassMetadata.DefaultInterface.GetOutputInterface();
+		FMetasoundFrontendInterfaceStyle OutputStyle;
+		for (const TPair<FVertexName, FOutputDataVertex>& OutputTuple : OutputInterface)
+		{
+			FMetasoundFrontendClassOutput ClassOutput;
+
+			ClassOutput.Name = OutputTuple.Value.GetVertexName();
+			ClassOutput.TypeName = OutputTuple.Value.GetDataTypeName();
+			ClassOutput.VertexID = FGuid::NewGuid();
+
+			const FDataVertexMetadata& VertexMetadata = OutputTuple.Value.GetMetadata();
+			ClassOutput.Metadata.SetSerializeText(false);
+			ClassOutput.Metadata.SetDescription(VertexMetadata.Description);
+			ClassOutput.Metadata.bIsAdvancedDisplay = VertexMetadata.bIsAdvancedDisplay;
+
+			// Advanced display items are pushed to bottom below non-advanced
+			int32 OrderIndex = OutputInterface.GetOrderIndex(OutputTuple.Key);
+			if (ClassOutput.Metadata.bIsAdvancedDisplay)
+			{
+				OrderIndex += OutputInterface.Num();
+			}
+			OutputStyle.DefaultSortOrder.Add(OrderIndex);
+
+			ClassInterface.Outputs.Add(MoveTemp(ClassOutput));
+		}
+
+		// Must set via direct accessor to avoid updating the change GUID
+		// (All instances of this generation call should be done for code
+		// defined classes only, which do not currently create a persistent
+		// change hash between builds and leave the guid 0'ed).
+		ClassInterface.OutputStyle = MoveTemp(OutputStyle);
+	}
+
+	for (auto& EnvTuple : InNodeClassMetadata.DefaultInterface.GetEnvironmentInterface())
+	{
+		FMetasoundFrontendClassEnvironmentVariable EnvVar;
+
+		EnvVar.Name = EnvTuple.Value.GetVertexName();
+		EnvVar.bIsRequired = true;
+
+		ClassInterface.Environment.Add(EnvVar);
+	}
+
+	return ClassInterface;
+}
+
 void FMetasoundFrontendClassMetadata::SetAuthor(const FText& InAuthor)
 {
 	using namespace Metasound::DocumentPrivate;
@@ -288,19 +388,19 @@ void FMetasoundFrontendClassMetadata::SetVersion(const FMetasoundFrontendVersion
 }
 
 #if WITH_EDITORONLY_DATA
-bool FMetasoundFrontendClass::CacheRegistryMetadata()
+bool FMetasoundFrontendClass::CacheGraphDependencyMetadataFromRegistry(FMetasoundFrontendClass& InOutDependency)
 {
 	using namespace Metasound::Frontend;
 
-	const FNodeRegistryKey Key = NodeRegistryKey::CreateKey(Metadata);
-	FMetasoundFrontendClass Class;
+	const FNodeRegistryKey Key = NodeRegistryKey::CreateKey(InOutDependency.Metadata);
+	FMetasoundFrontendClass RegistryClass;
 
 	FMetasoundFrontendRegistryContainer* Registry = FMetasoundFrontendRegistryContainer::Get();
 	if (ensure(Registry))
 	{
-		if (Registry->FindFrontendClassFromRegistered(Key, Class))
+		if (Registry->FindFrontendClassFromRegistered(Key, RegistryClass))
 		{
-			Metadata = Class.Metadata;
+			InOutDependency.Metadata = RegistryClass.Metadata;
 
 			using FNameTypeKey = TPair<FName, FName>;
 			TMap<FNameTypeKey, const FMetasoundFrontendVertexMetadata*> InterfaceMembers;
@@ -317,20 +417,22 @@ bool FMetasoundFrontendClass::CacheRegistryMetadata()
 				if (const FMetasoundFrontendVertexMetadata* RegVertex = InterfaceMembers.FindRef(Key))
 				{
 					OutVertex.Metadata = *RegVertex;
+					OutVertex.Metadata.SetSerializeText(false);
 				}
 			};
 
-			Algo::Transform(Class.Interface.Inputs, InterfaceMembers, [&](const FMetasoundFrontendClassInput& Input) { return MakePairFromVertex(Input); });
-			Algo::ForEach(Interface.Inputs, [&](FMetasoundFrontendClassInput& Input) { CacheRegistryVertexMetadata(Input); });
+			Algo::Transform(RegistryClass.Interface.Inputs, InterfaceMembers, [&](const FMetasoundFrontendClassInput& Input) { return MakePairFromVertex(Input); });
+			Algo::ForEach(InOutDependency.Interface.Inputs, [&](FMetasoundFrontendClassInput& Input) { CacheRegistryVertexMetadata(Input); });
 
 			InterfaceMembers.Reset();
 
-			Algo::Transform(Class.Interface.Outputs, InterfaceMembers, [&](const FMetasoundFrontendClassOutput& Output) { return MakePairFromVertex(Output); });
-			Algo::ForEach(Interface.Outputs, [&](FMetasoundFrontendClassOutput& Output) { CacheRegistryVertexMetadata(Output); });
+			Algo::Transform(RegistryClass.Interface.Outputs, InterfaceMembers, [&](const FMetasoundFrontendClassOutput& Output) { return MakePairFromVertex(Output); });
+			Algo::ForEach(InOutDependency.Interface.Outputs, [&](FMetasoundFrontendClassOutput& Output) { CacheRegistryVertexMetadata(Output); });
+
+			InOutDependency.Interface.InputStyle = RegistryClass.Interface.InputStyle;
+			InOutDependency.Interface.OutputStyle = RegistryClass.Interface.OutputStyle;
 
-			Interface.InputStyle = Class.Interface.InputStyle;
-			Interface.OutputStyle = Class.Interface.OutputStyle;
-			Style = Class.Style;
+			InOutDependency.Style = RegistryClass.Style;
 
 			return true;
 		}
diff --git a/Engine/Plugins/Runtime/Metasound/Source/MetasoundFrontend/Public/MetasoundFrontendDocument.h b/Engine/Plugins/Runtime/Metasound/Source/MetasoundFrontend/Public/MetasoundFrontendDocument.h
--- a/Engine/Plugins/Runtime/Metasound/Source/MetasoundFrontend/Public/MetasoundFrontendDocument.h
+++ b/Engine/Plugins/Runtime/Metasound/Source/MetasoundFrontend/Public/MetasoundFrontendDocument.h
@@ -838,6 +838,10 @@ private:
 	FMetasoundFrontendInterfaceStyle OutputStyle;
 
 public:
+	// Generates class interface intended to be used as a registry descriptor from FNodeClassMetadata.
+	// Does not initialize a change ID as it is not considered to be transactional.
+	static FMetasoundFrontendClassInterface GenerateClassDescription(const Metasound::FNodeClassMetadata& InNodeClassMetadata);
+
 	// Description of class inputs.
 	UPROPERTY(VisibleAnywhere, Category = CustomView)
 	TArray<FMetasoundFrontendClassInput> Inputs;
@@ -1170,12 +1174,13 @@ struct METASOUNDFRONTEND_API FMetasoundFrontendClass
 
 #if WITH_EDITORONLY_DATA
 	/*
-	 * Caches transient metadata (class & vertex) found in the registry
-	 * that is not necessary for serialization or core graph generation.
+	 * Caches transient style, class & vertex Metadata found in the registry
+	 * on a passed (presumed) dependency.  Only modifies properties that are
+	 * not necessary for serialization or core graph generation.
 	 *
 	 * @return - Whether class was found in the registry & data was cached successfully.
 	 */
-	bool CacheRegistryMetadata();
+	static bool CacheGraphDependencyMetadataFromRegistry(FMetasoundFrontendClass& InOutDependency);
 #endif // WITH_EDITORONLY_DATA
 };
 
