Fix for sparse ensure when unregistering due to AssetClassKeys sometimes being out-of-date on the document. Minor rename to clarify a graph has been modified vs. its dirty state (which may be true regardless of function's execution)
#rnx

diff --git a/Engine/Plugins/Runtime/Metasound/Source/MetasoundEditor/Private/MetasoundEditorGraphBuilder.cpp b/Engine/Plugins/Runtime/Metasound/Source/MetasoundEditor/Private/MetasoundEditorGraphBuilder.cpp
--- a/Engine/Plugins/Runtime/Metasound/Source/MetasoundEditor/Private/MetasoundEditorGraphBuilder.cpp
+++ b/Engine/Plugins/Runtime/Metasound/Source/MetasoundEditor/Private/MetasoundEditorGraphBuilder.cpp
@@ -1332,6 +1332,7 @@ namespace Metasound
 					{
 						if (FMetasoundAssetBase* EditedMetaSound = IMetasoundUObjectRegistry::Get().GetObjectAsAssetBase(Asset))
 						{
+							EditedMetaSound->RebuildReferencedAssetClassKeys();
 							if (EditedMetaSound->IsReferencedAsset(*MetaSoundAsset))
 							{
 								EditedMetaSound->SetSynchronizationRequired();
@@ -1635,12 +1636,12 @@ namespace Metasound
 
 		bool FGraphBuilder::SynchronizeGraph(UObject& InMetaSound)
 		{
-			bool bIsEditorGraphDirty = SynchronizeGraphVertices(InMetaSound);
-			bIsEditorGraphDirty |= SynchronizeNodeMembers(InMetaSound);
-			bIsEditorGraphDirty |= SynchronizeNodes(InMetaSound);
-			bIsEditorGraphDirty |= SynchronizeConnections(InMetaSound);
+			bool bEditorGraphModified = SynchronizeGraphVertices(InMetaSound);
+			bEditorGraphModified |= SynchronizeNodeMembers(InMetaSound);
+			bEditorGraphModified |= SynchronizeNodes(InMetaSound);
+			bEditorGraphModified |= SynchronizeConnections(InMetaSound);
 
-			if (bIsEditorGraphDirty)
+			if (bEditorGraphModified)
 			{
 				InMetaSound.MarkPackageDirty();
 			}
@@ -1658,7 +1659,7 @@ namespace Metasound
 		{
 			using namespace Frontend;
 
-			bool bIsEditorGraphDirty = false;
+			bool bEditorGraphModified = false;
 
 			FMetasoundAssetBase* MetaSoundAsset = IMetasoundUObjectRegistry::Get().GetObjectAsAssetBase(&InMetaSound);
 			check(MetaSoundAsset);
@@ -1690,7 +1691,7 @@ namespace Metasound
 								FText InputDisplayName = Node->GetDisplayName();
 								UE_LOG(LogMetasoundEditor, Verbose, TEXT("Editor Input Node '%s' interface versioned"), *InputDisplayName.ToString());
 
-								bIsEditorGraphDirty = true;
+								bEditorGraphModified = true;
 							}
 						}
 					}
@@ -1720,21 +1721,21 @@ namespace Metasound
 								FText OutputDisplayName = Node->GetDisplayName();
 								UE_LOG(LogMetasoundEditor, Verbose, TEXT("Editor Output Node '%s' interface versioned"), *OutputDisplayName.ToString());
 
-								bIsEditorGraphDirty = true;
+								bEditorGraphModified = true;
 							}
 						}
 					}
 				}
 			}
 
-			return bIsEditorGraphDirty;
+			return bEditorGraphModified;
 		}
 
 		bool FGraphBuilder::SynchronizeNodes(UObject& InMetaSound)
 		{
 			using namespace Frontend;
 
-			bool bIsEditorGraphDirty = false;
+			bool bEditorGraphModified = false;
 
 			// Get all external nodes from Frontend graph.  Input and output references will only be added/synchronized
 			// if required when synchronizing connections (as they are not required to inhabit editor graph).
@@ -1805,7 +1806,7 @@ namespace Metasound
 			// exist before attempting to synchronize connections.
 			for (UMetasoundEditorGraphNode* EditorNode : EditorNodes)
 			{
-				bIsEditorGraphDirty |= EditorGraph->RemoveNode(EditorNode);
+				bEditorGraphModified |= EditorGraph->RemoveNode(EditorNode);
 			}
 
 			// Add missing editor nodes marked as visible.
@@ -1818,7 +1819,7 @@ namespace Metasound
 				}
 
 				FMetasoundFrontendNodeStyle NewStyle = CurrentStyle;
-				bIsEditorGraphDirty = true;
+				bEditorGraphModified = true;
 
 				TArray<UMetasoundEditorGraphNode*> AddedNodes;
 				for (const TPair<FGuid, FVector2D>& Location : NewStyle.Display.Locations)
@@ -1854,11 +1855,11 @@ namespace Metasound
 			{
 				for (UMetasoundEditorGraphNode* EditorNode : IdNodePair.Value.EditorNodes)
 				{
-					bIsEditorGraphDirty |= SynchronizeNodePins(*EditorNode, IdNodePair.Value.Node);
+					bEditorGraphModified |= SynchronizeNodePins(*EditorNode, IdNodePair.Value.Node);
 				}
 			}
 
-			return bIsEditorGraphDirty;
+			return bEditorGraphModified;
 		}
 
 		bool FGraphBuilder::SynchronizeNodePins(UMetasoundEditorGraphNode& InEditorNode, Frontend::FConstNodeHandle InNode, bool bRemoveUnusedPins, bool bLogChanges)
@@ -2009,7 +2010,7 @@ namespace Metasound
 		{
 			using namespace Frontend;
 
-			bool bIsEditorGraphDirty = false;
+			bool bEditorGraphModified = false;
 
 			FMetasoundAssetBase* MetaSoundAsset = IMetasoundUObjectRegistry::Get().GetObjectAsAssetBase(&InMetaSound);
 			check(MetaSoundAsset);
@@ -2033,7 +2034,7 @@ namespace Metasound
 				constexpr bool bIncludeNamespace = true;
 				FText NodeDisplayName = FGraphBuilder::GetDisplayName(*NodeHandle, bIncludeNamespace);
 				UE_LOG(LogMetasoundEditor, Verbose, TEXT("Synchronizing Inputs: Added missing input '%s'."), *NodeDisplayName.ToString());
-				bIsEditorGraphDirty = true;
+				bEditorGraphModified = true;
 			}, EMetasoundFrontendClassType::Input);
 
 			// Collect all editor graph outputs with corresponding frontend outputs. 
@@ -2050,7 +2051,7 @@ namespace Metasound
 				constexpr bool bIncludeNamespace = true;
 				FText NodeDisplayName = FGraphBuilder::GetDisplayName(*NodeHandle, bIncludeNamespace);
 				UE_LOG(LogMetasoundEditor, Verbose, TEXT("Synchronizing Outputs: Added missing output '%s'."), *NodeDisplayName.ToString());
-				bIsEditorGraphDirty = true;
+				bEditorGraphModified = true;
 			}, EMetasoundFrontendClassType::Output);
 
 			// Collect editor inputs and outputs to remove which have no corresponding frontend input or output.
@@ -2073,7 +2074,7 @@ namespace Metasound
 			});
 
 			// Remove stale inputs and outputs.
-			bIsEditorGraphDirty |= !ToRemove.IsEmpty();
+			bEditorGraphModified |= !ToRemove.IsEmpty();
 			for (UMetasoundEditorGraphMember* GraphMember: ToRemove)
 			{
 				Graph->RemoveMember(*GraphMember);
@@ -2135,7 +2136,7 @@ namespace Metasound
 
 
 
-			return bIsEditorGraphDirty;
+			return bEditorGraphModified;
 		}
 	} // namespace Editor
 } // namespace Metasound
