Fix for invalid proxy generate classes showing up on metasound node drop downs.

diff --git a/Engine/Plugins/Runtime/Metasound/Source/MetasoundEditor/Private/MetasoundEditorGraphNode.cpp b/Engine/Plugins/Runtime/Metasound/Source/MetasoundEditor/Private/MetasoundEditorGraphNode.cpp
--- a/Engine/Plugins/Runtime/Metasound/Source/MetasoundEditor/Private/MetasoundEditorGraphNode.cpp
+++ b/Engine/Plugins/Runtime/Metasound/Source/MetasoundEditor/Private/MetasoundEditorGraphNode.cpp
@@ -341,7 +341,7 @@ TSet<FString> UMetasoundEditorGraphNode::GetDisallowedPinClassNames(const UEdGra
 			continue;
 		}
 
-		if (EditorModule.IsExplicitProxyClass(*Class) && Class->IsChildOf(ProxyGenClass))
+		if (EditorModule.IsExplicitProxyClass(*ProxyGenClass) && Class->IsChildOf(ProxyGenClass))
 		{
 			DisallowedClasses.Add(ClassIt->GetName());
 		}
