Fix for metasound wave player node using incorrect OnNearlyDone trigger index
#rnx

diff --git a/Engine/Plugins/Runtime/Metasound/Source/MetasoundEngine/Private/MetasoundWavePlayerNode.cpp b/Engine/Plugins/Runtime/Metasound/Source/MetasoundEngine/Private/MetasoundWavePlayerNode.cpp
--- a/Engine/Plugins/Runtime/Metasound/Source/MetasoundEngine/Private/MetasoundWavePlayerNode.cpp
+++ b/Engine/Plugins/Runtime/Metasound/Source/MetasoundEngine/Private/MetasoundWavePlayerNode.cpp
@@ -881,7 +881,7 @@ namespace Metasound
 					// alter the predicted OnDone frame between render blocks.
 					if (!bOnNearlyDoneTriggeredForWave)
 					{
-						TriggerOnNearlyDone->TriggerFrame(InOperatorStartFrame + Events.OnEOFFrameIndex);
+						TriggerOnNearlyDone->TriggerFrame(InOperatorStartFrame + Events.OnEOFFrameIndex - OperatorSettings.GetNumFramesPerBlock());
 						bOnNearlyDoneTriggeredForWave = true;
 					}
 				}
