Fix for deleting copied inputs causing them to sneak back into graph
#rnx

diff --git a/Engine/Plugins/Runtime/Metasound/Source/MetasoundEditor/Private/MetasoundEditor.cpp b/Engine/Plugins/Runtime/Metasound/Source/MetasoundEditor/Private/MetasoundEditor.cpp
--- a/Engine/Plugins/Runtime/Metasound/Source/MetasoundEditor/Private/MetasoundEditor.cpp
+++ b/Engine/Plugins/Runtime/Metasound/Source/MetasoundEditor/Private/MetasoundEditor.cpp
@@ -1815,7 +1815,7 @@ namespace Metasound
 					{
 						const FVector2D NewNodeLocation = FVector2D(GraphNode->NodePosX, GraphNode->NodePosY);
 						FMetasoundFrontendNodeStyle NodeStyle = NodeHandle->GetNodeStyle();
-						NodeStyle.Display.Locations.FindOrAdd(NodeHandle->GetID()) = NewNodeLocation;
+						NodeStyle.Display.Locations.FindOrAdd(MetasoundGraphNode->NodeGuid) = NewNodeLocation;
 						NodeHandle->SetNodeStyle(NodeStyle);
 					}
 				}
