Move asset loading logs to verbose to avoid spamming editor on load
#rnx

diff --git a/Engine/Plugins/Runtime/Metasound/Source/MetasoundEngine/Private/MetasoundAssetSubsystem.cpp b/Engine/Plugins/Runtime/Metasound/Source/MetasoundEngine/Private/MetasoundAssetSubsystem.cpp
--- a/Engine/Plugins/Runtime/Metasound/Source/MetasoundEngine/Private/MetasoundAssetSubsystem.cpp
+++ b/Engine/Plugins/Runtime/Metasound/Source/MetasoundEngine/Private/MetasoundAssetSubsystem.cpp
@@ -130,7 +130,7 @@ void UMetaSoundAssetSubsystem::AddAssetReferences(FMetasoundAssetBase& InAssetBa
 	{
 		if (!ContainsKey(ReferencedAssetClassKey))
 		{
-			UE_LOG(LogMetaSound, Log, TEXT("Missing referenced class '%s' asset entry."), *ReferencedAssetClassKey);
+			UE_LOG(LogMetaSound, Verbose, TEXT("Missing referenced class '%s' asset entry."), *ReferencedAssetClassKey);
 			bLoadFromPathCache = true;
 		}
 	}
@@ -141,7 +141,7 @@ void UMetaSoundAssetSubsystem::AddAssetReferences(FMetasoundAssetBase& InAssetBa
 		return;
 	}
 
-	UE_LOG(LogMetaSound, Log, TEXT("Attempting preemptive reference load..."));
+	UE_LOG(LogMetaSound, Verbose, TEXT("Attempting preemptive reference load..."));
 
 	// If keys are not loaded, iterate class cache paths to prime asset manager with
 	// hint paths as the registration is getting called before asset manager has 
@@ -155,7 +155,7 @@ void UMetaSoundAssetSubsystem::AddAssetReferences(FMetasoundAssetBase& InAssetBa
 			const FNodeRegistryKey ClassKey = NodeRegistryKey::CreateKey(ClassInfo);
 			if (!ContainsKey(ClassKey))
 			{
-				UE_LOG(LogMetaSound, Log,
+				UE_LOG(LogMetaSound, Verbose,
 					TEXT("Preemptive load of class '%s' from hint path '%s' due to early "
 						"registration request (asset scan likely not complete)."),
 					*ClassKey,
