Fix missing copy for style to registry from code-defined MetaSound node classes, causing math nodes to show inputs, outputs, names, & no icons & sort order for pins to be wrong.

diff --git a/Engine/Plugins/Runtime/Metasound/Source/MetasoundFrontend/Private/MetasoundFrontend.cpp b/Engine/Plugins/Runtime/Metasound/Source/MetasoundFrontend/Private/MetasoundFrontend.cpp
--- a/Engine/Plugins/Runtime/Metasound/Source/MetasoundFrontend/Private/MetasoundFrontend.cpp
+++ b/Engine/Plugins/Runtime/Metasound/Source/MetasoundFrontend/Private/MetasoundFrontend.cpp
@@ -36,7 +36,8 @@ namespace Metasound
 			FMetasoundFrontendClass ClassDescription;
 
 			ClassDescription.Metadata = FMetasoundFrontendClassMetadata::GenerateClassDescription(InNodeMetadata, ClassType);
-			ClassDescription.Interface = FMetasoundFrontendClassInterface::GenerateClassDescription(InNodeMetadata);
+			ClassDescription.Interface = FMetasoundFrontendClassInterface::GenerateClassDescription(InNodeMetadata.DefaultInterface);
+			ClassDescription.Style = FMetasoundFrontendClassStyle::GenerateClassDescription(InNodeMetadata.DisplayStyle);
 
 			return ClassDescription;
 		}
diff --git a/Engine/Plugins/Runtime/Metasound/Source/MetasoundFrontend/Private/MetasoundFrontendDocument.cpp b/Engine/Plugins/Runtime/Metasound/Source/MetasoundFrontend/Private/MetasoundFrontendDocument.cpp
--- a/Engine/Plugins/Runtime/Metasound/Source/MetasoundFrontend/Private/MetasoundFrontendDocument.cpp
+++ b/Engine/Plugins/Runtime/Metasound/Source/MetasoundFrontend/Private/MetasoundFrontendDocument.cpp
@@ -182,7 +182,7 @@ bool operator!=(const FMetasoundFrontendClassName& InLHS, const FMetasoundFronte
 	return !(InLHS == InRHS);
 }
 
-FMetasoundFrontendClassInterface FMetasoundFrontendClassInterface::GenerateClassDescription(const Metasound::FNodeClassMetadata& InNodeClassMetadata)
+FMetasoundFrontendClassInterface FMetasoundFrontendClassInterface::GenerateClassDescription(const Metasound::FVertexInterface& InDefaultInterface)
 {
 	using namespace Metasound;
 	using namespace Metasound::Frontend;
@@ -191,7 +191,7 @@ FMetasoundFrontendClassInterface FMetasoundFrontendClassInterface::GenerateClass
 
 	// Copy over inputs
 	{
-		const FInputVertexInterface& InputInterface = InNodeClassMetadata.DefaultInterface.GetInputInterface();
+		const FInputVertexInterface& InputInterface = InDefaultInterface.GetInputInterface();
 		FMetasoundFrontendInterfaceStyle InputStyle;
 		for (const TPair<FVertexName, FInputDataVertex>& InputTuple : InputInterface)
 		{
@@ -235,7 +235,7 @@ FMetasoundFrontendClassInterface FMetasoundFrontendClassInterface::GenerateClass
 
 	// Copy over outputs
 	{
-		const FOutputVertexInterface& OutputInterface = InNodeClassMetadata.DefaultInterface.GetOutputInterface();
+		const FOutputVertexInterface& OutputInterface = InDefaultInterface.GetOutputInterface();
 		FMetasoundFrontendInterfaceStyle OutputStyle;
 		for (const TPair<FVertexName, FOutputDataVertex>& OutputTuple : OutputInterface)
 		{
@@ -268,7 +268,7 @@ FMetasoundFrontendClassInterface FMetasoundFrontendClassInterface::GenerateClass
 		ClassInterface.OutputStyle = MoveTemp(OutputStyle);
 	}
 
-	for (auto& EnvTuple : InNodeClassMetadata.DefaultInterface.GetEnvironmentInterface())
+	for (auto& EnvTuple : InDefaultInterface.GetEnvironmentInterface())
 	{
 		FMetasoundFrontendClassEnvironmentVariable EnvVar;
 
@@ -306,7 +306,8 @@ void FMetasoundFrontendClassMetadata::SetCategoryHierarchy(const TArray<FText>&
 void FMetasoundFrontendClassMetadata::SetKeywords(const TArray<FText>& InKeywords)
 {
 	using namespace Metasound::DocumentPrivate;
-	SetWithChangeID(InKeywords, Keywords, ChangeID);
+	TArray<FText>& TextToSet = bSerializeText ? Keywords : KeywordsTransient;
+	SetWithChangeID(InKeywords, TextToSet, ChangeID);
 }
 
 void FMetasoundFrontendClassMetadata::SetClassName(const FMetasoundFrontendClassName& InClassName)
@@ -387,6 +388,17 @@ void FMetasoundFrontendClassMetadata::SetVersion(const FMetasoundFrontendVersion
 	SetWithChangeID(InVersion, Version, ChangeID);
 }
 
+FMetasoundFrontendClassStyle FMetasoundFrontendClassStyle::GenerateClassDescription(const Metasound::FNodeDisplayStyle& InNodeDisplayStyle)
+{
+	FMetasoundFrontendClassStyle Style;
+	Style.Display.bShowName = InNodeDisplayStyle.bShowName;
+	Style.Display.bShowInputNames = InNodeDisplayStyle.bShowInputNames;
+	Style.Display.bShowOutputNames = InNodeDisplayStyle.bShowOutputNames;
+	Style.Display.ImageName = InNodeDisplayStyle.ImageName;
+
+	return Style;
+}
+
 #if WITH_EDITORONLY_DATA
 bool FMetasoundFrontendClass::CacheGraphDependencyMetadataFromRegistry(FMetasoundFrontendClass& InOutDependency)
 {
@@ -401,6 +413,7 @@ bool FMetasoundFrontendClass::CacheGraphDependencyMetadataFromRegistry(FMetasoun
 		if (Registry->FindFrontendClassFromRegistered(Key, RegistryClass))
 		{
 			InOutDependency.Metadata = RegistryClass.Metadata;
+			InOutDependency.Style = RegistryClass.Style;
 
 			using FNameTypeKey = TPair<FName, FName>;
 			TMap<FNameTypeKey, const FMetasoundFrontendVertexMetadata*> InterfaceMembers;
@@ -432,8 +445,6 @@ bool FMetasoundFrontendClass::CacheGraphDependencyMetadataFromRegistry(FMetasoun
 			InOutDependency.Interface.InputStyle = RegistryClass.Interface.InputStyle;
 			InOutDependency.Interface.OutputStyle = RegistryClass.Interface.OutputStyle;
 
-			InOutDependency.Style = RegistryClass.Style;
-
 			return true;
 		}
 	}
diff --git a/Engine/Plugins/Runtime/Metasound/Source/MetasoundFrontend/Public/MetasoundFrontendDocument.h b/Engine/Plugins/Runtime/Metasound/Source/MetasoundFrontend/Public/MetasoundFrontendDocument.h
--- a/Engine/Plugins/Runtime/Metasound/Source/MetasoundFrontend/Public/MetasoundFrontendDocument.h
+++ b/Engine/Plugins/Runtime/Metasound/Source/MetasoundFrontend/Public/MetasoundFrontendDocument.h
@@ -840,7 +840,7 @@ private:
 public:
 	// Generates class interface intended to be used as a registry descriptor from FNodeClassMetadata.
 	// Does not initialize a change ID as it is not considered to be transactional.
-	static FMetasoundFrontendClassInterface GenerateClassDescription(const Metasound::FNodeClassMetadata& InNodeClassMetadata);
+	static FMetasoundFrontendClassInterface GenerateClassDescription(const Metasound::FVertexInterface& InNodeClassMetadata);
 
 	// Description of class inputs.
 	UPROPERTY(VisibleAnywhere, Category = CustomView)
@@ -1151,6 +1151,9 @@ struct FMetasoundFrontendClassStyle
 
 	UPROPERTY()
 	FMetasoundFrontendClassStyleDisplay Display;
+
+	// Generates class style from core node class metadata.
+	static FMetasoundFrontendClassStyle GenerateClassDescription(const Metasound::FNodeDisplayStyle& InNodeDisplayStyle);
 };
 
 USTRUCT()
