Retain all connections when control + click dragging multiple connections to another MetaSound node

diff --git a/Engine/Plugins/Runtime/Metasound/Source/MetasoundEditor/Private/MetasoundEditorGraphSchema.cpp b/Engine/Plugins/Runtime/Metasound/Source/MetasoundEditor/Private/MetasoundEditorGraphSchema.cpp
--- a/Engine/Plugins/Runtime/Metasound/Source/MetasoundEditor/Private/MetasoundEditorGraphSchema.cpp
+++ b/Engine/Plugins/Runtime/Metasound/Source/MetasoundEditor/Private/MetasoundEditorGraphSchema.cpp
@@ -322,6 +322,30 @@ namespace Metasound
 	} // namespace Editor
 } // namespace Metasound
 
+UEdGraphNode* FMetasoundGraphSchemaAction_NodeWithMultipleOutputs::PerformAction(class UEdGraph* ParentGraph, TArray<UEdGraphPin*>& FromPins, const FVector2D Location, bool bSelectNewNode/* = true*/)
+{
+	UEdGraphNode* ResultNode = NULL;
+
+	if (FromPins.Num() > 0)
+	{
+		ResultNode = PerformAction(ParentGraph, FromPins[0], Location, bSelectNewNode);
+
+		if (ResultNode)
+		{
+			// Try autowiring the rest of the pins
+			for (int32 Index = 1; Index < FromPins.Num(); ++Index)
+			{
+				ResultNode->AutowireNewNode(FromPins[Index]);
+			}
+		}
+	}
+	else
+	{
+		ResultNode = PerformAction(ParentGraph, NULL, Location, bSelectNewNode);
+	}
+
+	return ResultNode;
+}
 
 const FSlateBrush* FMetasoundGraphSchemaAction_NewNode::GetIconBrush() const
 {
@@ -382,7 +406,7 @@ UEdGraphNode* FMetasoundGraphSchemaAction_NewNode::PerformAction(UEdGraph* Paren
 }
 
 FMetasoundGraphSchemaAction_NewInput::FMetasoundGraphSchemaAction_NewInput(FText InNodeCategory, FText InDisplayName, FGuid InNodeID, FText InToolTip, Metasound::Editor::EPrimaryContextGroup InGroup)
-	: FMetasoundGraphSchemaAction(MoveTemp(InNodeCategory), MoveTemp(InDisplayName), MoveTemp(InToolTip), InGroup)
+	: FMetasoundGraphSchemaAction_NodeWithMultipleOutputs(MoveTemp(InNodeCategory), MoveTemp(InDisplayName), MoveTemp(InToolTip), InGroup)
 	, NodeID(InNodeID)
 {
 }
@@ -446,7 +470,7 @@ UEdGraphNode* FMetasoundGraphSchemaAction_NewInput::PerformAction(UEdGraph* Pare
 }
 
 FMetasoundGraphSchemaAction_PromoteToInput::FMetasoundGraphSchemaAction_PromoteToInput()
-	: FMetasoundGraphSchemaAction(
+	: FMetasoundGraphSchemaAction_NodeWithMultipleOutputs(
 		FText(),
 		LOCTEXT("PromoteToInputName", "Promote To Graph Input"),
 		LOCTEXT("PromoteToInputTooltip", "Promotes node input to graph input"),
@@ -503,7 +527,7 @@ UEdGraphNode* FMetasoundGraphSchemaAction_PromoteToInput::PerformAction(UEdGraph
 }
 
 FMetasoundGraphSchemaAction_PromoteToVariable_AccessorNode::FMetasoundGraphSchemaAction_PromoteToVariable_AccessorNode()
-	: FMetasoundGraphSchemaAction(
+	: FMetasoundGraphSchemaAction_NodeWithMultipleOutputs(
 		FText(),
 		LOCTEXT("PromoteToVariableGetterName", "Promote To Graph Variable"),
 		LOCTEXT("PromoteToInputTooltip", "Promotes node input to graph variable and creates a getter node"),
@@ -535,7 +559,7 @@ UEdGraphNode* FMetasoundGraphSchemaAction_PromoteToVariable_AccessorNode::Perfor
 }
 
 FMetasoundGraphSchemaAction_PromoteToVariable_DeferredAccessorNode::FMetasoundGraphSchemaAction_PromoteToVariable_DeferredAccessorNode()
-	: FMetasoundGraphSchemaAction(
+	: FMetasoundGraphSchemaAction_NodeWithMultipleOutputs(
 		FText(),
 		LOCTEXT("PromoteToVariableDeferredGetterName", "Promote To Graph Variable (Deferred)"),
 		LOCTEXT("PromoteToInputTooltip", "Promotes node input to graph variable and creates a deferred getter node"),
@@ -715,7 +739,7 @@ UEdGraphNode* FMetasoundGraphSchemaAction_PromoteToOutput::PerformAction(UEdGrap
 }
 
 FMetasoundGraphSchemaAction_NewVariableNode::FMetasoundGraphSchemaAction_NewVariableNode(FText InNodeCategory, FText InDisplayName, FGuid InVariableID, FText InToolTip)
-	: FMetasoundGraphSchemaAction(MoveTemp(InNodeCategory), MoveTemp(InDisplayName), MoveTemp(InToolTip), Metasound::Editor::EPrimaryContextGroup::Variables)
+	: FMetasoundGraphSchemaAction_NodeWithMultipleOutputs(MoveTemp(InNodeCategory), MoveTemp(InDisplayName), MoveTemp(InToolTip), Metasound::Editor::EPrimaryContextGroup::Variables)
 	, VariableID(InVariableID)
 {
 }
diff --git a/Engine/Plugins/Runtime/Metasound/Source/MetasoundEditor/Public/MetasoundEditorGraphSchema.h b/Engine/Plugins/Runtime/Metasound/Source/MetasoundEditor/Public/MetasoundEditorGraphSchema.h
--- a/Engine/Plugins/Runtime/Metasound/Source/MetasoundEditor/Public/MetasoundEditorGraphSchema.h
+++ b/Engine/Plugins/Runtime/Metasound/Source/MetasoundEditor/Public/MetasoundEditorGraphSchema.h
@@ -85,10 +85,29 @@ struct METASOUNDEDITOR_API FMetasoundGraphSchemaAction : public FEdGraphSchemaAc
 	}
 };
 
+/** This is used to combine functionality for nodes that can have multiple outputs and should never be directly instantiated. */
+USTRUCT()
+struct METASOUNDEDITOR_API FMetasoundGraphSchemaAction_NodeWithMultipleOutputs : public FMetasoundGraphSchemaAction
+{
+	GENERATED_USTRUCT_BODY();
+
+	FMetasoundGraphSchemaAction_NodeWithMultipleOutputs()
+		: FMetasoundGraphSchemaAction()
+	{}
+
+	FMetasoundGraphSchemaAction_NodeWithMultipleOutputs(FText InNodeCategory, FText InMenuDesc, FText InToolTip, Metasound::Editor::EPrimaryContextGroup InGroup, FText InKeywords = FText::GetEmpty())
+		: FMetasoundGraphSchemaAction(InNodeCategory, InMenuDesc, InToolTip, InGroup, InKeywords)
+	{}
+
+	//~ Begin FEdGraphSchemaAction Interface
+	virtual UEdGraphNode* PerformAction(class UEdGraph* ParentGraph, TArray<UEdGraphPin*>& FromPins, const FVector2D Location, bool bSelectNewNode = true) override;
+	virtual UEdGraphNode* PerformAction(class UEdGraph* ParentGraph, UEdGraphPin* FromPin, const FVector2D Location, bool bSelectNewNode = true) { return nullptr; }
+	//~ End FEdGraphSchemaAction Interface
+};
 
 /** Action to add an input reference to the graph */
 USTRUCT()
-struct METASOUNDEDITOR_API FMetasoundGraphSchemaAction_NewInput : public FMetasoundGraphSchemaAction
+struct METASOUNDEDITOR_API FMetasoundGraphSchemaAction_NewInput : public FMetasoundGraphSchemaAction_NodeWithMultipleOutputs
 {
 	GENERATED_USTRUCT_BODY();
 
@@ -96,7 +115,7 @@ struct METASOUNDEDITOR_API FMetasoundGraphSchemaAction_NewInput : public FMetaso
 	FGuid NodeID;
 
 	FMetasoundGraphSchemaAction_NewInput()
-		: FMetasoundGraphSchemaAction()
+		: FMetasoundGraphSchemaAction_NodeWithMultipleOutputs()
 	{}
 
 	FMetasoundGraphSchemaAction_NewInput(FText InNodeCategory, FText InDisplayName, FGuid InInputNodeID, FText InToolTip, Metasound::Editor::EPrimaryContextGroup InGroup);
@@ -114,7 +133,7 @@ struct METASOUNDEDITOR_API FMetasoundGraphSchemaAction_NewInput : public FMetaso
 
 /** Promotes an input to a graph input, using its respective literal value as the default value */
 USTRUCT()
-struct METASOUNDEDITOR_API FMetasoundGraphSchemaAction_PromoteToInput : public FMetasoundGraphSchemaAction
+struct METASOUNDEDITOR_API FMetasoundGraphSchemaAction_PromoteToInput : public FMetasoundGraphSchemaAction_NodeWithMultipleOutputs
 {
 	GENERATED_USTRUCT_BODY();
 
@@ -166,7 +185,7 @@ struct METASOUNDEDITOR_API FMetasoundGraphSchemaAction_PromoteToOutput : public
 
 /** Adds a variable node to the graph */
 USTRUCT()
-struct METASOUNDEDITOR_API FMetasoundGraphSchemaAction_NewVariableNode : public FMetasoundGraphSchemaAction
+struct METASOUNDEDITOR_API FMetasoundGraphSchemaAction_NewVariableNode : public FMetasoundGraphSchemaAction_NodeWithMultipleOutputs
 {
 	GENERATED_USTRUCT_BODY();
 
@@ -174,7 +193,7 @@ struct METASOUNDEDITOR_API FMetasoundGraphSchemaAction_NewVariableNode : public
 	FGuid VariableID;
 
 	FMetasoundGraphSchemaAction_NewVariableNode()
-		: FMetasoundGraphSchemaAction()
+		: FMetasoundGraphSchemaAction_NodeWithMultipleOutputs()
 	{}
 
 	FMetasoundGraphSchemaAction_NewVariableNode(FText InNodeCategory, FText InDisplayName, FGuid InVariableID, FText InToolTip);
@@ -248,7 +267,7 @@ protected:
 
 /** Promotes an input to a graph variable & respective getter node, using its respective literal value as the default value */
 USTRUCT()
-struct METASOUNDEDITOR_API FMetasoundGraphSchemaAction_PromoteToVariable_AccessorNode : public FMetasoundGraphSchemaAction
+struct METASOUNDEDITOR_API FMetasoundGraphSchemaAction_PromoteToVariable_AccessorNode : public FMetasoundGraphSchemaAction_NodeWithMultipleOutputs
 {
 	GENERATED_USTRUCT_BODY();
 
@@ -274,7 +293,7 @@ struct METASOUNDEDITOR_API FMetasoundGraphSchemaAction_PromoteToVariable_Mutator
 
 /** Promotes an input to a graph variable & respective deferred getter node, using its respective literal value as the default value */
 USTRUCT()
-struct METASOUNDEDITOR_API FMetasoundGraphSchemaAction_PromoteToVariable_DeferredAccessorNode : public FMetasoundGraphSchemaAction
+struct METASOUNDEDITOR_API FMetasoundGraphSchemaAction_PromoteToVariable_DeferredAccessorNode : public FMetasoundGraphSchemaAction_NodeWithMultipleOutputs
 {
 	GENERATED_USTRUCT_BODY();
 
@@ -288,7 +307,7 @@ struct METASOUNDEDITOR_API FMetasoundGraphSchemaAction_PromoteToVariable_Deferre
 
 /** Action to add a node to the graph */
 USTRUCT()
-struct METASOUNDEDITOR_API FMetasoundGraphSchemaAction_NewNode : public FMetasoundGraphSchemaAction
+struct METASOUNDEDITOR_API FMetasoundGraphSchemaAction_NewNode : public FMetasoundGraphSchemaAction_NodeWithMultipleOutputs
 {
 	GENERATED_USTRUCT_BODY();
 
@@ -296,11 +315,11 @@ struct METASOUNDEDITOR_API FMetasoundGraphSchemaAction_NewNode : public FMetasou
 	FMetasoundFrontendClassMetadata ClassMetadata;
 
 	FMetasoundGraphSchemaAction_NewNode() 
-		: FMetasoundGraphSchemaAction()
+		: FMetasoundGraphSchemaAction_NodeWithMultipleOutputs()
 	{}
 
 	FMetasoundGraphSchemaAction_NewNode(const FText& InNodeCategory, const FText& InMenuDesc, const FText& InToolTip, Metasound::Editor::EPrimaryContextGroup InGroup, FText InKeywords = FText::GetEmpty())
-		: FMetasoundGraphSchemaAction(InNodeCategory, InMenuDesc, InToolTip, InGroup, InKeywords)
+		: FMetasoundGraphSchemaAction_NodeWithMultipleOutputs(InNodeCategory, InMenuDesc, InToolTip, InGroup, InKeywords)
 	{}
 
 	//~ Begin FMetasoundGraphSchemaAction Interface
