MetaSound Composition Fixes
- Removed need for marking ed graph nodes for refresh in most places in favor of using frontend class metadata/interface guids to signify an ed node requiring a refresh
- Update CanAutoUpdate to always provide interface changes if the metadata/interface guids change and return true, even if the runtime properties do not report
 changes, so the editor can adequately update referenced classes that have only asthetic updates
#rnx

diff --git a/Engine/Plugins/Runtime/Metasound/Source/MetasoundEditor/Private/MetasoundEditor.cpp b/Engine/Plugins/Runtime/Metasound/Source/MetasoundEditor/Private/MetasoundEditor.cpp
--- a/Engine/Plugins/Runtime/Metasound/Source/MetasoundEditor/Private/MetasoundEditor.cpp
+++ b/Engine/Plugins/Runtime/Metasound/Source/MetasoundEditor/Private/MetasoundEditor.cpp
@@ -434,22 +434,25 @@ namespace Metasound
 			checkf(IMetasoundUObjectRegistry::Get().IsRegisteredClass(ObjectToEdit), TEXT("Object passed in was not registered as a valid metasound interface!"));
 			
 			// Support undo/redo
-			ObjectToEdit->SetFlags(RF_Transactional);
-
 			Metasound = ObjectToEdit;
-
-			FGraphBuilder::RegisterGraphWithFrontend(*Metasound);
-			FGraphBuilder::SynchronizeGraph(*Metasound);
+			Metasound->SetFlags(RF_Transactional);
 
 			GEditor->RegisterForUndo(this);
 
 			FGraphEditorCommands::Register();
 			FEditorCommands::Register();
 
-			BindGraphCommands();
+			FGraphBuilder::RegisterGraphWithFrontend(*Metasound);
 
+			BindGraphCommands();
 			CreateInternalWidgets();
 
+			// Call after create internal widgets as editor is
+			// referenced in validation step. TODO: Make this
+			// relationship more explicit.
+			constexpr bool bForceRefreshNodes = true;
+			FGraphBuilder::SynchronizeGraph(*Metasound, bForceRefreshNodes);
+
 			CreateAnalyzers();
 
 			const TSharedRef<FTabManager::FLayout> StandaloneDefaultLayout = FTabManager::NewLayout("Standalone_MetasoundEditor_Layout_v9")
@@ -597,6 +600,16 @@ namespace Metasound
 			}
 		}
 
+		void FEditor::NotifyNodePasteFailure_MultipleVariableSetters()
+		{
+			FNotificationInfo Info(LOCTEXT("NodePasteFailed_MultipleVariableSetters", "Node(s) not pasted: Only one variable setter node possible per graph."));
+			Info.bFireAndForget = true;
+			Info.bUseSuccessFailIcons = false;
+			Info.ExpireDuration = 5.0f;
+
+			MetasoundGraphEditor->AddNotification(Info, false /* bSuccess */);
+		}
+
 		void FEditor::NotifyNodePasteFailure_ReferenceLoop()
 		{
 			FNotificationInfo Info(LOCTEXT("NodePasteFailed_ReferenceLoop", "Node(s) not pasted: Nodes would create asset reference cycle."));
@@ -1700,6 +1713,7 @@ namespace Metasound
 			NodeTextToPaste.Empty();
 
 			bool bNotifyReferenceLoop = false;
+			bool bNotifyMultipleVariableSetters = false;
 
 			TArray<UEdGraphNode*> NodesToRemove;
 			for (UEdGraphNode* GraphNode : PastedGraphNodes)
@@ -1746,34 +1760,54 @@ namespace Metasound
 				}
 				else if (UMetasoundEditorGraphInputNode* InputNode = Cast<UMetasoundEditorGraphInputNode>(GraphNode))
 				{
-					if (!Graph.ContainsInput(InputNode->Input))
+					if (!InputNode->Input || !Graph.ContainsInput(*InputNode->Input))
 					{
 						NodesToRemove.Add(GraphNode);
 					}
 				}
 				else if (UMetasoundEditorGraphOutputNode* OutputNode = Cast<UMetasoundEditorGraphOutputNode>(GraphNode))
 				{
-					if (!Graph.ContainsOutput(OutputNode->Output))
+					if (OutputNode->Output && Graph.ContainsOutput(*OutputNode->Output))
+					{
+						auto NodeMatches = [OutputNodeID = OutputNode->GetNodeID()](const TObjectPtr<UEdGraphNode>& EdNode)
+						{
+							if (UMetasoundEditorGraphOutputNode* OutputNode = Cast<UMetasoundEditorGraphOutputNode>(EdNode))
+							{
+								return OutputNodeID == OutputNode->GetNodeID();
+							}
+							return false;
+						};
+
+						// Can only have one output reference node
+						if (Graph.Nodes.ContainsByPredicate(NodeMatches))
+						{
+							NodesToRemove.Add(GraphNode);
+						}
+					}
+					else
 					{
 						NodesToRemove.Add(GraphNode);
 					}
-
-					auto NodeMatches = [OutputNodeID = OutputNode->GetNodeID()](const TObjectPtr<UEdGraphNode>& EdNode)
+				}
+				else if (UMetasoundEditorGraphVariableNode* VariableNode = Cast<UMetasoundEditorGraphVariableNode>(GraphNode))
+				{
+					// Can only have one setter node
+					if (const UMetasoundEditorGraphVariable* Variable = VariableNode->Variable)
 					{
-						if (UMetasoundEditorGraphOutputNode* OutputNode = Cast<UMetasoundEditorGraphOutputNode>(EdNode))
+						FConstVariableHandle VariableHandle = Variable->GetConstVariableHandle();
+						FConstNodeHandle VariableMutatorNodeHandle = VariableHandle->FindMutatorNode();
+						if (VariableNode->GetNodeID() == VariableMutatorNodeHandle->GetID())
 						{
-							return OutputNodeID == OutputNode->GetNodeID();
+							bNotifyMultipleVariableSetters = true;
+							NodesToRemove.Add(GraphNode);
 						}
-						return false;
-					};
-
-					// Can only have one output reference node
-					if (Graph.Nodes.ContainsByPredicate(NodeMatches))
+					}
+					else
 					{
 						NodesToRemove.Add(GraphNode);
 					}
 				}
-				else if (!(GraphNode->IsA<UEdGraphNode_Comment>() || GraphNode->IsA<UMetasoundEditorGraphVariableNode>()))
+				else if (!GraphNode->IsA<UEdGraphNode_Comment>())
 				{
 					checkNoEntry();
 				}
@@ -1869,6 +1903,11 @@ namespace Metasound
 				NotifyNodePasteFailure_ReferenceLoop();
 			}
 
+			if (bNotifyMultipleVariableSetters)
+			{
+				NotifyNodePasteFailure_MultipleVariableSetters();
+			}
+
 			MetasoundGraphEditor->NotifyGraphChanged();
 		}
 
diff --git a/Engine/Plugins/Runtime/Metasound/Source/MetasoundEditor/Private/MetasoundEditor.h b/Engine/Plugins/Runtime/Metasound/Source/MetasoundEditor/Private/MetasoundEditor.h
--- a/Engine/Plugins/Runtime/Metasound/Source/MetasoundEditor/Private/MetasoundEditor.h
+++ b/Engine/Plugins/Runtime/Metasound/Source/MetasoundEditor/Private/MetasoundEditor.h
@@ -302,6 +302,7 @@ namespace Metasound
 			void EditObjectSettings();
 
 			void NotifyDocumentVersioned();
+			void NotifyNodePasteFailure_MultipleVariableSetters();
 			void NotifyNodePasteFailure_ReferenceLoop();
 
 			bool IsPlaying() const;
diff --git a/Engine/Plugins/Runtime/Metasound/Source/MetasoundEditor/Private/MetasoundEditorGraph.cpp b/Engine/Plugins/Runtime/Metasound/Source/MetasoundEditor/Private/MetasoundEditorGraph.cpp
--- a/Engine/Plugins/Runtime/Metasound/Source/MetasoundEditor/Private/MetasoundEditorGraph.cpp
+++ b/Engine/Plugins/Runtime/Metasound/Source/MetasoundEditor/Private/MetasoundEditorGraph.cpp
@@ -36,7 +36,6 @@ namespace Metasound
 				const bool bCommitChange = InChangeType != EPropertyChangeType::Interactive;
 				if (bCommitChange)
 				{
-					InMember.MarkNodesForRefresh();
 					if (UObject* MetaSound = InMember.GetOutermostObject())
 					{
 						FGraphBuilder::RegisterGraphWithFrontend(*MetaSound);
@@ -61,25 +60,6 @@ const UMetasoundEditorGraph* UMetasoundEditorGraphMember::GetOwningGraph() const
 	return Cast<const UMetasoundEditorGraph>(GetOuter());
 }
 
-void UMetasoundEditorGraphMember::MarkNodesForRefresh()
-{
-	using namespace Metasound;
-
-	UMetasoundEditorGraph* Graph = GetOwningGraph();
-	if (ensure(Graph))
-	{
-		FMetasoundAssetBase* MetasoundAsset = IMetasoundUObjectRegistry::Get().GetObjectAsAssetBase(&Graph->GetMetasoundChecked());
-		check(MetasoundAsset);
-		MetasoundAsset->SetSynchronizationRequired();
-
-		const TArray<UMetasoundEditorGraphNode*> Nodes = GetNodes();
-		for (UMetasoundEditorGraphNode* Node : Nodes)
-		{
-			Node->bRefreshNode = true;
-		}
-	}
-}
-
 void UMetasoundEditorGraphMember::ConformLiteralDataType()
 {
 	using namespace Metasound::Editor;
@@ -219,6 +199,7 @@ FName UMetasoundEditorGraphVertex::GetMemberName() const
 
 void UMetasoundEditorGraphVertex::SetMemberName(const FName& InNewName, bool bPostTransaction)
 {
+	using namespace Metasound;
 	using namespace Metasound::Editor;
 	using namespace Metasound::Frontend;
 
@@ -235,19 +216,9 @@ void UMetasoundEditorGraphVertex::SetMemberName(const FName& InNewName, bool bPo
 
 		GetNodeHandle()->SetNodeName(InNewName);
 
-		const TArray<UMetasoundEditorGraphNode*> Nodes = GetNodes();
-		for (UMetasoundEditorGraphNode* Node : Nodes)
-		{
-			const TArray<UEdGraphPin*>& Pins = Node->GetAllPins();
-			ensure(Pins.Num() == 1);
-
-			for (UEdGraphPin* Pin : Pins)
-			{
-				Pin->PinName = InNewName;
-			}
-		}
-
-		MarkNodesForRefresh();
+		FMetasoundAssetBase* MetasoundAsset = IMetasoundUObjectRegistry::Get().GetObjectAsAssetBase(&Graph->GetMetasoundChecked());
+		check(MetasoundAsset);
+		MetasoundAsset->SetSynchronizationRequired();
 	}
 
 	NameChanged.Broadcast(NodeID);
@@ -287,8 +258,6 @@ void UMetasoundEditorGraphVertex::SetDisplayName(const FText& InNewName, bool bP
 				Pin->PinFriendlyName = InNewName;
 			}
 		}
-
-		MarkNodesForRefresh();
 	}
 
 	NameChanged.Broadcast(NodeID);
@@ -477,7 +446,6 @@ void UMetasoundEditorGraphMemberDefaultLiteral::PostEditUndo()
 	UMetasoundEditorGraphMember* Member = CastChecked<UMetasoundEditorGraphMember>(GetOuter());
 	Member->UpdateFrontendDefaultLiteral(bPostTransaction);
 
-	Member->MarkNodesForRefresh();
 	if (UMetasoundEditorGraph* Graph = Member->GetOwningGraph())
 	{
 		Graph->SetSynchronizationRequired();
@@ -1178,10 +1146,12 @@ UMetasoundEditorGraphInputNode* UMetasoundEditorGraph::CreateInputNode(Metasound
 			NewInputNode->AllocateDefaultPins();
 		}
 
+		NewInputNode->CacheTitle();
+
 		return NewInputNode;
 	}
 	
-	return nullptr;	
+	return nullptr;
 }
 
 Metasound::Frontend::FDocumentHandle UMetasoundEditorGraph::GetDocumentHandle()
@@ -1513,14 +1483,19 @@ UMetasoundEditorGraphMember* UMetasoundEditorGraph::FindAdjacentMember(const UMe
 	return nullptr;
 }
 
-bool UMetasoundEditorGraph::ContainsInput(UMetasoundEditorGraphInput* InInput) const
+bool UMetasoundEditorGraph::ContainsInput(const UMetasoundEditorGraphInput& InInput) const
 {
-	return Inputs.Contains(InInput);
+	return Inputs.Contains(&InInput);
 }
 
-bool UMetasoundEditorGraph::ContainsOutput(UMetasoundEditorGraphOutput* InOutput) const
+bool UMetasoundEditorGraph::ContainsOutput(const UMetasoundEditorGraphOutput& InOutput) const
 {
-	return Outputs.Contains(InOutput);
+	return Outputs.Contains(&InOutput);
+}
+
+bool UMetasoundEditorGraph::ContainsVariable(const UMetasoundEditorGraphVariable& InVariable) const
+{
+	return Variables.Contains(&InVariable);
 }
 
 void UMetasoundEditorGraph::IterateInputs(TUniqueFunction<void(UMetasoundEditorGraphInput&)> InFunction) const
@@ -1565,6 +1540,21 @@ bool UMetasoundEditorGraph::IsEditable() const
 	return GetGraphHandle()->GetGraphStyle().bIsGraphEditable;
 }
 
+void UMetasoundEditorGraph::SetForceRefreshNodes()
+{
+	bForceRefreshNodes = true;
+}
+
+void UMetasoundEditorGraph::ClearForceRefreshNodes()
+{
+	bForceRefreshNodes = false;
+}
+
+bool UMetasoundEditorGraph::RequiresForceRefreshNodes() const
+{
+	return bForceRefreshNodes;
+}
+
 void UMetasoundEditorGraph::IterateOutputs(TUniqueFunction<void(UMetasoundEditorGraphOutput&)> InFunction) const
 {
 	for (UMetasoundEditorGraphOutput* Output : Outputs)
diff --git a/Engine/Plugins/Runtime/Metasound/Source/MetasoundEditor/Private/MetasoundEditorGraphBuilder.cpp b/Engine/Plugins/Runtime/Metasound/Source/MetasoundEditor/Private/MetasoundEditorGraphBuilder.cpp
--- a/Engine/Plugins/Runtime/Metasound/Source/MetasoundEditor/Private/MetasoundEditorGraphBuilder.cpp
+++ b/Engine/Plugins/Runtime/Metasound/Source/MetasoundEditor/Private/MetasoundEditorGraphBuilder.cpp
@@ -221,6 +221,7 @@ namespace Metasound
 			const FNodeRegistryKey RegistryKey = NodeRegistryKey::CreateKey(InNodeHandle->GetClassMetadata());
 			NewGraphNode->bIsClassNative = FMetasoundFrontendRegistryContainer::Get()->IsNodeNative(RegistryKey);
 			NewGraphNode->ClassName = InNodeHandle->GetClassMetadata().GetClassName();
+			NewGraphNode->CacheTitle();
 
 			NodeCreator.Finalize();
 			InitGraphNode(InNodeHandle, NewGraphNode, InMetaSound);
@@ -303,6 +304,7 @@ namespace Metasound
 			NewGraphNode = NodeCreator.CreateNode(bInSelectNewNode);
 			UMetasoundEditorGraph* MetasoundGraph = CastChecked<UMetasoundEditorGraph>(&Graph);
 			NewGraphNode->Output = MetasoundGraph->FindOrAddOutput(InNodeHandle);
+			NewGraphNode->CacheTitle();
 
 			NodeCreator.Finalize();
 			InitGraphNode(InNodeHandle, NewGraphNode, InMetaSound);
@@ -320,7 +322,7 @@ namespace Metasound
 			RebuildNodePins(*NewGraphNode);
 		}
 
-		bool FGraphBuilder::ValidateGraph(UObject& InMetaSound)
+		bool FGraphBuilder::ValidateGraph(UObject& InMetaSound, bool bForceRefreshNodes)
 		{
 			using namespace Frontend;
 
@@ -335,6 +337,9 @@ namespace Metasound
 			check(MetaSoundAsset);
 			UMetasoundEditorGraph& Graph = *CastChecked<UMetasoundEditorGraph>(&MetaSoundAsset->GetGraphChecked());
 
+			bForceRefreshNodes |= Graph.RequiresForceRefreshNodes();
+			Graph.ClearForceRefreshNodes();
+
 			FGraphValidationResults Results;
 
 			bool bMarkDirty = false;
@@ -343,13 +348,17 @@ namespace Metasound
 			for (const FGraphNodeValidationResult& Result : Results.GetResults())
 			{
 				bMarkDirty |= Result.bIsDirty;
-				if (GraphEditor.IsValid())
+				check(Result.Node);
+				const bool bMetadataChange = Result.Node->ContainsInterfaceChange();
+				const bool bInterfaceChange = Result.Node->ContainsMetadataChange();
+				if (Result.bIsDirty || bMetadataChange || bInterfaceChange || bForceRefreshNodes)
 				{
-					check(Result.Node);
-					if (Result.bIsDirty || Result.Node->bRefreshNode)
+					Result.Node->SyncChangeIDs();
+					Result.Node->CacheTitle();
+
+					if (GraphEditor.IsValid())
 					{
 						GraphEditor->RefreshNode(*Result.Node);
-						Result.Node->bRefreshNode = false;
 					}
 				}
 			}
@@ -1242,8 +1251,6 @@ namespace Metasound
 					AddPinToNode(InGraphNode, OutputHandle);
 				}
 			}
-
-			InGraphNode.bRefreshNode = true;
 		}
 
 		void FGraphBuilder::RefreshPinMetadata(UEdGraphPin& InPin, const FMetasoundFrontendVertexMetadata& InMetadata)
@@ -1259,11 +1266,6 @@ namespace Metasound
 				{
 					OwningNode->AdvancedPinDisplay = ENodeAdvancedPins::Hidden;
 				}
-
-				if (UMetasoundEditorGraphNode* MetaSoundNode = Cast<UMetasoundEditorGraphNode>(OwningNode))
-				{
-					MetaSoundNode->bRefreshNode = true;
-				}
 			}
 		}
 
@@ -1345,56 +1347,6 @@ namespace Metasound
 			MetaSoundAsset->UnregisterGraphWithFrontend();
 		}
 
-		void FGraphBuilder::MarkEditorNodesReferencingAssetForRefresh(UObject& InMetaSound)
-		{
-			using namespace Frontend;
-
-			FMetasoundAssetBase* MetaSoundAsset = IMetasoundUObjectRegistry::Get().GetObjectAsAssetBase(&InMetaSound);
-			if (!MetaSoundAsset)
-			{
-				return;
-			}
-
-			if (!GEditor)
-			{
-				return;
-			}
-
-			bool bGraphUpdated = false;
-
-			FMetasoundFrontendClassMetadata AssetClassMetadata = MetaSoundAsset->GetRootGraphHandle()->GetGraphMetadata();
-			AssetClassMetadata.SetType(EMetasoundFrontendClassType::External);
-			const FNodeRegistryKey AssetClassKey = NodeRegistryKey::CreateKey(AssetClassMetadata);
-
-			TArray<UObject*> EditedAssets = GEditor->GetEditorSubsystem<UAssetEditorSubsystem>()->GetAllEditedAssets();
-			for (UObject* EditedAsset : EditedAssets)
-			{
-				if (FMetasoundAssetBase* EditedMetaSoundAsset = IMetasoundUObjectRegistry::Get().GetObjectAsAssetBase(EditedAsset))
-				{
-					const UEdGraph& Graph = EditedMetaSoundAsset->GetGraphChecked();
-					TArray<UMetasoundEditorGraphExternalNode*> ExternalNodes;
-					Graph.GetNodesOfClass<UMetasoundEditorGraphExternalNode>(ExternalNodes);
-					for (UMetasoundEditorGraphExternalNode* Node : ExternalNodes)
-					{
-						FConstNodeHandle NodeHandle = Node->GetConstNodeHandle();
-						const FMetasoundFrontendClassMetadata& ClassMetadata = NodeHandle->GetClassMetadata();
-						const FNodeRegistryKey RegistryKey = NodeRegistryKey::CreateKey(ClassMetadata);
-
-						if (AssetClassKey == RegistryKey)
-						{
-							bGraphUpdated = true;
-							Node->bRefreshNode = true;
-						}
-					}
-				}
-			}
-
-			if (bGraphUpdated)
-			{
-				MetaSoundAsset->SetSynchronizationRequired();
-			}
-		}
-
 		bool FGraphBuilder::IsMatchingInputHandleAndPin(const Frontend::FConstInputHandle& InInputHandle, const UEdGraphPin& InEditorPin)
 		{
 			if (InEditorPin.Direction != EGPD_Input)
@@ -1491,7 +1443,6 @@ namespace Metasound
 				RefreshPinMetadata(*NewPin, InOutputHandle->GetMetadata());
 			}
 
-			InEditorNode.bRefreshNode = true;
 			return NewPin;
 		}
 
@@ -1634,7 +1585,7 @@ namespace Metasound
 			return bIsGraphDirty;
 		}
 
-		bool FGraphBuilder::SynchronizeGraph(UObject& InMetaSound)
+		bool FGraphBuilder::SynchronizeGraph(UObject& InMetaSound, bool bForceRefreshNodes)
 		{
 			bool bEditorGraphModified = SynchronizeGraphVertices(InMetaSound);
 			bEditorGraphModified |= SynchronizeNodeMembers(InMetaSound);
@@ -1646,7 +1597,8 @@ namespace Metasound
 				InMetaSound.MarkPackageDirty();
 			}
 
-			const bool bIsValid = FGraphBuilder::ValidateGraph(InMetaSound);
+			bForceRefreshNodes |= bEditorGraphModified;
+			const bool bIsValid = FGraphBuilder::ValidateGraph(InMetaSound, bForceRefreshNodes);
 
 			FMetasoundAssetBase* MetaSoundAsset = IMetasoundUObjectRegistry::Get().GetObjectAsAssetBase(&InMetaSound);
 			check(MetaSoundAsset);
@@ -1924,7 +1876,6 @@ namespace Metasound
 			if (bRemoveUnusedPins)
 			{
 				bIsNodeDirty |= !EditorPins.IsEmpty();
-				InEditorNode.bRefreshNode |= !EditorPins.IsEmpty();
 				for (UEdGraphPin* Pin : EditorPins)
 				{
 					if (bLogChanges)
diff --git a/Engine/Plugins/Runtime/Metasound/Source/MetasoundEditor/Private/MetasoundEditorGraphBuilder.h b/Engine/Plugins/Runtime/Metasound/Source/MetasoundEditor/Private/MetasoundEditorGraphBuilder.h
--- a/Engine/Plugins/Runtime/Metasound/Source/MetasoundEditor/Private/MetasoundEditorGraphBuilder.h
+++ b/Engine/Plugins/Runtime/Metasound/Source/MetasoundEditor/Private/MetasoundEditorGraphBuilder.h
@@ -43,7 +43,7 @@ namespace Metasound
 			static void InitGraphNode(Frontend::FNodeHandle& InNodeHandle, UMetasoundEditorGraphNode* NewGraphNode, UObject& InMetaSound);
 
 			// Validates MetaSound graph.
-			static bool ValidateGraph(UObject& InMetaSound);
+			static bool ValidateGraph(UObject& InMetaSound, bool bForceRefreshNodes);
 
 		public:
 			static const FName PinCategoryAudio;
@@ -75,8 +75,6 @@ namespace Metasound
 			// Wraps RegisterGraphWithFrontend logic in Frontend with any additional logic required to refresh editor & respective editor object state.
 			static void UnregisterGraphWithFrontend(UObject& InMetaSound);
 
-			static void MarkEditorNodesReferencingAssetForRefresh(UObject& InMetaSound);
-
 			// Returns a display name for a node. If the node has an empty or whitespace
 			// only DisplayName, then the NodeName is used. 
 			static FText GetDisplayName(const Frontend::INodeController& InFrontendNode, bool bInIncludeNamespace);
@@ -210,8 +208,9 @@ namespace Metasound
 
 			// Adds and removes nodes, pins and connections so that the UEdGraph of the MetaSound matches the FMetasoundFrontendDocumentModel
 			//
+			// @param bForceRefreshNodes - Refreshes all transient data stored on editor nodes & redraws them
 			// @return True if the UEdGraph is synchronized and is in valid state, false otherwise.
-			static bool SynchronizeGraph(UObject& InMetaSound);
+			static bool SynchronizeGraph(UObject& InMetaSound, bool bForceRefreshNodes = false);
 
 			// Synchronizes editor nodes with frontend nodes, removing editor nodes that are not represented in the frontend, and adding editor nodes to represent missing frontend nodes.
 			//
diff --git a/Engine/Plugins/Runtime/Metasound/Source/MetasoundEditor/Private/MetasoundEditorGraphMemberDefaults.cpp b/Engine/Plugins/Runtime/Metasound/Source/MetasoundEditor/Private/MetasoundEditorGraphMemberDefaults.cpp
--- a/Engine/Plugins/Runtime/Metasound/Source/MetasoundEditor/Private/MetasoundEditorGraphMemberDefaults.cpp
+++ b/Engine/Plugins/Runtime/Metasound/Source/MetasoundEditor/Private/MetasoundEditorGraphMemberDefaults.cpp
@@ -214,10 +214,13 @@ void UMetasoundEditorGraphMemberDefaultFloat::PostEditChangeChainProperty(FPrope
 		OnClampChanged.Broadcast(ClampDefault);
 	}
 
+	// TODO: Remove this once widget Metadata is migrated to frontend
+	// Metadata, which will inherently update the change guid and issue
+	// resync.
 	UMetasoundEditorGraphMember* Member = GetParentMember();
 	if (ensure(Member))
 	{
-		Member->MarkNodesForRefresh();
+		Member->GetOwningGraph()->SetForceRefreshNodes();
 	}
 
 	Super::PostEditChangeChainProperty(PropertyChangedEvent);
diff --git a/Engine/Plugins/Runtime/Metasound/Source/MetasoundEditor/Private/MetasoundEditorGraphNode.cpp b/Engine/Plugins/Runtime/Metasound/Source/MetasoundEditor/Private/MetasoundEditorGraphNode.cpp
--- a/Engine/Plugins/Runtime/Metasound/Source/MetasoundEditor/Private/MetasoundEditorGraphNode.cpp
+++ b/Engine/Plugins/Runtime/Metasound/Source/MetasoundEditor/Private/MetasoundEditorGraphNode.cpp
@@ -172,6 +172,41 @@ void UMetasoundEditorGraphNode::AllocateDefaultPins()
 	Editor::FGraphBuilder::RebuildNodePins(*this);
 }
 
+void UMetasoundEditorGraphNode::SyncChangeIDs()
+{
+	using namespace Metasound::Frontend;
+	FConstNodeHandle NodeHandle = GetConstNodeHandle();
+
+	MetadataChangeID = NodeHandle->GetClassMetadata().GetChangeID();
+	InterfaceChangeID = NodeHandle->GetClassInterface().GetChangeID();
+
+}
+
+void UMetasoundEditorGraphNode::CacheTitle()
+{
+	using namespace Metasound;
+	using namespace Metasound::Frontend;
+
+	FConstNodeHandle NodeHandle = GetNodeHandle();
+	CachedTitle = NodeHandle->GetDisplayTitle();
+}
+
+bool UMetasoundEditorGraphNode::ContainsMetadataChange() const
+{
+	using namespace Metasound::Frontend;
+	FConstNodeHandle NodeHandle = GetConstNodeHandle();
+
+	return MetadataChangeID != NodeHandle->GetClassMetadata().GetChangeID();
+}
+
+bool UMetasoundEditorGraphNode::ContainsInterfaceChange() const
+{
+	using namespace Metasound::Frontend;
+	FConstNodeHandle NodeHandle = GetConstNodeHandle();
+
+	return InterfaceChangeID != NodeHandle->GetClassInterface().GetChangeID();
+}
+
 void UMetasoundEditorGraphNode::ReconstructNode()
 {
 	using namespace Metasound::Editor;
@@ -186,6 +221,8 @@ void UMetasoundEditorGraphNode::ReconstructNode()
 	{
 		FGraphBuilder::SynchronizeNodePins(*this, NodeHandle, false /* bRemoveUnusedPins */, false /* bLogChanges */);
 	}
+
+	CacheTitle();
 }
 
 void UMetasoundEditorGraphNode::AutowireNewNode(UEdGraphPin* FromPin)
@@ -244,10 +281,7 @@ FString UMetasoundEditorGraphNode::GetDocumentationLink() const
 
 FText UMetasoundEditorGraphNode::GetNodeTitle(ENodeTitleType::Type TitleType) const
 {
-	using namespace Metasound::Frontend;
-
-	FConstNodeHandle NodeHandle = GetNodeHandle();
-	return NodeHandle->GetDisplayTitle();
+	return CachedTitle;
 }
 
 void UMetasoundEditorGraphNode::GetPinHoverText(const UEdGraphPin& Pin, FString& OutHoverText) const
@@ -641,6 +675,14 @@ bool UMetasoundEditorGraphExternalNode::RefreshPinMetadata()
 	return bModified;
 }
 
+void UMetasoundEditorGraphExternalNode::ReconstructNode()
+{
+	using namespace Metasound::Editor;
+	using namespace Metasound::Frontend;
+
+	Super::ReconstructNode();
+}
+
 FMetasoundFrontendVersionNumber UMetasoundEditorGraphExternalNode::FindHighestVersionInRegistry() const
 {
 	return GetConstNodeHandle()->FindHighestVersionInRegistry();
@@ -655,7 +697,34 @@ bool UMetasoundEditorGraphExternalNode::CanAutoUpdate() const
 {
 	using namespace Metasound::Frontend;
 
-	return GetConstNodeHandle()->CanAutoUpdate();
+	FClassInterfaceUpdates InterfaceUpdates;
+	return GetConstNodeHandle()->CanAutoUpdate(InterfaceUpdates);
+}
+
+void UMetasoundEditorGraphExternalNode::CacheTitle()
+{
+	using namespace Metasound;
+	using namespace Metasound::Frontend;
+
+	Super::CacheTitle();
+
+	// If version is missing from the registry, then this node will not cache
+	// a useful title.  In that case, just display the next highest available name.
+	if (CachedTitle.IsEmpty())
+	{
+		FMetasoundFrontendClass ClassWithHighestVersion;
+		if (ISearchEngine::Get().FindClassWithHighestVersion(ClassName, ClassWithHighestVersion))
+		{
+			CachedTitle = ClassWithHighestVersion.Metadata.GetDisplayName();
+		}
+	}
+
+	// If that cannot be found, build a title from the cached node registry FName.
+	if (CachedTitle.IsEmpty())
+	{
+		CachedTitle = FText::FromString(ClassName.ToString());
+	}
+
 }
 
 UMetasoundEditorGraphExternalNode* UMetasoundEditorGraphExternalNode::UpdateToVersion(const FMetasoundFrontendVersionNumber& InNewVersion, bool bInPropagateErrorMessages)
@@ -704,7 +773,6 @@ UMetasoundEditorGraphExternalNode* UMetasoundEditorGraphExternalNode::UpdateToVe
 		}
 	}
 
-	bRefreshNode = true;
 	return ReplacementEdNode;
 }
 
@@ -875,27 +943,27 @@ bool UMetasoundEditorGraphExternalNode::Validate(Metasound::Editor::FGraphNodeVa
 			EMessageSeverity::Type Severity;
 
 			const bool bClassVersionExists = SortedClasses.ContainsByPredicate([InCurrentVersion = &CurrentVersion](const FMetasoundFrontendClass& AvailableClass)
-				{
-					return AvailableClass.Metadata.GetVersion() == *InCurrentVersion;
-				});
+			{
+				return AvailableClass.Metadata.GetVersion() == *InCurrentVersion;
+			});
 			if (bClassVersionExists)
 			{
 				NodeMsg = FString::Format(TEXT("Node class '{0} {1}' is prior version: Eligible for upgrade to {2}"),
-					{
-						*Metadata.GetClassName().ToString(),
-						*Metadata.GetVersion().ToString(),
-						*HighestRegistryClass.Metadata.GetVersion().ToString()
-					});
+				{
+					*Metadata.GetClassName().ToString(),
+					*Metadata.GetVersion().ToString(),
+					*HighestRegistryClass.Metadata.GetVersion().ToString()
+				});
 				Severity = EMessageSeverity::Warning;
 			}
 			else
 			{
 				NodeMsg = FString::Format(TEXT("Node class '{0} {1}' is missing and ineligible for auto-update.  Highest version '{2}' found."),
-					{
-						*Metadata.GetClassName().ToString(),
-						*Metadata.GetVersion().ToString(),
-						*HighestRegistryClass.Metadata.GetVersion().ToString()
-					});
+				{
+					*Metadata.GetClassName().ToString(),
+					*Metadata.GetVersion().ToString(),
+					*HighestRegistryClass.Metadata.GetVersion().ToString()
+				});
 				Severity = EMessageSeverity::Error;
 				OutResult.bIsInvalid = true;
 			}
@@ -908,11 +976,11 @@ bool UMetasoundEditorGraphExternalNode::Validate(Metasound::Editor::FGraphNodeVa
 			if (InterfaceUpdates.ContainsChanges())
 			{
 				GraphNodePrivate::SetGraphNodeMessage(*this, EMessageSeverity::Error,
-					FString::Format(TEXT("Node & registered class interface mismatch: '{0} {1}'. Class either versioned improperly, class key collision exists, or AutoUpdate disabled in 'MetaSound' Developer Settings."),
-						{
-							*Metadata.GetClassName().ToString(),
-							*Metadata.GetVersion().ToString()
-						}));
+				FString::Format(TEXT("Node & registered class interface mismatch: '{0} {1}'. Class either versioned improperly, class key collision exists, or AutoUpdate disabled in 'MetaSound' Developer Settings."),
+				{
+					*Metadata.GetClassName().ToString(),
+					*Metadata.GetVersion().ToString()
+				}));
 				OutResult.bIsDirty = true;
 				OutResult.bIsInvalid = true;
 			}
@@ -921,11 +989,11 @@ bool UMetasoundEditorGraphExternalNode::Validate(Metasound::Editor::FGraphNodeVa
 		{
 			GraphNodePrivate::SetGraphNodeMessage(*this, EMessageSeverity::Error,
 				FString::Format(TEXT("Node with class '{0} {1}' interface version higher than that of highest minor revision ({2}) in class registry."),
-					{
-						*Metadata.GetClassName().ToString(),
-						*Metadata.GetVersion().ToString(),
-						*HighestRegistryClass.Metadata.GetVersion().ToString()
-					}));
+				{
+					*Metadata.GetClassName().ToString(),
+					*Metadata.GetVersion().ToString(),
+					*HighestRegistryClass.Metadata.GetVersion().ToString()
+				}));
 			OutResult.bIsDirty = true;
 			OutResult.bIsInvalid = true;
 		}
@@ -934,11 +1002,6 @@ bool UMetasoundEditorGraphExternalNode::Validate(Metasound::Editor::FGraphNodeVa
 	return !OutResult.bIsInvalid;
 }
 
-void UMetasoundEditorGraphExternalNode::PostLoad()
-{
-	Super::PostLoad();
-}
-
 FLinearColor UMetasoundEditorGraphExternalNode::GetNodeTitleColor() const
 {
 	if (const UMetasoundEditorSettings* EditorSettings = GetDefault<UMetasoundEditorSettings>())
diff --git a/Engine/Plugins/Runtime/Metasound/Source/MetasoundEditor/Private/MetasoundEditorModule.cpp b/Engine/Plugins/Runtime/Metasound/Source/MetasoundEditor/Private/MetasoundEditorModule.cpp
--- a/Engine/Plugins/Runtime/Metasound/Source/MetasoundEditor/Private/MetasoundEditorModule.cpp
+++ b/Engine/Plugins/Runtime/Metasound/Source/MetasoundEditor/Private/MetasoundEditorModule.cpp
@@ -294,10 +294,21 @@ namespace Metasound
 					// MetaSoundAssetSubsystem so as to properly refresh respective open editors.
 					constexpr bool bReregisterWithFrontend = false;
 					AssetSubsystem->RenameAsset(InAssetData, bReregisterWithFrontend);
-					if (UObject* AssetObject = InAssetData.GetAsset())
+					UObject* AssetObject = InAssetData.GetAsset();
+					FGraphBuilder::RegisterGraphWithFrontend(*AssetObject);
+
+					// Request sync of other open MetaSounds to ensure rename is reflected in other asset editors with
+					// nodes referencing the renamed asset.
+					TArray<UObject*> EditedAssets = GEditor->GetEditorSubsystem<UAssetEditorSubsystem>()->GetAllEditedAssets();
+					for (UObject* EditedAsset : EditedAssets)
 					{
-						FGraphBuilder::RegisterGraphWithFrontend(*AssetObject);
-						FGraphBuilder::MarkEditorNodesReferencingAssetForRefresh(*AssetObject);
+						if (AssetObject != EditedAsset)
+						{
+							if (FMetasoundAssetBase* EditedMetaSoundAsset = IMetasoundUObjectRegistry::Get().GetObjectAsAssetBase(EditedAsset))
+							{
+								EditedMetaSoundAsset->SetSynchronizationRequired();
+							}
+						}
 					}
 				}
 			}
diff --git a/Engine/Plugins/Runtime/Metasound/Source/MetasoundEditor/Private/SMetasoundGraphNode.cpp b/Engine/Plugins/Runtime/Metasound/Source/MetasoundEditor/Private/SMetasoundGraphNode.cpp
--- a/Engine/Plugins/Runtime/Metasound/Source/MetasoundEditor/Private/SMetasoundGraphNode.cpp
+++ b/Engine/Plugins/Runtime/Metasound/Source/MetasoundEditor/Private/SMetasoundGraphNode.cpp
@@ -613,7 +613,6 @@ namespace Metasound
 								DefaultFloat->SetDefault(Output);
 								GraphMember->UpdateFrontendDefaultLiteral(bPostTransaction);
 
-								GraphMember->MarkNodesForRefresh();
 								if (UMetasoundEditorGraph* Graph = GraphMember->GetOwningGraph())
 								{
 									Graph->SetSynchronizationRequired();
diff --git a/Engine/Plugins/Runtime/Metasound/Source/MetasoundEditor/Private/SMetasoundGraphPin.h b/Engine/Plugins/Runtime/Metasound/Source/MetasoundEditor/Private/SMetasoundGraphPin.h
--- a/Engine/Plugins/Runtime/Metasound/Source/MetasoundEditor/Private/SMetasoundGraphPin.h
+++ b/Engine/Plugins/Runtime/Metasound/Source/MetasoundEditor/Private/SMetasoundGraphPin.h
@@ -185,9 +185,6 @@ namespace Metasound
 											}
 										}
 
-										// Mark node for refresh in case non-pin widget on node needs updating (ex. sliders, knobs)
-										Node->bRefreshNode = true;
-
 										// Full node synchronization required as custom
 										// node-level widgets may need to be refreshed
 										MetaSoundGraph->SetSynchronizationRequired();
diff --git a/Engine/Plugins/Runtime/Metasound/Source/MetasoundEditor/Public/MetasoundEditorGraph.h b/Engine/Plugins/Runtime/Metasound/Source/MetasoundEditor/Public/MetasoundEditorGraph.h
--- a/Engine/Plugins/Runtime/Metasound/Source/MetasoundEditor/Public/MetasoundEditorGraph.h
+++ b/Engine/Plugins/Runtime/Metasound/Source/MetasoundEditor/Public/MetasoundEditorGraph.h
@@ -144,9 +144,6 @@ public:
 	 * a UMetasoundEditorGraph, returns a nullptr. */
 	const UMetasoundEditorGraph* GetOwningGraph() const;
 
-	/** Marks all member's nodes as requiring a display refresh. */
-	void MarkNodesForRefresh();
-
 	/* Whether this member can be renamed. */
 	virtual bool CanRename() const PURE_VIRTUAL(UMetasoundEditorGraphMember::CanRename, return false;);
 
@@ -374,13 +371,18 @@ public:
 	void IterateInputs(TUniqueFunction<void(UMetasoundEditorGraphInput&)> InFunction) const;
 	void IterateOutputs(TUniqueFunction<void(UMetasoundEditorGraphOutput&)> InFunction) const;
 
-	bool ContainsInput(UMetasoundEditorGraphInput* InInput) const;
-	bool ContainsOutput(UMetasoundEditorGraphOutput* InOutput) const;
+	bool ContainsInput(const UMetasoundEditorGraphInput& InInput) const;
+	bool ContainsOutput(const UMetasoundEditorGraphOutput& InOutput) const;
+	bool ContainsVariable(const UMetasoundEditorGraphVariable& InVariable) const;
 
 	void SetPreviewID(uint32 InPreviewID);
 	bool IsPreviewing() const;
 	bool IsEditable() const;
 
+	void SetForceRefreshNodes();
+	void ClearForceRefreshNodes();
+	bool RequiresForceRefreshNodes() const;
+
 	// UMetasoundEditorGraphBase Implementation
 	virtual void RegisterGraphWithFrontend() override;
 	virtual void SetSynchronizationRequired() override;
@@ -401,6 +403,12 @@ private:
 
 	bool bVersionedOnLoad = false;
 
+	// Used as a means of forcing the graph to rebuild nodes on next tick.
+	// TODO: Will no longer require this once all editor metadata is migrated
+	// to the frontend & the system can adequately rely on the changeIDs as a
+	// mechanism for selectively updating nodes.
+	bool bForceRefreshNodes = false;
+
 	UPROPERTY()
 	TArray<TObjectPtr<UMetasoundEditorGraphInput>> Inputs;
 
diff --git a/Engine/Plugins/Runtime/Metasound/Source/MetasoundEditor/Public/MetasoundEditorGraphNode.h b/Engine/Plugins/Runtime/Metasound/Source/MetasoundEditor/Public/MetasoundEditorGraphNode.h
--- a/Engine/Plugins/Runtime/Metasound/Source/MetasoundEditor/Public/MetasoundEditorGraphNode.h
+++ b/Engine/Plugins/Runtime/Metasound/Source/MetasoundEditor/Public/MetasoundEditorGraphNode.h
@@ -98,11 +98,26 @@ public:
 	virtual FMetasoundFrontendClassName GetClassName() const { return FMetasoundFrontendClassName(); }
 	virtual FGuid GetNodeID() const { return FGuid(); }
 	virtual FText GetDisplayName() const;
+	virtual void CacheTitle();
 
-	// Marks node to be reconstructed on next editor tick.
-	bool bRefreshNode = false;
+	// Mark node for refresh
+	void SyncChangeIDs();
+
+
+	// Returns whether or not Metadata has been changed since the last node refresh
+	bool ContainsMetadataChange() const;
+
+	// Returns whether or not Interface has been changed since the last node refresh
+	bool ContainsInterfaceChange() const;
 
 protected:
+	FGuid InterfaceChangeID;
+	FGuid MetadataChangeID;
+
+	// Not be serialized to avoid text desync as the registry can provide
+	// a new name if the external definition changes between application sessions.
+	FText CachedTitle;
+
 	virtual void SetNodeID(FGuid InNodeID) { }
 
 	friend class Metasound::Editor::FGraphBuilder;
@@ -177,13 +192,14 @@ protected:
 	bool bIsClassNative = true;
 
 public:
-	virtual void PostLoad() override;
-
 	virtual FMetasoundFrontendClassName GetClassName() const override { return ClassName; }
 	virtual FGuid GetNodeID() const override { return NodeID; }
 	virtual FLinearColor GetNodeTitleColor() const override;
 	virtual FSlateIcon GetNodeTitleIcon() const override;
 
+	virtual void ReconstructNode() override;
+	virtual void CacheTitle() override;
+
 	FMetasoundFrontendVersionNumber FindHighestVersionInRegistry() const;
 	FMetasoundFrontendVersionNumber FindHighestMinorVersionInRegistry() const;
 	bool CanAutoUpdate() const;
diff --git a/Engine/Plugins/Runtime/Metasound/Source/MetasoundFrontend/Private/MetasoundAssetBase.cpp b/Engine/Plugins/Runtime/Metasound/Source/MetasoundFrontend/Private/MetasoundAssetBase.cpp
--- a/Engine/Plugins/Runtime/Metasound/Source/MetasoundFrontend/Private/MetasoundAssetBase.cpp
+++ b/Engine/Plugins/Runtime/Metasound/Source/MetasoundFrontend/Private/MetasoundAssetBase.cpp
@@ -276,8 +276,6 @@ void FMetasoundAssetBase::SetMetadata(FMetasoundFrontendClassMetadata& InMetadat
 		UE_LOG(LogMetaSound, Display, TEXT("Forcing class type to EMetasoundFrontendClassType::Graph on root graph metadata"));
 		Doc.RootGraph.Metadata.SetType(EMetasoundFrontendClassType::Graph);
 	}
-
-	MarkMetasoundDocumentDirty();
 }
 
 bool FMetasoundAssetBase::GetDeclaredInterfaces(TArray<const Metasound::Frontend::IInterfaceRegistryEntry*>& OutInterfaces) const
diff --git a/Engine/Plugins/Runtime/Metasound/Source/MetasoundFrontend/Private/MetasoundFrontendInvalidController.h b/Engine/Plugins/Runtime/Metasound/Source/MetasoundFrontend/Private/MetasoundFrontendInvalidController.h
--- a/Engine/Plugins/Runtime/Metasound/Source/MetasoundFrontend/Private/MetasoundFrontendInvalidController.h
+++ b/Engine/Plugins/Runtime/Metasound/Source/MetasoundFrontend/Private/MetasoundFrontendInvalidController.h
@@ -224,7 +224,7 @@ namespace Metasound
 			virtual bool IsInterfaceMember() const override { return false; }
 
 			virtual bool DiffAgainstRegistryInterface(FClassInterfaceUpdates& OutInterfaceUpdates, bool bInUseHighestMinorVersion) const override { return false; }
-			virtual bool CanAutoUpdate(FClassInterfaceUpdates* OutInterfaceUpdates = nullptr) const override { return false; }
+			virtual bool CanAutoUpdate(FClassInterfaceUpdates& OutInterfaceUpdates) const override { OutInterfaceUpdates = { }; return false; }
 			virtual FMetasoundFrontendVersionNumber FindHighestVersionInRegistry() const override { return FMetasoundFrontendVersionNumber::GetInvalid(); }
 			virtual FMetasoundFrontendVersionNumber FindHighestMinorVersionInRegistry() const override { return FMetasoundFrontendVersionNumber::GetInvalid(); }
 
diff --git a/Engine/Plugins/Runtime/Metasound/Source/MetasoundFrontend/Private/MetasoundFrontendNodeController.cpp b/Engine/Plugins/Runtime/Metasound/Source/MetasoundFrontend/Private/MetasoundFrontendNodeController.cpp
--- a/Engine/Plugins/Runtime/Metasound/Source/MetasoundFrontend/Private/MetasoundFrontendNodeController.cpp
+++ b/Engine/Plugins/Runtime/Metasound/Source/MetasoundFrontend/Private/MetasoundFrontendNodeController.cpp
@@ -14,12 +14,12 @@
 
 #define LOCTEXT_NAMESPACE "MetasoundFrontendNodeController"
 
-static int32 MetaSoundAutoUpdateNativeClassCVar = 1;
+static int32 MetaSoundAutoUpdateNativeClassesOfEqualVersionCVar = 1;
 FAutoConsoleVariableRef CVarMetaSoundAutoUpdateNativeClass(
-	TEXT("au.MetaSounds.AutoUpdate.NativeClasses"),
-	MetaSoundAutoUpdateNativeClassCVar,
+	TEXT("au.MetaSounds.AutoUpdate.NativeClassesOfEqualVersion"),
+	MetaSoundAutoUpdateNativeClassesOfEqualVersionCVar,
 	TEXT("If true, node references to native class that share a version number will attempt to auto-update if the interface is different, which results in slower graph load times.\n")
-	TEXT("0: Don't auto-update native classes, !0: Auto-update native classes (default)"),
+	TEXT("0: Don't auto-update native classes of the same version with interface discrepancies, !0: Auto-update native classes of the same version with interface discrepancies (default)"),
 	ECVF_Default);
 
 namespace Metasound
@@ -937,10 +937,12 @@ namespace Metasound
 			return true;
 		}
 
-		bool FBaseNodeController::CanAutoUpdate(FClassInterfaceUpdates* OutInterfaceUpdates) const
+		bool FBaseNodeController::CanAutoUpdate(FClassInterfaceUpdates& OutInterfaceUpdates) const
 		{
 			METASOUND_TRACE_CPUPROFILER_EVENT_SCOPE(BaseNodeController::CanAutoUpdate);
 
+			OutInterfaceUpdates = { };
+
 			const FMetasoundFrontendClassMetadata& NodeClassMetadata = GetClassMetadata();
 			if (IMetaSoundAssetManager* AssetManager = IMetaSoundAssetManager::Get())
 			{
@@ -959,32 +961,25 @@ namespace Metasound
 				return false;
 			}
 
+			// 1. Document's class version is somehow higher than registries, so can't update.
 			if (RegistryClass.Metadata.GetVersion() < NodeClassMetadata.GetVersion())
 			{
 				return false;
 			}
 
+			// 2. Document's class version is equal, so have to dif and check change IDs
+			// to ensure a change wasn't created that didn't contain a version change.
 			if (RegistryClass.Metadata.GetVersion() == NodeClassMetadata.GetVersion())
 			{
 				// TODO: Merge these paths.  Shouldn't use different logic to
-				// define changes in native vs asset class definitions.
+				// define changes in native vs asset class definitions. Need to
+				// hash a changeID from natively defined node classes in order to
+				// merge branches.
 				const FNodeRegistryKey RegistryKey = NodeRegistryKey::CreateKey(RegistryClass.Metadata);
 				const bool bIsClassNative = FMetasoundFrontendRegistryContainer::Get()->IsNodeNative(RegistryKey);
 				if (bIsClassNative)
 				{
-					if (!MetaSoundAutoUpdateNativeClassCVar)
-					{
-						return false;
-					}
-
-					FClassInterfaceUpdates InterfaceUpdates;
-					DiffAgainstRegistryInterface(InterfaceUpdates, true /* bUseHighestMinorVersion */);
-					if (OutInterfaceUpdates)
-					{
-						*OutInterfaceUpdates = InterfaceUpdates;
-					}
-
-					if (!InterfaceUpdates.ContainsChanges())
+					if (!MetaSoundAutoUpdateNativeClassesOfEqualVersionCVar)
 					{
 						return false;
 					}
@@ -1000,8 +995,15 @@ namespace Metasound
 						}
 					}
 				}
+
+				DiffAgainstRegistryInterface(OutInterfaceUpdates, true /* bUseHighestMinorVersion */);
+				return OutInterfaceUpdates.ContainsChanges();
 			}
 
+			// 3. Document's class version is out-of-date, so dif and always return true that can auto-update
+			// (Unlike the case where the version is equal, the version must be updated even if the interface
+			// contains no changes).
+			DiffAgainstRegistryInterface(OutInterfaceUpdates, true /* bUseHighestMinorVersion */);
 			return true;
 		}
 
diff --git a/Engine/Plugins/Runtime/Metasound/Source/MetasoundFrontend/Private/MetasoundFrontendNodeController.h b/Engine/Plugins/Runtime/Metasound/Source/MetasoundFrontend/Private/MetasoundFrontendNodeController.h
--- a/Engine/Plugins/Runtime/Metasound/Source/MetasoundFrontend/Private/MetasoundFrontendNodeController.h
+++ b/Engine/Plugins/Runtime/Metasound/Source/MetasoundFrontend/Private/MetasoundFrontendNodeController.h
@@ -58,7 +58,7 @@ namespace Metasound
 
 			virtual bool DiffAgainstRegistryInterface(FClassInterfaceUpdates& OutInterfaceUpdates, bool bInUseHighestMinorVersion) const override;
 
-			virtual bool CanAutoUpdate(FClassInterfaceUpdates* OutInterfaceUpdates = nullptr) const override;
+			virtual bool CanAutoUpdate(FClassInterfaceUpdates& OutInterfaceUpdates) const override;
 			virtual FNodeHandle ReplaceWithVersion(const FMetasoundFrontendVersionNumber& InNewVersion) override;
 			virtual FMetasoundFrontendVersionNumber FindHighestVersionInRegistry() const override;
 			virtual FMetasoundFrontendVersionNumber FindHighestMinorVersionInRegistry() const override;
diff --git a/Engine/Plugins/Runtime/Metasound/Source/MetasoundFrontend/Private/MetasoundFrontendTransform.cpp b/Engine/Plugins/Runtime/Metasound/Source/MetasoundFrontend/Private/MetasoundFrontendTransform.cpp
--- a/Engine/Plugins/Runtime/Metasound/Source/MetasoundFrontend/Private/MetasoundFrontendTransform.cpp
+++ b/Engine/Plugins/Runtime/Metasound/Source/MetasoundFrontend/Private/MetasoundFrontendTransform.cpp
@@ -481,7 +481,7 @@ namespace Metasound
 				}
 
 				FClassInterfaceUpdates InterfaceUpdates;
-				if (!NodeHandle->CanAutoUpdate(&InterfaceUpdates))
+				if (!NodeHandle->CanAutoUpdate(InterfaceUpdates))
 				{
 					return;
 				}
@@ -490,14 +490,25 @@ namespace Metasound
 				if (UpdateVersion.IsValid() && UpdateVersion > ClassMetadata.GetVersion())
 				{
 					UE_LOG(LogMetaSound, Display, TEXT("Auto-Updating '%s' node class '%s': Newer version '%s' found."), *DebugAssetPath, *ClassMetadata.GetClassName().ToString(), *UpdateVersion.ToString());
+
+					NodesToUpdate.Add(TPair<FNodeHandle, FMetasoundFrontendVersionNumber>(NodeHandle, UpdateVersion));
 				}
 				else if (InterfaceUpdates.ContainsChanges())
 				{
 					UpdateVersion = ClassMetadata.GetVersion();
 					UE_LOG(LogMetaSound, Display, TEXT("Auto-Updating '%s' node class '%s (%s)': Interface change detected."), *DebugAssetPath, *ClassMetadata.GetClassName().ToString(), *UpdateVersion.ToString());
+
+					NodesToUpdate.Add(TPair<FNodeHandle, FMetasoundFrontendVersionNumber>(NodeHandle, UpdateVersion));
 				}
 
-				NodesToUpdate.Add(TPair<FNodeHandle, FMetasoundFrontendVersionNumber>(NodeHandle, UpdateVersion));
+				// Only update the node at this point if editor data is loaded. If it isn't and their are no interface
+				// changes but auto-update returned it was eligible, then the auto-update only contains cosmetic changes.
+#if WITH_EDITORONLY_DATA
+				else
+				{
+					NodesToUpdate.Add(TPair<FNodeHandle, FMetasoundFrontendVersionNumber>(NodeHandle, UpdateVersion));
+				}
+#endif // WITH_EDITORONLY_DATA
 			}, EMetasoundFrontendClassType::External);
 
 			if (PresetReferencedMetaSoundAsset)
diff --git a/Engine/Plugins/Runtime/Metasound/Source/MetasoundFrontend/Public/MetasoundFrontendController.h b/Engine/Plugins/Runtime/Metasound/Source/MetasoundFrontend/Public/MetasoundFrontendController.h
--- a/Engine/Plugins/Runtime/Metasound/Source/MetasoundFrontend/Public/MetasoundFrontendController.h
+++ b/Engine/Plugins/Runtime/Metasound/Source/MetasoundFrontend/Public/MetasoundFrontendController.h
@@ -529,9 +529,13 @@ namespace Metasound
 
 			/** Returns whether the node is eligible for auto-updating (i.e.
 			  * has undergone minor revision or the interface has changed, but
-			  * no higher major revision is available. Optionally provide
-			  * interface updates to be populate with any information regarding interface updates. */
-			virtual bool CanAutoUpdate(FClassInterfaceUpdates* OutInterfaceUpdates = nullptr) const = 0;
+			  * no higher major revision is available). Provides interface updates
+			  * to populate with any information regarding interface updates.
+			  * Can return true if the interface has changed but only cosmetic
+			  * differences (ex. DisplayName only used in editor) but no runtime
+			  * behavior has been modified.
+			  */
+			virtual bool CanAutoUpdate(FClassInterfaceUpdates& OutInterfaceUpdates) const = 0;
 
 			/** Description of the given node. */
 			virtual const FText& GetDescription() const = 0;
diff --git a/Engine/Plugins/Runtime/Metasound/Source/MetasoundStandardNodes/Private/MetasoundOscillatorNodes.cpp b/Engine/Plugins/Runtime/Metasound/Source/MetasoundStandardNodes/Private/MetasoundOscillatorNodes.cpp
--- a/Engine/Plugins/Runtime/Metasound/Source/MetasoundStandardNodes/Private/MetasoundOscillatorNodes.cpp
+++ b/Engine/Plugins/Runtime/Metasound/Source/MetasoundStandardNodes/Private/MetasoundOscillatorNodes.cpp
@@ -761,7 +761,7 @@ namespace Metasound
 				FNodeClassMetadata Info;
 				Info.ClassName = { StandardNodes::Namespace, TEXT("Sine"), StandardNodes::AudioVariant };
 				Info.MajorVersion = 1;
-				Info.MinorVersion = 0;
+				Info.MinorVersion = 1;
 				Info.DisplayName = LOCTEXT("Metasound_SineNodeDisplayName", "Sine");
 				Info.Description = LOCTEXT("Metasound_SineNodeDescription", "Emits an audio signal of a sinusoid.");
 				Info.Author = PluginAuthor;
