Tidy up feedback from earlier CL
- Rename 'GenerateClassDescription' functions/comments for clarity regarding what they're acting upon
- Remove global, unused 'GenerateClassMetadata'
#rnx

diff --git a/Engine/Plugins/Runtime/Metasound/Source/MetasoundFrontend/Private/MetasoundFrontend.cpp b/Engine/Plugins/Runtime/Metasound/Source/MetasoundFrontend/Private/MetasoundFrontend.cpp
--- a/Engine/Plugins/Runtime/Metasound/Source/MetasoundFrontend/Private/MetasoundFrontend.cpp
+++ b/Engine/Plugins/Runtime/Metasound/Source/MetasoundFrontend/Private/MetasoundFrontend.cpp
@@ -25,24 +25,18 @@ namespace Metasound
 {
 	namespace Frontend
 	{
-		// gets all metadata (name, description, author, what to say if it's missing) for a given node.
-		FMetasoundFrontendClassMetadata GenerateClassMetadata(const FNodeRegistryKey& InKey)
-		{
-			return GenerateClassDescription(InKey).Metadata;
-		}
-
-		FMetasoundFrontendClass GenerateClassDescription(const FNodeClassMetadata& InNodeMetadata, EMetasoundFrontendClassType ClassType)
+		FMetasoundFrontendClass GenerateClass(const FNodeClassMetadata& InNodeMetadata, EMetasoundFrontendClassType ClassType)
 		{
 			FMetasoundFrontendClass ClassDescription;
 
-			ClassDescription.Metadata = FMetasoundFrontendClassMetadata::GenerateClassDescription(InNodeMetadata, ClassType);
-			ClassDescription.Interface = FMetasoundFrontendClassInterface::GenerateClassDescription(InNodeMetadata.DefaultInterface);
-			ClassDescription.Style = FMetasoundFrontendClassStyle::GenerateClassDescription(InNodeMetadata.DisplayStyle);
+			ClassDescription.Metadata = FMetasoundFrontendClassMetadata::GenerateClassMetadata(InNodeMetadata, ClassType);
+			ClassDescription.Interface = FMetasoundFrontendClassInterface::GenerateClassInterface(InNodeMetadata.DefaultInterface);
+			ClassDescription.Style = FMetasoundFrontendClassStyle::GenerateClassStyle(InNodeMetadata.DisplayStyle);
 
 			return ClassDescription;
 		}
 
-		FMetasoundFrontendClass GenerateClassDescription(const FNodeRegistryKey& InKey)
+		FMetasoundFrontendClass GenerateClass(const FNodeRegistryKey& InKey)
 		{
 			FMetasoundFrontendClass OutClass;
 
diff --git a/Engine/Plugins/Runtime/Metasound/Source/MetasoundFrontend/Private/MetasoundFrontendDocument.cpp b/Engine/Plugins/Runtime/Metasound/Source/MetasoundFrontend/Private/MetasoundFrontendDocument.cpp
--- a/Engine/Plugins/Runtime/Metasound/Source/MetasoundFrontend/Private/MetasoundFrontendDocument.cpp
+++ b/Engine/Plugins/Runtime/Metasound/Source/MetasoundFrontend/Private/MetasoundFrontendDocument.cpp
@@ -182,7 +182,7 @@ bool operator!=(const FMetasoundFrontendClassName& InLHS, const FMetasoundFronte
 	return !(InLHS == InRHS);
 }
 
-FMetasoundFrontendClassInterface FMetasoundFrontendClassInterface::GenerateClassDescription(const Metasound::FVertexInterface& InDefaultInterface)
+FMetasoundFrontendClassInterface FMetasoundFrontendClassInterface::GenerateClassInterface(const Metasound::FVertexInterface& InVertexInterface)
 {
 	using namespace Metasound;
 	using namespace Metasound::Frontend;
@@ -191,7 +191,7 @@ FMetasoundFrontendClassInterface FMetasoundFrontendClassInterface::GenerateClass
 
 	// Copy over inputs
 	{
-		const FInputVertexInterface& InputInterface = InDefaultInterface.GetInputInterface();
+		const FInputVertexInterface& InputInterface = InVertexInterface.GetInputInterface();
 		FMetasoundFrontendInterfaceStyle InputStyle;
 		for (const TPair<FVertexName, FInputDataVertex>& InputTuple : InputInterface)
 		{
@@ -235,7 +235,7 @@ FMetasoundFrontendClassInterface FMetasoundFrontendClassInterface::GenerateClass
 
 	// Copy over outputs
 	{
-		const FOutputVertexInterface& OutputInterface = InDefaultInterface.GetOutputInterface();
+		const FOutputVertexInterface& OutputInterface = InVertexInterface.GetOutputInterface();
 		FMetasoundFrontendInterfaceStyle OutputStyle;
 		for (const TPair<FVertexName, FOutputDataVertex>& OutputTuple : OutputInterface)
 		{
@@ -268,7 +268,7 @@ FMetasoundFrontendClassInterface FMetasoundFrontendClassInterface::GenerateClass
 		ClassInterface.OutputStyle = MoveTemp(OutputStyle);
 	}
 
-	for (auto& EnvTuple : InDefaultInterface.GetEnvironmentInterface())
+	for (auto& EnvTuple : InVertexInterface.GetEnvironmentInterface())
 	{
 		FMetasoundFrontendClassEnvironmentVariable EnvVar;
 
@@ -388,7 +388,7 @@ void FMetasoundFrontendClassMetadata::SetVersion(const FMetasoundFrontendVersion
 	SetWithChangeID(InVersion, Version, ChangeID);
 }
 
-FMetasoundFrontendClassStyle FMetasoundFrontendClassStyle::GenerateClassDescription(const Metasound::FNodeDisplayStyle& InNodeDisplayStyle)
+FMetasoundFrontendClassStyle FMetasoundFrontendClassStyle::GenerateClassStyle(const Metasound::FNodeDisplayStyle& InNodeDisplayStyle)
 {
 	FMetasoundFrontendClassStyle Style;
 	Style.Display.bShowName = InNodeDisplayStyle.bShowName;
@@ -453,12 +453,10 @@ bool FMetasoundFrontendClass::CacheGraphDependencyMetadataFromRegistry(FMetasoun
 }
 #endif // WITH_EDITORONLY_DATA
 
-FMetasoundFrontendClassMetadata FMetasoundFrontendClassMetadata::GenerateClassDescription(const Metasound::FNodeClassMetadata& InNodeClassMetadata, EMetasoundFrontendClassType InType)
+FMetasoundFrontendClassMetadata FMetasoundFrontendClassMetadata::GenerateClassMetadata(const Metasound::FNodeClassMetadata& InNodeClassMetadata, EMetasoundFrontendClassType InType)
 {
 	FMetasoundFrontendClassMetadata NewMetadata;
 
-	NewMetadata.SetSerializeText(false);
-
 	NewMetadata.Type = InType;
 
 	// TODO: This flag is only used by the graph class' metadata.
@@ -470,6 +468,7 @@ FMetasoundFrontendClassMetadata FMetasoundFrontendClassMetadata::GenerateClassDe
 	NewMetadata.ClassName = InNodeClassMetadata.ClassName;
 	NewMetadata.Version = { InNodeClassMetadata.MajorVersion, InNodeClassMetadata.MinorVersion };
 
+	NewMetadata.SetSerializeText(false);
 	NewMetadata.SetDisplayName(InNodeClassMetadata.DisplayName);
 	NewMetadata.SetDescription(InNodeClassMetadata.Description);
 	NewMetadata.SetPromptIfMissing(InNodeClassMetadata.PromptIfMissing);
diff --git a/Engine/Plugins/Runtime/Metasound/Source/MetasoundFrontend/Private/MetasoundFrontendDocumentController.cpp b/Engine/Plugins/Runtime/Metasound/Source/MetasoundFrontend/Private/MetasoundFrontendDocumentController.cpp
--- a/Engine/Plugins/Runtime/Metasound/Source/MetasoundFrontend/Private/MetasoundFrontendDocumentController.cpp
+++ b/Engine/Plugins/Runtime/Metasound/Source/MetasoundFrontend/Private/MetasoundFrontendDocumentController.cpp
@@ -285,7 +285,7 @@ namespace Metasound
 						// TODO: Assuming we want to recheck classes when they add another
 						// node, this should be replace with a call to synchronize a 
 						// single class.
-						FMetasoundFrontendClass NewClass = GenerateClassDescription(InKey);
+						FMetasoundFrontendClass NewClass = GenerateClass(InKey);
 						if (NewClass.Metadata.GetVersion().Major != MetasoundClass->Metadata.GetVersion().Major)
 						{
 							return AddClass(MoveTemp(NewClass), FGuid::NewGuid());
@@ -295,14 +295,14 @@ namespace Metasound
 					if (bInRefreshFromRegistry)
 					{
 						FGuid ClassID = MetasoundClass->ID;
-						FMetasoundFrontendClass NewClass = GenerateClassDescription(InKey);
+						FMetasoundFrontendClass NewClass = GenerateClass(InKey);
 						return AddClass(MoveTemp(NewClass), ClassID);
 					}
 
 					return ClassPtr;
 				}
 
-				FMetasoundFrontendClass NewClass = GenerateClassDescription(InKey);
+				FMetasoundFrontendClass NewClass = GenerateClass(InKey);
 				return AddClass(MoveTemp(NewClass), FGuid::NewGuid());
 			}
 
diff --git a/Engine/Plugins/Runtime/Metasound/Source/MetasoundFrontend/Public/MetasoundDataTypeRegistrationMacro.h b/Engine/Plugins/Runtime/Metasound/Source/MetasoundFrontend/Public/MetasoundDataTypeRegistrationMacro.h
--- a/Engine/Plugins/Runtime/Metasound/Source/MetasoundFrontend/Public/MetasoundDataTypeRegistrationMacro.h
+++ b/Engine/Plugins/Runtime/Metasound/Source/MetasoundFrontend/Public/MetasoundDataTypeRegistrationMacro.h
@@ -349,25 +349,25 @@ namespace Metasound
 					if constexpr (bIsParsable)
 					{
 						TInputNode<TDataType> InputPrototype(TEXT(""), FGuid(), TEXT(""), FLiteral());
-						InputClass = Metasound::Frontend::GenerateClassDescription(InputPrototype.GetMetadata(), EMetasoundFrontendClassType::Input);
+						InputClass = Metasound::Frontend::GenerateClass(InputPrototype.GetMetadata(), EMetasoundFrontendClassType::Input);
 
 						TOutputNode<TDataType> OutputPrototype(TEXT(""), FGuid(), TEXT(""));
-						OutputClass = Metasound::Frontend::GenerateClassDescription(OutputPrototype.GetMetadata(), EMetasoundFrontendClassType::Output);
+						OutputClass = Metasound::Frontend::GenerateClass(OutputPrototype.GetMetadata(), EMetasoundFrontendClassType::Output);
 
 						TLiteralNode<TDataType> LiteralPrototype(TEXT(""), FGuid(), FLiteral());
-						LiteralClass = Metasound::Frontend::GenerateClassDescription(LiteralPrototype.GetMetadata(), EMetasoundFrontendClassType::Literal);
+						LiteralClass = Metasound::Frontend::GenerateClass(LiteralPrototype.GetMetadata(), EMetasoundFrontendClassType::Literal);
 
 						TVariableNode<TDataType> VariablePrototype(TEXT(""), FGuid(), FLiteral());
-						VariableClass = Metasound::Frontend::GenerateClassDescription(VariablePrototype.GetMetadata(), EMetasoundFrontendClassType::Variable);
+						VariableClass = Metasound::Frontend::GenerateClass(VariablePrototype.GetMetadata(), EMetasoundFrontendClassType::Variable);
 
 						TVariableMutatorNode<TDataType> VariableMutatorPrototype(TEXT(""), FGuid());
-						VariableMutatorClass = Metasound::Frontend::GenerateClassDescription(VariableMutatorPrototype.GetMetadata(), EMetasoundFrontendClassType::VariableMutator);
+						VariableMutatorClass = Metasound::Frontend::GenerateClass(VariableMutatorPrototype.GetMetadata(), EMetasoundFrontendClassType::VariableMutator);
 
 						TVariableAccessorNode<TDataType> VariableAccessorPrototype(TEXT(""), FGuid());
-						VariableAccessorClass = Metasound::Frontend::GenerateClassDescription(VariableAccessorPrototype.GetMetadata(), EMetasoundFrontendClassType::VariableAccessor);
+						VariableAccessorClass = Metasound::Frontend::GenerateClass(VariableAccessorPrototype.GetMetadata(), EMetasoundFrontendClassType::VariableAccessor);
 
 						TVariableDeferredAccessorNode<TDataType> VariableDeferredAccessorPrototype(TEXT(""), FGuid());
-						VariableDeferredAccessorClass = Metasound::Frontend::GenerateClassDescription(VariableDeferredAccessorPrototype.GetMetadata(), EMetasoundFrontendClassType::VariableDeferredAccessor);
+						VariableDeferredAccessorClass = Metasound::Frontend::GenerateClass(VariableDeferredAccessorPrototype.GetMetadata(), EMetasoundFrontendClassType::VariableDeferredAccessor);
 					}
 				}
 				virtual ~FDataTypeRegistryEntry() {}
diff --git a/Engine/Plugins/Runtime/Metasound/Source/MetasoundFrontend/Public/MetasoundFrontend.h b/Engine/Plugins/Runtime/Metasound/Source/MetasoundFrontend/Public/MetasoundFrontend.h
--- a/Engine/Plugins/Runtime/Metasound/Source/MetasoundFrontend/Public/MetasoundFrontend.h
+++ b/Engine/Plugins/Runtime/Metasound/Source/MetasoundFrontend/Public/MetasoundFrontend.h
@@ -17,48 +17,37 @@ namespace Metasound
 {	
 	namespace Frontend
 	{
-		class IRegistryTransaction;
-		using FRegistryTransactionID = int32;
-
-		/** Returns all metadata (name, description, author, what to say if it's missing) for a given node.
-		 *
-		 * @param InInfo - Class info for a already registered external node.
-		 *
-		 * @return Metadata for node.
-		 */
-		METASOUNDFRONTEND_API FMetasoundFrontendClassMetadata GenerateClassMetadata(const FNodeRegistryKey& InKey);
-
 		/** Generates a new FMetasoundFrontendClass from Node Metadata 
 		 *
 		 * @param InNodeMetadata - Metadata describing an external node.
 		 *
-		 * @return Class description for external node.
+		 * @return FrontendClass for natively-defined node.
 		 */
-		METASOUNDFRONTEND_API FMetasoundFrontendClass GenerateClassDescription(const FNodeClassMetadata& InNodeMetadata, EMetasoundFrontendClassType ClassType=EMetasoundFrontendClassType::External);
+		METASOUNDFRONTEND_API FMetasoundFrontendClass GenerateClass(const FNodeClassMetadata& InNodeMetadata, EMetasoundFrontendClassType ClassType=EMetasoundFrontendClassType::External);
 
 		/** Generates a new FMetasoundFrontendClass from node lookup info.
 		 *
 		 * @param InInfo - Class info for a already registered external node.
 		 *
-		 * @return Class description for external node.
+		 * @return FrontendClass for natively-defined node.
 		 */
-		METASOUNDFRONTEND_API FMetasoundFrontendClass GenerateClassDescription(const Metasound::Frontend::FNodeRegistryKey& InKey);
+		METASOUNDFRONTEND_API FMetasoundFrontendClass GenerateClass(const Metasound::Frontend::FNodeRegistryKey& InKey);
 
 		/** Generates a new FMetasoundFrontendClass from Node init data
 		 *
 		 * @tparam NodeType - Type of node to instantiate.
 		 * @param InNodeInitData - Data used to call constructor of node.
 		 *
-		 * @return Class description for external node.
+		 * @return FrontendClass for natively-defined node.
 		 */
 		template<typename NodeType>
-		FMetasoundFrontendClass GenerateClassDescription(const FNodeInitData& InNodeInitData)
+		FMetasoundFrontendClass GenerateClass(const FNodeInitData& InNodeInitData)
 		{
 			TUniquePtr<INode> Node = MakeUnique<NodeType>(InNodeInitData);
 
 			if (ensure(Node.IsValid()))
 			{
-				return GenerateClassDescription(Node->GetMetadata());
+				return GenerateClass(Node->GetMetadata());
 			}
 
 			return FMetasoundFrontendClass();
@@ -68,15 +57,15 @@ namespace Metasound
 		 *
 		 * @tparam NodeType - Type of node.
 		 *
-		 * @return Class description for external node.
+		 * @return FrontendClass for natively-defined node.
 		 */
 		template<typename NodeType>
-		FMetasoundFrontendClass GenerateClassDescription()
+		FMetasoundFrontendClass GenerateClass()
 		{
 			FNodeInitData InitData;
-			InitData.InstanceName = TEXT("GenerateClassDescriptionForNode");
+			InitData.InstanceName = "GeneratedClass";
 
-			return GenerateClassDescription<NodeType>(InitData);
+			return GenerateClass<NodeType>(InitData);
 		}
 
 		// Takes a JSON string and deserializes it into a Metasound document struct.
diff --git a/Engine/Plugins/Runtime/Metasound/Source/MetasoundFrontend/Public/MetasoundFrontendDocument.h b/Engine/Plugins/Runtime/Metasound/Source/MetasoundFrontend/Public/MetasoundFrontendDocument.h
--- a/Engine/Plugins/Runtime/Metasound/Source/MetasoundFrontend/Public/MetasoundFrontendDocument.h
+++ b/Engine/Plugins/Runtime/Metasound/Source/MetasoundFrontend/Public/MetasoundFrontendDocument.h
@@ -840,7 +840,7 @@ private:
 public:
 	// Generates class interface intended to be used as a registry descriptor from FNodeClassMetadata.
 	// Does not initialize a change ID as it is not considered to be transactional.
-	static FMetasoundFrontendClassInterface GenerateClassDescription(const Metasound::FVertexInterface& InNodeClassMetadata);
+	static FMetasoundFrontendClassInterface GenerateClassInterface(const Metasound::FVertexInterface& InVertexInterface);
 
 	// Description of class inputs.
 	UPROPERTY(VisibleAnywhere, Category = CustomView)
@@ -962,7 +962,7 @@ struct METASOUNDFRONTEND_API FMetasoundFrontendClassMetadata
 	GENERATED_BODY()
 
 	// Generates class metadata intended to be used as a registry descriptor from FNodeClassMetadata. Does not initialize a change ID as it is not considered to be transactional.
-	static FMetasoundFrontendClassMetadata GenerateClassDescription(const Metasound::FNodeClassMetadata& InNodeClassMetadata, EMetasoundFrontendClassType InType);
+	static FMetasoundFrontendClassMetadata GenerateClassMetadata(const Metasound::FNodeClassMetadata& InNodeClassMetadata, EMetasoundFrontendClassType InType);
 
 private:
 	UPROPERTY(VisibleAnywhere, Category = Metasound)
@@ -1153,7 +1153,7 @@ struct FMetasoundFrontendClassStyle
 	FMetasoundFrontendClassStyleDisplay Display;
 
 	// Generates class style from core node class metadata.
-	static FMetasoundFrontendClassStyle GenerateClassDescription(const Metasound::FNodeDisplayStyle& InNodeDisplayStyle);
+	static FMetasoundFrontendClassStyle GenerateClassStyle(const Metasound::FNodeDisplayStyle& InNodeDisplayStyle);
 };
 
 USTRUCT()
diff --git a/Engine/Plugins/Runtime/Metasound/Source/MetasoundFrontend/Public/MetasoundNodeRegistrationMacro.h b/Engine/Plugins/Runtime/Metasound/Source/MetasoundFrontend/Public/MetasoundNodeRegistrationMacro.h
--- a/Engine/Plugins/Runtime/Metasound/Source/MetasoundFrontend/Public/MetasoundNodeRegistrationMacro.h
+++ b/Engine/Plugins/Runtime/Metasound/Source/MetasoundFrontend/Public/MetasoundNodeRegistrationMacro.h
@@ -45,8 +45,8 @@ namespace Metasound
 		{
 		public:
 			FNodeRegistryEntry(const Metasound::FNodeClassMetadata& InMetadata)
-			: ClassInfo(FMetasoundFrontendClassMetadata::GenerateClassDescription(InMetadata, EMetasoundFrontendClassType::External))
-			, FrontendClass(Metasound::Frontend::GenerateClassDescription(InMetadata)) 
+			: ClassInfo(FMetasoundFrontendClassMetadata::GenerateClassMetadata(InMetadata, EMetasoundFrontendClassType::External))
+			, FrontendClass(Metasound::Frontend::GenerateClass(InMetadata))
 			{
 			}
 
