Fix for ensure hit when calling CreateInternalWidgets prior to synchronizing graph for some MetaSounds assets (subsequent CL had deeper fix that avoided calling refresh on nodes prematurally, so the incorrect expectation that widgets had to be generated prior to synchronization no longer holds).  This is a partial revert from previous CL.
#rnx

diff --git a/Engine/Plugins/Runtime/Metasound/Source/MetasoundEditor/Private/MetasoundEditor.cpp b/Engine/Plugins/Runtime/Metasound/Source/MetasoundEditor/Private/MetasoundEditor.cpp
--- a/Engine/Plugins/Runtime/Metasound/Source/MetasoundEditor/Private/MetasoundEditor.cpp
+++ b/Engine/Plugins/Runtime/Metasound/Source/MetasoundEditor/Private/MetasoundEditor.cpp
@@ -442,17 +442,12 @@ namespace Metasound
 			FGraphEditorCommands::Register();
 			FEditorCommands::Register();
 
+			constexpr bool bForceRefreshNodes = true;
 			FGraphBuilder::RegisterGraphWithFrontend(*Metasound);
+			FGraphBuilder::SynchronizeGraph(*Metasound, bForceRefreshNodes);
 
 			BindGraphCommands();
 			CreateInternalWidgets();
-
-			// Call after create internal widgets as editor is
-			// referenced in validation step. TODO: Make this
-			// relationship more explicit.
-			constexpr bool bForceRefreshNodes = true;
-			FGraphBuilder::SynchronizeGraph(*Metasound, bForceRefreshNodes);
-
 			CreateAnalyzers();
 
 			const TSharedRef<FTabManager::FLayout> StandaloneDefaultLayout = FTabManager::NewLayout("Standalone_MetasoundEditor_Layout_v9")
