- Version MetaSound asset graph name FText being generated from file name & copied to MetaSound asset.
- Add missing user-defined fields to MetaSound graph options (ex. IsDeprecated, Keywords, etc.) and force update on transient style changeID (TODO: Probably would make more sense for Metadata FText to live on style struct and share that transient ChangeID, so that non-runtime changes to FText metadata wouldn't bump the change ID)
#rnx

diff --git a/Engine/Plugins/Runtime/Metasound/Source/MetasoundEditor/Private/MetasoundDetailCustomization.cpp b/Engine/Plugins/Runtime/Metasound/Source/MetasoundEditor/Private/MetasoundDetailCustomization.cpp
--- a/Engine/Plugins/Runtime/Metasound/Source/MetasoundEditor/Private/MetasoundDetailCustomization.cpp
+++ b/Engine/Plugins/Runtime/Metasound/Source/MetasoundEditor/Private/MetasoundDetailCustomization.cpp
@@ -343,6 +343,9 @@ namespace Metasound
 					const FName AuthorPropertyPath = BuildChildPath(GetMetadataPropertyPath(), FMetasoundFrontendClassMetadata::GetAuthorPropertyName());
 					const FName CategoryHierarchyPropertyPath = BuildChildPath(GetMetadataPropertyPath(), FMetasoundFrontendClassMetadata::GetCategoryHierarchyPropertyName());
 					const FName DescPropertyPath = BuildChildPath(GetMetadataPropertyPath(), FMetasoundFrontendClassMetadata::GetDescriptionPropertyName());
+					const FName DisplayNamePropertyPath = BuildChildPath(GetMetadataPropertyPath(), FMetasoundFrontendClassMetadata::GetDisplayNamePropertyName());
+					const FName KeywordsPropertyPath = BuildChildPath(GetMetadataPropertyPath(), FMetasoundFrontendClassMetadata::GetKeywordsPropertyName());
+					const FName IsDeprecatedPropertyPath = BuildChildPath(GetMetadataPropertyPath(), FMetasoundFrontendClassMetadata::GetIsDeprecatedPropertyName());
 
 					const FName ClassNamePropertyPath = BuildChildPath(GetMetadataPropertyPath(), FMetasoundFrontendClassMetadata::GetClassNamePropertyName());
 					const FName ClassNameNamePropertyPath = BuildChildPath(ClassNamePropertyPath, GET_MEMBER_NAME_CHECKED(FMetasoundFrontendClassName, Name));
@@ -356,10 +359,13 @@ namespace Metasound
 					TSharedPtr<IPropertyHandle> AuthorHandle = DetailLayout.GetProperty(AuthorPropertyPath);
 					TSharedPtr<IPropertyHandle> CategoryHierarchyHandle = DetailLayout.GetProperty(CategoryHierarchyPropertyPath);
 					TSharedPtr<IPropertyHandle> ClassNameHandle = DetailLayout.GetProperty(ClassNameNamePropertyPath);
+					TSharedPtr<IPropertyHandle> DisplayNameHandle = DetailLayout.GetProperty(DisplayNamePropertyPath);
 					TSharedPtr<IPropertyHandle> DescHandle = DetailLayout.GetProperty(DescPropertyPath);
+					TSharedPtr<IPropertyHandle> KeywordsHandle = DetailLayout.GetProperty(KeywordsPropertyPath);
+					TSharedPtr<IPropertyHandle> IsDeprecatedHandle = DetailLayout.GetProperty(IsDeprecatedPropertyPath);
+					TSharedPtr<IPropertyHandle> InterfaceVersionsHandle = DetailLayout.GetProperty(InterfaceVersionsPropertyPath);
 					TSharedPtr<IPropertyHandle> MajorVersionHandle = DetailLayout.GetProperty(MajorVersionPropertyPath);
 					TSharedPtr<IPropertyHandle> MinorVersionHandle = DetailLayout.GetProperty(MinorVersionPropertyPath);
-					TSharedPtr<IPropertyHandle> InterfaceVersionsHandle = DetailLayout.GetProperty(InterfaceVersionsPropertyPath);
 
 					// Invalid for UMetaSounds
 					TSharedPtr<IPropertyHandle> OutputFormat = DetailLayout.GetProperty(GET_MEMBER_NAME_CHECKED(UMetaSoundSource, OutputFormat));
@@ -409,8 +415,32 @@ namespace Metasound
 						OutputFormat->MarkHiddenByCustomization();
 					}
 
-					GeneralCategoryBuilder.AddProperty(AuthorHandle);
+					// Updates FText properties on open editors if required
+					{
+						FSimpleDelegate RegisterOnChange = FSimpleDelegate::CreateLambda([this]()
+							{
+								if (MetaSound.IsValid())
+								{
+									if (FMetasoundAssetBase* MetaSoundAsset = IMetasoundUObjectRegistry::Get().GetObjectAsAssetBase(MetaSound.Get()))
+									{
+										MetaSoundAsset->GetDocumentChecked().RootGraph.Style.UpdateChangeID();
+									}
+									constexpr bool bForceViewSynchronization = true;
+									FGraphBuilder::RegisterGraphWithFrontend(*MetaSound.Get(), bForceViewSynchronization);
+								}
+							});
+						AuthorHandle->SetOnChildPropertyValueChanged(RegisterOnChange);
+						DescHandle->SetOnPropertyValueChanged(RegisterOnChange);
+						DisplayNameHandle->SetOnPropertyValueChanged(RegisterOnChange);
+						KeywordsHandle->SetOnPropertyValueChanged(RegisterOnChange);
+						KeywordsHandle->SetOnChildPropertyValueChanged(RegisterOnChange);
+						IsDeprecatedHandle->SetOnPropertyValueChanged(RegisterOnChange);
+					}
+
+					GeneralCategoryBuilder.AddProperty(DisplayNameHandle);
 					GeneralCategoryBuilder.AddProperty(DescHandle);
+					GeneralCategoryBuilder.AddProperty(AuthorHandle);
+					GeneralCategoryBuilder.AddProperty(IsDeprecatedHandle);
 					GeneralCategoryBuilder.AddProperty(MajorVersionHandle);
 					GeneralCategoryBuilder.AddProperty(MinorVersionHandle);
 
@@ -431,6 +461,7 @@ namespace Metasound
 						ClassNameHandle->CreatePropertyValueWidget()
 					];
 					GeneralCategoryBuilder.AddProperty(CategoryHierarchyHandle);
+					GeneralCategoryBuilder.AddProperty(KeywordsHandle);
 
 					CustomizeInterfaces(DetailLayout);
 
diff --git a/Engine/Plugins/Runtime/Metasound/Source/MetasoundEditor/Private/MetasoundEditorGraphBuilder.cpp b/Engine/Plugins/Runtime/Metasound/Source/MetasoundEditor/Private/MetasoundEditorGraphBuilder.cpp
--- a/Engine/Plugins/Runtime/Metasound/Source/MetasoundEditor/Private/MetasoundEditorGraphBuilder.cpp
+++ b/Engine/Plugins/Runtime/Metasound/Source/MetasoundEditor/Private/MetasoundEditorGraphBuilder.cpp
@@ -85,29 +85,81 @@ namespace Metasound
 			}
 		} // namespace GraphBuilderPrivate
 
-		FText FGraphBuilder::GetDisplayName(const Frontend::INodeController& InFrontendNode, bool bInIncludeNamespace)
+		FText FGraphBuilder::GetDisplayName(const FMetasoundFrontendClassMetadata& InClassMetadata, FName InNodeName, bool bInIncludeNamespace)
 		{
+			using namespace Frontend;
+
 			FName Namespace;
 			FName ParameterName;
-			Audio::FParameterPath::SplitName(InFrontendNode.GetNodeName(), Namespace, ParameterName);
+			Audio::FParameterPath::SplitName(InNodeName, Namespace, ParameterName);
 
-			FText DisplayName = InFrontendNode.GetDisplayName();
+			FText DisplayName;
+			auto GetAssetDisplayNameFromMetadata = [&DisplayName](const FMetasoundFrontendClassMetadata& Metadata)
+			{
+				DisplayName = Metadata.GetDisplayName();
+				if (DisplayName.IsEmptyOrWhitespace())
+				{
+					const FNodeRegistryKey RegistryKey = NodeRegistryKey::CreateKey(Metadata);
+					bool bIsClassNative = FMetasoundFrontendRegistryContainer::Get()->IsNodeNative(RegistryKey);
+					if (!bIsClassNative)
+					{
+						if (IMetaSoundAssetManager* AssetManager = IMetaSoundAssetManager::Get())
+						{
+							if (const FSoftObjectPath* Path = AssetManager->FindObjectPathFromKey(RegistryKey))
+							{
+								DisplayName = FText::FromString(Path->GetAssetName());
+							}
+						}
+					}
+				}
+			};
+
+			// 1. Try to get display name from metadata or asset if one can be found from the asset manager
+			GetAssetDisplayNameFromMetadata(InClassMetadata);
+
+			// 2. If version is missing from the registry or from asset system, then this node
+			// will not provide a useful DisplayName.  In that case, attempt to find the next highest
+			// class & associated DisplayName.
 			if (DisplayName.IsEmptyOrWhitespace())
 			{
-				DisplayName = FText::FromName(ParameterName);
+				FMetasoundFrontendClass ClassWithHighestVersion;
+				if (ISearchEngine::Get().FindClassWithHighestVersion(InClassMetadata.GetClassName(), ClassWithHighestVersion))
+				{
+					GetAssetDisplayNameFromMetadata(ClassWithHighestVersion.Metadata);
+				}
+			}
+
+			// 3. If that cannot be found, build a title from the cached node registry FName.
+			if (DisplayName.IsEmptyOrWhitespace())
+			{
+				DisplayName = FText::FromString(ParameterName.ToString());
 			}
 
+			// 4. Tack on the namespace if requested
 			if (bInIncludeNamespace)
 			{
 				if (!Namespace.IsNone())
 				{
-					return FText::Format(LOCTEXT("MemberDisplayNameWithNamespaceFormat", "{0} ({1})"), DisplayName, FText::FromName(Namespace));
+					return FText::Format(LOCTEXT("ClassMetadataDisplayNameWithNamespaceFormat", "{0} ({1})"), DisplayName, FText::FromName(Namespace));
 				}
 			}
 
 			return DisplayName;
 		}
 
+		FText FGraphBuilder::GetDisplayName(const Frontend::INodeController& InFrontendNode, bool bInIncludeNamespace)
+		{
+			using namespace Frontend;
+
+			FText DisplayName = InFrontendNode.GetDisplayName();
+			if (!DisplayName.IsEmptyOrWhitespace())
+			{
+				return DisplayName;
+			}
+
+			return GetDisplayName(InFrontendNode.GetClassMetadata(), InFrontendNode.GetNodeName(), bInIncludeNamespace);
+		}
+
 		FText FGraphBuilder::GetDisplayName(const Frontend::IInputController& InFrontendInput)
 		{
 			FText DisplayName = InFrontendInput.GetDisplayName();
@@ -349,9 +401,10 @@ namespace Metasound
 			{
 				bMarkDirty |= Result.bIsDirty;
 				check(Result.Node);
-				const bool bMetadataChange = Result.Node->ContainsInterfaceChange();
-				const bool bInterfaceChange = Result.Node->ContainsMetadataChange();
-				if (Result.bIsDirty || bMetadataChange || bInterfaceChange || bForceRefreshNodes)
+				const bool bInterfaceChange = Result.Node->ContainsInterfaceChange();
+				const bool bMetadataChange = Result.Node->ContainsMetadataChange();
+				const bool bStyleChange = Result.Node->ContainsStyleChange();
+				if (Result.bIsDirty || bMetadataChange || bInterfaceChange || bStyleChange || bForceRefreshNodes)
 				{
 					Result.Node->SyncChangeIDs();
 					Result.Node->CacheTitle();
@@ -1269,7 +1322,7 @@ namespace Metasound
 			}
 		}
 
-		void FGraphBuilder::RegisterGraphWithFrontend(UObject& InMetaSound)
+		void FGraphBuilder::RegisterGraphWithFrontend(UObject& InMetaSound, bool bInForceViewSynchronization)
 		{
 			using namespace Frontend;
 
@@ -1298,8 +1351,8 @@ namespace Metasound
 
 			FMetaSoundAssetRegistrationOptions RegOptions;
 			RegOptions.bForceReregister = true;
+			RegOptions.bForceViewSynchronization = bInForceViewSynchronization;
 			RegOptions.bRegisterDependencies = true;
-
 			// if EditedReferencingMetaSounds is empty, then no MetaSounds are open
 			// that reference this MetaSound, so just register this asset. Otherwise,
 			// this graph will recursively get updated when the open referencing graphs
diff --git a/Engine/Plugins/Runtime/Metasound/Source/MetasoundEditor/Private/MetasoundEditorGraphBuilder.h b/Engine/Plugins/Runtime/Metasound/Source/MetasoundEditor/Private/MetasoundEditorGraphBuilder.h
--- a/Engine/Plugins/Runtime/Metasound/Source/MetasoundEditor/Private/MetasoundEditorGraphBuilder.h
+++ b/Engine/Plugins/Runtime/Metasound/Source/MetasoundEditor/Private/MetasoundEditorGraphBuilder.h
@@ -70,13 +70,22 @@ namespace Metasound
 			static TSharedPtr<FEditor> GetEditorForGraph(const UEdGraph& InEdGraph);
 
 			// Wraps RegisterGraphWithFrontend logic in Frontend with any additional logic required to refresh editor & respective editor object state.
-			static void RegisterGraphWithFrontend(UObject& InMetaSound);
+			// @param InMetaSound - MetaSound to register
+			// @param bInForceSynchronize - Forces the synchronize flag for all open graphs being registered by this call (all referenced graphs and
+			// referencing graphs open in editors)
+			static void RegisterGraphWithFrontend(UObject& InMetaSound, bool bInForceViewSynchronization = false);
 
 			// Wraps RegisterGraphWithFrontend logic in Frontend with any additional logic required to refresh editor & respective editor object state.
 			static void UnregisterGraphWithFrontend(UObject& InMetaSound);
 
 			// Returns a display name for a node. If the node has an empty or whitespace
-			// only DisplayName, then the NodeName is used. 
+			// only DisplayName, first attempts to use the asset name if class is defined
+			// in an asset, and finally the NodeName is used.
+			static FText GetDisplayName(const FMetasoundFrontendClassMetadata& InClassMetadata, FName InNodeName, bool bInIncludeNamespace);
+
+			// Returns a display name for a node. If the node has an empty or whitespace
+			// only DisplayName, first attempts to use the asset name if class is defined
+			// in an asset, and finally the NodeName is used.
 			static FText GetDisplayName(const Frontend::INodeController& InFrontendNode, bool bInIncludeNamespace);
 
 			// Returns a display name for an input. If the input has an empty or whitespace
diff --git a/Engine/Plugins/Runtime/Metasound/Source/MetasoundEditor/Private/MetasoundEditorGraphMemberDefaults.cpp b/Engine/Plugins/Runtime/Metasound/Source/MetasoundEditor/Private/MetasoundEditorGraphMemberDefaults.cpp
--- a/Engine/Plugins/Runtime/Metasound/Source/MetasoundEditor/Private/MetasoundEditorGraphMemberDefaults.cpp
+++ b/Engine/Plugins/Runtime/Metasound/Source/MetasoundEditor/Private/MetasoundEditorGraphMemberDefaults.cpp
@@ -215,7 +215,7 @@ void UMetasoundEditorGraphMemberDefaultFloat::PostEditChangeChainProperty(FPrope
 	}
 
 	// TODO: Remove this once widget Metadata is migrated to frontend
-	// Metadata, which will inherently update the change guid and issue
+	// style, which will inherently update the change guid and issue
 	// resync.
 	UMetasoundEditorGraphMember* Member = GetParentMember();
 	if (ensure(Member))
diff --git a/Engine/Plugins/Runtime/Metasound/Source/MetasoundEditor/Private/MetasoundEditorGraphNode.cpp b/Engine/Plugins/Runtime/Metasound/Source/MetasoundEditor/Private/MetasoundEditorGraphNode.cpp
--- a/Engine/Plugins/Runtime/Metasound/Source/MetasoundEditor/Private/MetasoundEditorGraphNode.cpp
+++ b/Engine/Plugins/Runtime/Metasound/Source/MetasoundEditor/Private/MetasoundEditorGraphNode.cpp
@@ -9,6 +9,7 @@
 #include "GraphEditorActions.h"
 #include "Logging/TokenizedMessage.h"
 #include "Metasound.h"
+#include "MetasoundAssetManager.h"
 #include "MetasoundEditorCommands.h"
 #include "MetasoundEditorGraph.h"
 #include "MetasoundEditorGraphBuilder.h"
@@ -179,7 +180,7 @@ void UMetasoundEditorGraphNode::SyncChangeIDs()
 
 	MetadataChangeID = NodeHandle->GetClassMetadata().GetChangeID();
 	InterfaceChangeID = NodeHandle->GetClassInterface().GetChangeID();
-
+	StyleChangeID = NodeHandle->GetClassStyle().GetChangeID();
 }
 
 void UMetasoundEditorGraphNode::CacheTitle()
@@ -207,6 +208,14 @@ bool UMetasoundEditorGraphNode::ContainsInterfaceChange() const
 	return InterfaceChangeID != NodeHandle->GetClassInterface().GetChangeID();
 }
 
+bool UMetasoundEditorGraphNode::ContainsStyleChange() const
+{
+	using namespace Metasound::Frontend;
+	FConstNodeHandle NodeHandle = GetConstNodeHandle();
+
+	return StyleChangeID != NodeHandle->GetClassStyle().GetChangeID();
+}
+
 void UMetasoundEditorGraphNode::ReconstructNode()
 {
 	using namespace Metasound::Editor;
@@ -703,28 +712,12 @@ bool UMetasoundEditorGraphExternalNode::CanAutoUpdate() const
 
 void UMetasoundEditorGraphExternalNode::CacheTitle()
 {
-	using namespace Metasound;
+	using namespace Metasound::Editor;
 	using namespace Metasound::Frontend;
 
-	Super::CacheTitle();
-
-	// If version is missing from the registry, then this node will not cache
-	// a useful title.  In that case, just display the next highest available name.
-	if (CachedTitle.IsEmpty())
-	{
-		FMetasoundFrontendClass ClassWithHighestVersion;
-		if (ISearchEngine::Get().FindClassWithHighestVersion(ClassName, ClassWithHighestVersion))
-		{
-			CachedTitle = ClassWithHighestVersion.Metadata.GetDisplayName();
-		}
-	}
-
-	// If that cannot be found, build a title from the cached node registry FName.
-	if (CachedTitle.IsEmpty())
-	{
-		CachedTitle = FText::FromString(ClassName.ToString());
-	}
-
+	constexpr bool bIncludeNamespace = false;
+	FConstNodeHandle NodeHandle = GetNodeHandle();
+	CachedTitle = FGraphBuilder::GetDisplayName(*NodeHandle, bIncludeNamespace);
 }
 
 UMetasoundEditorGraphExternalNode* UMetasoundEditorGraphExternalNode::UpdateToVersion(const FMetasoundFrontendVersionNumber& InNewVersion, bool bInPropagateErrorMessages)
@@ -908,16 +901,16 @@ bool UMetasoundEditorGraphExternalNode::Validate(Metasound::Editor::FGraphNodeVa
 	{
 		if (RegisteredClass.Metadata.GetIsDeprecated())
 		{
+			constexpr bool bIncludeNamespace = true;
 			GraphNodePrivate::SetGraphNodeMessage(*this, EMessageSeverity::Warning,
 				FString::Format(TEXT("Class '{0} {1}' is deprecated."),
 				{
-					*RegisteredClass.Metadata.GetClassName().ToString(),
+					*FGraphBuilder::GetDisplayName(RegisteredClass.Metadata, { }, bIncludeNamespace).ToString(),
 					*RegisteredClass.Metadata.GetVersion().ToString()
 				}));
 		}
 	}
 
-
 	// Find all available versions & report if upgrade available
 	const Metasound::FNodeClassName NodeClassName = Metadata.GetClassName().ToNodeClassName();
 	const TArray<FMetasoundFrontendClass> SortedClasses = ISearchEngine::Get().FindClassesWithName(NodeClassName, true /* bSortByVersion */);
diff --git a/Engine/Plugins/Runtime/Metasound/Source/MetasoundEditor/Private/MetasoundEditorGraphSchema.cpp b/Engine/Plugins/Runtime/Metasound/Source/MetasoundEditor/Private/MetasoundEditorGraphSchema.cpp
--- a/Engine/Plugins/Runtime/Metasound/Source/MetasoundEditor/Private/MetasoundEditorGraphSchema.cpp
+++ b/Engine/Plugins/Runtime/Metasound/Source/MetasoundEditor/Private/MetasoundEditorGraphSchema.cpp
@@ -1544,19 +1544,19 @@ void UMetasoundEditorGraphSchema::GetConversionActions(FGraphActionMenuBuilder&
 			FText KeywordsText = FText::Join(SchemaPrivate::KeywordDelim, Metadata.GetKeywords());
 			const FText CategoryText = FText::Join(SchemaPrivate::CategoryDelim, Metadata.GetCategoryHierarchy());
 
+			constexpr bool bShowNamespace = false; // TODO: Make this an option to display
 			TSharedPtr<FMetasoundGraphSchemaAction_NewNode> NewNodeAction = MakeShared<FMetasoundGraphSchemaAction_NewNode>
 			(
 				CategoryText,
-				Metadata.GetDisplayName(),
+				FGraphBuilder::GetDisplayName(Metadata, FName(), bShowNamespace),
 				Tooltip,
 				EPrimaryContextGroup::Conversions,
 				KeywordsText
 			);
 
 			Metadata.SetType(EMetasoundFrontendClassType::External);
-			NewNodeAction->ClassMetadata = Metadata;
+			NewNodeAction->ClassMetadata = MoveTemp(Metadata);
 			ActionMenuBuilder.AddAction(NewNodeAction);
-
 		}
 	}
 }
@@ -1679,10 +1679,11 @@ void UMetasoundEditorGraphSchema::GetFunctionActions(FGraphActionMenuBuilder& Ac
 			const FText KeywordsText = FText::Join(KeywordDelim, Metadata.GetKeywords());
 			const FText CategoryText = FText::Join(CategoryDelim, CategoryHierarchy);
 
+			constexpr bool bShowNamespace = false; // TODO: Make this an option to display
 			TSharedPtr<FMetasoundGraphSchemaAction_NewNode> NewNodeAction = MakeShared<FMetasoundGraphSchemaAction_NewNode>
 			(
 				CategoryText,
-				Metadata.GetDisplayName(),
+				FGraphBuilder::GetDisplayName(Metadata, FName(), bShowNamespace),
 				Tooltip,
 				ContextGroup,
 				KeywordsText
diff --git a/Engine/Plugins/Runtime/Metasound/Source/MetasoundEditor/Public/MetasoundEditorGraphNode.h b/Engine/Plugins/Runtime/Metasound/Source/MetasoundEditor/Public/MetasoundEditorGraphNode.h
--- a/Engine/Plugins/Runtime/Metasound/Source/MetasoundEditor/Public/MetasoundEditorGraphNode.h
+++ b/Engine/Plugins/Runtime/Metasound/Source/MetasoundEditor/Public/MetasoundEditorGraphNode.h
@@ -110,9 +110,13 @@ public:
 	// Returns whether or not Interface has been changed since the last node refresh
 	bool ContainsInterfaceChange() const;
 
+	// Returns whether or not Style has been changed since the last node refresh
+	bool ContainsStyleChange() const;
+
 protected:
 	FGuid InterfaceChangeID;
 	FGuid MetadataChangeID;
+	FGuid StyleChangeID;
 
 	// Not be serialized to avoid text desync as the registry can provide
 	// a new name if the external definition changes between application sessions.
diff --git a/Engine/Plugins/Runtime/Metasound/Source/MetasoundEngine/Private/MetasoundAssetSubsystem.cpp b/Engine/Plugins/Runtime/Metasound/Source/MetasoundEngine/Private/MetasoundAssetSubsystem.cpp
--- a/Engine/Plugins/Runtime/Metasound/Source/MetasoundEngine/Private/MetasoundAssetSubsystem.cpp
+++ b/Engine/Plugins/Runtime/Metasound/Source/MetasoundEngine/Private/MetasoundAssetSubsystem.cpp
@@ -379,8 +379,7 @@ void UMetaSoundAssetSubsystem::RenameAsset(const FAssetData& InAssetData, bool b
 	auto PerformRename = [this, &InAssetData]()
 	{
 		RemoveAsset(InAssetData);
-
-		SynchronizeAssetClassDisplayName(InAssetData);
+		ResetAssetClassDisplayName(InAssetData);
 		AddOrUpdateAsset(InAssetData);
 	};
 
@@ -399,7 +398,7 @@ void UMetaSoundAssetSubsystem::RenameAsset(const FAssetData& InAssetData, bool b
 	}
 }
 
-void UMetaSoundAssetSubsystem::SynchronizeAssetClassDisplayName(const FAssetData& InAssetData)
+void UMetaSoundAssetSubsystem::ResetAssetClassDisplayName(const FAssetData& InAssetData)
 {
 	UObject* Object = nullptr;
 	FSoftObjectPath Path(InAssetData.ObjectPath);
@@ -414,7 +413,8 @@ void UMetaSoundAssetSubsystem::SynchronizeAssetClassDisplayName(const FAssetData
 
 	FMetasoundAssetBase* MetaSoundAsset = Metasound::IMetasoundUObjectRegistry::Get().GetObjectAsAssetBase(Object);
 	check(MetaSoundAsset);
-	Metasound::Frontend::FSynchronizeAssetClassDisplayName(InAssetData.AssetName).Transform(MetaSoundAsset->GetDocumentHandle());
+	FMetasoundFrontendGraphClass& Class = MetaSoundAsset->GetDocumentChecked().RootGraph;
+	Class.Metadata.SetDisplayName(FText());
 }
 
 void UMetaSoundAssetSubsystem::SearchAndIterateDirectoryAssets(const TArray<FDirectoryPath>& InDirectories, TFunctionRef<void(const FAssetData&)> InFunction)
diff --git a/Engine/Plugins/Runtime/Metasound/Source/MetasoundEngine/Public/MetasoundAssetSubsystem.h b/Engine/Plugins/Runtime/Metasound/Source/MetasoundEngine/Public/MetasoundAssetSubsystem.h
--- a/Engine/Plugins/Runtime/Metasound/Source/MetasoundEngine/Public/MetasoundAssetSubsystem.h
+++ b/Engine/Plugins/Runtime/Metasound/Source/MetasoundEngine/Public/MetasoundAssetSubsystem.h
@@ -60,8 +60,8 @@ protected:
 	void PostEngineInit();
 	void PostInitAssetScan();
 	void RebuildDenyListCache(const UAssetManager& InAssetManager);
+	void ResetAssetClassDisplayName(const FAssetData& InAssetData);
 	void SearchAndIterateDirectoryAssets(const TArray<FDirectoryPath>& InDirectories, TFunctionRef<void(const FAssetData&)> InFunction);
-	void SynchronizeAssetClassDisplayName(const FAssetData& InAssetData);
 
 private:
 	int32 AutoUpdateDenyListChangeID = INDEX_NONE;
diff --git a/Engine/Plugins/Runtime/Metasound/Source/MetasoundFrontend/Private/MetasoundAssetBase.cpp b/Engine/Plugins/Runtime/Metasound/Source/MetasoundFrontend/Private/MetasoundAssetBase.cpp
--- a/Engine/Plugins/Runtime/Metasound/Source/MetasoundFrontend/Private/MetasoundAssetBase.cpp
+++ b/Engine/Plugins/Runtime/Metasound/Source/MetasoundFrontend/Private/MetasoundAssetBase.cpp
@@ -113,9 +113,17 @@ void FMetasoundAssetBase::RegisterGraphWithFrontend(Metasound::Frontend::FMetaSo
 		FString OwningAssetName = GetOwningAssetName();
 		const bool bAutoUpdated = FAutoUpdateRootGraph(MoveTemp(OwningAssetName)).Transform(GetDocumentHandle());
 #if WITH_EDITORONLY_DATA
-		if (bAutoUpdated)
+		if (bAutoUpdated || InRegistrationOptions.bForceViewSynchronization)
+		{
+ 			SetSynchronizationRequired();
+		}
+#endif // WITH_EDITORONLY_DATA
+	}
+	else
+	{
+#if WITH_EDITORONLY_DATA
+		if (InRegistrationOptions.bForceViewSynchronization)
 		{
- 
 			SetSynchronizationRequired();
 		}
 #endif // WITH_EDITORONLY_DATA
diff --git a/Engine/Plugins/Runtime/Metasound/Source/MetasoundFrontend/Private/MetasoundFrontendTransform.cpp b/Engine/Plugins/Runtime/Metasound/Source/MetasoundFrontend/Private/MetasoundFrontendTransform.cpp
--- a/Engine/Plugins/Runtime/Metasound/Source/MetasoundFrontend/Private/MetasoundFrontendTransform.cpp
+++ b/Engine/Plugins/Runtime/Metasound/Source/MetasoundFrontend/Private/MetasoundFrontendTransform.cpp
@@ -757,22 +757,6 @@ namespace Metasound
 			return ClassOutputs;
 		}
 
-		bool FSynchronizeAssetClassDisplayName::Transform(FDocumentHandle InDocument) const
-		{
-			const FMetasoundFrontendClassMetadata& Metadata = InDocument->GetRootGraphClass().Metadata;
-			const FText NewAssetName = FText::FromString(AssetName.ToString());
-
-			if (Metadata.GetDisplayName().CompareTo(NewAssetName) != 0)
-			{
-				FMetasoundFrontendClassMetadata NewMetadata = Metadata;
-				NewMetadata.SetDisplayName(NewAssetName);
-				InDocument->GetRootGraph()->SetGraphMetadata(NewMetadata);
-				return true;
-			}
-
-			return false;
-		}
-
 		bool FRegenerateAssetClassName::Transform(FDocumentHandle InDocument) const
 		{
 			FMetasoundFrontendClassMetadata Metadata = InDocument->GetRootGraph()->GetGraphMetadata();
@@ -985,7 +969,14 @@ namespace Metasound
 
 			void TransformInternal(FDocumentHandle InDocument) const override
 			{
-				FSynchronizeAssetClassDisplayName(AssetName).Transform(InDocument);
+				const FMetasoundFrontendClassMetadata& Metadata = InDocument->GetRootGraphClass().Metadata;
+				const FText NewAssetName = FText::FromString(AssetName.ToString());
+				if (Metadata.GetDisplayName().CompareTo(NewAssetName) != 0)
+				{
+					FMetasoundFrontendClassMetadata NewMetadata = Metadata;
+					NewMetadata.SetDisplayName(NewAssetName);
+					InDocument->GetRootGraph()->SetGraphMetadata(NewMetadata);
+				}
 			}
 
 		private:
@@ -1128,6 +1119,28 @@ namespace Metasound
 			}
 		};
 
+		/** Versions document from 1.8 to 1.9. */
+		class FVersionDocument_1_9 : public FVersionDocumentTransform
+		{
+		public:
+			FVersionDocument_1_9() = default;
+
+			FMetasoundFrontendVersionNumber GetTargetVersion() const override
+			{
+				return { 1, 9 };
+			}
+
+			void TransformInternal(FDocumentHandle InDocument) const override
+			{
+				// Display name text is no longer copied at this versioning point for assets
+				// from the asset's FName to avoid FText warnings regarding generation from
+				// an FString.  It also avoids desync if asset gets moved.
+				FMetasoundFrontendGraphClass RootGraphClass = InDocument->GetRootGraphClass();
+				RootGraphClass.Metadata.SetDisplayName(FText());
+				InDocument->SetRootGraphClass(MoveTemp(RootGraphClass));
+			}
+		};
+
 		FVersionDocument::FVersionDocument(FName InName, const FString& InPath)
 			: Name(InName)
 			, Path(InPath)
@@ -1154,12 +1167,13 @@ namespace Metasound
 			bWasUpdated |= FVersionDocument_1_6().Transform(InDocument);
 			bWasUpdated |= FVersionDocument_1_7().Transform(InDocument);
 			bWasUpdated |= FVersionDocument_1_8().Transform(InDocument);
+			bWasUpdated |= FVersionDocument_1_9().Transform(InDocument);
 
 			if (bWasUpdated)
 			{
 				const FText& DisplayName = InDocument->GetRootGraph()->GetGraphMetadata().GetDisplayName();
 				const FMetasoundFrontendVersionNumber NewVersionNumber = InDocument->GetMetadata().Version.Number;
-				UE_LOG(LogMetaSound, Display, TEXT("MetaSound Graph '%s' Parent Document Versioned: '%s' --> '%s'"), *DisplayName.ToString(), *InitVersionNumber.ToString(), *NewVersionNumber.ToString());
+				UE_LOG(LogMetaSound, Display, TEXT("MetaSound at '%s' Document Versioned: '%s' --> '%s'"), *Path, *InitVersionNumber.ToString(), *NewVersionNumber.ToString());
 			}
 
 			return bWasUpdated;
diff --git a/Engine/Plugins/Runtime/Metasound/Source/MetasoundFrontend/Public/MetasoundAssetManager.h b/Engine/Plugins/Runtime/Metasound/Source/MetasoundFrontend/Public/MetasoundAssetManager.h
--- a/Engine/Plugins/Runtime/Metasound/Source/MetasoundFrontend/Public/MetasoundAssetManager.h
+++ b/Engine/Plugins/Runtime/Metasound/Source/MetasoundFrontend/Public/MetasoundAssetManager.h
@@ -36,6 +36,9 @@ namespace Metasound
 			// if the following option 'bRegisterDependencies' is enabled).
 			bool bForceReregister = true;
 
+			// If true, forces flag to resync all view (editor) data pertaining to the given asset(s) being registered.
+			bool bForceViewSynchronization = true;
+
 			// If true, recursively attempts to register dependencies.
 			bool bRegisterDependencies = false;
 
diff --git a/Engine/Plugins/Runtime/Metasound/Source/MetasoundFrontend/Public/MetasoundFrontendDocument.h b/Engine/Plugins/Runtime/Metasound/Source/MetasoundFrontend/Public/MetasoundFrontendDocument.h
--- a/Engine/Plugins/Runtime/Metasound/Source/MetasoundFrontend/Public/MetasoundFrontendDocument.h
+++ b/Engine/Plugins/Runtime/Metasound/Source/MetasoundFrontend/Public/MetasoundFrontendDocument.h
@@ -974,7 +974,7 @@ private:
 	UPROPERTY(VisibleAnywhere, Category = Metasound)
 	EMetasoundFrontendClassType Type = EMetasoundFrontendClassType::Invalid;
 
-	UPROPERTY(VisibleAnywhere, Category = Metasound)
+	UPROPERTY(EditAnywhere, Category = Metasound)
 	FText DisplayName;
 
 	UPROPERTY(Transient)
@@ -999,7 +999,7 @@ private:
 	UPROPERTY(Transient)
 	FText AuthorTransient;
 
-	UPROPERTY(VisibleAnywhere, Category = Metasound)
+	UPROPERTY(EditAnywhere, Category = Metasound)
 	TArray<FText> Keywords;
 
 	UPROPERTY(Transient)
@@ -1012,7 +1012,7 @@ private:
 	TArray<FText> CategoryHierarchyTransient;
 
 	// If true, this node is deprecated and should not be used in new MetaSounds.
-	UPROPERTY(VisibleAnywhere, Category = Metasound)
+	UPROPERTY(EditAnywhere, Category = Metasound)
 	bool bIsDeprecated = false;
 
 	// If true, auto-update will manage (add and remove)
@@ -1045,11 +1045,26 @@ public:
 		return GET_MEMBER_NAME_CHECKED(FMetasoundFrontendClassMetadata, CategoryHierarchy);
 	}
 
+	static FName GetDisplayNamePropertyName()
+	{
+		return GET_MEMBER_NAME_CHECKED(FMetasoundFrontendClassMetadata, DisplayName);
+	}
+
 	static FName GetDescriptionPropertyName()
 	{
 		return GET_MEMBER_NAME_CHECKED(FMetasoundFrontendClassMetadata, Description);
 	}
 
+	static FName GetIsDeprecatedPropertyName()
+	{
+		return GET_MEMBER_NAME_CHECKED(FMetasoundFrontendClassMetadata, bIsDeprecated);
+	}
+
+	static FName GetKeywordsPropertyName()
+	{
+		return GET_MEMBER_NAME_CHECKED(FMetasoundFrontendClassMetadata, Keywords);
+	}
+
 	static FName GetClassNamePropertyName()
 	{
 		return GET_MEMBER_NAME_CHECKED(FMetasoundFrontendClassMetadata, ClassName);
@@ -1154,6 +1169,23 @@ struct FMetasoundFrontendClassStyle
 
 	// Generates class style from core node class metadata.
 	static FMetasoundFrontendClassStyle GenerateClassStyle(const Metasound::FNodeDisplayStyle& InNodeDisplayStyle);
+
+#if WITH_EDITORONLY_DATA
+	// Editor only ID that allows for pumping view to reflect changes to class.
+	void UpdateChangeID()
+	{
+		ChangeID = FGuid::NewGuid();
+	}
+
+	FGuid GetChangeID() const
+	{
+		return ChangeID;
+	}
+
+private:
+	UPROPERTY(Transient)
+	FGuid ChangeID;
+#endif // WITH_EDITORONLY_DATA
 };
 
 USTRUCT()
diff --git a/Engine/Plugins/Runtime/Metasound/Source/MetasoundFrontend/Public/MetasoundFrontendTransform.h b/Engine/Plugins/Runtime/Metasound/Source/MetasoundFrontend/Public/MetasoundFrontendTransform.h
--- a/Engine/Plugins/Runtime/Metasound/Source/MetasoundFrontend/Public/MetasoundFrontendTransform.h
+++ b/Engine/Plugins/Runtime/Metasound/Source/MetasoundFrontend/Public/MetasoundFrontendTransform.h
@@ -136,21 +136,6 @@ namespace Metasound
 			const FString DebugAssetPath;
 		};
 
-		/** Synchronizes the document's root graph's display name with that of the asset. */
-		class METASOUNDFRONTEND_API FSynchronizeAssetClassDisplayName : public IDocumentTransform
-		{
-		public:
-			FSynchronizeAssetClassDisplayName(FName InAssetName)
-				: AssetName(InAssetName)
-			{
-			}
-
-			bool Transform(FDocumentHandle InDocument) const override;
-
-		private:
-			FName AssetName;
-		};
-
 		/** Regenerates the class' name, effectively causing the class to get registered as
 		  * a new class (useful for when an asset is duplicated). */
 		class METASOUNDFRONTEND_API FRegenerateAssetClassName : public IDocumentTransform
@@ -171,7 +156,7 @@ namespace Metasound
 		public:
 			static FMetasoundFrontendVersionNumber GetMaxVersion()
 			{
-				return FMetasoundFrontendVersionNumber { 1, 8 };
+				return FMetasoundFrontendVersionNumber { 1, 9 };
 			}
 
 			FVersionDocument(FName InName, const FString& InPath);
