Move MetaSound asset type into the new MetaSounds asset submenu

#rnx

diff --git a/Engine/Plugins/Runtime/Metasound/Source/MetasoundEditor/Private/MetasoundAssetTypeActions.cpp b/Engine/Plugins/Runtime/Metasound/Source/MetasoundEditor/Private/MetasoundAssetTypeActions.cpp
--- a/Engine/Plugins/Runtime/Metasound/Source/MetasoundEditor/Private/MetasoundAssetTypeActions.cpp
+++ b/Engine/Plugins/Runtime/Metasound/Source/MetasoundEditor/Private/MetasoundAssetTypeActions.cpp
@@ -131,6 +131,16 @@ namespace Metasound
 			RegisterPresetAction<UMetaSound, UMetaSoundFactory>(PresetLabel, PresetToolTip);
 		}
 
+		const TArray<FText>& FAssetTypeActions_MetaSound::GetSubMenus() const
+		{
+			static const TArray<FText> SubMenus
+			{
+				LOCTEXT("AssetSoundMetaSoundsSubMenu", "MetaSounds"),
+			};
+
+			return SubMenus;
+		}
+
 		UClass* FAssetTypeActions_MetaSoundSource::GetSupportedClass() const
 		{
 			return UMetaSoundSource::StaticClass();
diff --git a/Engine/Plugins/Runtime/Metasound/Source/MetasoundEditor/Private/MetasoundAssetTypeActions.h b/Engine/Plugins/Runtime/Metasound/Source/MetasoundEditor/Private/MetasoundAssetTypeActions.h
--- a/Engine/Plugins/Runtime/Metasound/Source/MetasoundEditor/Private/MetasoundAssetTypeActions.h
+++ b/Engine/Plugins/Runtime/Metasound/Source/MetasoundEditor/Private/MetasoundAssetTypeActions.h
@@ -20,6 +20,7 @@ namespace Metasound
 			virtual FColor GetTypeColor() const override { return FColor(13, 55, 13); }
 			virtual UClass* GetSupportedClass() const override;
 			virtual uint32 GetCategories() override { return EAssetTypeCategories::Sounds; }
+			virtual const TArray<FText>& GetSubMenus() const override;
 			virtual void OpenAssetEditor(const TArray<UObject*>& InObjects, TSharedPtr<IToolkitHost> ToolkitHost) override;
 
 			static void RegisterMenuActions();
