Fixing duplicate loc keys in audio

diff --git a/Engine/Plugins/Runtime/Metasound/Source/MetasoundEditor/Private/MetasoundDetailCustomization.cpp b/Engine/Plugins/Runtime/Metasound/Source/MetasoundEditor/Private/MetasoundDetailCustomization.cpp
--- a/Engine/Plugins/Runtime/Metasound/Source/MetasoundEditor/Private/MetasoundDetailCustomization.cpp
+++ b/Engine/Plugins/Runtime/Metasound/Source/MetasoundEditor/Private/MetasoundDetailCustomization.cpp
@@ -172,7 +172,7 @@ namespace Metasound
 					InterfaceComboBox->RefreshOptions();
 					MetaSoundAsset->SetUpdateDetailsOnSynchronization();
 					MetaSoundAsset->SetSynchronizationRequired();
-				}), LOCTEXT("RemoveInterfaceTooltip", "Removes all interfaces from the given MetaSound."))
+				}), LOCTEXT("RemoveInterfaceTooltip1", "Removes all interfaces from the given MetaSound."))
 			];
 
 			const FText HeaderName = LOCTEXT("InterfacesGroupDisplayName", "Interfaces");
@@ -217,7 +217,7 @@ namespace Metasound
 					InterfaceComboBox->RefreshOptions();
 					MetaSoundAsset->SetUpdateDetailsOnSynchronization();
 					MetaSoundAsset->SetSynchronizationRequired();
-				}), LOCTEXT("RemoveInterfaceTooltip", "Removes the associated interface from the MetaSound."));
+				}), LOCTEXT("RemoveInterfaceTooltip2", "Removes the associated interface from the MetaSound."));
 
 				return SNew(SHorizontalBox)
 					+ SHorizontalBox::Slot()
diff --git a/Engine/Plugins/Runtime/Metasound/Source/MetasoundEditor/Private/MetasoundEditor.cpp b/Engine/Plugins/Runtime/Metasound/Source/MetasoundEditor/Private/MetasoundEditor.cpp
--- a/Engine/Plugins/Runtime/Metasound/Source/MetasoundEditor/Private/MetasoundEditor.cpp
+++ b/Engine/Plugins/Runtime/Metasound/Source/MetasoundEditor/Private/MetasoundEditor.cpp
@@ -1431,7 +1431,7 @@ namespace Metasound
 				UMetasoundEditorGraphMember* NextToSelect = Graph.FindAdjacentMember(*GraphMember);
 
 				{
-					const FScopedTransaction Transaction(LOCTEXT("MetaSoundEditorDeleteSelectedNode", "Delete MetaSound Graph Member"));
+					const FScopedTransaction Transaction(LOCTEXT("MetaSoundEditorDeleteSelectedNode1", "Delete MetaSound Graph Member"));
 					Metasound->Modify();
 					Graph.Modify();
 					Graph.RemoveMember(*GraphMember);
@@ -1526,7 +1526,7 @@ namespace Metasound
 			const FGraphPanelSelectionSet SelectedNodes = MetasoundGraphEditor->GetSelectedNodes();
 			MetasoundGraphEditor->ClearSelectionSet();
 
-			const FScopedTransaction Transaction(LOCTEXT("MetaSoundEditorDeleteSelectedNode", "Delete Selected MetaSound Node(s)"));
+			const FScopedTransaction Transaction(LOCTEXT("MetaSoundEditorDeleteSelectedNode2", "Delete Selected MetaSound Node(s)"));
 			check(Metasound);
 			Metasound->Modify();
 			UEdGraph* Graph = MetasoundGraphEditor->GetCurrentGraph();
diff --git a/Engine/Plugins/Runtime/Metasound/Source/MetasoundEditor/Private/MetasoundEditorGraphSchema.cpp b/Engine/Plugins/Runtime/Metasound/Source/MetasoundEditor/Private/MetasoundEditorGraphSchema.cpp
--- a/Engine/Plugins/Runtime/Metasound/Source/MetasoundEditor/Private/MetasoundEditorGraphSchema.cpp
+++ b/Engine/Plugins/Runtime/Metasound/Source/MetasoundEditor/Private/MetasoundEditorGraphSchema.cpp
@@ -473,7 +473,7 @@ FMetasoundGraphSchemaAction_PromoteToInput::FMetasoundGraphSchemaAction_PromoteT
 	: FMetasoundGraphSchemaAction_NodeWithMultipleOutputs(
 		FText(),
 		LOCTEXT("PromoteToInputName", "Promote To Graph Input"),
-		LOCTEXT("PromoteToInputTooltip", "Promotes node input to graph input"),
+		LOCTEXT("PromoteToInputTooltip2", "Promotes node input to graph input"),
 		Metasound::Editor::EPrimaryContextGroup::Common)
 {
 }
@@ -530,7 +530,7 @@ FMetasoundGraphSchemaAction_PromoteToVariable_AccessorNode::FMetasoundGraphSchem
 	: FMetasoundGraphSchemaAction_NodeWithMultipleOutputs(
 		FText(),
 		LOCTEXT("PromoteToVariableGetterName", "Promote To Graph Variable"),
-		LOCTEXT("PromoteToInputTooltip", "Promotes node input to graph variable and creates a getter node"),
+		LOCTEXT("PromoteToInputTooltip2", "Promotes node input to graph variable and creates a getter node"),
 		Metasound::Editor::EPrimaryContextGroup::Common)
 {
 }
@@ -562,7 +562,7 @@ FMetasoundGraphSchemaAction_PromoteToVariable_DeferredAccessorNode::FMetasoundGr
 	: FMetasoundGraphSchemaAction_NodeWithMultipleOutputs(
 		FText(),
 		LOCTEXT("PromoteToVariableDeferredGetterName", "Promote To Graph Variable (Deferred)"),
-		LOCTEXT("PromoteToInputTooltip", "Promotes node input to graph variable and creates a deferred getter node"),
+		LOCTEXT("PromoteToInputTooltip1", "Promotes node input to graph variable and creates a deferred getter node"),
 		Metasound::Editor::EPrimaryContextGroup::Common)
 {
 }
@@ -594,7 +594,7 @@ FMetasoundGraphSchemaAction_PromoteToVariable_MutatorNode::FMetasoundGraphSchema
 	: FMetasoundGraphSchemaAction(
 		FText(),
 		LOCTEXT("PromoteToVariableSetterName", "Promote To Graph Variable"),
-		LOCTEXT("PromoteToVariableSetterTooltip", "Promotes node input to graph variable and creates a setter node"),
+		LOCTEXT("PromoteToVariableSetterTooltip2", "Promotes node input to graph variable and creates a setter node"),
 		Metasound::Editor::EPrimaryContextGroup::Common)
 {
 }
@@ -668,7 +668,7 @@ UEdGraphNode* FMetasoundGraphSchemaAction_NewOutput::PerformAction(UEdGraph* Par
 		return nullptr;
 	}
 
-	const FScopedTransaction Transaction(LOCTEXT("AddNewOutputNode", "Add New MetaSound Output Node"));
+	const FScopedTransaction Transaction(LOCTEXT("AddNewOutputNode2", "Add New MetaSound Output Node"));
 	ParentMetasound.Modify();
 	ParentGraph->Modify();
 
@@ -854,7 +854,7 @@ UEdGraphNode* FMetasoundGraphSchemaAction_NewComment::PerformAction(UEdGraph* Pa
 {
 	using namespace Metasound::Editor;
 
-	const FScopedTransaction Transaction(LOCTEXT("AddNewOutputNode", "Add Comment to MetaSound Graph"));
+	const FScopedTransaction Transaction(LOCTEXT("AddNewOutputNode1", "Add Comment to MetaSound Graph"));
 	ParentGraph->Modify();
 
 	UEdGraphNode_Comment* CommentTemplate = NewObject<UEdGraphNode_Comment>();
@@ -1104,7 +1104,7 @@ const FPinConnectionResponse UMetasoundEditorGraphSchema::CanCreateConnection(co
 
 	if (ConnectionCausesLoop(InputPin, OutputPin))
 	{
-		return FPinConnectionResponse(CONNECT_RESPONSE_DISALLOW, LOCTEXT("ConnectionLoop", "Connection would cause loop"));
+		return FPinConnectionResponse(CONNECT_RESPONSE_DISALLOW, LOCTEXT("ConnectionLoop2", "Connection causes loop"));
 	}
 
 	if (InputPin->PinType.PinCategory != OutputPin->PinType.PinCategory)
@@ -1156,7 +1156,7 @@ const FPinConnectionResponse UMetasoundEditorGraphSchema::CanCreateConnection(co
 			}
 			else if (Frontend::FConnectability::EReason::CausesLoop == Connectability.Reason)
 			{
-				return FPinConnectionResponse(CONNECT_RESPONSE_DISALLOW, LOCTEXT("ConnectionLoop", "Connection causes loop"));
+				return FPinConnectionResponse(CONNECT_RESPONSE_DISALLOW, LOCTEXT("ConnectionLoop1", "Connection causes loop"));
 			}
 			else
 			{
@@ -1538,7 +1538,7 @@ void UMetasoundEditorGraphSchema::GetConversionActions(FGraphActionMenuBuilder&
 		FMetasoundFrontendClassMetadata Metadata = FrontendClass.Metadata;
 		const FText Tooltip = Metadata.GetAuthor().IsEmpty()
 			? Metadata.GetDescription()
-			: FText::Format(LOCTEXT("MetasoundTooltipAuthorFormat", "{0}\nAuthor: {1}"), Metadata.GetDescription(), Metadata.GetAuthor());
+			: FText::Format(LOCTEXT("MetasoundTooltipAuthorFormat1", "{0}\nAuthor: {1}"), Metadata.GetDescription(), Metadata.GetAuthor());
 		if (!Metadata.GetCategoryHierarchy().IsEmpty() && !Metadata.GetCategoryHierarchy()[0].CompareTo(NodeCategories::Conversions))
 		{
 			FText KeywordsText = FText::Join(SchemaPrivate::KeywordDelim, Metadata.GetKeywords());
@@ -1669,7 +1669,7 @@ void UMetasoundEditorGraphSchema::GetFunctionActions(FGraphActionMenuBuilder& Ac
 		FMetasoundFrontendClassMetadata Metadata = FrontendClass.Metadata;
 		const FText Tooltip = Metadata.GetAuthor().IsEmpty()
 			? Metadata.GetDescription()
-			: FText::Format(LOCTEXT("MetasoundTooltipAuthorFormat", "{0}\nAuthor: {1}"), Metadata.GetDescription(), Metadata.GetAuthor());
+			: FText::Format(LOCTEXT("MetasoundTooltipAuthorFormat2", "{0}\nAuthor: {1}"), Metadata.GetDescription(), Metadata.GetAuthor());
 
 		if (Metadata.GetCategoryHierarchy().IsEmpty() || Metadata.GetCategoryHierarchy()[0].CompareTo(Metasound::NodeCategories::Conversions))
 		{
diff --git a/Engine/Plugins/Runtime/Metasound/Source/MetasoundFrontend/Private/MetasoundSourceInterface.cpp b/Engine/Plugins/Runtime/Metasound/Source/MetasoundFrontend/Private/MetasoundSourceInterface.cpp
--- a/Engine/Plugins/Runtime/Metasound/Source/MetasoundFrontend/Private/MetasoundSourceInterface.cpp
+++ b/Engine/Plugins/Runtime/Metasound/Source/MetasoundFrontend/Private/MetasoundSourceInterface.cpp
@@ -122,7 +122,7 @@ namespace Metasound
 							},
 							{
 								LOCTEXT("GraphNameDisplayName", "Graph Name"),
-								LOCTEXT("AudioDeviceIDDescription", "Name of source graph (for debugging/logging)."),
+								LOCTEXT("GraphNameDescription", "Name of source graph (for debugging/logging)."),
 								GetMetasoundDataTypeName<FString>(),
 								Environment::GraphName
 							},
@@ -139,8 +139,8 @@ namespace Metasound
 								Environment::TransmitterID
 							},
 							{
-								LOCTEXT("AudioDeviceIDDisplayName", "Sound Unique ID"),
-								LOCTEXT("AudioDeviceIDDescription", "ID of unique source instance."),
+								LOCTEXT("SoundUniqueIdDisplayName", "Sound Unique ID"),
+								LOCTEXT("SoundUniqueIdDescription", "ID of unique source instance."),
 								FName(), // TODO: Align environment data types with environment (ex. this is actually set/get as a uint32)
 								Environment::SoundUniqueID
 							}
@@ -195,7 +195,7 @@ namespace Metasound
 						Environment =
 						{
 							{
-								LOCTEXT("AudioDeviceIDDisplayName", "Audio Device ID"),
+								LOCTEXT("AudioDeviceIDDisplayName1", "Audio Device ID"),
 								LOCTEXT("AudioDeviceIDDescription", "ID of AudioDevice source is played from."),
 								FName(), // TODO: Align environment data types with environment (ex. this is actually set/get as a uint32)
 								Environment::DeviceID
@@ -219,7 +219,7 @@ namespace Metasound
 								Environment::TransmitterID
 							},
 							{
-								LOCTEXT("AudioDeviceIDDisplayName", "Sound Unique ID"),
+								LOCTEXT("SoundUniqueDisplayName", "Sound Unique ID"),
 								LOCTEXT("AudioDeviceIDDescription", "ID of unique source instance."),
 								FName(), // TODO: Align environment data types with environment (ex. this is actually set/get as a uint32)
 								Environment::SoundUniqueID
diff --git a/Engine/Plugins/Runtime/Metasound/Source/MetasoundStandardNodes/Private/MetasoundBandsplitterNode.cpp b/Engine/Plugins/Runtime/Metasound/Source/MetasoundStandardNodes/Private/MetasoundBandsplitterNode.cpp
--- a/Engine/Plugins/Runtime/Metasound/Source/MetasoundStandardNodes/Private/MetasoundBandsplitterNode.cpp
+++ b/Engine/Plugins/Runtime/Metasound/Source/MetasoundStandardNodes/Private/MetasoundBandsplitterNode.cpp
@@ -89,8 +89,8 @@ namespace Metasound
 			auto CreateNodeClassMetadataMono = []() -> FNodeClassMetadata
 			{
 				FName OperatorName = *FString::Printf(TEXT("Band Splitter (Mono, %d)"), NumBands);
-				FText NodeDisplayName = FText::Format(LOCTEXT("BandSplitterDisplayNamePattern", "Mono Band Splitter ({0})"), NumBands);
-				FText NodeDescription = LOCTEXT("BandSplitterDescription", "Will split incoming audio into separate frequency bands.");
+				FText NodeDisplayName = FText::Format(LOCTEXT("BandSplitterDisplayNamePattern1", "Mono Band Splitter ({0})"), NumBands);
+				FText NodeDescription = LOCTEXT("BandSplitterDescription1", "Will split incoming audio into separate frequency bands.");
 				FVertexInterface NodeInterface = GetDefaultInterface();
 
 				return CreateNodeClassMetadata(OperatorName, NodeDisplayName, NodeDescription, NodeInterface);
@@ -100,8 +100,8 @@ namespace Metasound
 			auto CreateNodeClassMetadataStereo = []() -> FNodeClassMetadata
 			{
 				FName OperatorName = *FString::Printf(TEXT("Band Splitter (Stereo, %d)"), NumBands);
-				FText NodeDisplayName = FText::Format(LOCTEXT("BandSplitterDisplayNamePattern", "Stereo Band Splitter ({0})"), NumBands);
-				FText NodeDescription = LOCTEXT("BandSplitterDescription", "Will split incoming audio into separate frequency bands.");
+				FText NodeDisplayName = FText::Format(LOCTEXT("BandSplitterDisplayNamePattern2", "Stereo Band Splitter ({0})"), NumBands);
+				FText NodeDescription = LOCTEXT("BandSplitterDescription2", "Will split incoming audio into separate frequency bands.");
 				FVertexInterface NodeInterface = GetDefaultInterface();
 
 				return  CreateNodeClassMetadata(OperatorName, NodeDisplayName, NodeDescription, NodeInterface);
@@ -111,8 +111,8 @@ namespace Metasound
 			auto CreateNodeClassMetadataMultiChan = []() -> FNodeClassMetadata
 			{
 				FName OperatorName = *FString::Printf(TEXT("Band Splitter (%d-Channel, %d)"), NumChannels, NumBands);
-				FText NodeDisplayName = FText::Format(LOCTEXT("BandSplitterDisplayNamePattern", "{0}-channel Band Splitter ({1})"), NumChannels, NumBands);
-				FText NodeDescription = LOCTEXT("BandSplitterDescription", "Will split incoming audio into separate frequency bands.");
+				FText NodeDisplayName = FText::Format(LOCTEXT("BandSplitterDisplayNamePattern3", "{0}-channel Band Splitter ({1})"), NumChannels, NumBands);
+				FText NodeDescription = LOCTEXT("BandSplitterDescription3", "Will split incoming audio into separate frequency bands.");
 				FVertexInterface NodeInterface = GetDefaultInterface();
 
 				return  CreateNodeClassMetadata(OperatorName, NodeDisplayName, NodeDescription, NodeInterface);
diff --git a/Engine/Plugins/Runtime/Metasound/Source/MetasoundStandardNodes/Private/MetasoundMathNodes.cpp b/Engine/Plugins/Runtime/Metasound/Source/MetasoundStandardNodes/Private/MetasoundMathNodes.cpp
--- a/Engine/Plugins/Runtime/Metasound/Source/MetasoundStandardNodes/Private/MetasoundMathNodes.cpp
+++ b/Engine/Plugins/Runtime/Metasound/Source/MetasoundStandardNodes/Private/MetasoundMathNodes.cpp
@@ -66,7 +66,7 @@ namespace Metasound
 		template<typename DataType>
 		const FText GetAddDisplayName()
 		{
-			static const FText DisplayName = FText::Format(LOCTEXT("AddNodeDisplayNamePattern", "Add ({0})"), GetMetasoundDataTypeDisplayText<DataType>());
+			static const FText DisplayName = FText::Format(LOCTEXT("AddNodeDisplayNamePattern1", "Add ({0})"), GetMetasoundDataTypeDisplayText<DataType>());
 			return DisplayName;
 		}
 
@@ -74,7 +74,7 @@ namespace Metasound
 		const FText GetAddDisplayName()
 		{
 			static const FText DisplayName = FText::Format(
-				LOCTEXT("AddNodeDisplayNamePattern", "Add ({0} to {1})"),
+				LOCTEXT("AddNodeDisplayNamePattern2", "Add ({0} to {1})"),
 				GetMetasoundDataTypeDisplayText<OperandDataType>(),
 				GetMetasoundDataTypeDisplayText<DataType>());
 			return DisplayName;
@@ -671,8 +671,8 @@ namespace Metasound
 		{
 			static const FVertexInterface Interface(
 				FInputVertexInterface(
-					TInputDataVertexModel<FAudioBuffer>(MathOpNames::PrimaryOperandName, LOCTEXT("MathOpAddendTooltip", "First attend.")),
-					TInputDataVertexModel<FAudioBuffer>(MathOpNames::AdditionalOperandsName, LOCTEXT("MathOpAddendAdditionalTooltip", "Additional attends."))
+					TInputDataVertexModel<FAudioBuffer>(MathOpNames::PrimaryOperandName, LOCTEXT("MathOpAddendTooltip1", "First addend.")),
+					TInputDataVertexModel<FAudioBuffer>(MathOpNames::AdditionalOperandsName, LOCTEXT("MathOpAddendAdditionalTooltip1", "Additional addends."))
 				),
 				FOutputVertexInterface(
 					TOutputDataVertexModel<FAudioBuffer>(TEXT("Out"), LOCTEXT("MathOpOutTooltip", "Math operation result"))
@@ -748,8 +748,8 @@ namespace Metasound
 		{
 			static const FVertexInterface Interface(
 				FInputVertexInterface(
-					TInputDataVertexModel<FTime>(MathOpNames::PrimaryOperandName, LOCTEXT("MathOpAddendTooltip", "First attend."), 0.0f),
-					TInputDataVertexModel<FTime>(MathOpNames::AdditionalOperandsName, LOCTEXT("MathOpAddendAdditionalTooltip", "Additional attends."), 0.0f)
+					TInputDataVertexModel<FTime>(MathOpNames::PrimaryOperandName, LOCTEXT("MathOpAddendTooltip2", "First addend."), 0.0f),
+					TInputDataVertexModel<FTime>(MathOpNames::AdditionalOperandsName, LOCTEXT("MathOpAddendAdditionalTooltip2", "Additional addends."), 0.0f)
 				),
 				FOutputVertexInterface(
 					TOutputDataVertexModel<FTime>(TEXT("Out"), LOCTEXT("MathOpOutTooltip", "Math operation result"))
@@ -793,7 +793,7 @@ namespace Metasound
 			static const FVertexInterface Interface(
 				FInputVertexInterface(
 					TInputDataVertexModel<FAudioBuffer>(MathOpNames::PrimaryOperandName, LOCTEXT("MathOpAddendAudioTooltip", "Audio Buffer to add offset(s) to.")),
-					TInputDataVertexModel<float>(MathOpNames::AdditionalOperandsName, LOCTEXT("MathOpAddendAdditionalTooltip", "Float attends of which to offset buffer samples."), 0.0f)
+					TInputDataVertexModel<float>(MathOpNames::AdditionalOperandsName, LOCTEXT("MathOpAddendAdditionalTooltip2", "Float addends of which to offset buffer samples."), 0.0f)
 				),
 				FOutputVertexInterface(
 					TOutputDataVertexModel<FAudioBuffer>(TEXT("Out"), LOCTEXT("MathOpAudioFloatAddOutTooltip", "Resulting buffer"))
diff --git a/Engine/Plugins/Runtime/Metasound/Source/MetasoundStandardNodes/Private/MetasoundMixerNode.cpp b/Engine/Plugins/Runtime/Metasound/Source/MetasoundStandardNodes/Private/MetasoundMixerNode.cpp
--- a/Engine/Plugins/Runtime/Metasound/Source/MetasoundStandardNodes/Private/MetasoundMixerNode.cpp
+++ b/Engine/Plugins/Runtime/Metasound/Source/MetasoundStandardNodes/Private/MetasoundMixerNode.cpp
@@ -85,8 +85,8 @@ namespace Metasound
 			auto CreateNodeClassMetadataMono = []() -> FNodeClassMetadata
 			{
 				FName OperatorName = *FString::Printf(TEXT("Audio Mixer (Mono, %d)"), NumInputs);
-				FText NodeDisplayName = FText::Format(LOCTEXT("AudioMixerDisplayNamePattern", "Mono Mixer ({0})"), NumInputs);
-				FText NodeDescription = LOCTEXT("MixerDescription", "Will scale input channels by their corresponding gain value and sum them together.");
+				FText NodeDisplayName = FText::Format(LOCTEXT("MonoMixer", "Mono Mixer ({0})"), NumInputs);
+				FText NodeDescription = LOCTEXT("MixerDescription1", "Will scale input channels by their corresponding gain value and sum them together.");
 				FVertexInterface NodeInterface = GetDefaultInterface();
 
 				return CreateNodeClassMetadata(OperatorName, NodeDisplayName, NodeDescription, NodeInterface);
@@ -96,8 +96,8 @@ namespace Metasound
 			auto CreateNodeClassMetadataStereo = []() -> FNodeClassMetadata
 			{
 				FName OperatorName = *FString::Printf(TEXT("Audio Mixer (Stereo, %d)"), NumInputs);
-				FText NodeDisplayName = FText::Format(LOCTEXT("AudioMixerDisplayNamePattern", "Stereo Mixer ({0})"), NumInputs);
-				FText NodeDescription = LOCTEXT("MixerDescription", "Will scale input channels by their corresponding gain value and sum them together.");
+				FText NodeDisplayName = FText::Format(LOCTEXT("StereoMixer", "Stereo Mixer ({0})"), NumInputs);
+				FText NodeDescription = LOCTEXT("MixerDescription2", "Will scale input channels by their corresponding gain value and sum them together.");
 				FVertexInterface NodeInterface = GetDefaultInterface();
 
 				return  CreateNodeClassMetadata(OperatorName, NodeDisplayName, NodeDescription, NodeInterface);
@@ -107,8 +107,8 @@ namespace Metasound
 			auto CreateNodeClassMetadataMultiChan = []() -> FNodeClassMetadata
 			{
 				FName OperatorName = *FString::Printf(TEXT("Audio Mixer (%d-Channel, %d)"), NumChannels, NumInputs);
-				FText NodeDisplayName = FText::Format(LOCTEXT("AudioMixerDisplayNamePattern", "{0}-channel Mixer ({1})"), NumChannels, NumInputs);
-				FText NodeDescription = LOCTEXT("MixerDescription", "Will scale input audio by their corresponding gain value and sum them together.");
+				FText NodeDisplayName = FText::Format(LOCTEXT("NChannelMixer", "{0}-channel Mixer ({1})"), NumChannels, NumInputs);
+				FText NodeDescription = LOCTEXT("MixerDescription3", "Will scale input audio by their corresponding gain value and sum them together.");
 				FVertexInterface NodeInterface = GetDefaultInterface();
 
 				return  CreateNodeClassMetadata(OperatorName, NodeDisplayName, NodeDescription, NodeInterface);
diff --git a/Engine/Plugins/Runtime/Metasound/Source/MetasoundStandardNodes/Private/MetasoundOscillatorNodes.cpp b/Engine/Plugins/Runtime/Metasound/Source/MetasoundStandardNodes/Private/MetasoundOscillatorNodes.cpp
--- a/Engine/Plugins/Runtime/Metasound/Source/MetasoundStandardNodes/Private/MetasoundOscillatorNodes.cpp
+++ b/Engine/Plugins/Runtime/Metasound/Source/MetasoundStandardNodes/Private/MetasoundOscillatorNodes.cpp
@@ -983,7 +983,7 @@ namespace Metasound
 
 	DEFINE_METASOUND_ENUM_BEGIN(ESquareGenerationType, FEnumSquareGenerationType, "SquareGenerationType")
 		DEFINE_METASOUND_ENUM_ENTRY(ESquareGenerationType::PolySmooth, LOCTEXT("SquarePolySmoothDescription", "Poly Smooth"), LOCTEXT("PolySmoothDescriptionTT", "PolySmooth (i.e. BLEP)")),
-		DEFINE_METASOUND_ENUM_ENTRY(ESquareGenerationType::Trivial, LOCTEXT("SquareTrivialDescription", "Trivial"), LOCTEXT("TrivialDescriptionTT", "The most basic raw implementation. Does not obey pulse width.")),
+		DEFINE_METASOUND_ENUM_ENTRY(ESquareGenerationType::Trivial, LOCTEXT("SquareTrivialDescription", "Trivial"), LOCTEXT("SquareTrivialDescriptionTT", "The most basic raw implementation. Does not obey pulse width.")),
 		//DEFINE_METASOUND_ENUM_ENTRY(ESquareGenerationType::Wavetable, LOCTEXT("SquareWavetableDescription", "Wavetable"), LOCTEXT("SquareWavetableDescriptionTT", "Use a Wavetable interpolation to generate the Waveform"))
 	DEFINE_METASOUND_ENUM_END()
 	
@@ -1167,7 +1167,7 @@ namespace Metasound
 
 	DEFINE_METASOUND_ENUM_BEGIN(ETriangleGenerationType, FEnumTriangleGenerationType, "TriangleGenerationType")
 		DEFINE_METASOUND_ENUM_ENTRY(ETriangleGenerationType::PolySmooth, LOCTEXT("TrianglePolySmoothDescription", "Poly Smooth"), LOCTEXT("PolySmoothDescriptionTT", "PolySmooth (i.e. BLEP)")),
-		DEFINE_METASOUND_ENUM_ENTRY(ETriangleGenerationType::Trivial, LOCTEXT("TriangleTrivialDescription", "Trivial"), LOCTEXT("TrivialDescriptionTT", "The most basic raw implementation")),
+		DEFINE_METASOUND_ENUM_ENTRY(ETriangleGenerationType::Trivial, LOCTEXT("TriangleTrivialDescription", "Trivial"), LOCTEXT("TriangleTrivialDescriptionTT", "The most basic raw implementation")),
 		//DEFINE_METASOUND_ENUM_ENTRY(ETriangleGenerationType::Wavetable, LOCTEXT("TriangleWavetableDescription", "Wavetable"), LOCTEXT("TriangleWavetableDescriptionTT", "Use a Wavetable iterpolation to generate the Waveform"))
 	DEFINE_METASOUND_ENUM_END()
 
diff --git a/Engine/Plugins/Runtime/Metasound/Source/MetasoundStandardNodes/Private/MetasoundTriggerToggleNode.cpp b/Engine/Plugins/Runtime/Metasound/Source/MetasoundStandardNodes/Private/MetasoundTriggerToggleNode.cpp
--- a/Engine/Plugins/Runtime/Metasound/Source/MetasoundStandardNodes/Private/MetasoundTriggerToggleNode.cpp
+++ b/Engine/Plugins/Runtime/Metasound/Source/MetasoundStandardNodes/Private/MetasoundTriggerToggleNode.cpp
@@ -134,7 +134,7 @@ namespace Metasound
 			),
 			FOutputVertexInterface(
 				TOutputDataVertexModel<FTrigger>(GetOutputTriggerName(), LOCTEXT("TriggerGateOutputTriggerTT", "Triggers output when gate is toggled.")),
-				TOutputDataVertexModel<bool>(GetOutputValueName(), LOCTEXT("TriggerGateOutputTriggerTT", "Current output value of the toggle."))
+				TOutputDataVertexModel<bool>(GetOutputValueName(), LOCTEXT("BoolOutputTT", "Current output value of the toggle."))
 			)
 		);
 
