Fixing random get "shared state" to use GUIDs vs metasound IDs to preserve state correctly.

diff --git a/Engine/Plugins/Runtime/Metasound/Source/MetasoundFrontend/Private/MetasoundArrayRandomNode.cpp b/Engine/Plugins/Runtime/Metasound/Source/MetasoundFrontend/Private/MetasoundArrayRandomNode.cpp
--- a/Engine/Plugins/Runtime/Metasound/Source/MetasoundFrontend/Private/MetasoundArrayRandomNode.cpp
+++ b/Engine/Plugins/Runtime/Metasound/Source/MetasoundFrontend/Private/MetasoundArrayRandomNode.cpp
@@ -39,7 +39,9 @@ namespace Metasound
 	{
 		if (InSeed == INDEX_NONE)
 		{
-			RandomStream.Initialize(FPlatformTime::Cycles());
+			static FThreadSafeCounter SeedCounter(FMath::Rand());
+			SeedCounter.Increment();
+			RandomStream.Initialize(SeedCounter.GetValue());
 		}
 		else
 		{
@@ -197,35 +199,35 @@ namespace Metasound
 		RandomGets.Add(InArgs.SharedStateId, MakeUnique<FArrayRandomGet>(InArgs.Seed, InArgs.NumElements, InArgs.Weights, InArgs.NoRepeatOrder));
 	}
 
-	int32 FSharedStateRandomGetManager::NextValue(uint32 InSharedStateId)
+	int32 FSharedStateRandomGetManager::NextValue(const FGuid& InSharedStateId)
 	{
 		FScopeLock Lock(&CritSect);
 		TUniquePtr<FArrayRandomGet>* RG = RandomGets.Find(InSharedStateId);
 		return (*RG)->NextValue();
 	}
 
-	void FSharedStateRandomGetManager::SetSeed(uint32 InSharedStateId, int32 InSeed)
+	void FSharedStateRandomGetManager::SetSeed(const FGuid& InSharedStateId, int32 InSeed)
 	{
 		FScopeLock Lock(&CritSect);
 		TUniquePtr<FArrayRandomGet>* RG = RandomGets.Find(InSharedStateId);
 		(*RG)->SetSeed(InSeed);
 	}
 
-	void FSharedStateRandomGetManager::SetNoRepeatOrder(uint32 InSharedStateId, int32 InNoRepeatOrder)
+	void FSharedStateRandomGetManager::SetNoRepeatOrder(const FGuid& InSharedStateId, int32 InNoRepeatOrder)
 	{
 		FScopeLock Lock(&CritSect);
 		TUniquePtr<FArrayRandomGet>* RG = RandomGets.Find(InSharedStateId);
 		(*RG)->SetNoRepeatOrder(InNoRepeatOrder);
 	}
 
-	void FSharedStateRandomGetManager::SetRandomWeights(uint32 InSharedStateId, const TArray<float>& InRandomWeights)
+	void FSharedStateRandomGetManager::SetRandomWeights(const FGuid& InSharedStateId, const TArray<float>& InRandomWeights)
 	{
 		FScopeLock Lock(&CritSect);
 		TUniquePtr<FArrayRandomGet>* RG = RandomGets.Find(InSharedStateId);
 		(*RG)->SetRandomWeights(InRandomWeights);
 	}
 
-	void FSharedStateRandomGetManager::ResetSeed(uint32 InSharedStateId)
+	void FSharedStateRandomGetManager::ResetSeed(const FGuid& InSharedStateId)
 	{
 		FScopeLock Lock(&CritSect);
 		TUniquePtr<FArrayRandomGet>* RG = RandomGets.Find(InSharedStateId);
diff --git a/Engine/Plugins/Runtime/Metasound/Source/MetasoundFrontend/Public/MetasoundArrayRandomNode.h b/Engine/Plugins/Runtime/Metasound/Source/MetasoundFrontend/Public/MetasoundArrayRandomNode.h
--- a/Engine/Plugins/Runtime/Metasound/Source/MetasoundFrontend/Public/MetasoundArrayRandomNode.h
+++ b/Engine/Plugins/Runtime/Metasound/Source/MetasoundFrontend/Public/MetasoundArrayRandomNode.h
@@ -71,7 +71,7 @@ namespace Metasound
 
 	struct InitSharedStateArgs
 	{
-		uint32 SharedStateId = 0;
+		FGuid SharedStateId;
 		int32 Seed = INDEX_NONE;
 		int32 NumElements = 0;
 		int32 NoRepeatOrder = 0;
@@ -85,11 +85,11 @@ namespace Metasound
 		static FSharedStateRandomGetManager& Get();
 
 		void InitSharedState(InitSharedStateArgs& InArgs);
-		int32 NextValue(uint32 InSharedStateId);
-		void SetSeed(uint32 InSharedStateId, int32 InSeed);
-		void SetNoRepeatOrder(uint32 InSharedStateId, int32 InNoRepeatOrder);
-		void SetRandomWeights(uint32 InSharedStateId, const TArray<float>& InRandomWeights);
-		void ResetSeed(uint32 InSharedStateId);
+		int32 NextValue(const FGuid& InSharedStateId);
+		void SetSeed(const FGuid& InSharedStateId, int32 InSeed);
+		void SetNoRepeatOrder(const FGuid& InSharedStateId, int32 InNoRepeatOrder);
+		void SetRandomWeights(const FGuid& InSharedStateId, const TArray<float>& InRandomWeights);
+		void ResetSeed(const FGuid& InSharedStateId);
 
 	private:
 		FSharedStateRandomGetManager() = default;
@@ -97,7 +97,7 @@ namespace Metasound
 
 		FCriticalSection CritSect;
 
-		TMap<uint32, TUniquePtr<FArrayRandomGet>> RandomGets;
+		TMap<FGuid, TUniquePtr<FArrayRandomGet>> RandomGets;
 	};
 
 	template<typename ArrayType>
@@ -211,8 +211,8 @@ namespace Metasound
 				if (*bEnableSharedState)
 				{
 					// Get the environment variable for the unique ID of the sound
-					SharedStateUniqueId = InParams.Environment.GetValue<uint32>(SourceInterface::Environment::SoundUniqueID);
-					check(SharedStateUniqueId != INDEX_NONE);
+					SharedStateUniqueId = InParams.Node.GetInstanceID();
+					check(SharedStateUniqueId.IsValid());
 
 					bIsPreviewSound = InParams.Environment.GetValue<bool>(SourceInterface::Environment::IsPreview);
 
@@ -309,7 +309,7 @@ namespace Metasound
 			{
 				PrevSeedValue = *SeedValue;
 
-				if (SharedStateUniqueId != INDEX_NONE)
+				if (SharedStateUniqueId.IsValid())
 				{
 					FSharedStateRandomGetManager& RGM = FSharedStateRandomGetManager::Get();
 					RGM.SetSeed(SharedStateUniqueId, PrevSeedValue);
@@ -324,7 +324,7 @@ namespace Metasound
 			if (PrevNoRepeatOrder != *NoRepeatOrder)
 			{
 				PrevNoRepeatOrder = *NoRepeatOrder;
-				if (SharedStateUniqueId != INDEX_NONE)
+				if (SharedStateUniqueId.IsValid())
 				{
 					FSharedStateRandomGetManager& RGM = FSharedStateRandomGetManager::Get();
 					RGM.SetNoRepeatOrder(SharedStateUniqueId, PrevNoRepeatOrder);
@@ -337,7 +337,7 @@ namespace Metasound
 			}
 
 			WeightsArray = *InputWeightsArray;
-			if (SharedStateUniqueId != INDEX_NONE)
+			if (SharedStateUniqueId.IsValid())
 			{
 				FSharedStateRandomGetManager& RGM = FSharedStateRandomGetManager::Get();
 				RGM.SetRandomWeights(SharedStateUniqueId, WeightsArray);
@@ -360,7 +360,7 @@ namespace Metasound
 				},
 				[this](int32 StartFrame, int32 EndFrame)
 				{
-					if (SharedStateUniqueId != INDEX_NONE)
+					if (SharedStateUniqueId.IsValid())
 					{
 						FSharedStateRandomGetManager& RGM = FSharedStateRandomGetManager::Get();
 						RGM.ResetSeed(SharedStateUniqueId);
@@ -383,7 +383,7 @@ namespace Metasound
 					const ArrayType& InputArrayRef = *InputArray;
 					int32 OutRandomIndex = INDEX_NONE;
 
-					if (SharedStateUniqueId != INDEX_NONE)
+					if (SharedStateUniqueId.IsValid())
 					{
 						FSharedStateRandomGetManager& RGM = FSharedStateRandomGetManager::Get();
 						OutRandomIndex = RGM.NextValue(SharedStateUniqueId);
@@ -429,7 +429,7 @@ namespace Metasound
 		TArray<float> WeightsArray;
 		int32 PrevSeedValue = INDEX_NONE;
 		int32 PrevNoRepeatOrder = INDEX_NONE;
-		uint32 SharedStateUniqueId = INDEX_NONE;
+		FGuid SharedStateUniqueId;
 		int32 PrevArraySize = 0;
 		bool bIsPreviewSound = false;
 		bool bHasLoggedEmptyArrayWarning = false;
