Fix for send / receive nodes failing on 2nd time playing sound. Fixes some failing sounds in AncientGame

diff --git a/Engine/Plugins/Runtime/Metasound/Source/MetasoundFrontend/Public/MetasoundReceiveNode.h b/Engine/Plugins/Runtime/Metasound/Source/MetasoundFrontend/Public/MetasoundReceiveNode.h
--- a/Engine/Plugins/Runtime/Metasound/Source/MetasoundFrontend/Public/MetasoundReceiveNode.h
+++ b/Engine/Plugins/Runtime/Metasound/Source/MetasoundFrontend/Public/MetasoundReceiveNode.h
@@ -148,12 +148,19 @@ namespace Metasound
 
 			private:
 
+				FSendAddress GetSendAddressWithDataType(const FSendAddress& InAddress) const 
+				{
+					// The data type of a send address is inferred by the underlying
+					// data type of this node. A full send address, including the data type,
+					// cannot be constructed from a literal. 
+					return FSendAddress{ InAddress.GetChannelName(), GetMetasoundDataTypeName<TDataType>(), InAddress.GetInstanceID() };
+				}
+
 				TReceiverPtr<TDataType> CreateNewReceiver() const
 				{
 					if (ensure(SendAddress->GetDataType().IsNone() || (GetMetasoundDataTypeName<TDataType>() == SendAddress->GetDataType())))
 					{
-						FSendAddress DataChannelKey(SendAddress->GetChannelName(), GetMetasoundDataTypeName<TDataType>(), SendAddress->GetInstanceID());
-						return FDataTransmissionCenter::Get().RegisterNewReceiver<TDataType>(DataChannelKey, CachedReceiverParams);
+						return FDataTransmissionCenter::Get().RegisterNewReceiver<TDataType>(GetSendAddressWithDataType(*SendAddress), CachedReceiverParams);
 					}
 					return TReceiverPtr<TDataType>(nullptr);
 				}
@@ -161,7 +168,7 @@ namespace Metasound
 				void ResetReceiverAndCleanupChannel()
 				{
 					Receiver.Reset();
-					FDataTransmissionCenter::Get().UnregisterDataChannelIfUnconnected(CachedSendAddress);
+					FDataTransmissionCenter::Get().UnregisterDataChannelIfUnconnected(GetSendAddressWithDataType(CachedSendAddress));
 				}
 
 				bool bHasNotReceivedData;
diff --git a/Engine/Plugins/Runtime/Metasound/Source/MetasoundFrontend/Public/MetasoundSendNode.h b/Engine/Plugins/Runtime/Metasound/Source/MetasoundFrontend/Public/MetasoundSendNode.h
--- a/Engine/Plugins/Runtime/Metasound/Source/MetasoundFrontend/Public/MetasoundSendNode.h
+++ b/Engine/Plugins/Runtime/Metasound/Source/MetasoundFrontend/Public/MetasoundSendNode.h
@@ -124,13 +124,19 @@ namespace Metasound
 				}
 
 			private:
+				FSendAddress GetSendAddressWithDataType(const FSendAddress& InAddress) const 
+				{
+					// The data type of a send address is inferred by the underlying
+					// data type of this node. A full send address, including the data type,
+					// cannot be constructed from a literal. 
+					return FSendAddress{ InAddress.GetChannelName(), GetMetasoundDataTypeName<TDataType>(), InAddress.GetInstanceID() };
+				}
 
 				TSenderPtr<TDataType> CreateNewSender() const
 				{
 					if (ensure(SendAddress->GetDataType().IsNone() || (GetMetasoundDataTypeName<TDataType>() == SendAddress->GetDataType())))
 					{
-						FSendAddress DataChannelKey(SendAddress->GetChannelName(), GetMetasoundDataTypeName<TDataType>(), SendAddress->GetInstanceID());
-						return FDataTransmissionCenter::Get().RegisterNewSender<TDataType>(DataChannelKey, CachedSenderParams);
+						return FDataTransmissionCenter::Get().RegisterNewSender<TDataType>(GetSendAddressWithDataType(*SendAddress), CachedSenderParams);
 					}
 					return TSenderPtr<TDataType>(nullptr);
 				}
@@ -138,7 +144,7 @@ namespace Metasound
 				void ResetSenderAndCleanupChannel()
 				{
 					Sender.Reset();
-					FDataTransmissionCenter::Get().UnregisterDataChannelIfUnconnected(CachedSendAddress);
+					FDataTransmissionCenter::Get().UnregisterDataChannelIfUnconnected(GetSendAddressWithDataType(CachedSendAddress));
 				}
 
 				TDataReadReference<TDataType> InputData;
