Fix for race condition in metasound BP interface causing data channels to get removed

diff --git a/Engine/Plugins/Runtime/Metasound/Source/MetasoundFrontend/Private/MetasoundParameterTransmitter.cpp b/Engine/Plugins/Runtime/Metasound/Source/MetasoundFrontend/Private/MetasoundParameterTransmitter.cpp
--- a/Engine/Plugins/Runtime/Metasound/Source/MetasoundFrontend/Private/MetasoundParameterTransmitter.cpp
+++ b/Engine/Plugins/Runtime/Metasound/Source/MetasoundFrontend/Private/MetasoundParameterTransmitter.cpp
@@ -158,7 +158,15 @@ namespace Metasound
 
 		for (const FSendInfo& SendInfo : SendInfos)
 		{
-			bSuccess &= FDataTransmissionCenter::Get().UnregisterDataChannel(SendInfo.Address);
+			if (InputSends.Remove(SendInfo.ParameterName))
+			{
+				// Only unregister the data channel if we had a sender using that 
+				// data channel. This protects against removing the data channel 
+				// multiple times. Multiple removals of data channels has caused
+				// race conditions between newly created transmitters and transmitters
+				// being cleaned up.
+				bSuccess &= FDataTransmissionCenter::Get().UnregisterDataChannel(SendInfo.Address);
+			}
 		}
 
 		return bSuccess;
