Fix AudioModulations Metasound Node audio device in PIE

diff --git a/Engine/Plugins/Runtime/AudioModulation/Source/AudioModulation/Private/MetaSoundNodes/GetModulatorValueNode.cpp b/Engine/Plugins/Runtime/AudioModulation/Source/AudioModulation/Private/MetaSoundNodes/GetModulatorValueNode.cpp
--- a/Engine/Plugins/Runtime/AudioModulation/Source/AudioModulation/Private/MetaSoundNodes/GetModulatorValueNode.cpp
+++ b/Engine/Plugins/Runtime/AudioModulation/Source/AudioModulation/Private/MetaSoundNodes/GetModulatorValueNode.cpp
@@ -77,6 +77,7 @@ namespace AudioModulation
 			, Normalized(InNormalized)
 			, OutValue(Metasound::TDataWriteReferenceFactory<float>::CreateAny(InParams.OperatorSettings))
 		{
+			Execute();
 		}
 
 		virtual ~FGetModulatorValueNodeOperator() = default;
diff --git a/Engine/Plugins/Runtime/AudioModulation/Source/AudioModulation/Private/MetaSoundNodes/MixModulatorsNode.cpp b/Engine/Plugins/Runtime/AudioModulation/Source/AudioModulation/Private/MetaSoundNodes/MixModulatorsNode.cpp
--- a/Engine/Plugins/Runtime/AudioModulation/Source/AudioModulation/Private/MetaSoundNodes/MixModulatorsNode.cpp
+++ b/Engine/Plugins/Runtime/AudioModulation/Source/AudioModulation/Private/MetaSoundNodes/MixModulatorsNode.cpp
@@ -88,6 +88,7 @@ namespace AudioModulation
 			, Parameter(InParameter)
 			, OutValue(Metasound::TDataWriteReferenceFactory<float>::CreateAny(InParams.OperatorSettings))
 		{
+			Execute();
 		}
 
 		virtual ~FMixModulatorsNodeOperator() = default;
diff --git a/Engine/Plugins/Runtime/Metasound/Source/MetasoundEngine/Private/MetasoundSource.cpp b/Engine/Plugins/Runtime/Metasound/Source/MetasoundEngine/Private/MetasoundSource.cpp
--- a/Engine/Plugins/Runtime/Metasound/Source/MetasoundEngine/Private/MetasoundSource.cpp
+++ b/Engine/Plugins/Runtime/Metasound/Source/MetasoundEngine/Private/MetasoundSource.cpp
@@ -42,38 +42,6 @@ namespace Metasound
 	{
 		using FFormatOutputVertexKeyMap = TMap<EMetasoundSourceAudioFormat, TArray<Metasound::FVertexName>>;
 
-		// Creates a base Metasound Environment that contains basic
-		// data associated with the AudioEngine (non-source/instance specific data).
-		FMetasoundEnvironment CreateEnvironment(UWorld* World)
-		{
-			using namespace Frontend;
-
-			FMetasoundEnvironment Environment;
-			// Add audio device ID to environment.
-			Audio::FDeviceId AudioDeviceID = INDEX_NONE;
-
-			if (World)
-			{
-				FAudioDeviceHandle DeviceHandle = World->GetAudioDevice();
-				if (DeviceHandle.IsValid())
-				{
-					AudioDeviceID = DeviceHandle.GetDeviceID();
-				}
-			}
-
-			if (INDEX_NONE == AudioDeviceID)
-			{
-				if (FAudioDeviceManager* DeviceManager = FAudioDeviceManager::Get())
-				{
-					AudioDeviceID = DeviceManager->GetMainAudioDeviceID();
-				}
-			}
-
-			Environment.SetValue<Audio::FDeviceId>(SourceInterface::Environment::DeviceID, AudioDeviceID);
-
-			return Environment;
-		}
-
 		const FFormatOutputVertexKeyMap& GetFormatOutputVertexKeys()
 		{
 			auto CreateVertexKeyMap = []()
@@ -767,7 +735,7 @@ Metasound::FMetasoundEnvironment UMetaSoundSource::CreateEnvironment() const
 	using namespace Metasound;
 	using namespace Metasound::Frontend;
 
-	FMetasoundEnvironment Environment = SourcePrivate::CreateEnvironment(GetWorld());
+	FMetasoundEnvironment Environment;
 	Environment.SetValue<uint32>(SourceInterface::Environment::SoundUniqueID, GetUniqueID());
 
 	return Environment;
@@ -782,6 +750,7 @@ Metasound::FMetasoundEnvironment UMetaSoundSource::CreateEnvironment(const FSoun
 	FMetasoundEnvironment Environment = CreateEnvironment();
 	Environment.SetValue<bool>(SourceInterface::Environment::IsPreview, InParams.bIsPreviewSound);
 	Environment.SetValue<uint64>(SourceInterface::Environment::TransmitterID, InParams.InstanceID);
+	Environment.SetValue<Audio::FDeviceId>(SourceInterface::Environment::DeviceID, InParams.AudioDeviceID);
 
 #if WITH_METASOUND_DEBUG_ENVIRONMENT
 	Environment.SetValue<FString>(SourceInterface::Environment::GraphName, GetFullName());
diff --git a/Engine/Source/Runtime/AudioMixer/Private/AudioMixerSource.cpp b/Engine/Source/Runtime/AudioMixer/Private/AudioMixerSource.cpp
--- a/Engine/Source/Runtime/AudioMixer/Private/AudioMixerSource.cpp
+++ b/Engine/Source/Runtime/AudioMixer/Private/AudioMixerSource.cpp
@@ -865,6 +865,7 @@ namespace Audio
 		}
 
 		FMixerSourceBufferInitArgs BufferInitArgs;
+		BufferInitArgs.AudioDeviceID = AudioDevice->DeviceID;
 		BufferInitArgs.InstanceID = InstanceID;
 		BufferInitArgs.SampleRate = AudioDevice->GetSampleRate();
 		BufferInitArgs.Buffer = MixerBuffer;
diff --git a/Engine/Source/Runtime/AudioMixer/Private/AudioMixerSourceBuffer.cpp b/Engine/Source/Runtime/AudioMixer/Private/AudioMixerSourceBuffer.cpp
--- a/Engine/Source/Runtime/AudioMixer/Private/AudioMixerSourceBuffer.cpp
+++ b/Engine/Source/Runtime/AudioMixer/Private/AudioMixerSourceBuffer.cpp
@@ -92,6 +92,7 @@ namespace Audio
 		if (bProcedural)
 		{
 			FSoundGeneratorInitParams InitParams;
+			InitParams.AudioDeviceID = InArgs.AudioDeviceID;
 			InitParams.SampleRate = InArgs.SampleRate;
 			InitParams.NumChannels = NumChannels;
 			InitParams.NumFramesPerCallback = MONO_PCM_BUFFER_SAMPLES;
diff --git a/Engine/Source/Runtime/AudioMixer/Private/AudioMixerSourceBuffer.h b/Engine/Source/Runtime/AudioMixer/Private/AudioMixerSourceBuffer.h
--- a/Engine/Source/Runtime/AudioMixer/Private/AudioMixerSourceBuffer.h
+++ b/Engine/Source/Runtime/AudioMixer/Private/AudioMixerSourceBuffer.h
@@ -48,6 +48,7 @@ namespace Audio
 
 	struct FMixerSourceBufferInitArgs
 	{
+		FDeviceId AudioDeviceID = 0;
 		uint32 InstanceID = 0;
 		int32 SampleRate = 0;
 		FMixerBuffer* Buffer = nullptr;
diff --git a/Engine/Source/Runtime/AudioMixer/Private/SoundWaveDecoder.cpp b/Engine/Source/Runtime/AudioMixer/Private/SoundWaveDecoder.cpp
--- a/Engine/Source/Runtime/AudioMixer/Private/SoundWaveDecoder.cpp
+++ b/Engine/Source/Runtime/AudioMixer/Private/SoundWaveDecoder.cpp
@@ -13,6 +13,7 @@ namespace Audio
 {
 	FDecodingSoundSource::FDecodingSoundSource(FAudioDevice* AudioDevice, const FSourceDecodeInit& InitData)
 		: Handle(InitData.Handle)
+		, AudioDeviceID(0)
 		, SoundWave(InitData.SoundWave)
 		, MixerBuffer(nullptr)
 		, SampleRate(INDEX_NONE)
@@ -23,6 +24,11 @@ namespace Audio
 		SourceInfo.VolumeParam.SetValue(InitData.VolumeScale);
 		SourceInfo.PitchScale = InitData.PitchScale;
 
+		if (nullptr != AudioDevice)
+		{
+			AudioDeviceID = AudioDevice->DeviceID;
+		}
+
 		MixerBuffer = FMixerBuffer::Init(AudioDevice, InitData.SoundWave, true);
 	}
 
@@ -59,6 +65,7 @@ namespace Audio
 			FScopeLock Lock(&MixerSourceBufferCritSec);
 
 			FMixerSourceBufferInitArgs Args;
+			Args.AudioDeviceID = AudioDeviceID;
 			Args.SampleRate = InSampleRate;
 			Args.Buffer = MixerBuffer;
 			Args.SoundWave = SoundWave;
diff --git a/Engine/Source/Runtime/AudioMixer/Public/SoundWaveDecoder.h b/Engine/Source/Runtime/AudioMixer/Public/SoundWaveDecoder.h
--- a/Engine/Source/Runtime/AudioMixer/Public/SoundWaveDecoder.h
+++ b/Engine/Source/Runtime/AudioMixer/Public/SoundWaveDecoder.h
@@ -93,6 +93,8 @@ namespace Audio
 		// Handle to the decoding source
 		FDecodingSoundSourceHandle Handle;
 
+		FDeviceId AudioDeviceID;
+
 		// The sound wave object with which this sound is generating
 		USoundWave* SoundWave;
 
@@ -274,4 +276,4 @@ namespace Audio
 		TQueue<TFunction<void()>> CommandQueue;
 	};
 
-}
\ No newline at end of file
+}
diff --git a/Engine/Source/Runtime/Engine/Classes/Sound/SoundGenerator.h b/Engine/Source/Runtime/Engine/Classes/Sound/SoundGenerator.h
--- a/Engine/Source/Runtime/Engine/Classes/Sound/SoundGenerator.h
+++ b/Engine/Source/Runtime/Engine/Classes/Sound/SoundGenerator.h
@@ -2,6 +2,7 @@
 
 #pragma once
 
+#include "AudioDeviceManager.h"
 #include "CoreMinimal.h"
 #include "Containers/Queue.h"
 #include "Templates/Function.h"
@@ -10,6 +11,7 @@
 // Parameters used for constructing a new ISoundGenerator.
 struct FSoundGeneratorInitParams
 {
+	Audio::FDeviceId AudioDeviceID;
 	float SampleRate = 0.0f;
 	int32 NumChannels = 0;
 	int32 NumFramesPerCallback = 0;
