Add metasound name to debug print node

diff --git a/Engine/Plugins/Runtime/Metasound/Source/MetasoundStandardNodes/Private/MetasoundDebugLogNode.cpp b/Engine/Plugins/Runtime/Metasound/Source/MetasoundStandardNodes/Private/MetasoundDebugLogNode.cpp
--- a/Engine/Plugins/Runtime/Metasound/Source/MetasoundStandardNodes/Private/MetasoundDebugLogNode.cpp
+++ b/Engine/Plugins/Runtime/Metasound/Source/MetasoundStandardNodes/Private/MetasoundDebugLogNode.cpp
@@ -3,6 +3,7 @@
 #include "MetasoundNodeRegistrationMacro.h"
 #include "MetasoundStandardNodesCategories.h"
 #include "MetasoundFacade.h"
+#include "MetasoundSourceInterface.h"
 #include "MetasoundStandardNodesNames.h"
 #include "MetasoundTrigger.h"
 #include "MetasoundPrimitives.h"
@@ -113,14 +114,21 @@ namespace Metasound
 				TDataReadReference<FString> Label = InputCollection.GetDataReadReferenceOrConstructWithVertexDefault<FString>(InputInterface, PrintLogVertexNames::GetLabelPrintLogName(), InParams.OperatorSettings);
 				TDataReadReference<PrintLogType> ValueToLog = InputCollection.GetDataReadReferenceOrConstructWithVertexDefault<PrintLogType>(InputInterface, PrintLogVertexNames::GetToLogPrintLogName(), InParams.OperatorSettings);
 
-				return MakeUnique<TPrintLogOperator<PrintLogType>>(InParams.OperatorSettings, Trigger, Label, ValueToLog);
+				FString GraphNameFull = InParams.Environment.GetValue<FString>(Frontend::SourceInterface::Environment::GraphName);
+				TArray<FString> ParsedString;
+				GraphNameFull.ParseIntoArray(ParsedString, TEXT("."), true);
+
+				const FString& GraphName = ParsedString.Last();
+
+				return MakeUnique<TPrintLogOperator<PrintLogType>>(InParams.OperatorSettings, Trigger, Label, ValueToLog, GraphName);
 			}
 
 
-			TPrintLogOperator(const FOperatorSettings& InSettings, TDataReadReference<FTrigger> InTrigger, TDataReadReference<FString> InLabelPrintLog, TDataReadReference<PrintLogType> InValueToLogPrintLog)
+			TPrintLogOperator(const FOperatorSettings& InSettings, TDataReadReference<FTrigger> InTrigger, TDataReadReference<FString> InLabelPrintLog, TDataReadReference<PrintLogType> InValueToLogPrintLog, const FString& InGraphName)
 				: Trigger(InTrigger)
 				, Label(InLabelPrintLog)
 				, ValueToLog(InValueToLogPrintLog)
+				, GraphName(InGraphName)
 			{
 			}
 
@@ -149,7 +157,7 @@ namespace Metasound
 			{
 				if (*Trigger)
 				{
-					UE_LOG(LogMetaSound, Display, TEXT("%s %s"), *(*Label), *LexToString(*ValueToLog));
+					UE_LOG(LogMetaSound, Display, TEXT("[%s]: %s %s"), *GraphName, *(*Label), *LexToString(*ValueToLog));
 				}
 			}
 
@@ -158,6 +166,8 @@ namespace Metasound
 			TDataReadReference<FTrigger> Trigger;
 			TDataReadReference<FString> Label;
 			TDataReadReference<PrintLogType> ValueToLog;
+
+			FString GraphName;
 	};
 
 	/** TPrintLogNode
