- Fix propagation issues with Preset input defaults
- Add UX for overriding preset inputs (allowing auto-update to pass through referenced value if not managed)
- Clarify code by moving bAutoUpdateManagestInterface to bIsPreset in new FMetasoundFrontendGraphClassPresetOptions struct
#rnx

diff --git a/Engine/Plugins/Runtime/Metasound/Source/MetasoundEditor/Private/MetasoundEditorGraphBuilder.cpp b/Engine/Plugins/Runtime/Metasound/Source/MetasoundEditor/Private/MetasoundEditorGraphBuilder.cpp
--- a/Engine/Plugins/Runtime/Metasound/Source/MetasoundEditor/Private/MetasoundEditorGraphBuilder.cpp
+++ b/Engine/Plugins/Runtime/Metasound/Source/MetasoundEditor/Private/MetasoundEditorGraphBuilder.cpp
@@ -1178,9 +1178,13 @@ namespace Metasound
 			Style.bIsGraphEditable = false;
 			PresetGraphHandle->SetGraphStyle(Style);
 
-			FMetasoundFrontendClassMetadata Metadata = PresetGraphHandle->GetGraphMetadata();
-			Metadata.SetAutoUpdateManagesInterface(true);
-			PresetGraphHandle->SetGraphMetadata(Metadata);
+			// Mark all inputs as inherited by default
+			TSet<FName> InputsInheritingDefault;
+			Algo::Transform(PresetGraphHandle->GetInputNodes(), InputsInheritingDefault, [](FConstNodeHandle NodeHandle)
+			{
+				return NodeHandle->GetNodeName();
+			});
+			PresetGraphHandle->SetInputsInheritingDefault(MoveTemp(InputsInheritingDefault));
 
 			FGraphBuilder::RegisterGraphWithFrontend(InMetaSoundReferenced);
 
@@ -2086,7 +2090,7 @@ namespace Metasound
 
 			UMetasoundEditorGraphMember* Member = nullptr;
 
-			auto SynchronizeMember = [](UMetasoundEditorGraphVertex& InVertex)
+			auto SynchronizeMemberDataType = [&](UMetasoundEditorGraphVertex& InVertex)
 			{
 				FConstNodeHandle NodeHandle = InVertex.GetConstNodeHandle();
 				TArray<FConstInputHandle> InputHandles = NodeHandle->GetConstInputs();
@@ -2120,12 +2124,28 @@ namespace Metasound
 				}
 			};
 
-			// Synchronize data types of input nodes.
+			// Synchronize data types & default values for input nodes.
 			GraphHandle->IterateNodes([&](FNodeHandle NodeHandle)
 			{
 				if (UMetasoundEditorGraphInput* Input = Graph->FindInput(NodeHandle->GetID()))
 				{
-					SynchronizeMember(*CastChecked<UMetasoundEditorGraphVertex>(Input));
+					SynchronizeMemberDataType(*Input);
+
+					if (UMetasoundEditorGraphMemberDefaultLiteral* Literal = Input->GetLiteral())
+					{
+						const FName NodeName = NodeHandle->GetNodeName();
+						const FGuid VertexID = GraphHandle->GetVertexIDForInputVertex(NodeName);
+						FMetasoundFrontendLiteral DefaultLiteral = GraphHandle->GetDefaultInput(VertexID);
+						if (!DefaultLiteral.IsEquivalent(Literal->GetDefault()))
+						{
+							if (DefaultLiteral.GetType() != EMetasoundFrontendLiteralType::None)
+							{
+								UE_LOG(LogMetasoundEditor, Verbose, TEXT("Synchronizing default value to '%s' for input '%s'"), *DefaultLiteral.ToString(), *NodeName.ToString());
+								Literal->SetFromLiteral(DefaultLiteral);
+								bEditorGraphModified = true;
+							}
+						}
+					}
 				}
 			}, EMetasoundFrontendClassType::Input);
 
@@ -2134,7 +2154,7 @@ namespace Metasound
 			{
 				if (UMetasoundEditorGraphOutput* Output = Graph->FindOutput(NodeHandle->GetID()))
 				{
-					SynchronizeMember(*CastChecked<UMetasoundEditorGraphVertex>(Output));
+					SynchronizeMemberDataType(*Output);
 				}
 			}, EMetasoundFrontendClassType::Output);
 
diff --git a/Engine/Plugins/Runtime/Metasound/Source/MetasoundEditor/Private/MetasoundEditorGraphMemberDefaults.cpp b/Engine/Plugins/Runtime/Metasound/Source/MetasoundEditor/Private/MetasoundEditorGraphMemberDefaults.cpp
--- a/Engine/Plugins/Runtime/Metasound/Source/MetasoundEditor/Private/MetasoundEditorGraphMemberDefaults.cpp
+++ b/Engine/Plugins/Runtime/Metasound/Source/MetasoundEditor/Private/MetasoundEditorGraphMemberDefaults.cpp
@@ -38,6 +38,7 @@ namespace Metasound
 			template <typename TPODType, typename TLiteralType = TPODType>
 			void ConvertLiteralToArray(const FMetasoundFrontendLiteral& InLiteral, TArray<TLiteralType>& OutArray)
 			{
+				OutArray.Reset();
 				const FName TypeName = CreateArrayTypeNameFromElementTypeName(GetMetasoundDataTypeName<TPODType>());
 				TArray<TPODType> NewValue = InLiteral.ToLiteral(TypeName).Value.Get<TArray<TPODType>>();
 				Algo::Transform(NewValue, OutArray, [](const TPODType& InValue) { return TLiteralType { InValue }; });
@@ -157,6 +158,12 @@ EMetasoundFrontendLiteralType UMetasoundEditorGraphMemberDefaultFloat::GetLitera
 	return EMetasoundFrontendLiteralType::Float;
 }
 
+void UMetasoundEditorGraphMemberDefaultFloat::ForceRefresh()
+{
+	OnRangeChanged.Broadcast(Range);
+	SetDefault(FMath::Clamp(Default, Range.X, Range.Y));
+}
+
 void UMetasoundEditorGraphMemberDefaultFloat::SetFromLiteral(const FMetasoundFrontendLiteral& InLiteral)
 {
 	Metasound::Editor::MemberDefaultsPrivate::ConvertLiteral<float>(InLiteral, Default);
@@ -197,7 +204,7 @@ void UMetasoundEditorGraphMemberDefaultFloat::PostEditChangeChainProperty(FPrope
 		{
 			// if Range.Y < Range.X, set Range.X to Range.Y
 			Range.X = FMath::Min(Range.X, Range.Y);
-			SetRange(Range);
+			ForceRefresh();
 		}
 	}
 	else if (PropertyChangedEvent.GetPropertyName().IsEqual(GET_MEMBER_NAME_CHECKED(UMetasoundEditorGraphMemberDefaultFloat, ClampDefault)))
@@ -228,8 +235,11 @@ void UMetasoundEditorGraphMemberDefaultFloat::PostEditChangeChainProperty(FPrope
 
 void UMetasoundEditorGraphMemberDefaultFloat::SetDefault(const float InDefault)
 {
-	Default = InDefault;
-	OnDefaultValueChanged.Broadcast(InDefault);
+	if (!FMath::IsNearlyEqual(Default, InDefault))
+	{
+		Default = InDefault;
+		OnDefaultValueChanged.Broadcast(InDefault);
+	}
 }
 
 float UMetasoundEditorGraphMemberDefaultFloat::GetDefault()
@@ -244,9 +254,12 @@ FVector2D UMetasoundEditorGraphMemberDefaultFloat::GetRange()
 
 void UMetasoundEditorGraphMemberDefaultFloat::SetRange(const FVector2D InRange)
 {
-	Range = InRange;
-	OnRangeChanged.Broadcast(InRange);
-	SetDefault(FMath::Clamp(Default, Range.X, Range.Y));
+	if (!(Range - InRange).IsNearlyZero())
+	{
+		Range = InRange;
+		OnRangeChanged.Broadcast(InRange);
+		SetDefault(FMath::Clamp(Default, Range.X, Range.Y));
+	}
 }
 
 FMetasoundFrontendLiteral UMetasoundEditorGraphMemberDefaultFloatArray::GetDefault() const
@@ -357,6 +370,7 @@ void UMetasoundEditorGraphMemberDefaultObjectArray::SetFromLiteral(const FMetaso
 	TArray<UObject*> ObjectArray;
 	ensure(InLiteral.TryGet(ObjectArray));
 
+	Default.Reset();
 	Algo::Transform(ObjectArray, Default, [](UObject* InValue) { return FMetasoundEditorGraphMemberDefaultObjectRef { InValue }; });
 }
 
diff --git a/Engine/Plugins/Runtime/Metasound/Source/MetasoundEditor/Private/MetasoundNodeDetailCustomization.cpp b/Engine/Plugins/Runtime/Metasound/Source/MetasoundEditor/Private/MetasoundNodeDetailCustomization.cpp
--- a/Engine/Plugins/Runtime/Metasound/Source/MetasoundEditor/Private/MetasoundNodeDetailCustomization.cpp
+++ b/Engine/Plugins/Runtime/Metasound/Source/MetasoundEditor/Private/MetasoundNodeDetailCustomization.cpp
@@ -3,33 +3,21 @@
 
 #include "Components/AudioComponent.h"
 #include "Containers/Set.h"
-#include "CoreMinimal.h"
 #include "Delegates/Delegate.h"
-#include "DetailCategoryBuilder.h"
-#include "DetailLayoutBuilder.h"
-#include "DetailWidgetRow.h"
 #include "Framework/Notifications/NotificationManager.h"
 #include "IDetailChildrenBuilder.h"
 #include "IDetailGroup.h"
 #include "Internationalization/Text.h"
-#include "MetasoundAssetBase.h"
 #include "MetasoundDataReference.h"
 #include "MetasoundDataReferenceMacro.h"
-#include "MetasoundEditorGraphBuilder.h"
-#include "MetasoundEditorGraphNode.h"
 #include "MetasoundEditorGraphInputNode.h"
-#include "MetasoundEditorGraphMemberDefaults.h"
 #include "MetasoundEditorGraphSchema.h"
-#include "MetasoundEditorModule.h"
 #include "MetasoundFrontend.h"
 #include "MetasoundFrontendController.h"
 #include "MetasoundFrontendDataTypeRegistry.h"
 #include "MetasoundFrontendRegistries.h"
-#include "MetasoundUObjectRegistry.h"
 #include "PropertyCustomizationHelpers.h"
 #include "PropertyEditorDelegates.h"
-#include "PropertyHandle.h"
-#include "PropertyRestriction.h"
 #include "SAssetDropTarget.h"
 #include "SlateCore/Public/Styling/SlateColor.h"
 #include "SMetasoundActionMenu.h"
@@ -41,9 +29,6 @@
 #include "UObject/WeakObjectPtrTemplates.h"
 #include "Widgets/Images/SImage.h"
 #include "Widgets/Input/SButton.h"
-#include "Widgets/Input/SCheckBox.h"
-#include "Widgets/Input/SMultiLineEditableTextBox.h"
-#include "Widgets/Input/STextComboBox.h"
 #include "Widgets/Notifications/SNotificationList.h"
 #include "Widgets/SToolTip.h"
 #include "Widgets/SWidget.h"
@@ -66,6 +51,11 @@ namespace Metasound
 				"Audio:Stereo"
 			};
 
+			static const FText OverrideInputDefaultText = LOCTEXT("OverridePresetInputDefault", "Override Inherited Default");
+			static const FText OverrideInputDefaultTooltip = LOCTEXT("OverridePresetInputTooltip",
+				"Enables overriding the input's inherited default value otherwise provided by the referenced graph."
+				"Setting to true disables auto-updating the input's default value if modified on the referenced asset.");
+
 			void GetDataTypeFromElementPropertyHandle(TSharedPtr<IPropertyHandle> ElementPropertyHandle, Frontend::FDataTypeRegistryInfo& OutDataTypeInfo)
 			{
 				using namespace Frontend;
@@ -100,25 +90,26 @@ namespace Metasound
 			if (FloatLiteral.IsValid())
 			{
 				FloatLiteral->OnClampChanged.Remove(OnClampChangedDelegateHandle);
-				FloatLiteral->OnRangeChanged.Remove(OnRangeChangedDelegateHandle);
 			}
 		}
 
-		void FMetasoundFloatLiteralCustomization::CustomizeLiteral(UMetasoundEditorGraphMemberDefaultLiteral& InLiteral, IDetailLayoutBuilder& InDetailLayout)
+		TArray<IDetailPropertyRow*> FMetasoundFloatLiteralCustomization::CustomizeLiteral(UMetasoundEditorGraphMemberDefaultLiteral& InLiteral, IDetailLayoutBuilder& InDetailLayout)
 		{
 			check(DefaultCategoryBuilder);
 
 			UMetasoundEditorGraphMemberDefaultFloat* DefaultFloat = Cast<UMetasoundEditorGraphMemberDefaultFloat>(&InLiteral);
 			if (!ensure(DefaultFloat))
 			{
-				return;
+				return { };
 			}
 			FloatLiteral = DefaultFloat;
 
+			TArray<IDetailPropertyRow*> DefaultRows;
 			TSharedPtr<IPropertyHandle> DefaultValueHandle;
 			IDetailPropertyRow* Row = DefaultCategoryBuilder->AddExternalObjectProperty(TArray<UObject*>({ DefaultFloat }), UMetasoundEditorGraphMemberDefaultFloat::GetDefaultPropertyName());
 			if (ensure(Row))
 			{
+				DefaultRows.Add(Row);
 				DefaultValueHandle = Row->GetPropertyHandle();
 			}
 
@@ -129,23 +120,28 @@ namespace Metasound
 			const bool bUsingWidget = DefaultFloat->WidgetType != EMetasoundMemberDefaultWidget::None;
 			const bool bShouldClampDefaultValue = (!bUsingWidget && DefaultFloat->ClampDefault) || (bUsingWidget && !bVolumeWidgetWithLinearOutput);
 
-			Row = DefaultCategoryBuilder->AddExternalObjectProperty(TArray<UObject*>({ DefaultFloat }), GET_MEMBER_NAME_CHECKED(UMetasoundEditorGraphMemberDefaultFloat, ClampDefault));
-			if (ensure(Row))
+			IDetailPropertyRow* ClampRow = DefaultCategoryBuilder->AddExternalObjectProperty(TArray<UObject*>({ DefaultFloat }), GET_MEMBER_NAME_CHECKED(UMetasoundEditorGraphMemberDefaultFloat, ClampDefault));
+			if (ensure(ClampRow))
 			{
+				DefaultRows.Add(ClampRow);
 				if (bVolumeWidgetWithLinearOutput)
 				{
 					DefaultFloat->SetRange(FVector2D(0.0f, 1.0f));
 				}
-				if (bShouldClampDefaultValue)
-				{
-					FVector2D Range = DefaultFloat->GetRange();
-					DefaultValueHandle->SetInstanceMetaData("ClampMin", FString::Printf(TEXT("%f"), Range.X));
-					DefaultValueHandle->SetInstanceMetaData("ClampMax", FString::Printf(TEXT("%f"), Range.Y));
-				}
-				else // Stop clamping
+
+				if (DefaultValueHandle.IsValid())
 				{
-					DefaultValueHandle->SetInstanceMetaData("ClampMin", "");
-					DefaultValueHandle->SetInstanceMetaData("ClampMax", "");
+					if (bShouldClampDefaultValue)
+					{
+						FVector2D Range = DefaultFloat->GetRange();
+						DefaultValueHandle->SetInstanceMetaData("ClampMin", FString::Printf(TEXT("%f"), Range.X));
+						DefaultValueHandle->SetInstanceMetaData("ClampMax", FString::Printf(TEXT("%f"), Range.Y));
+					}
+					else // Stop clamping
+					{
+						DefaultValueHandle->SetInstanceMetaData("ClampMin", "");
+						DefaultValueHandle->SetInstanceMetaData("ClampMax", "");
+					}
 				}
 
 				DefaultFloat->OnClampChanged.Remove(OnClampChangedDelegateHandle);
@@ -163,7 +159,11 @@ namespace Metasound
 
 				if (bShouldClampDefaultValue)
 				{
-					DefaultCategoryBuilder->AddExternalObjectProperty(TArray<UObject*>({ DefaultFloat }), GET_MEMBER_NAME_CHECKED(UMetasoundEditorGraphMemberDefaultFloat, Range));
+					IDetailPropertyRow* RangeRow = DefaultCategoryBuilder->AddExternalObjectProperty(TArray<UObject*>({ DefaultFloat }), GET_MEMBER_NAME_CHECKED(UMetasoundEditorGraphMemberDefaultFloat, Range));
+					if (ensure(RangeRow))
+					{
+						DefaultRows.Add(RangeRow);
+					}
 				}
 			}
 
@@ -180,18 +180,20 @@ namespace Metasound
 			if (bIsGraphEditable)
 			{
 				IDetailCategoryBuilder& WidgetCategoryBuilder = InDetailLayout.EditCategory("EditorOptions");
-				WidgetCategoryBuilder.AddExternalObjectProperty(TArray<UObject*>({ DefaultFloat }), GET_MEMBER_NAME_CHECKED(UMetasoundEditorGraphMemberDefaultFloat, WidgetType));
-				WidgetCategoryBuilder.AddExternalObjectProperty(TArray<UObject*>({ DefaultFloat }), GET_MEMBER_NAME_CHECKED(UMetasoundEditorGraphMemberDefaultFloat, WidgetOrientation));
-				WidgetCategoryBuilder.AddExternalObjectProperty(TArray<UObject*>({ DefaultFloat }), GET_MEMBER_NAME_CHECKED(UMetasoundEditorGraphMemberDefaultFloat, WidgetValueType));
+				DefaultRows.Add(WidgetCategoryBuilder.AddExternalObjectProperty(TArray<UObject*>({ DefaultFloat }), GET_MEMBER_NAME_CHECKED(UMetasoundEditorGraphMemberDefaultFloat, WidgetType)));
+				DefaultRows.Add(WidgetCategoryBuilder.AddExternalObjectProperty(TArray<UObject*>({ DefaultFloat }), GET_MEMBER_NAME_CHECKED(UMetasoundEditorGraphMemberDefaultFloat, WidgetOrientation)));
+				DefaultRows.Add(WidgetCategoryBuilder.AddExternalObjectProperty(TArray<UObject*>({ DefaultFloat }), GET_MEMBER_NAME_CHECKED(UMetasoundEditorGraphMemberDefaultFloat, WidgetValueType)));
 				if (DefaultFloat->WidgetValueType == EMetasoundMemberDefaultWidgetValueType::Volume)
 				{
-					WidgetCategoryBuilder.AddExternalObjectProperty(TArray<UObject*>({ DefaultFloat }), GET_MEMBER_NAME_CHECKED(UMetasoundEditorGraphMemberDefaultFloat, VolumeWidgetUseLinearOutput));
+					DefaultRows.Add(WidgetCategoryBuilder.AddExternalObjectProperty(TArray<UObject*>({ DefaultFloat }), GET_MEMBER_NAME_CHECKED(UMetasoundEditorGraphMemberDefaultFloat, VolumeWidgetUseLinearOutput)));
 
 				}
 			}
+
+			return DefaultRows;
 		}
 
-		void FMetasoundObjectArrayLiteralCustomization::CustomizeLiteral(UMetasoundEditorGraphMemberDefaultLiteral& InLiteral, IDetailLayoutBuilder& InDetailLayout)
+		TArray<IDetailPropertyRow*> FMetasoundObjectArrayLiteralCustomization::CustomizeLiteral(UMetasoundEditorGraphMemberDefaultLiteral& InLiteral, IDetailLayoutBuilder& InDetailLayout)
 		{
 			check(DefaultCategoryBuilder);
 
@@ -262,6 +264,8 @@ namespace Metasound
 					DefaultValueHandle->CreatePropertyValueWidget()
 				]
 			];
+
+			return { Row };
 		}
 
 		FText FMetasoundMemberDefaultBoolDetailCustomization::GetPropertyNameOverride() const
@@ -826,6 +830,179 @@ namespace Metasound
 
 		const FText FMetasoundInputDetailCustomization::MemberNameText = LOCTEXT("InputGraphMemberLabel", "Input");
 
+		Frontend::FDocumentHandle FMetasoundInputDetailCustomization::GetDocumentHandle() const
+		{
+			if (MemberDefaultLiteral.IsValid())
+			{
+				if (UMetasoundEditorGraphMember* Parent = MemberDefaultLiteral->GetParentMember())
+				{
+					if (UMetasoundEditorGraph* Graph = Parent->GetOwningGraph())
+					{
+						return Graph->GetDocumentHandle();
+					}
+				}
+			}
+
+			return Frontend::IDocumentController::GetInvalidHandle();
+		}
+
+		bool FMetasoundInputDetailCustomization::GetInputInheritsDefault() const
+		{
+			if (!MemberDefaultLiteral.IsValid())
+			{
+				return false;
+			}
+
+			if (const UMetasoundEditorGraphVertex* Vertex = Cast<UMetasoundEditorGraphVertex>(MemberDefaultLiteral->GetParentMember()))
+			{
+				const TSet<FName>& InputsInheritingDefault = GetDocumentHandle()->GetRootGraph()->GetInputsInheritingDefault();
+				FName NodeName = Vertex->GetConstNodeHandle()->GetNodeName();
+				return InputsInheritingDefault.Contains(NodeName);
+			}
+
+			return false;
+		}
+
+		void FMetasoundInputDetailCustomization::SetInputInheritsDefault()
+		{
+			if (!MemberDefaultLiteral.IsValid())
+			{
+				return;
+			}
+
+			if (UMetasoundEditorGraphVertex* Vertex = Cast<UMetasoundEditorGraphVertex>(MemberDefaultLiteral->GetParentMember()))
+			{
+				FScopedTransaction(LOCTEXT("SetPresetInputOverrideTransaction", "Set MetaSound Preset Input Overridden"));
+
+				Vertex->GetOutermost()->Modify();
+				Vertex->Modify();
+				MemberDefaultLiteral->Modify();
+
+				constexpr bool bDefaultIsInherited = true;
+				const FName NodeName = Vertex->GetConstNodeHandle()->GetNodeName();
+				GetDocumentHandle()->GetRootGraph()->SetInputInheritsDefault(NodeName, bDefaultIsInherited);
+
+				if (UObject* Metasound = Vertex->GetOutermostObject())
+				{
+					FGraphBuilder::RegisterGraphWithFrontend(*Metasound);
+				}
+			}
+		}
+
+		void FMetasoundInputDetailCustomization::ClearInputInheritsDefault()
+		{
+			if (!MemberDefaultLiteral.IsValid())
+			{
+				return;
+			}
+
+			if (UMetasoundEditorGraphVertex* Vertex = Cast<UMetasoundEditorGraphVertex>(MemberDefaultLiteral->GetParentMember()))
+			{
+				FScopedTransaction(LOCTEXT("ClearPresetInputOverrideTransaction", "Clear MetaSound Preset Input Overridden"));
+
+				Vertex->GetOutermost()->Modify();
+				Vertex->Modify();
+				MemberDefaultLiteral->Modify();
+
+				constexpr bool bDefaultIsInherited = false;
+				const FName NodeName = Vertex->GetConstNodeHandle()->GetNodeName();
+				GetDocumentHandle()->GetRootGraph()->SetInputInheritsDefault(NodeName, bDefaultIsInherited);
+
+				Vertex->UpdateFrontendDefaultLiteral(false /* bPostTransaction */);
+
+				if (UMetasoundEditorGraphMemberDefaultLiteral* Literal = Vertex->GetLiteral())
+				{
+					Literal->ForceRefresh();
+				}
+
+				if (UObject* Metasound = Vertex->GetOutermostObject())
+				{
+					FGraphBuilder::RegisterGraphWithFrontend(*Metasound);
+				}
+			}
+		}
+
+		void FMetasoundInputDetailCustomization::CustomizeDetails(IDetailLayoutBuilder& InDetailLayout)
+		{
+			CacheMemberData(InDetailLayout);
+			if (!GraphMember.IsValid())
+			{
+				return;
+			}
+
+			CustomizeGeneralCategory(InDetailLayout);
+
+			if (!MemberDefaultLiteral.IsValid())
+			{
+				return;
+			}
+
+			// Build preset row first if graph has managed interface, not default constructed, & not a trigger
+			const bool bIsPreset = GetDocumentHandle()->GetRootGraphClass().PresetOptions.bIsPreset;
+			const bool bIsDefaultConstructed = MemberDefaultLiteral->GetLiteralType() == EMetasoundFrontendLiteralType::None;
+			const bool bIsTriggerDataType = GraphMember->GetDataType() == GetMetasoundDataTypeName<FTrigger>();
+
+			if (bIsPreset && !bIsDefaultConstructed && !bIsTriggerDataType)
+			{
+				GetDefaultCategoryBuilder(InDetailLayout)
+				.AddCustomRow(MemberCustomizationPrivate::OverrideInputDefaultText)
+				.NameContent()
+				[
+					SNew(STextBlock)
+					.Text(MemberCustomizationPrivate::OverrideInputDefaultText)
+					.Font(IDetailLayoutBuilder::GetDetailFontBold())
+					.ToolTipText(MemberCustomizationPrivate::OverrideInputDefaultTooltip)
+				]
+				.ValueContent()
+				[
+					SNew(SCheckBox)
+					.OnCheckStateChanged_Lambda([this](ECheckBoxState State)
+					{
+						switch(State)
+						{
+							case ECheckBoxState::Checked:
+							{
+								ClearInputInheritsDefault();
+								break;
+							}
+							case ECheckBoxState::Unchecked:
+							case ECheckBoxState::Undetermined:
+							default:
+							{
+								SetInputInheritsDefault();
+							}
+						}
+					})
+					.IsChecked_Lambda([this]() { return GetInputInheritsDefault() ? ECheckBoxState::Unchecked : ECheckBoxState::Checked; })
+					.ToolTipText(MemberCustomizationPrivate::OverrideInputDefaultTooltip)
+				];
+			}
+
+			TArray<IDetailPropertyRow*> DefaultPropertyRows = CustomizeDefaultCategory(InDetailLayout);
+
+			if (bIsPreset && !bIsDefaultConstructed && !bIsTriggerDataType)
+			{
+				const UMetasoundEditorGraphInput* Input = Cast<UMetasoundEditorGraphInput>(MemberDefaultLiteral->GetParentMember());
+				if (ensure(Input))
+				{
+					auto PropertyEnabled = TAttribute<bool>::CreateLambda([this] { return !GetInputInheritsDefault(); });
+					for (IDetailPropertyRow* DefaultPropertyRow : DefaultPropertyRows)
+					{
+						DefaultPropertyRow->EditCondition(PropertyEnabled, { });
+						FResetToDefaultOverride ResetOverride = FResetToDefaultOverride::Create(
+							FIsResetToDefaultVisible::CreateLambda([this](TSharedPtr<IPropertyHandle> /* PropertyHandle */) { return !GetInputInheritsDefault(); }),
+							FResetToDefaultHandler::CreateLambda([this](TSharedPtr<IPropertyHandle> /* PropertyHandle */) { SetInputInheritsDefault(); }));
+						DefaultPropertyRow->OverrideResetToDefault(ResetOverride);
+					}
+				}
+			}
+		}
+
+		bool FMetasoundInputDetailCustomization::IsDefaultEditable() const
+		{
+			return !GetInputInheritsDefault();
+		}
+
 		bool FMetasoundInputDetailCustomization::IsInterfaceMember() const
 		{
 			if (GraphMember.IsValid())
diff --git a/Engine/Plugins/Runtime/Metasound/Source/MetasoundEditor/Private/MetasoundNodeDetailCustomization.h b/Engine/Plugins/Runtime/Metasound/Source/MetasoundEditor/Private/MetasoundNodeDetailCustomization.h
--- a/Engine/Plugins/Runtime/Metasound/Source/MetasoundEditor/Private/MetasoundNodeDetailCustomization.h
+++ b/Engine/Plugins/Runtime/Metasound/Source/MetasoundEditor/Private/MetasoundNodeDetailCustomization.h
@@ -8,6 +8,7 @@
 #include "EdGraph/EdGraphSchema.h"
 #include "Framework/Application/SlateApplication.h"
 #include "IDetailCustomization.h"
+#include "IDetailPropertyRow.h"
 #include "IPropertyTypeCustomization.h"
 #include "Layout/Visibility.h"
 #include "MetasoundAssetBase.h"
@@ -53,9 +54,6 @@ namespace Metasound
 		{
 			TWeakObjectPtr<UMetasoundEditorGraphMemberDefaultFloat> FloatLiteral;
 
-			// Delegate for updating the clamp min/max of the input value when the slider range is changed 
-			FDelegateHandle OnRangeChangedDelegateHandle;
-
 			// Delegate for clamping the input value or not
 			FDelegateHandle OnClampChangedDelegateHandle;
 
@@ -66,7 +64,7 @@ namespace Metasound
 			}
 			virtual ~FMetasoundFloatLiteralCustomization();
 
-			virtual void CustomizeLiteral(UMetasoundEditorGraphMemberDefaultLiteral& InLiteral, IDetailLayoutBuilder& InDetailLayout) override;
+			virtual TArray<IDetailPropertyRow*> CustomizeLiteral(UMetasoundEditorGraphMemberDefaultLiteral& InLiteral, IDetailLayoutBuilder& InDetailLayout) override;
 		};
 
 		// Customization to support drag-and-drop of Proxy UObject types on underlying members that are structs.
@@ -81,7 +79,7 @@ namespace Metasound
 
 			virtual ~FMetasoundObjectArrayLiteralCustomization() = default;
 
-			virtual void CustomizeLiteral(UMetasoundEditorGraphMemberDefaultLiteral& InLiteral, IDetailLayoutBuilder& InDetailLayout) override;
+			virtual TArray<IDetailPropertyRow*> CustomizeLiteral(UMetasoundEditorGraphMemberDefaultLiteral& InLiteral, IDetailLayoutBuilder& InDetailLayout) override;
 		};
 
 		class FMetasoundDefaultLiteralCustomizationFactory : public IMemberDefaultLiteralCustomizationFactory
@@ -193,7 +191,7 @@ namespace Metasound
 		protected:
 			TWeakObjectPtr<TMemberType> GraphMember;
 			TSharedPtr<SEditableTextBox> NameEditableTextBox;
-			TSharedPtr<SEditableTextBox> DisplayNameEditableTextBox;
+// 			TSharedPtr<SEditableTextBox> DisplayNameEditableTextBox;
 			TSharedPtr<FMetasoundDataTypeSelector> DataTypeSelector;
 
 			bool bIsNameInvalid = false;
@@ -226,6 +224,11 @@ namespace Metasound
 				return Name;
 			}
 
+			static IDetailCategoryBuilder& GetDefaultCategoryBuilder(IDetailLayoutBuilder& InDetailLayout)
+			{
+				return InDetailLayout.EditCategory("DefaultValue");
+			}
+
 			void UpdateRenameDelegate(UMetasoundEditorGraphMemberDefaultLiteral& InMemberDefaultLiteral)
 			{
 				if (UMetasoundEditorGraphMember* Member = InMemberDefaultLiteral.GetParentMember())
@@ -244,29 +247,32 @@ namespace Metasound
 				}
 			}
 
-			virtual bool IsInterfaceMember() const
-			{
-				return false;
-			}
-
-			// IDetailCustomization interface
-			virtual void CustomizeDetails(IDetailLayoutBuilder& InDetailLayout) override
+			void CacheMemberData(IDetailLayoutBuilder& InDetailLayout)
 			{
-				using namespace Frontend;
-
 				TArray<TWeakObjectPtr<UObject>> Objects;
 				InDetailLayout.GetObjectsBeingCustomized(Objects);
-				if (Objects.IsEmpty())
+				if (!Objects.IsEmpty())
 				{
-					return;
-				}
+					GraphMember = Cast<TMemberType>(Objects.Last().Get());
 
-				GraphMember = Cast<TMemberType>(Objects.Last().Get());
-				if (!GraphMember.IsValid())
-				{
-					return;
+					TSharedPtr<IPropertyHandle> LiteralHandle = InDetailLayout.GetProperty(GET_MEMBER_NAME_CHECKED(TMemberType, Literal));
+					if (ensure(GraphMember.IsValid()) && ensure(LiteralHandle.IsValid()))
+					{
+						// Always hide, even if no customization (LiteralObject isn't found) as this is the case
+						// where the default object is not required (i.e. Default Member is default constructed)
+						LiteralHandle->MarkHiddenByCustomization();
+
+						UObject* LiteralObject = nullptr;
+						if (LiteralHandle->GetValue(LiteralObject) == FPropertyAccess::Success)
+						{
+							MemberDefaultLiteral = Cast<UMetasoundEditorGraphMemberDefaultLiteral>(LiteralObject);
+						}
+					}
 				}
+			}
 
+			void CustomizeGeneralCategory(IDetailLayoutBuilder& InDetailLayout)
+			{
 				IDetailCategoryBuilder& CategoryBuilder = InDetailLayout.EditCategory("General");
 
 				const bool bIsInterfaceMember = IsInterfaceMember();
@@ -280,11 +286,11 @@ namespace Metasound
 					.SelectAllTextWhenFocused(true)
 					.Font(IDetailLayoutBuilder::GetDetailFont());
 
-				DisplayNameEditableTextBox = SNew(SEditableTextBox)
-					.Text(this, &TChildClass::GetDisplayName)
-					.OnTextCommitted(this, &TChildClass::OnDisplayNameCommitted)
-					.IsReadOnly(bIsInterfaceMember || !bIsGraphEditable)
-					.Font(IDetailLayoutBuilder::GetDetailFont());
+				// 				DisplayNameEditableTextBox = SNew(SEditableTextBox)
+				// 					.Text(this, &TChildClass::GetDisplayName)
+				// 					.OnTextCommitted(this, &TChildClass::OnDisplayNameCommitted)
+				// 					.IsReadOnly(bIsInterfaceMember || !bIsGraphEditable)
+				// 					.Font(IDetailLayoutBuilder::GetDetailFont());
 
 				CategoryBuilder.AddCustomRow(TChildClass::MemberNameText)
 					.EditCondition(!bIsInterfaceMember && bIsGraphEditable, nullptr)
@@ -292,8 +298,8 @@ namespace Metasound
 					[
 						SNew(STextBlock)
 						.Font(IDetailLayoutBuilder::GetDetailFontBold())
-					.Text(TChildClass::MemberNameText)
-					.ToolTipText(GetMemberNameDescriptionText())
+						.Text(TChildClass::MemberNameText)
+						.ToolTipText(GetMemberNameDescriptionText())
 					]
 					.ValueContent()
 					[
@@ -316,60 +322,77 @@ namespace Metasound
 // 				];
 
 				CategoryBuilder.AddCustomRow(MemberCustomizationStyle::NodeTooltipText)
-				.EditCondition(!bIsInterfaceMember && bIsGraphEditable, nullptr)
-				.NameContent()
-				[
-					SNew(STextBlock)
-					.Font(IDetailLayoutBuilder::GetDetailFontBold())
-				.Text(MemberCustomizationStyle::NodeTooltipText)
-				]
-				.ValueContent()
-				[
-					SNew(SMultiLineEditableTextBox)
-					.Text(this, &TChildClass::GetTooltip)
-					.OnTextCommitted(this, &TChildClass::OnTooltipCommitted)
-					.IsReadOnly(bIsInterfaceMember || !bIsGraphEditable)
-					.ModiferKeyForNewLine(EModifierKey::Shift)
-					.RevertTextOnEscape(true)
-					.WrapTextAt(MemberCustomizationStyle::DetailsTitleMaxWidth - MemberCustomizationStyle::DetailsTitleWrapPadding)
-					.Font(IDetailLayoutBuilder::GetDetailFont())
-				];
+					.EditCondition(!bIsInterfaceMember && bIsGraphEditable, nullptr)
+					.NameContent()
+					[
+						SNew(STextBlock)
+						.Font(IDetailLayoutBuilder::GetDetailFontBold())
+					.Text(MemberCustomizationStyle::NodeTooltipText)
+					]
+					.ValueContent()
+					[
+						SNew(SMultiLineEditableTextBox)
+						.Text(this, &TChildClass::GetTooltip)
+						.OnTextCommitted(this, &TChildClass::OnTooltipCommitted)
+						.IsReadOnly(bIsInterfaceMember || !bIsGraphEditable)
+						.ModiferKeyForNewLine(EModifierKey::Shift)
+						.RevertTextOnEscape(true)
+						.WrapTextAt(MemberCustomizationStyle::DetailsTitleMaxWidth - MemberCustomizationStyle::DetailsTitleWrapPadding)
+						.Font(IDetailLayoutBuilder::GetDetailFont())
+					];
 
 				DataTypeSelector->AddDataTypeSelector(InDetailLayout, MemberCustomizationStyle::DataTypeNameText, GraphMember, !bIsInterfaceMember && bIsGraphEditable);
+			}
 
-				IDetailCategoryBuilder& DefaultCategoryBuilder = InDetailLayout.EditCategory("DefaultValue");
-				TSharedPtr<IPropertyHandle> LiteralHandle = InDetailLayout.GetProperty(GET_MEMBER_NAME_CHECKED(TMemberType, Literal));
-				if (ensure(GraphMember.IsValid()) && ensure(LiteralHandle.IsValid()))
+			TArray<IDetailPropertyRow*> CustomizeDefaultCategory(IDetailLayoutBuilder& InDetailLayout)
+			{
+				TArray<IDetailPropertyRow*> DefaultPropertyRows;
+
+				if (MemberDefaultLiteral.IsValid())
 				{
-					// Always hide, even if no customization (LiteralObject isn't found) as this is the case
-					// where the default object is not required (i.e. Default Member is default constructed)
-					LiteralHandle->MarkHiddenByCustomization();
+					UpdateRenameDelegate(*MemberDefaultLiteral);
+
+					UClass* MemberClass = MemberDefaultLiteral->GetClass();
+					check(MemberClass);
 
-					UObject* LiteralObject = nullptr;
-					if (LiteralHandle->GetValue(LiteralObject) == FPropertyAccess::Success)
+					IDetailCategoryBuilder& DefaultCategoryBuilder = GetDefaultCategoryBuilder(InDetailLayout);
+					IMetasoundEditorModule& EditorModule = FModuleManager::GetModuleChecked<IMetasoundEditorModule>("MetaSoundEditor");
+					TUniquePtr<FMetasoundDefaultLiteralCustomizationBase> LiteralCustomization = EditorModule.CreateMemberDefaultLiteralCustomization(*MemberClass, DefaultCategoryBuilder);
+					if (LiteralCustomization.IsValid())
 					{
-						if (UMetasoundEditorGraphMemberDefaultLiteral * MemberDefaultLiteral = Cast<UMetasoundEditorGraphMemberDefaultLiteral>(LiteralObject))
-						{
-							UpdateRenameDelegate(*MemberDefaultLiteral);
+						DefaultPropertyRows = LiteralCustomization->CustomizeLiteral(*MemberDefaultLiteral, InDetailLayout);
+					}
+					else
+					{
+						IDetailPropertyRow* DefaultPropertyRow = DefaultCategoryBuilder.AddExternalObjectProperty(TArray<UObject*>({ MemberDefaultLiteral.Get() }), "Default");
+						ensureMsgf(DefaultPropertyRow, TEXT("Class '%s' missing expected 'Default' member."
+							"Either add/rename default member or register customization to display default value/opt out appropriately."),
+							*MemberClass->GetName());
+						DefaultPropertyRows.Add(DefaultPropertyRow);
+					}
+				}
 
-							UClass* MemberClass = MemberDefaultLiteral->GetClass();
-							check(MemberClass);
+				return DefaultPropertyRows;
+			}
 
-							IMetasoundEditorModule& EditorModule = FModuleManager::GetModuleChecked<IMetasoundEditorModule>("MetaSoundEditor");
-							LiteralCustomization = EditorModule.CreateMemberDefaultLiteralCustomization(*MemberClass, DefaultCategoryBuilder);
-							if (LiteralCustomization.IsValid())
-							{
-								LiteralCustomization->CustomizeLiteral(*MemberDefaultLiteral, InDetailLayout);
-							}
-							else
-							{
-								IDetailPropertyRow* Row = DefaultCategoryBuilder.AddExternalObjectProperty(TArray<UObject*>({ MemberDefaultLiteral }), "Default");
-								ensureMsgf(Row, TEXT("Class '%s' missing expected 'Default' member."
-									"Either add/rename default member or register customization to display default value/opt out appropriately."),
-									*MemberClass->GetName());
-							}
-						}
-					}
+			virtual bool IsDefaultEditable() const
+			{
+				return true;
+			}
+
+			virtual bool IsInterfaceMember() const
+			{
+				return false;
+			}
+
+			// IDetailCustomization interface
+			virtual void CustomizeDetails(IDetailLayoutBuilder& InDetailLayout) override
+			{
+				CacheMemberData(InDetailLayout);
+				if (GraphMember.IsValid())
+				{
+					CustomizeGeneralCategory(InDetailLayout);
+					CustomizeDefaultCategory(InDetailLayout);
 				}
 			}
 			// End of IDetailCustomization interface
@@ -379,7 +402,8 @@ namespace Metasound
 				using namespace Frontend;
 
 				bIsNameInvalid = false;
-				DisplayNameEditableTextBox->SetError(FText::GetEmpty());
+				NameEditableTextBox->SetError(FText::GetEmpty());
+// 				DisplayNameEditableTextBox->SetError(FText::GetEmpty());
 
 				if (!ensure(GraphMember.IsValid()))
 				{
@@ -390,7 +414,8 @@ namespace Metasound
 				if (!GraphMember->CanRename(InNewName, Error))
 				{
 					bIsNameInvalid = true;
-					DisplayNameEditableTextBox->SetError(Error);
+					NameEditableTextBox->SetError(Error);
+// 					DisplayNameEditableTextBox->SetError(Error);
 				}
 			}
 
@@ -429,19 +454,19 @@ namespace Metasound
 				return FText::GetEmpty();
 			}
 
-			void OnDisplayNameCommitted(const FText& InNewName, ETextCommit::Type InTextCommit)
-			{
-				using namespace Frontend;
-
-				if (!bIsNameInvalid && GraphMember.IsValid())
-				{
-					constexpr bool bPostTransaction = true;
-					GraphMember->SetDisplayName(FText::GetEmpty(), bPostTransaction);
-				}
-
-				DisplayNameEditableTextBox->SetError(FText::GetEmpty());
-				bIsNameInvalid = false;
-			}
+// 			void OnDisplayNameCommitted(const FText& InNewName, ETextCommit::Type InTextCommit)
+// 			{
+// 				using namespace Frontend;
+// 
+// 				if (!bIsNameInvalid && GraphMember.IsValid())
+// 				{
+// 					constexpr bool bPostTransaction = true;
+// 					GraphMember->SetDisplayName(FText::GetEmpty(), bPostTransaction);
+// 				}
+// 
+// 				DisplayNameEditableTextBox->SetError(FText::GetEmpty());
+// 				bIsNameInvalid = false;
+// 			}
 
 			void OnTooltipCommitted(const FText& InNewText, ETextCommit::Type InTextCommit)
 			{
@@ -478,7 +503,7 @@ namespace Metasound
 					GraphMember->SetMemberName(*InNewName.ToString(), bPostTransaction);
 				}
 
-				DisplayNameEditableTextBox->SetError(FText::GetEmpty());
+				NameEditableTextBox->SetError(FText::GetEmpty());
 				bIsNameInvalid = false;
 			}
 
@@ -492,8 +517,8 @@ namespace Metasound
 				return ECheckBoxState::Unchecked;
 			}
 
-		private:
-			TUniquePtr<FMetasoundDefaultLiteralCustomizationBase> LiteralCustomization;
+		protected:
+			TWeakObjectPtr<UMetasoundEditorGraphMemberDefaultLiteral> MemberDefaultLiteral;
 
 			FDelegateHandle RenameRequestedHandle;
 		};
@@ -508,12 +533,16 @@ namespace Metasound
 
 			virtual ~FMetasoundInputDetailCustomization() = default;
 
+			virtual void CustomizeDetails(IDetailLayoutBuilder& InDetailLayout) override;
+			virtual bool IsDefaultEditable() const override;
 			virtual bool IsInterfaceMember() const override;
-
 			static const FText MemberNameText;
 
 		private:
-			TUniquePtr<FMetasoundDefaultLiteralCustomizationBase> LiteralCustomization;
+			Frontend::FDocumentHandle GetDocumentHandle() const;
+			bool GetInputInheritsDefault() const;
+			void SetInputInheritsDefault();
+			void ClearInputInheritsDefault();
 		};
 
 		class FMetasoundOutputDetailCustomization : public TMetasoundGraphMemberDetailCustomization<UMetasoundEditorGraphOutput, FMetasoundOutputDetailCustomization>
@@ -529,9 +558,6 @@ namespace Metasound
 			virtual ~FMetasoundOutputDetailCustomization() = default;
 
 			virtual bool IsInterfaceMember() const override;
-
-		private:
-			TUniquePtr<FMetasoundDefaultLiteralCustomizationBase> LiteralCustomization;
 		};
 	} // namespace Editor
 } // namespace Metasound
diff --git a/Engine/Plugins/Runtime/Metasound/Source/MetasoundEditor/Public/MetasoundEditorGraph.h b/Engine/Plugins/Runtime/Metasound/Source/MetasoundEditor/Public/MetasoundEditorGraph.h
--- a/Engine/Plugins/Runtime/Metasound/Source/MetasoundEditor/Public/MetasoundEditorGraph.h
+++ b/Engine/Plugins/Runtime/Metasound/Source/MetasoundEditor/Public/MetasoundEditorGraph.h
@@ -60,6 +60,10 @@ public:
 	{
 	}
 
+	virtual void ForceRefresh()
+	{
+	}
+
 #if WITH_EDITOR
 	virtual void PostEditChangeProperty(FPropertyChangedEvent& InPropertyChangedEvent) override;
 	virtual void PostEditChangeChainProperty(FPropertyChangedChainEvent& InPropertyChangedEvent) override;
diff --git a/Engine/Plugins/Runtime/Metasound/Source/MetasoundEditor/Public/MetasoundEditorGraphMemberDefaults.h b/Engine/Plugins/Runtime/Metasound/Source/MetasoundEditor/Public/MetasoundEditorGraphMemberDefaults.h
--- a/Engine/Plugins/Runtime/Metasound/Source/MetasoundEditor/Public/MetasoundEditorGraphMemberDefaults.h
+++ b/Engine/Plugins/Runtime/Metasound/Source/MetasoundEditor/Public/MetasoundEditorGraphMemberDefaults.h
@@ -175,6 +175,7 @@ public:
 	FOnMetasoundRangeChangedEvent OnRangeChanged;
 	FOnMetasoundInputClampDefaultChangedEvent OnClampChanged;
 
+	virtual void ForceRefresh() override;
 	virtual FMetasoundFrontendLiteral GetDefault() const override;
 	virtual EMetasoundFrontendLiteralType GetLiteralType() const override;
 	virtual void SetFromLiteral(const FMetasoundFrontendLiteral& InLiteral) override;
diff --git a/Engine/Plugins/Runtime/Metasound/Source/MetasoundEditor/Public/MetasoundEditorModule.h b/Engine/Plugins/Runtime/Metasound/Source/MetasoundEditor/Public/MetasoundEditorModule.h
--- a/Engine/Plugins/Runtime/Metasound/Source/MetasoundEditor/Public/MetasoundEditorModule.h
+++ b/Engine/Plugins/Runtime/Metasound/Source/MetasoundEditor/Public/MetasoundEditorModule.h
@@ -13,6 +13,7 @@
 
 // Forward Declarations
 class IDetailLayoutBuilder;
+class IDetailPropertyRow;
 class UMetasoundEditorGraph;
 class UMetasoundEditorGraphMemberDefaultLiteral;
 
@@ -50,7 +51,9 @@ namespace Metasound
 
 			virtual ~FMetasoundDefaultLiteralCustomizationBase() = default;
 
-			virtual void CustomizeLiteral(UMetasoundEditorGraphMemberDefaultLiteral& InLiteral, IDetailLayoutBuilder& InDetailLayout) { }
+			// Customizes the given literal for the provided DetailLayoutBuilder.
+			// @return the DetailPropertyRow created for the default parameter set by this customization.
+			virtual TArray<IDetailPropertyRow*> CustomizeLiteral(UMetasoundEditorGraphMemberDefaultLiteral& InLiteral, IDetailLayoutBuilder& InDetailLayout) { return { }; };
 		};
 
 		class METASOUNDEDITOR_API IMemberDefaultLiteralCustomizationFactory
diff --git a/Engine/Plugins/Runtime/Metasound/Source/MetasoundFrontend/Private/MetasoundAssetBase.cpp b/Engine/Plugins/Runtime/Metasound/Source/MetasoundFrontend/Private/MetasoundAssetBase.cpp
--- a/Engine/Plugins/Runtime/Metasound/Source/MetasoundFrontend/Private/MetasoundAssetBase.cpp
+++ b/Engine/Plugins/Runtime/Metasound/Source/MetasoundFrontend/Private/MetasoundAssetBase.cpp
@@ -623,9 +623,9 @@ void FMetasoundAssetBase::ConvertFromPreset()
 	Style.bIsGraphEditable = true;
 	GraphHandle->SetGraphStyle(Style);
 
-	FMetasoundFrontendClassMetadata Metadata = GraphHandle->GetGraphMetadata();
-	Metadata.SetAutoUpdateManagesInterface(false);
-	GraphHandle->SetGraphMetadata(Metadata);
+	FMetasoundFrontendGraphClassPresetOptions PresetOptions = GraphHandle->GetGraphPresetOptions();
+	PresetOptions.bIsPreset = false;
+	GraphHandle->SetGraphPresetOptions(PresetOptions);
 }
 
 TArray<FMetasoundAssetBase::FSendInfoAndVertexName> FMetasoundAssetBase::GetSendInfos(uint64 InInstanceID) const
diff --git a/Engine/Plugins/Runtime/Metasound/Source/MetasoundFrontend/Private/MetasoundFrontendDocument.cpp b/Engine/Plugins/Runtime/Metasound/Source/MetasoundFrontend/Private/MetasoundFrontendDocument.cpp
--- a/Engine/Plugins/Runtime/Metasound/Source/MetasoundFrontend/Private/MetasoundFrontendDocument.cpp
+++ b/Engine/Plugins/Runtime/Metasound/Source/MetasoundFrontend/Private/MetasoundFrontendDocument.cpp
@@ -289,12 +289,6 @@ void FMetasoundFrontendClassMetadata::SetAuthor(const FText& InAuthor)
 	SetWithChangeID(InAuthor, TextToSet, ChangeID);
 }
 
-void FMetasoundFrontendClassMetadata::SetAutoUpdateManagesInterface(bool bInAutoUpdateManagesInterface)
-{
-	using namespace Metasound::DocumentPrivate;
-	SetWithChangeID(bInAutoUpdateManagesInterface, bAutoUpdateManagesInterface, ChangeID);
-}
-
 void FMetasoundFrontendClassMetadata::SetCategoryHierarchy(const TArray<FText>& InCategoryHierarchy)
 {
 	using namespace Metasound::DocumentPrivate;
diff --git a/Engine/Plugins/Runtime/Metasound/Source/MetasoundFrontend/Private/MetasoundFrontendGraphController.cpp b/Engine/Plugins/Runtime/Metasound/Source/MetasoundFrontend/Private/MetasoundFrontendGraphController.cpp
--- a/Engine/Plugins/Runtime/Metasound/Source/MetasoundFrontend/Private/MetasoundFrontendGraphController.cpp
+++ b/Engine/Plugins/Runtime/Metasound/Source/MetasoundFrontend/Private/MetasoundFrontendGraphController.cpp
@@ -663,6 +663,7 @@ namespace Metasound
 				GraphClass->Graph.Edges.Reset();
 				GraphClass->Interface.Inputs.Reset();
 				GraphClass->Interface.Outputs.Reset();
+				GraphClass->PresetOptions.InputsInheritingDefault.Reset();
 				OwningDocument->ClearInterfaceVersions();
 				OwningDocument->RemoveUnreferencedDependencies();
 			}
@@ -734,6 +735,54 @@ namespace Metasound
 			return GetNodesByPredicate(IsInputNode);
 		}
 
+		const TSet<FName>& FGraphController::GetInputsInheritingDefault() const
+		{
+			if (const FMetasoundFrontendGraphClass* GraphClass = GraphClassPtr.Get())
+			{
+				if (GraphClass->PresetOptions.bIsPreset)
+				{
+					return GraphClass->PresetOptions.InputsInheritingDefault;
+				}
+			}
+
+			return Invalid::GetInvalidNameSet();
+		}
+
+		bool FGraphController::SetInputInheritsDefault(FName InName, bool bInputInheritsDefault)
+		{
+			if (bInputInheritsDefault)
+			{
+				if (FMetasoundFrontendGraphClass* GraphClass = GraphClassPtr.Get())
+				{
+					if (GraphClass->PresetOptions.bIsPreset)
+					{
+						return GraphClass->PresetOptions.InputsInheritingDefault.Add(InName).IsValidId();
+					}
+				}
+			}
+			else
+			{
+				if (FMetasoundFrontendGraphClass* GraphClass = GraphClassPtr.Get())
+				{
+					if (GraphClass->PresetOptions.bIsPreset)
+					{
+						return GraphClass->PresetOptions.InputsInheritingDefault.Remove(InName) > 0;
+					}
+				}
+			}
+
+			return false;
+		}
+
+		void FGraphController::SetInputsInheritingDefault(TSet<FName>&& InNames)
+		{
+			if (FMetasoundFrontendGraphClass* GraphClass = GraphClassPtr.Get())
+			{
+				GraphClass->PresetOptions.bIsPreset = true;
+				GraphClass->PresetOptions.InputsInheritingDefault = MoveTemp(InNames);
+			}
+		}
+
 		bool FGraphController::ContainsOutputVertex(const FVertexName& InName, const FName& InTypeName) const
 		{
 			if (const FMetasoundFrontendGraphClass* GraphClass = GraphClassPtr.Get())
@@ -1278,7 +1327,23 @@ namespace Metasound
 			return false;
 		}
 
-		// Returns the metadata for the current graph, including the name, description and author.
+		const FMetasoundFrontendGraphClassPresetOptions& FGraphController::GetGraphPresetOptions() const
+		{
+			if (const FMetasoundFrontendGraphClass* GraphClass = GraphClassPtr.Get())
+			{
+				return GraphClass->PresetOptions;
+			}
+			return Invalid::GetInvalidGraphClassPresetOptions();
+		}
+
+		void FGraphController::SetGraphPresetOptions(const FMetasoundFrontendGraphClassPresetOptions& InPresetOptions)
+		{
+			if (FMetasoundFrontendGraphClass* GraphClass = GraphClassPtr.Get())
+			{
+				GraphClass->PresetOptions = InPresetOptions;
+			}
+		}
+
 		const FMetasoundFrontendClassMetadata& FGraphController::GetGraphMetadata() const
 		{
 			if (const FMetasoundFrontendGraphClass* GraphClass = GraphClassPtr.Get())
diff --git a/Engine/Plugins/Runtime/Metasound/Source/MetasoundFrontend/Private/MetasoundFrontendGraphController.h b/Engine/Plugins/Runtime/Metasound/Source/MetasoundFrontend/Private/MetasoundFrontendGraphController.h
--- a/Engine/Plugins/Runtime/Metasound/Source/MetasoundFrontend/Private/MetasoundFrontendGraphController.h
+++ b/Engine/Plugins/Runtime/Metasound/Source/MetasoundFrontend/Private/MetasoundFrontendGraphController.h
@@ -69,6 +69,10 @@ namespace Metasound
 			TArray<FNodeHandle> GetInputNodes() override;
 			TArray<FConstNodeHandle> GetConstInputNodes() const override;
 
+			virtual const TSet<FName>& GetInputsInheritingDefault() const override;
+			virtual bool SetInputInheritsDefault(FName InName, bool bInputInheritsDefault) override;
+			virtual void SetInputsInheritingDefault(TSet<FName>&& InNames) override;
+
 			virtual FVariableHandle AddVariable(const FName& InDataType) override;
 			virtual FVariableHandle FindVariable(const FGuid& InVariableID) override;
 			virtual FConstVariableHandle FindVariable(const FGuid& InVariableID) const override;
@@ -157,9 +161,10 @@ namespace Metasound
 			// On success, invalidates the received node handle.
 			bool RemoveNode(INodeController& InNode) override;
 
-			// Returns the metadata for the current graph, including the name, description and author.
-			const FMetasoundFrontendClassMetadata& GetGraphMetadata() const override;
+			const FMetasoundFrontendGraphClassPresetOptions& GetGraphPresetOptions() const override;
+			void SetGraphPresetOptions(const FMetasoundFrontendGraphClassPresetOptions& InPresetOptions) override;
 
+			const FMetasoundFrontendClassMetadata& GetGraphMetadata() const override;
 			void SetGraphMetadata(const FMetasoundFrontendClassMetadata& InClassMetadata) override;
 
 			FNodeHandle CreateEmptySubgraph(const FMetasoundFrontendClassMetadata& InInfo) override;
diff --git a/Engine/Plugins/Runtime/Metasound/Source/MetasoundFrontend/Private/MetasoundFrontendInvalidController.cpp b/Engine/Plugins/Runtime/Metasound/Source/MetasoundFrontend/Private/MetasoundFrontendInvalidController.cpp
--- a/Engine/Plugins/Runtime/Metasound/Source/MetasoundFrontend/Private/MetasoundFrontendInvalidController.cpp
+++ b/Engine/Plugins/Runtime/Metasound/Source/MetasoundFrontend/Private/MetasoundFrontendInvalidController.cpp
@@ -57,6 +57,12 @@ namespace Metasound
 				return Invalid;
 			}
 
+			const FMetasoundFrontendGraphClassPresetOptions& GetInvalidGraphClassPresetOptions()
+			{
+				static const FMetasoundFrontendGraphClassPresetOptions Invalid;
+				return Invalid;
+			}
+
 			const FMetasoundFrontendInterfaceStyle& GetInvalidInterfaceStyle()
 			{
 				static const FMetasoundFrontendInterfaceStyle Invalid;
@@ -93,6 +99,12 @@ namespace Metasound
 				return Invalid;
 			}
 
+			const TSet<FName>& GetInvalidNameSet()
+			{
+				static TSet<FName> Invalid;
+				return Invalid;
+			}
+
 			const FMetasoundFrontendDocumentMetadata& GetInvalidDocumentMetadata()
 			{
 				static const FMetasoundFrontendDocumentMetadata Invalid;
diff --git a/Engine/Plugins/Runtime/Metasound/Source/MetasoundFrontend/Private/MetasoundFrontendInvalidController.h b/Engine/Plugins/Runtime/Metasound/Source/MetasoundFrontend/Private/MetasoundFrontendInvalidController.h
--- a/Engine/Plugins/Runtime/Metasound/Source/MetasoundFrontend/Private/MetasoundFrontendInvalidController.h
+++ b/Engine/Plugins/Runtime/Metasound/Source/MetasoundFrontend/Private/MetasoundFrontendInvalidController.h
@@ -28,12 +28,14 @@ namespace Metasound
 			const FMetasoundFrontendLiteral& GetInvalidLiteral();
 			const FMetasoundFrontendClassInterface& GetInvalidClassInterface();
 			const FMetasoundFrontendClassMetadata& GetInvalidClassMetadata();
+			const FMetasoundFrontendGraphClassPresetOptions& GetInvalidGraphClassPresetOptions();
 			const FMetasoundFrontendInterfaceStyle& GetInvalidInterfaceStyle();
 			const FMetasoundFrontendClassStyle& GetInvalidClassStyle();
 			const FMetasoundFrontendGraphStyle& GetInvalidGraphStyle();
 			const FMetasoundFrontendGraphClass& GetInvalidGraphClass();
 			const TArray<FMetasoundFrontendClass>& GetInvalidClassArray();
 			const TArray<FMetasoundFrontendGraphClass>& GetInvalidGraphClassArray();
+			const TSet<FName>& GetInvalidNameSet();
 			const FMetasoundFrontendDocumentMetadata& GetInvalidDocumentMetadata();
 		}
 
@@ -286,6 +288,10 @@ namespace Metasound
 			virtual TArray<FConstNodeHandle> GetConstOutputNodes() const override { return TArray<FConstNodeHandle>(); }
 			virtual TArray<FConstNodeHandle> GetConstInputNodes() const override { return TArray<FConstNodeHandle>(); }
 
+			virtual const TSet<FName>& GetInputsInheritingDefault() const { return Invalid::GetInvalidNameSet(); }
+			virtual bool SetInputInheritsDefault(FName InName, bool bDefaultIsInherited) { return false; }
+			virtual void SetInputsInheritingDefault(TSet<FName>&& InNames) { }
+
 			virtual FVariableHandle AddVariable(const FName& InDataTypeName) override { return IVariableController::GetInvalidHandle(); }
 			virtual FVariableHandle FindVariable(const FGuid& InVariableID) override { return IVariableController::GetInvalidHandle(); }
 			virtual FConstVariableHandle FindVariable(const FGuid& InVariableID) const { return IVariableController::GetInvalidHandle(); }
@@ -364,6 +370,11 @@ namespace Metasound
 			// On success, invalidates the received node handle.
 			virtual bool RemoveNode(INodeController& InNode) override { return false; }
 
+			// Returns the metadata for the current graph, including the name, description and author.
+			virtual const FMetasoundFrontendGraphClassPresetOptions& GetGraphPresetOptions() const override { return Invalid::GetInvalidGraphClassPresetOptions(); }
+
+			virtual void SetGraphPresetOptions(const FMetasoundFrontendGraphClassPresetOptions& InMetadata) override { }
+
 			// Returns the metadata for the current graph, including the name, description and author.
 			virtual const FMetasoundFrontendClassMetadata& GetGraphMetadata() const override { return Invalid::GetInvalidClassMetadata(); }
 
diff --git a/Engine/Plugins/Runtime/Metasound/Source/MetasoundFrontend/Private/MetasoundFrontendTransform.cpp b/Engine/Plugins/Runtime/Metasound/Source/MetasoundFrontend/Private/MetasoundFrontendTransform.cpp
--- a/Engine/Plugins/Runtime/Metasound/Source/MetasoundFrontend/Private/MetasoundFrontendTransform.cpp
+++ b/Engine/Plugins/Runtime/Metasound/Source/MetasoundFrontend/Private/MetasoundFrontendTransform.cpp
@@ -462,14 +462,14 @@ namespace Metasound
 			TArray<TPair<FNodeHandle, FMetasoundFrontendVersionNumber>> NodesToUpdate;
 
 			FGraphHandle RootGraph = InDocument->GetRootGraph();
-			const bool bInterfaceManaged = RootGraph->GetGraphMetadata().GetAutoUpdateManagesInterface();
+			const bool bIsPreset = RootGraph->GetGraphPresetOptions().bIsPreset;
 			RootGraph->IterateNodes([&](FNodeHandle NodeHandle)
 			{
 				using namespace Metasound::Frontend;
 
 				const FMetasoundFrontendClassMetadata& ClassMetadata = NodeHandle->GetClassMetadata();
 
-				if (bInterfaceManaged)
+				if (bIsPreset)
 				{
 					const FNodeRegistryKey RegistryKey = FMetasoundFrontendRegistryContainer::Get()->GetRegistryKey(ClassMetadata);
 					PresetReferencedMetaSoundAsset = IMetaSoundAssetManager::GetChecked().TryLoadAssetFromKey(RegistryKey);
@@ -513,18 +513,18 @@ namespace Metasound
 
 			if (PresetReferencedMetaSoundAsset)
 			{
-				if (bInterfaceManaged)
+				if (bIsPreset)
 				{
 					bDidEdit |= FRebuildPresetRootGraph(PresetReferencedMetaSoundAsset->GetDocumentHandle()).Transform(InDocument);
 					if (bDidEdit)
 					{
-						FMetasoundFrontendClassMetadata ParentMetadata = InDocument->GetRootGraphClass().Metadata;
-						ParentMetadata.SetType(EMetasoundFrontendClassType::External);
-						const FNodeRegistryKey RegistryKey = FMetasoundFrontendRegistryContainer::Get()->GetRegistryKey(ParentMetadata);
-						FMetasoundAssetBase* ParentMetaSoundAsset = IMetaSoundAssetManager::GetChecked().TryLoadAssetFromKey(RegistryKey);
-						if (ensure(ParentMetaSoundAsset))
+						FMetasoundFrontendClassMetadata PresetMetadata = InDocument->GetRootGraphClass().Metadata;
+						PresetMetadata.SetType(EMetasoundFrontendClassType::External);
+						const FNodeRegistryKey RegistryKey = FMetasoundFrontendRegistryContainer::Get()->GetRegistryKey(PresetMetadata);
+						FMetasoundAssetBase* PresetMetaSoundAsset = IMetaSoundAssetManager::GetChecked().TryLoadAssetFromKey(RegistryKey);
+						if (ensure(PresetMetaSoundAsset))
 						{
-							ParentMetaSoundAsset->ConformObjectDataToInterfaces();
+							PresetMetaSoundAsset->ConformObjectDataToInterfaces();
 						}
 
 						InDocument->RemoveUnreferencedDependencies();
@@ -554,8 +554,8 @@ namespace Metasound
 
 		bool FRebuildPresetRootGraph::Transform(FDocumentHandle InDocument) const
 		{
-			FGraphHandle RootGraphHandle = InDocument->GetRootGraph();
-			if (!ensure(RootGraphHandle->IsValid()))
+			FGraphHandle PresetGraphHandle = InDocument->GetRootGraph();
+			if (!ensure(PresetGraphHandle->IsValid()))
 			{
 				return false;
 			}
@@ -565,7 +565,7 @@ namespace Metasound
 			// arises where this transform is used outside of AutoUpdate, then this
 			// early exist should be removed as it's mostly here to protect against
 			// accidental manipulation of metasound graphs.
-			if (!ensure(RootGraphHandle->GetGraphMetadata().GetAutoUpdateManagesInterface()))
+			if (!ensure(PresetGraphHandle->GetGraphPresetOptions().bIsPreset))
 			{
 				return false;
 			}
@@ -578,17 +578,18 @@ namespace Metasound
 
 			// Determine the inputs and outputs needed in the wrapping graph. Also
 			// cache any exiting literals that have been set on the wrapping graph.
-			TArray<FMetasoundFrontendClassInput> ClassInputs = GenerateRequiredClassInputs(RootGraphHandle);
-			TArray<FMetasoundFrontendClassOutput> ClassOutputs = GenerateRequiredClassOutputs(RootGraphHandle);
+			TSet<FName> InputsInheritingDefault;
+			TArray<FMetasoundFrontendClassInput> ClassInputs = GenerateRequiredClassInputs(PresetGraphHandle, InputsInheritingDefault);
+			TArray<FMetasoundFrontendClassOutput> ClassOutputs = GenerateRequiredClassOutputs(PresetGraphHandle);
 
 			FGuid PresetNodeID;
-			RootGraphHandle->IterateConstNodes([InPresetNodeID = &PresetNodeID](FConstNodeHandle PresetNodeHandle)
+			PresetGraphHandle->IterateConstNodes([InPresetNodeID = &PresetNodeID](FConstNodeHandle PresetNodeHandle)
 			{
 				*InPresetNodeID = PresetNodeHandle->GetID();
 			}, EMetasoundFrontendClassType::External);
 
 			// Clear the root graph so it can be rebuilt.
-			RootGraphHandle->ClearGraph();
+			PresetGraphHandle->ClearGraph();
 
 			// Ensure preset interfaces match those found in referenced graph.  Referenced graph is assumed to be
 			// well-formed (i.e. all inputs/outputs/environment variables declared by interfaces are present, and
@@ -605,7 +606,7 @@ namespace Metasound
 			ReferencedClassMetadata.SetType(EMetasoundFrontendClassType::External);
 
 			// Set node location.
-			FNodeHandle ReferencedNodeHandle = RootGraphHandle->AddNode(ReferencedClassMetadata, PresetNodeID);
+			FNodeHandle ReferencedNodeHandle = PresetGraphHandle->AddNode(ReferencedClassMetadata, PresetNodeID);
 			FMetasoundFrontendNodeStyle RefNodeStyle;
 
 			// Offset to be to the right of input nodes
@@ -613,20 +614,22 @@ namespace Metasound
 			ReferencedNodeHandle->SetNodeStyle(RefNodeStyle);
 
 			// Connect parent graph to referenced graph
-			AddAndConnectInputs(ClassInputs, RootGraphHandle, ReferencedNodeHandle);
-			AddAndConnectOutputs(ClassOutputs, RootGraphHandle, ReferencedNodeHandle);
+			PresetGraphHandle->SetInputsInheritingDefault(MoveTemp(InputsInheritingDefault));
+
+			AddAndConnectInputs(ClassInputs, PresetGraphHandle, ReferencedNodeHandle);
+			AddAndConnectOutputs(ClassOutputs, PresetGraphHandle, ReferencedNodeHandle);
 
 			return true;
 		}
 
-		void FRebuildPresetRootGraph::AddAndConnectInputs(const TArray<FMetasoundFrontendClassInput>& InClassInputs, FGraphHandle& InParentGraphHandle, FNodeHandle& InReferencedNode) const
+		void FRebuildPresetRootGraph::AddAndConnectInputs(const TArray<FMetasoundFrontendClassInput>& InClassInputs, FGraphHandle& InPresetGraphHandle, FNodeHandle& InReferencedNode) const
 		{
 			// Add inputs and space appropriately
 			FVector2D InputNodeLocation = FVector2D::ZeroVector;
 
 			for (const FMetasoundFrontendClassInput& ClassInput : InClassInputs)
 			{
-				FNodeHandle InputNode = InParentGraphHandle->AddInputVertex(ClassInput);
+				FNodeHandle InputNode = InPresetGraphHandle->AddInputVertex(ClassInput);
 
 				if (ensure(InputNode->IsValid()))
 				{
@@ -645,14 +648,14 @@ namespace Metasound
 			}
 		}
 
-		void FRebuildPresetRootGraph::AddAndConnectOutputs(const TArray<FMetasoundFrontendClassOutput>& InClassOutputs, FGraphHandle& InParentGraphHandle, FNodeHandle& InReferencedNode) const
+		void FRebuildPresetRootGraph::AddAndConnectOutputs(const TArray<FMetasoundFrontendClassOutput>& InClassOutputs, FGraphHandle& InPresetGraphHandle, FNodeHandle& InReferencedNode) const
 		{
 			// Add outputs and space appropriately
 			FVector2D OutputNodeLocation = (2 * DisplayStyle::NodeLayout::DefaultOffsetX);
 
 			for (const FMetasoundFrontendClassOutput& ClassOutput : InClassOutputs)
 			{
-				FNodeHandle OutputNode = InParentGraphHandle->AddOutputVertex(ClassOutput);
+				FNodeHandle OutputNode = InPresetGraphHandle->AddOutputVertex(ClassOutput);
 				
 				if (ensure(OutputNode->IsValid()))
 				{
@@ -671,7 +674,7 @@ namespace Metasound
 			}
 		}
 
-		TArray<FMetasoundFrontendClassInput> FRebuildPresetRootGraph::GenerateRequiredClassInputs(const FConstGraphHandle& InParentGraph) const
+		TArray<FMetasoundFrontendClassInput> FRebuildPresetRootGraph::GenerateRequiredClassInputs(const FConstGraphHandle& InPresetGraph, TSet<FName>& OutInputsInheritingDefault) const
 		{
 			TArray<FMetasoundFrontendClassInput> ClassInputs;
 
@@ -694,17 +697,26 @@ namespace Metasound
 
 					ClassInput.VertexID = FGuid::NewGuid();
 
-					if (const FMetasoundFrontendClassInput* ExistingClassInput = InParentGraph->FindClassInputWithName(NodeName).Get())
+					if (const FMetasoundFrontendClassInput* ExistingClassInput = InPresetGraph->FindClassInputWithName(NodeName).Get())
 					{
 						ClassInput.NodeID = ExistingClassInput->NodeID;
 					}
 
-					if (InParentGraph->ContainsInputVertex(NodeName, ClassInput.TypeName))
+					if (InPresetGraph->ContainsInputVertex(NodeName, ClassInput.TypeName))
 					{
 						// If the input vertex already exists in the parent graph,
-						// use the default literal value from the parent graph.
-						FGuid VertexID = InParentGraph->GetVertexIDForInputVertex(NodeName);
-						ClassInput.DefaultLiteral = InParentGraph->GetDefaultInput(VertexID);
+						// check if parent should be used or not from set of managed
+						// input names.
+						if (InPresetGraph->GetInputsInheritingDefault().Contains(NodeName))
+						{
+							const FGuid ReferencedVertexID = ReferencedGraph->GetVertexIDForInputVertex(NodeName);
+							ClassInput.DefaultLiteral = ReferencedGraph->GetDefaultInput(ReferencedVertexID);
+						}
+						else
+						{
+							FGuid VertexID = InPresetGraph->GetVertexIDForInputVertex(NodeName);
+							ClassInput.DefaultLiteral = InPresetGraph->GetDefaultInput(VertexID);
+						}
 					}
 					else
 					{
@@ -719,10 +731,23 @@ namespace Metasound
 				}
 			}, EMetasoundFrontendClassType::Input);
 
+			// Fill new managed inputs set with names of all class inputs & if the old input was explicitly not
+			// marked as a managed input, then remove it from the new managed inputs if found.
+			OutInputsInheritingDefault.Reset();
+			Algo::Transform(ClassInputs, OutInputsInheritingDefault, [](const FMetasoundFrontendClassInput& Input) { return Input.Name; });
+			const TSet<FName>& InputsInheritingDefault = InPresetGraph->GetInputsInheritingDefault();
+			InPresetGraph->IterateConstNodes([&InputsInheritingDefault, &OutInputsInheritingDefault](FConstNodeHandle Input)
+			{
+				if (!InputsInheritingDefault.Contains(Input->GetNodeName()))
+				{
+					OutInputsInheritingDefault.Remove(Input->GetNodeName());
+				}
+			}, EMetasoundFrontendClassType::Input);
+
 			return ClassInputs;
 		}
 
-		TArray<FMetasoundFrontendClassOutput> FRebuildPresetRootGraph::GenerateRequiredClassOutputs(const FConstGraphHandle& InParentGraph) const
+		TArray<FMetasoundFrontendClassOutput> FRebuildPresetRootGraph::GenerateRequiredClassOutputs(const FConstGraphHandle& InPresetGraph) const
 		{
 			TArray<FMetasoundFrontendClassOutput> ClassOutputs;
 
@@ -745,7 +770,7 @@ namespace Metasound
 
 					ClassOutput.VertexID = FGuid::NewGuid();
 
-					if (const FMetasoundFrontendClassOutput* ExistingClassOutput = InParentGraph->FindClassOutputWithName(NodeName).Get())
+					if (const FMetasoundFrontendClassOutput* ExistingClassOutput = InPresetGraph->FindClassOutputWithName(NodeName).Get())
 					{
 						ClassOutput.NodeID = ExistingClassOutput->NodeID;
 					}
@@ -1141,6 +1166,26 @@ namespace Metasound
 			}
 		};
 
+		/** Versions document from 1.9 to 1.10. */
+		class FVersionDocument_1_10 : public FVersionDocumentTransform
+		{
+		public:
+			FVersionDocument_1_10() = default;
+
+			FMetasoundFrontendVersionNumber GetTargetVersion() const override
+			{
+				return { 1, 10 };
+			}
+
+			void TransformInternal(FDocumentHandle InDocument) const override
+			{
+				FMetasoundFrontendGraphClass Class = InDocument->GetRootGraphClass();
+				FMetasoundFrontendGraphClassPresetOptions PresetOptions = Class.PresetOptions;
+				Class.PresetOptions.bIsPreset = Class.Metadata.GetAndClearAutoUpdateManagesInterface_Deprecated();
+				InDocument->SetRootGraphClass(MoveTemp(Class));
+			}
+		};
+
 		FVersionDocument::FVersionDocument(FName InName, const FString& InPath)
 			: Name(InName)
 			, Path(InPath)
@@ -1168,6 +1213,7 @@ namespace Metasound
 			bWasUpdated |= FVersionDocument_1_7().Transform(InDocument);
 			bWasUpdated |= FVersionDocument_1_8().Transform(InDocument);
 			bWasUpdated |= FVersionDocument_1_9().Transform(InDocument);
+			bWasUpdated |= FVersionDocument_1_10().Transform(InDocument);
 
 			if (bWasUpdated)
 			{
diff --git a/Engine/Plugins/Runtime/Metasound/Source/MetasoundFrontend/Public/MetasoundFrontendController.h b/Engine/Plugins/Runtime/Metasound/Source/MetasoundFrontend/Public/MetasoundFrontendController.h
--- a/Engine/Plugins/Runtime/Metasound/Source/MetasoundFrontend/Public/MetasoundFrontendController.h
+++ b/Engine/Plugins/Runtime/Metasound/Source/MetasoundFrontend/Public/MetasoundFrontendController.h
@@ -587,6 +587,12 @@ namespace Metasound
 			/** Returns the ClassID associated with this graph. */
 			virtual FGuid GetClassID() const = 0;
 
+			/** Return the preset options for the current graph. */
+			virtual const FMetasoundFrontendGraphClassPresetOptions& GetGraphPresetOptions() const = 0;
+
+			/** Sets the preset options for the current graph. */
+			virtual void SetGraphPresetOptions(const FMetasoundFrontendGraphClassPresetOptions& InPresetOptions) = 0;
+
 			/** Return the metadata for the current graph. */
 			virtual const FMetasoundFrontendClassMetadata& GetGraphMetadata() const = 0;
 
@@ -629,6 +635,18 @@ namespace Metasound
 			/** Returns all input nodes in the graph. */
 			virtual TArray<FConstNodeHandle> GetConstInputNodes() const = 0;
 
+			/** Returns a set of all input names that are managed */
+			virtual const TSet<FName>& GetInputsInheritingDefault() const = 0;
+
+			/** If true, adds an item to the set of all input names
+			  * that are managed.
+			  * Returns true if successfully added/removed, false if not.
+			  */
+			virtual bool SetInputInheritsDefault(FName InName, bool bDefaultIsInherited) = 0;
+
+			/** Sets managed input names */
+			virtual void SetInputsInheritingDefault(TSet<FName>&& InNames) = 0;
+
 			/** Adds a new variable to the graph */
 			virtual FVariableHandle AddVariable(const FName& InDataTypeName) = 0;
 
diff --git a/Engine/Plugins/Runtime/Metasound/Source/MetasoundFrontend/Public/MetasoundFrontendDocument.h b/Engine/Plugins/Runtime/Metasound/Source/MetasoundFrontend/Public/MetasoundFrontendDocument.h
--- a/Engine/Plugins/Runtime/Metasound/Source/MetasoundFrontend/Public/MetasoundFrontendDocument.h
+++ b/Engine/Plugins/Runtime/Metasound/Source/MetasoundFrontend/Public/MetasoundFrontendDocument.h
@@ -1080,11 +1080,6 @@ public:
 		return ClassName;
 	}
 
-	bool GetAutoUpdateManagesInterface() const
-	{
-		return bAutoUpdateManagesInterface;
-	}
-
 	void SetClassName(const FMetasoundFrontendClassName& InClassName);
 
 	EMetasoundFrontendClassType GetType() const
@@ -1138,7 +1133,6 @@ public:
 	}
 
 	void SetAuthor(const FText& InAuthor);
-	void SetAutoUpdateManagesInterface(bool bInAutoUpdateManagesInterface);
 	void SetCategoryHierarchy(const TArray<FText>& InCategoryHierarchy);
 	void SetDescription(const FText& InDescription);
 	void SetDisplayName(const FText& InDisplayName);
@@ -1157,6 +1151,14 @@ public:
 		// External/Internal should probably be a separate field.
 		// ChangeID = FGuid::NewGuid();
 	}
+
+	// Deprecated field in favor of GraphClass PresetOptions
+	bool GetAndClearAutoUpdateManagesInterface_Deprecated()
+	{
+		bool bToReturn = bAutoUpdateManagesInterface;
+		bAutoUpdateManagesInterface = false;
+		return bToReturn;
+	}
 };
 
 USTRUCT()
@@ -1219,6 +1221,26 @@ struct METASOUNDFRONTEND_API FMetasoundFrontendClass
 #endif // WITH_EDITORONLY_DATA
 };
 
+// Preset options related to a parent graph class.  A graph class with bIsPreset set to true
+// auto-updates to mirror the interface members (inputs & outputs) of the single, referenced
+// node. It also connects all of these nodes' interface members on update to corresponding inputs
+// & outputs, and inherits input defaults from the referenced node unless otherwise specified.
+USTRUCT()
+struct METASOUNDFRONTEND_API FMetasoundFrontendGraphClassPresetOptions
+{
+	GENERATED_BODY()
+
+	// Whether or not graph class is a preset or not.
+	UPROPERTY()
+	bool bIsPreset = false;
+
+	// Names of all inputs inheriting default values from the referenced node. All input names
+	// in this set have their default value set on update when registered with the Frontend Class
+	// Registry.  Omitted inputs remain using the pre-existing, serialized default values.
+	UPROPERTY()
+	TSet<FName> InputsInheritingDefault;
+};
+
 USTRUCT()
 struct METASOUNDFRONTEND_API FMetasoundFrontendGraphClass : public FMetasoundFrontendClass
 {
@@ -1230,6 +1252,9 @@ struct METASOUNDFRONTEND_API FMetasoundFrontendGraphClass : public FMetasoundFro
 
 	UPROPERTY()
 	FMetasoundFrontendGraph Graph;
+
+	UPROPERTY()
+	FMetasoundFrontendGraphClassPresetOptions PresetOptions;
 };
 
 USTRUCT()
diff --git a/Engine/Plugins/Runtime/Metasound/Source/MetasoundFrontend/Public/MetasoundFrontendTransform.h b/Engine/Plugins/Runtime/Metasound/Source/MetasoundFrontend/Public/MetasoundFrontendTransform.h
--- a/Engine/Plugins/Runtime/Metasound/Source/MetasoundFrontend/Public/MetasoundFrontendTransform.h
+++ b/Engine/Plugins/Runtime/Metasound/Source/MetasoundFrontend/Public/MetasoundFrontendTransform.h
@@ -105,8 +105,9 @@ namespace Metasound
 		private:
 
 			// Get the class inputs needed for this preset. Input literals set on 
-			// the parent graph will be used if they exist. 
-			TArray<FMetasoundFrontendClassInput> GenerateRequiredClassInputs(const FConstGraphHandle& InParentGraph) const;
+			// the preset graph will be used if they are set and are marked as inheriting
+			// the default from the referenced graph.
+			TArray<FMetasoundFrontendClassInput> GenerateRequiredClassInputs(const FConstGraphHandle& InParentGraph, TSet<FName>& OutInputsInheritingDefault) const;
 
 			// Get the class Outputs needed for this preset.
 			TArray<FMetasoundFrontendClassOutput> GenerateRequiredClassOutputs(const FConstGraphHandle& InParentGraph) const;
@@ -156,7 +157,7 @@ namespace Metasound
 		public:
 			static FMetasoundFrontendVersionNumber GetMaxVersion()
 			{
-				return FMetasoundFrontendVersionNumber { 1, 9 };
+				return FMetasoundFrontendVersionNumber { 1, 10 };
 			}
 
 			FVersionDocument(FName InName, const FString& InPath);
