- Reduce MetaSounds loaded on Editor startup by moving additional MetaSound loads/registration to occur when MetaSoundEditor is loaded
- Add UX to communicate on first asset editor load that MetaSounds are still loading and class picker results may be incomplete
#rnx

diff --git a/Engine/Plugins/Runtime/Metasound/Source/MetasoundEditor/Private/MetasoundEditor.cpp b/Engine/Plugins/Runtime/Metasound/Source/MetasoundEditor/Private/MetasoundEditor.cpp
--- a/Engine/Plugins/Runtime/Metasound/Source/MetasoundEditor/Private/MetasoundEditor.cpp
+++ b/Engine/Plugins/Runtime/Metasound/Source/MetasoundEditor/Private/MetasoundEditor.cpp
@@ -432,7 +432,14 @@ namespace Metasound
 
 			check(ObjectToEdit);
 			checkf(IMetasoundUObjectRegistry::Get().IsRegisteredClass(ObjectToEdit), TEXT("Object passed in was not registered as a valid metasound interface!"));
-			
+
+			IMetasoundEditorModule& MetaSoundEditorModule = FModuleManager::GetModuleChecked<IMetasoundEditorModule>("MetaSoundEditor");
+			bPrimingRegistry = MetaSoundEditorModule.GetAssetRegistryPrimeStatus() <= EAssetPrimeStatus::InProgress;
+			if (MetaSoundEditorModule.GetAssetRegistryPrimeStatus() < EAssetPrimeStatus::InProgress)
+			{
+				MetaSoundEditorModule.PrimeAssetRegistryAsync();
+			}
+
 			// Support undo/redo
 			Metasound = ObjectToEdit;
 			Metasound->SetFlags(RF_Transactional);
@@ -450,6 +457,12 @@ namespace Metasound
 			CreateInternalWidgets();
 			CreateAnalyzers();
 
+			// Has to be run after widgets are initialized to properly display
+			if (bPrimingRegistry)
+			{
+				NotifyAssetPrimeInProgress();
+			}
+
 			const TSharedRef<FTabManager::FLayout> StandaloneDefaultLayout = FTabManager::NewLayout("Standalone_MetasoundEditor_Layout_v9")
 				->AddArea
 				(
@@ -569,6 +582,36 @@ namespace Metasound
 			FSlateApplication::Get().DismissAllMenus();
 		}
 
+		void FEditor::NotifyAssetPrimeInProgress()
+		{
+			if (MetasoundGraphEditor.IsValid())
+			{
+				FNotificationInfo Info(LOCTEXT("MetaSoundScanInProgressNotificationText", "Registering MetaSound Assets..."));
+				Info.SubText = LOCTEXT("MetaSoundScanInProgressNotificationSubText", "Class selector results may be incomplete");
+				Info.bUseThrobber = true;
+				Info.bFireAndForget = true;
+				Info.bUseSuccessFailIcons = false;
+				Info.ExpireDuration = 3.0f;
+				Info.FadeOutDuration = 1.0f;
+
+				MetasoundGraphEditor->AddNotification(Info, false /* bSuccess */);
+			}
+		}
+
+		void FEditor::NotifyAssetPrimeComplete()
+		{
+			if (MetasoundGraphEditor.IsValid())
+			{
+				FNotificationInfo Info(LOCTEXT("MetaSoundScanInProgressNotification", "MetaSound Asset Registration Complete"));
+				Info.bFireAndForget = true;
+				Info.bUseSuccessFailIcons = true;
+				Info.ExpireDuration = 3.0f;
+				Info.FadeOutDuration = 1.0f;
+
+				MetasoundGraphEditor->AddNotification(Info, true /* bSuccess */);
+			}
+		}
+
 		void FEditor::NotifyDocumentVersioned()
 		{
 			if (MetasoundGraphEditor.IsValid())
@@ -2361,6 +2404,16 @@ namespace Metasound
 			FMetasoundAssetBase* MetasoundAsset = IMetasoundUObjectRegistry::Get().GetObjectAsAssetBase(Metasound);
 			check(MetasoundAsset);
 
+			if (bPrimingRegistry)
+			{
+				IMetasoundEditorModule& MetaSoundEditorModule = FModuleManager::GetModuleChecked<IMetasoundEditorModule>("MetaSoundEditor");
+				if (MetaSoundEditorModule.GetAssetRegistryPrimeStatus() == EAssetPrimeStatus::Complete)
+				{
+					bPrimingRegistry = false;
+					NotifyAssetPrimeComplete();
+				}
+			}
+
 			if (MetasoundAsset->GetSynchronizationRequired())
 			{
 				MetasoundAsset->CacheRegistryMetadata();
diff --git a/Engine/Plugins/Runtime/Metasound/Source/MetasoundEditor/Private/MetasoundEditor.h b/Engine/Plugins/Runtime/Metasound/Source/MetasoundEditor/Private/MetasoundEditor.h
--- a/Engine/Plugins/Runtime/Metasound/Source/MetasoundEditor/Private/MetasoundEditor.h
+++ b/Engine/Plugins/Runtime/Metasound/Source/MetasoundEditor/Private/MetasoundEditor.h
@@ -301,6 +301,8 @@ namespace Metasound
 
 			void EditObjectSettings();
 
+			void NotifyAssetPrimeComplete();
+			void NotifyAssetPrimeInProgress();
 			void NotifyDocumentVersioned();
 			void NotifyNodePasteFailure_MultipleVariableSetters();
 			void NotifyNodePasteFailure_ReferenceLoop();
@@ -346,6 +348,9 @@ namespace Metasound
 			/** Boolean state for when selection change handle should not respond due to selection state
 			  * being manually applied in code */
 			bool bManuallyClearingGraphSelection = false;
+
+			/** Whether or not registry is currently being primed due to assets still loading */
+			bool bPrimingRegistry = false;
 		};
 	} // namespace Editor
 } // namespace Metasound
diff --git a/Engine/Plugins/Runtime/Metasound/Source/MetasoundEditor/Private/MetasoundEditorGraphBuilder.cpp b/Engine/Plugins/Runtime/Metasound/Source/MetasoundEditor/Private/MetasoundEditorGraphBuilder.cpp
--- a/Engine/Plugins/Runtime/Metasound/Source/MetasoundEditor/Private/MetasoundEditorGraphBuilder.cpp
+++ b/Engine/Plugins/Runtime/Metasound/Source/MetasoundEditor/Private/MetasoundEditorGraphBuilder.cpp
@@ -1356,7 +1356,6 @@ namespace Metasound
 			FMetaSoundAssetRegistrationOptions RegOptions;
 			RegOptions.bForceReregister = true;
 			RegOptions.bForceViewSynchronization = bInForceViewSynchronization;
-			RegOptions.bRegisterDependencies = true;
 			// if EditedReferencingMetaSounds is empty, then no MetaSounds are open
 			// that reference this MetaSound, so just register this asset. Otherwise,
 			// this graph will recursively get updated when the open referencing graphs
diff --git a/Engine/Plugins/Runtime/Metasound/Source/MetasoundEditor/Private/MetasoundEditorModule.cpp b/Engine/Plugins/Runtime/Metasound/Source/MetasoundEditor/Private/MetasoundEditorModule.cpp
--- a/Engine/Plugins/Runtime/Metasound/Source/MetasoundEditor/Private/MetasoundEditorModule.cpp
+++ b/Engine/Plugins/Runtime/Metasound/Source/MetasoundEditor/Private/MetasoundEditorModule.cpp
@@ -162,42 +162,101 @@ namespace Metasound
 
 		class FModule : public IMetasoundEditorModule
 		{
-			void AddOrUpdateClassRegistryAsset(const FAssetData& InAssetData)
+			void LoadAndRegisterAsset(const FAssetData& InAssetData)
 			{
-				using namespace Metasound::Frontend;
-
-				if (IsMetaSoundAssetClass(InAssetData.AssetClass))
+				if (InAssetData.IsAssetLoaded())
 				{
-					// in case they need to be 
-					check(GEngine);
-					UMetaSoundAssetSubsystem* AssetSubsystem = GEngine->GetEngineSubsystem<UMetaSoundAssetSubsystem>();
-					check(AssetSubsystem);
+					if (UObject* AssetObject = InAssetData.GetAsset())
+					{
+						FMetasoundAssetBase* MetaSoundAsset = IMetasoundUObjectRegistry::Get().GetObjectAsAssetBase(AssetObject);
+						check(MetaSoundAsset);
+						Frontend::FMetaSoundAssetRegistrationOptions RegOptions;
+						RegOptions.bForceReregister = false;
+						MetaSoundAsset->RegisterGraphWithFrontend(RegOptions);
+					}
+				}
+				else
+				{
+					if (!MetaSoundEditorAsyncRegistrationEnabledCVar)
+					{
+						return;
+					}
 
-					// Use the editor version of `RegisterGraphWithFrontend` so it re-registers any open MetaSound editors
-					AssetSubsystem->AddOrUpdateAsset(InAssetData);
-					
-					// Loading all assets necessary only in editor to register &
-					// populate potential graphs to reference in MetaSound editor.
-					if (InAssetData.IsAssetLoaded())
+					if (AssetPrimeStatus == EAssetPrimeStatus::NotRequested)
 					{
-						if (UObject* AssetObject = InAssetData.GetAsset())
-						{
-							FGraphBuilder::RegisterGraphWithFrontend(*AssetObject);
-						}
+						return;
 					}
-					else if (MetaSoundEditorAsyncRegistrationEnabledCVar)
+
+					ActiveAsyncAssetLoadRequests++;
+
+					FSoftObjectPath AssetPath = InAssetData.ToSoftObjectPath();
+					auto LoadAndRegister = [this, ObjectPath = AssetPath](const FName& PackageName, UPackage* Package, EAsyncLoadingResult::Type Result)
 					{
-						FSoftObjectPath Path = InAssetData.ToSoftObjectPath();
-						LoadPackageAsync(Path.GetLongPackageName(), FLoadPackageAsyncDelegate::CreateLambda(
-							[this, ObjectPath = MoveTemp(Path)](const FName& PackageName, UPackage* Package, EAsyncLoadingResult::Type Result)
+						if (Result == EAsyncLoadingResult::Succeeded)
+						{
+							FMetasoundAssetBase* MetaSoundAsset = IMetasoundUObjectRegistry::Get().GetObjectAsAssetBase(ObjectPath.ResolveObject());
+							check(MetaSoundAsset);
+							if (!MetaSoundAsset->IsRegistered())
 							{
-								if (Result == EAsyncLoadingResult::Succeeded)
-								{
-									FGraphBuilder::RegisterGraphWithFrontend(*ObjectPath.ResolveObject());
-								}
+								Frontend::FMetaSoundAssetRegistrationOptions RegOptions;
+								RegOptions.bForceReregister = false;
+								MetaSoundAsset->RegisterGraphWithFrontend(RegOptions);
 							}
-						));
-					}
+						}
+
+						ActiveAsyncAssetLoadRequests--;
+						if (AssetPrimeStatus == EAssetPrimeStatus::InProgress && ActiveAsyncAssetLoadRequests == 0)
+						{
+							AssetPrimeStatus = EAssetPrimeStatus::Complete;
+						}
+					};
+					LoadPackageAsync(AssetPath.GetLongPackageName(), FLoadPackageAsyncDelegate::CreateLambda(LoadAndRegister));
+				}
+			}
+
+			void AddClassRegistryAsset(const FAssetData& InAssetData)
+			{
+				using namespace Frontend;
+
+				if (!IsMetaSoundAssetClass(InAssetData.AssetClass))
+				{
+					return;
+				}
+
+				check(GEngine);
+				UMetaSoundAssetSubsystem* AssetSubsystem = GEngine->GetEngineSubsystem<UMetaSoundAssetSubsystem>();
+				check(AssetSubsystem);
+
+				const FNodeRegistryKey RegistryKey = AssetSubsystem->AddOrUpdateAsset(InAssetData);
+				const bool bPrimeRequested = AssetPrimeStatus > EAssetPrimeStatus::NotRequested;
+				const bool bIsRegistered = FMetasoundFrontendRegistryContainer::Get()->IsNodeRegistered(RegistryKey);
+				if (bPrimeRequested && !bIsRegistered)
+				{
+					LoadAndRegisterAsset(InAssetData);
+				}
+			}
+
+			void UpdateClassRegistryAsset(const FAssetData& InAssetData)
+			{
+				using namespace Frontend;
+
+				if (!IsMetaSoundAssetClass(InAssetData.AssetClass))
+				{
+					return;
+				}
+
+				check(GEngine);
+				UMetaSoundAssetSubsystem* AssetSubsystem = GEngine->GetEngineSubsystem<UMetaSoundAssetSubsystem>();
+				check(AssetSubsystem);
+
+				const FNodeRegistryKey RegistryKey = AssetSubsystem->AddOrUpdateAsset(InAssetData);
+				const bool bPrimeRequested = AssetPrimeStatus > EAssetPrimeStatus::NotRequested;
+				const bool bIsRegistered = FMetasoundFrontendRegistryContainer::Get()->IsNodeRegistered(RegistryKey);
+
+				// Have to re-register even if prime was not requested to avoid registry desync.
+				if (bPrimeRequested || bIsRegistered)
+				{
+					LoadAndRegisterAsset(InAssetData);
 				}
 			}
 
@@ -250,13 +309,16 @@ namespace Metasound
 
 			void OnAssetScanFinished()
 			{
-				FARFilter Filter;
-				Filter.ClassNames = MetaSoundClassNames;
+				AssetScanStatus = EAssetScanStatus::Complete;
+
+				if (AssetPrimeStatus == EAssetPrimeStatus::Requested)
+				{
+					PrimeAssetRegistryAsync();
+				}
 
 				FAssetRegistryModule& AssetRegistryModule = FModuleManager::LoadModuleChecked<FAssetRegistryModule>("AssetRegistry");
-				AssetRegistryModule.Get().EnumerateAssets(Filter, [this](const FAssetData& AssetData) { AddOrUpdateClassRegistryAsset(AssetData); return true; });
-				AssetRegistryModule.Get().OnAssetAdded().AddRaw(this, &FModule::AddOrUpdateClassRegistryAsset);
-				AssetRegistryModule.Get().OnAssetUpdated().AddRaw(this, &FModule::AddOrUpdateClassRegistryAsset);
+				AssetRegistryModule.Get().OnAssetAdded().AddRaw(this, &FModule::AddClassRegistryAsset);
+				AssetRegistryModule.Get().OnAssetUpdated().AddRaw(this, &FModule::UpdateClassRegistryAsset);
 				AssetRegistryModule.Get().OnAssetRemoved().AddRaw(this, &FModule::RemoveAssetFromClassRegistry);
 				AssetRegistryModule.Get().OnAssetRenamed().AddRaw(this, &FModule::RenameAssetInClassRegistry);
 
@@ -466,6 +528,36 @@ namespace Metasound
 				}
 			}
 
+			virtual void PrimeAssetRegistryAsync() override
+			{
+				// Ignore step if still loading assets from initial scan but set prime status as requested.
+				if (AssetScanStatus <= EAssetScanStatus::InProgress)
+				{
+					AssetPrimeStatus = EAssetPrimeStatus::Requested;
+					return;
+				}
+
+				if (AssetPrimeStatus != EAssetPrimeStatus::InProgress)
+				{
+					AssetPrimeStatus = EAssetPrimeStatus::InProgress;
+
+					FARFilter Filter;
+					Filter.ClassNames = MetaSoundClassNames;
+
+					FAssetRegistryModule& AssetRegistryModule = FModuleManager::LoadModuleChecked<FAssetRegistryModule>("AssetRegistry");
+					AssetRegistryModule.Get().EnumerateAssets(Filter, [this](const FAssetData& AssetData)
+					{
+						AddClassRegistryAsset(AssetData);
+						return true;
+					});
+				}
+			}
+
+			virtual EAssetPrimeStatus GetAssetRegistryPrimeStatus() const override
+			{
+				return AssetPrimeStatus;
+			}
+
 			virtual void RegisterExplicitProxyClass(const UClass& InClass) override
 			{
 				using namespace Metasound::Frontend;
@@ -602,7 +694,15 @@ namespace Metasound
 				FAssetTypeActions_MetaSoundSource::RegisterMenuActions();
 
 				FAssetRegistryModule& AssetRegistryModule = FModuleManager::LoadModuleChecked<FAssetRegistryModule>("AssetRegistry");
-				AssetRegistryModule.Get().OnFilesLoaded().AddRaw(this, &FModule::OnAssetScanFinished);
+				if (AssetRegistryModule.Get().IsLoadingAssets())
+				{
+					AssetScanStatus = EAssetScanStatus::InProgress;
+					AssetRegistryModule.Get().OnFilesLoaded().AddRaw(this, &FModule::OnAssetScanFinished);
+				}
+				else
+				{
+					AssetScanStatus = EAssetScanStatus::Complete;
+				}
 
 				RegisterExplicitProxyClass(*USoundWave::StaticClass());
 			}
@@ -661,6 +761,10 @@ namespace Metasound
 			TSharedPtr<FSlateStyleSet> StyleSet;
 
 			TSet<const UClass*> ExplicitProxyClasses;
+
+			EAssetPrimeStatus AssetPrimeStatus = EAssetPrimeStatus::NotRequested;
+			EAssetScanStatus AssetScanStatus = EAssetScanStatus::NotRequested;
+			int32 ActiveAsyncAssetLoadRequests = 0;
 		};
 	} // namespace Editor
 } // namespace Metasound
diff --git a/Engine/Plugins/Runtime/Metasound/Source/MetasoundEditor/Public/MetasoundEditorModule.h b/Engine/Plugins/Runtime/Metasound/Source/MetasoundEditor/Public/MetasoundEditorModule.h
--- a/Engine/Plugins/Runtime/Metasound/Source/MetasoundEditor/Public/MetasoundEditorModule.h
+++ b/Engine/Plugins/Runtime/Metasound/Source/MetasoundEditor/Public/MetasoundEditorModule.h
@@ -26,6 +26,25 @@ namespace Metasound
 	{
 		using FDataTypeRegistryInfo = Frontend::FDataTypeRegistryInfo;
 
+		// Status of initial asset scan when editor loads up.
+		enum class EAssetScanStatus : uint8
+		{
+			NotRequested = 0,
+			InProgress = 2,
+			Complete = 3
+		};
+
+		// Primes status of MetaSound assets.  Priming an asset
+		// effectively loading the asset asynchronously (if not already loaded)
+		// & registers it with the MetaSound Class Registry.
+		enum class EAssetPrimeStatus : uint8
+		{
+			NotRequested = 0,
+			Requested = 1,
+			InProgress = 2,
+			Complete = 3
+		};
+
 		struct FEditorDataType
 		{
 			FEdGraphPinType PinType;
@@ -82,6 +101,11 @@ namespace Metasound
 
 			virtual bool IsRegisteredDataType(FName InDataTypeName) const = 0;
 
+			// Primes MetaSound assets, effectively loading the asset asynchronously (if not already
+			// loaded) & registers them if not already registered with the MetaSound Class Registry.
+			virtual void PrimeAssetRegistryAsync() = 0;
+			virtual EAssetPrimeStatus GetAssetRegistryPrimeStatus() const = 0;
+
 			virtual void IterateDataTypes(TUniqueFunction<void(const FEditorDataType&)> InDataTypeFunction) const = 0;
 
 			virtual TUniquePtr<FMetasoundDefaultLiteralCustomizationBase> CreateMemberDefaultLiteralCustomization(UClass& InClass, IDetailCategoryBuilder& DefaultCategoryBuilder) const = 0;
diff --git a/Engine/Plugins/Runtime/Metasound/Source/MetasoundEngine/Private/MetasoundAssetSubsystem.cpp b/Engine/Plugins/Runtime/Metasound/Source/MetasoundEngine/Private/MetasoundAssetSubsystem.cpp
--- a/Engine/Plugins/Runtime/Metasound/Source/MetasoundEngine/Private/MetasoundAssetSubsystem.cpp
+++ b/Engine/Plugins/Runtime/Metasound/Source/MetasoundEngine/Private/MetasoundAssetSubsystem.cpp
@@ -171,7 +171,7 @@ void UMetaSoundAssetSubsystem::AddAssetReferences(FMetasoundAssetBase& InAssetBa
 	}
 }
 
-void UMetaSoundAssetSubsystem::AddOrUpdateAsset(UObject& InObject)
+Metasound::Frontend::FNodeRegistryKey UMetaSoundAssetSubsystem::AddOrUpdateAsset(const UObject& InObject)
 {
 	using namespace Metasound;
 	using namespace Metasound::AssetSubsystemPrivate;
@@ -179,15 +179,17 @@ void UMetaSoundAssetSubsystem::AddOrUpdateAsset(UObject& InObject)
 
 	METASOUND_TRACE_CPUPROFILER_EVENT_SCOPE(UMetaSoundAssetSubsystem::AddOrUpdateAsset);
 
-	FMetasoundAssetBase* MetaSoundAsset = Metasound::IMetasoundUObjectRegistry::Get().GetObjectAsAssetBase(&InObject);
+	const FMetasoundAssetBase* MetaSoundAsset = Metasound::IMetasoundUObjectRegistry::Get().GetObjectAsAssetBase(&InObject);
 	check(MetaSoundAsset);
 
 	FNodeClassInfo ClassInfo = MetaSoundAsset->GetAssetClassInfo();
 	const FNodeRegistryKey RegistryKey = NodeRegistryKey::CreateKey(ClassInfo);
 	PathMap.FindOrAdd(RegistryKey) = InObject.GetPathName();
+
+	return RegistryKey;
 }
 
-void UMetaSoundAssetSubsystem::AddOrUpdateAsset(const FAssetData& InAssetData)
+Metasound::Frontend::FNodeRegistryKey UMetaSoundAssetSubsystem::AddOrUpdateAsset(const FAssetData& InAssetData)
 {
 	using namespace Metasound;
 	using namespace Metasound::AssetSubsystemPrivate;
@@ -215,13 +217,14 @@ void UMetaSoundAssetSubsystem::AddOrUpdateAsset(const FAssetData& InAssetData)
 
 		if (Object)
 		{
-			AddOrUpdateAsset(*Object);
-			return;
+			return AddOrUpdateAsset(*Object);
 		}
 	}
 
 	const FNodeRegistryKey RegistryKey = NodeRegistryKey::CreateKey(ClassInfo);
 	PathMap.FindOrAdd(RegistryKey) = InAssetData.ObjectPath;
+
+	return RegistryKey;
 }
 
 bool UMetaSoundAssetSubsystem::CanAutoUpdate(const FMetasoundFrontendClassName& InClassName) const
@@ -348,12 +351,12 @@ FMetasoundAssetBase* UMetaSoundAssetSubsystem::TryLoadAsset(const FSoftObjectPat
 	return Metasound::IMetasoundUObjectRegistry::Get().GetObjectAsAssetBase(InObjectPath.TryLoad());
 }
 
-void UMetaSoundAssetSubsystem::RemoveAsset(UObject& InObject)
+void UMetaSoundAssetSubsystem::RemoveAsset(const UObject& InObject)
 {
 	using namespace Metasound;
 	using namespace Metasound::Frontend;
 
-	if (FMetasoundAssetBase* MetaSoundAsset = Metasound::IMetasoundUObjectRegistry::Get().GetObjectAsAssetBase(&InObject))
+	if (const FMetasoundAssetBase* MetaSoundAsset = Metasound::IMetasoundUObjectRegistry::Get().GetObjectAsAssetBase(&InObject))
 	{
 		const FNodeClassInfo ClassInfo = MetaSoundAsset->GetAssetClassInfo();
 		FNodeRegistryKey RegistryKey = FMetasoundFrontendRegistryContainer::Get()->GetRegistryKey(ClassInfo);
diff --git a/Engine/Plugins/Runtime/Metasound/Source/MetasoundEngine/Private/MetasoundSource.cpp b/Engine/Plugins/Runtime/Metasound/Source/MetasoundEngine/Private/MetasoundSource.cpp
--- a/Engine/Plugins/Runtime/Metasound/Source/MetasoundEngine/Private/MetasoundSource.cpp
+++ b/Engine/Plugins/Runtime/Metasound/Source/MetasoundEngine/Private/MetasoundSource.cpp
@@ -415,7 +415,6 @@ void UMetaSoundSource::InitResources()
 
 	FMetaSoundAssetRegistrationOptions RegOptions;
 	RegOptions.bForceReregister = false;
-	RegOptions.bRegisterDependencies = true;
 	RegisterGraphWithFrontend(RegOptions);
 }
 
diff --git a/Engine/Plugins/Runtime/Metasound/Source/MetasoundEngine/Public/Metasound.h b/Engine/Plugins/Runtime/Metasound/Source/MetasoundEngine/Public/Metasound.h
--- a/Engine/Plugins/Runtime/Metasound/Source/MetasoundEngine/Public/Metasound.h
+++ b/Engine/Plugins/Runtime/Metasound/Source/MetasoundEngine/Public/Metasound.h
@@ -65,7 +65,6 @@ namespace Metasound
 		// it to register.
 		FMetaSoundAssetRegistrationOptions NonEditorRegistrationOptions;
 		NonEditorRegistrationOptions.bForceReregister = false;
-		NonEditorRegistrationOptions.bRegisterDependencies = true;
 
 #if WITH_EDITORONLY_DATA
 		if (UMetasoundEditorGraphBase* MetaSoundGraph = Cast<UMetasoundEditorGraphBase>(InMetaSound.GetGraph()))
diff --git a/Engine/Plugins/Runtime/Metasound/Source/MetasoundEngine/Public/MetasoundAssetSubsystem.h b/Engine/Plugins/Runtime/Metasound/Source/MetasoundEngine/Public/MetasoundAssetSubsystem.h
--- a/Engine/Plugins/Runtime/Metasound/Source/MetasoundEngine/Public/MetasoundAssetSubsystem.h
+++ b/Engine/Plugins/Runtime/Metasound/Source/MetasoundEngine/Public/MetasoundAssetSubsystem.h
@@ -34,13 +34,13 @@ class METASOUNDENGINE_API UMetaSoundAssetSubsystem : public UEngineSubsystem, pu
 public:
 	virtual void Initialize(FSubsystemCollectionBase& InCollection) override;
 
-	void AddOrUpdateAsset(const FAssetData& InAssetData);
-	void RemoveAsset(UObject& InObject);
+	Metasound::Frontend::FNodeRegistryKey AddOrUpdateAsset(const FAssetData& InAssetData);
+	void RemoveAsset(const UObject& InObject);
 	void RemoveAsset(const FAssetData& InAssetData);
 	void RenameAsset(const FAssetData& InAssetData, bool bInReregisterWithFrontend = true);
 
 	virtual void AddAssetReferences(FMetasoundAssetBase& InAssetBase) override;
-	virtual void AddOrUpdateAsset(UObject& InObject) override;
+	virtual Metasound::Frontend::FNodeRegistryKey  AddOrUpdateAsset(const UObject& InObject) override;
 	virtual bool CanAutoUpdate(const FMetasoundFrontendClassName& InClassName) const override;
 	virtual bool ContainsKey(const Metasound::Frontend::FNodeRegistryKey& InRegistryKey) const override;
 	virtual const FSoftObjectPath* FindObjectPathFromKey(const Metasound::Frontend::FNodeRegistryKey& RegistryKey) const override;
diff --git a/Engine/Plugins/Runtime/Metasound/Source/MetasoundFrontend/Public/MetasoundAssetManager.h b/Engine/Plugins/Runtime/Metasound/Source/MetasoundFrontend/Public/MetasoundAssetManager.h
--- a/Engine/Plugins/Runtime/Metasound/Source/MetasoundFrontend/Public/MetasoundAssetManager.h
+++ b/Engine/Plugins/Runtime/Metasound/Source/MetasoundFrontend/Public/MetasoundAssetManager.h
@@ -39,8 +39,10 @@ namespace Metasound
 			// If true, forces flag to resync all view (editor) data pertaining to the given asset(s) being registered.
 			bool bForceViewSynchronization = true;
 
-			// If true, recursively attempts to register dependencies.
-			bool bRegisterDependencies = false;
+			// If true, recursively attempts to register dependencies. (TODO: Determine if this option should be removed.
+			// Must validate that failed dependency updates due to auto-update for ex. being disabled is handled gracefully
+			// at runtime.)
+			bool bRegisterDependencies = true;
 
 			// Attempt to auto-update (Only runs if class not registered or set to force re-register.
 			// Will not respect being set to true if project-level MetaSoundSettings specify to not run auto-update.)
@@ -83,7 +85,7 @@ namespace Metasound
 			virtual void AddAssetReferences(FMetasoundAssetBase& InAssetBase) = 0;
 
 			// Add or Update a MetaSound Asset's entry data
-			virtual void AddOrUpdateAsset(UObject& InObject) = 0;
+			virtual Metasound::Frontend::FNodeRegistryKey AddOrUpdateAsset(const UObject& InObject) = 0;
 
 			// Whether or not the class is eligible for auto-update
 			virtual bool CanAutoUpdate(const FMetasoundFrontendClassName& InClassName) const = 0;
