Fix for multiple composition nodes referencing same class failing to auto-update properly due to additional class definition added and then immediately removed when it should be updating the pre-existing entry.
- Fix node display title being 'none' when you initially add an externally-defined class
#rnx

diff --git a/Engine/Plugins/Runtime/Metasound/Source/MetasoundEditor/Private/MetasoundEditorGraphSchema.cpp b/Engine/Plugins/Runtime/Metasound/Source/MetasoundEditor/Private/MetasoundEditorGraphSchema.cpp
--- a/Engine/Plugins/Runtime/Metasound/Source/MetasoundEditor/Private/MetasoundEditorGraphSchema.cpp
+++ b/Engine/Plugins/Runtime/Metasound/Source/MetasoundEditor/Private/MetasoundEditorGraphSchema.cpp
@@ -391,7 +391,8 @@ UEdGraphNode* FMetasoundGraphSchemaAction_NewNode::PerformAction(UEdGraph* Paren
 	using namespace Metasound::Frontend;
 
 	const FScopedTransaction Transaction(LOCTEXT("AddNewNode", "Add New MetaSound Node"));
-	UObject& ParentMetasound = CastChecked<UMetasoundEditorGraph>(ParentGraph)->GetMetasoundChecked();
+	UMetasoundEditorGraph* MetaSoundGraph = CastChecked<UMetasoundEditorGraph>(ParentGraph);
+	UObject& ParentMetasound = MetaSoundGraph->GetMetasoundChecked();
 	ParentMetasound.Modify();
 	ParentGraph->Modify();
 
@@ -399,6 +400,7 @@ UEdGraphNode* FMetasoundGraphSchemaAction_NewNode::PerformAction(UEdGraph* Paren
 	{
 		NewGraphNode->Modify();
 		SchemaPrivate::TryConnectNewNodeToMatchingDataTypePin(*NewGraphNode, FromPin);
+		MetaSoundGraph->SetSynchronizationRequired();
 		return NewGraphNode;
 	}
 
diff --git a/Engine/Plugins/Runtime/Metasound/Source/MetasoundFrontend/Private/MetasoundFrontendDocumentController.cpp b/Engine/Plugins/Runtime/Metasound/Source/MetasoundFrontend/Private/MetasoundFrontendDocumentController.cpp
--- a/Engine/Plugins/Runtime/Metasound/Source/MetasoundFrontend/Private/MetasoundFrontendDocumentController.cpp
+++ b/Engine/Plugins/Runtime/Metasound/Source/MetasoundFrontend/Private/MetasoundFrontendDocumentController.cpp
@@ -260,7 +260,7 @@ namespace Metasound
 		{
 			if (FMetasoundFrontendDocument* Document = DocumentPtr.Get())
 			{
-				FConstClassAccessPtr ClassPtr = FindClass(InKey);
+				FClassAccessPtr ClassPtr = DocumentPtr.GetClassWithRegistryKey(InKey);
 
 				auto AddClass = [=](FMetasoundFrontendClass&& NewClassDescription, const FGuid& NewClassID)
 				{
@@ -277,7 +277,7 @@ namespace Metasound
 					return NewClassPtr;
 				};
 
-				if (const FMetasoundFrontendClass* MetasoundClass = ClassPtr.Get())
+				if (FMetasoundFrontendClass* MetasoundClass = ClassPtr.Get())
 				{
 					// External node classes must match version to return shared definition.
 					if (MetasoundClass->Metadata.GetType() == EMetasoundFrontendClassType::External)
@@ -295,8 +295,9 @@ namespace Metasound
 					if (bInRefreshFromRegistry)
 					{
 						FGuid ClassID = MetasoundClass->ID;
-						FMetasoundFrontendClass NewClass = GenerateClass(InKey);
-						return AddClass(MoveTemp(NewClass), ClassID);
+						*MetasoundClass = GenerateClass(InKey);
+						MetasoundClass->ID = ClassID;
+						return FindClass(InKey);
 					}
 
 					return ClassPtr;
