- Fix & re-enable manual update of node class if auto-update is disabled or node class is updated and original class is missing
- Fix updating default value for a graph input not compositionally propagating to references
- Fix 'Reset to Class Default' sometimes showing when it isn't applicable (ex. on triggers or when the class default is out-of-date)
#rnx

diff --git a/Engine/Plugins/Runtime/Metasound/Source/MetasoundEditor/Private/MetasoundEditor.cpp b/Engine/Plugins/Runtime/Metasound/Source/MetasoundEditor/Private/MetasoundEditor.cpp
--- a/Engine/Plugins/Runtime/Metasound/Source/MetasoundEditor/Private/MetasoundEditor.cpp
+++ b/Engine/Plugins/Runtime/Metasound/Source/MetasoundEditor/Private/MetasoundEditor.cpp
@@ -901,6 +901,10 @@ namespace Metasound
 			ToolkitCommands->MapAction(
 				FGenericCommands::Get().Rename,
 				FExecuteAction::CreateSP(this, &FEditor::RenameSelected));
+
+			ToolkitCommands->MapAction(
+				FEditorCommands::Get().UpdateNodeClass,
+				FExecuteAction::CreateSP(this, &FEditor::UpdateSelectedNodeClasses));
 		}
 
 		void FEditor::Import()
@@ -1357,31 +1361,8 @@ namespace Metasound
 				GraphEditorCommands->MapAction(FGraphEditorCommands::Get().CreateComment,
 					FExecuteAction::CreateSP(this, &FEditor::OnCreateComment));
 
-				GraphEditorCommands->MapAction(FEditorCommands::Get().UpdateNodes,
-					FExecuteAction::CreateLambda([this]()
-					{
-						const FScopedTransaction Transaction(LOCTEXT("NodeVersionUpgrade", "Upgrade MetaSound Node(s)"));
-						check(Metasound);
-						Metasound->Modify();
-
-						UMetasoundEditorGraph& Graph = GetMetaSoundGraphChecked();
-						Graph.Modify();
-
-						const FGraphPanelSelectionSet Selection = MetasoundGraphEditor->GetSelectedNodes();
-						for (UObject* Object : Selection)
-						{
-							if (UMetasoundEditorGraphExternalNode* ExternalNode = Cast<UMetasoundEditorGraphExternalNode>(Object))
-							{
-								FMetasoundFrontendVersionNumber MajorUpdateVersion = ExternalNode->FindHighestVersionInRegistry();
-								Metasound::Frontend::FConstNodeHandle NodeHandle = ExternalNode->GetConstNodeHandle();
-								if (MajorUpdateVersion.IsValid() && MajorUpdateVersion > NodeHandle->GetClassMetadata().GetVersion())
-								{
-									ExternalNode->UpdateToVersion(MajorUpdateVersion, false /* bInPropagateErrorMessages */);
-									Graph.SetSynchronizationRequired();
-								}
-							}
-						}
-					}));
+				GraphEditorCommands->MapAction(FEditorCommands::Get().UpdateNodeClass,
+					FExecuteAction::CreateSP(this, &FEditor::UpdateSelectedNodeClasses));
 			}
 
 			SGraphEditor::FGraphEditorEvents GraphEvents;
@@ -1914,7 +1895,19 @@ namespace Metasound
 						FMetasoundFrontendLiteral LiteralValue;
 						if (FGraphBuilder::GetPinLiteral(*Pin, LiteralValue))
 						{
-							InputHandle->SetLiteral(LiteralValue);
+							if (const FMetasoundFrontendLiteral* ClassDefault = InputHandle->GetClassDefaultLiteral())
+							{
+								// Check equivalence with class default and don't set if they are equal. Copied node
+								// pin has no information to indicate whether or not the literal was already set.
+								if (!LiteralValue.IsEqual(*ClassDefault))
+								{
+									InputHandle->SetLiteral(LiteralValue);
+								}
+							}
+							else
+							{
+								InputHandle->SetLiteral(LiteralValue);
+							}
 						}
 					}
 
@@ -2077,6 +2070,49 @@ namespace Metasound
 			}
 		}
 
+		void FEditor::UpdateSelectedNodeClasses()
+		{
+			using namespace Metasound::Frontend;
+
+			const FScopedTransaction Transaction(LOCTEXT("NodeVersionUpdate", "Update MetaSound Node(s) Class(es)"));
+			check(Metasound);
+			Metasound->Modify();
+
+			UMetasoundEditorGraph& Graph = GetMetaSoundGraphChecked();
+			Graph.Modify();
+
+			bool bReplacedNodes = false;
+			const FGraphPanelSelectionSet Selection = MetasoundGraphEditor->GetSelectedNodes();
+			for (UObject* Object : Selection)
+			{
+				if (UMetasoundEditorGraphExternalNode* ExternalNode = Cast<UMetasoundEditorGraphExternalNode>(Object))
+				{
+					FMetasoundFrontendVersionNumber HighestVersion = ExternalNode->FindHighestVersionInRegistry();
+					Metasound::Frontend::FConstNodeHandle NodeHandle = ExternalNode->GetConstNodeHandle();
+					const FMetasoundFrontendClassMetadata& Metadata = NodeHandle->GetClassMetadata();
+					const bool bHasNewVersion = HighestVersion.IsValid() && HighestVersion > Metadata.GetVersion();
+
+					const FNodeRegistryKey RegistryKey = NodeRegistryKey::CreateKey(Metadata);
+					const bool bIsClassNative = FMetasoundFrontendRegistryContainer::Get()->IsNodeNative(RegistryKey);
+
+					if (bHasNewVersion || !bIsClassNative)
+					{
+						FNodeHandle ExistingNode = ExternalNode->GetNodeHandle();
+						FNodeHandle NewNode = ExistingNode->ReplaceWithVersion(HighestVersion);
+						bReplacedNodes = true;
+					}
+				}
+			}
+
+			if (bReplacedNodes)
+			{
+				FDocumentHandle DocumentHandle = Graph.GetDocumentHandle();
+				DocumentHandle->RemoveUnreferencedDependencies();
+				DocumentHandle->SynchronizeDependencyMetadata();
+				Graph.SetSynchronizationRequired();
+			}
+		}
+
 		bool FEditor::CanPasteNodes()
 		{
 			if (!IsGraphEditable())
diff --git a/Engine/Plugins/Runtime/Metasound/Source/MetasoundEditor/Private/MetasoundEditor.h b/Engine/Plugins/Runtime/Metasound/Source/MetasoundEditor/Private/MetasoundEditor.h
--- a/Engine/Plugins/Runtime/Metasound/Source/MetasoundEditor/Private/MetasoundEditor.h
+++ b/Engine/Plugins/Runtime/Metasound/Source/MetasoundEditor/Private/MetasoundEditor.h
@@ -136,6 +136,9 @@ namespace Metasound
 			/** Forces all UX pertaining to the root graph's interface to be refreshed. */
 			void RefreshInterface();
 
+			/** Updates selected node classes to highest class found in the MetaSound Class Registry. */
+			void UpdateSelectedNodeClasses();
+
 			/* Whether the displayed graph is marked as editable */
 			bool IsGraphEditable() const;
 
diff --git a/Engine/Plugins/Runtime/Metasound/Source/MetasoundEditor/Private/MetasoundEditorCommands.cpp b/Engine/Plugins/Runtime/Metasound/Source/MetasoundEditor/Private/MetasoundEditorCommands.cpp
--- a/Engine/Plugins/Runtime/Metasound/Source/MetasoundEditor/Private/MetasoundEditorCommands.cpp
+++ b/Engine/Plugins/Runtime/Metasound/Source/MetasoundEditor/Private/MetasoundEditorCommands.cpp
@@ -23,7 +23,7 @@ namespace Metasound
 			UI_COMMAND(EditMetasoundSettings, "MetaSound", "Edit MetaSound settings", EUserInterfaceActionType::Button, FInputChord());
 			UI_COMMAND(EditSourceSettings, "Source", "Edit source specific settings", EUserInterfaceActionType::Button, FInputChord());
 
-			UI_COMMAND(UpdateNodes, "Update Nodes", "Update selected node(s) that have an available update.", EUserInterfaceActionType::Button, FInputChord());
+			UI_COMMAND(UpdateNodeClass, "Update Node Class", "Update selected node(s) class(es) that have an available update.", EUserInterfaceActionType::Button, FInputChord());
 
 			UI_COMMAND(ConvertFromPreset, "Convert From Preset", "Converts this preset to a fully accessible MetaSound.", EUserInterfaceActionType::Button, FInputChord());
 			UI_COMMAND(Delete, "Delete Selected", "Delete selected items.", EUserInterfaceActionType::None, FInputChord(EKeys::Delete));
diff --git a/Engine/Plugins/Runtime/Metasound/Source/MetasoundEditor/Private/MetasoundEditorGraphBuilder.cpp b/Engine/Plugins/Runtime/Metasound/Source/MetasoundEditor/Private/MetasoundEditorGraphBuilder.cpp
--- a/Engine/Plugins/Runtime/Metasound/Source/MetasoundEditor/Private/MetasoundEditorGraphBuilder.cpp
+++ b/Engine/Plugins/Runtime/Metasound/Source/MetasoundEditor/Private/MetasoundEditorGraphBuilder.cpp
@@ -670,7 +670,11 @@ namespace Metasound
 			{
 				case ELiteralType::Boolean:
 				{
-					OutDefaultLiteral.Set(FCString::ToBool(*InStringValue));
+					// Currently don't support triggers being initialized to boolean in-graph
+					if (GetMetasoundDataTypeName<FTrigger>() != TypeName)
+					{
+						OutDefaultLiteral.Set(FCString::ToBool(*InStringValue));
+					}
 				}
 				break;
 
@@ -2135,7 +2139,7 @@ namespace Metasound
 						const FName NodeName = NodeHandle->GetNodeName();
 						const FGuid VertexID = GraphHandle->GetVertexIDForInputVertex(NodeName);
 						FMetasoundFrontendLiteral DefaultLiteral = GraphHandle->GetDefaultInput(VertexID);
-						if (!DefaultLiteral.IsEquivalent(Literal->GetDefault()))
+						if (!DefaultLiteral.IsEqual(Literal->GetDefault()))
 						{
 							if (DefaultLiteral.GetType() != EMetasoundFrontendLiteralType::None)
 							{
diff --git a/Engine/Plugins/Runtime/Metasound/Source/MetasoundEditor/Private/MetasoundEditorGraphNode.cpp b/Engine/Plugins/Runtime/Metasound/Source/MetasoundEditor/Private/MetasoundEditorGraphNode.cpp
--- a/Engine/Plugins/Runtime/Metasound/Source/MetasoundEditor/Private/MetasoundEditorGraphNode.cpp
+++ b/Engine/Plugins/Runtime/Metasound/Source/MetasoundEditor/Private/MetasoundEditorGraphNode.cpp
@@ -720,55 +720,6 @@ void UMetasoundEditorGraphExternalNode::CacheTitle()
 	CachedTitle = FGraphBuilder::GetDisplayName(*NodeHandle, bIncludeNamespace);
 }
 
-UMetasoundEditorGraphExternalNode* UMetasoundEditorGraphExternalNode::UpdateToVersion(const FMetasoundFrontendVersionNumber& InNewVersion, bool bInPropagateErrorMessages)
-{
-	using namespace Metasound::Editor;
-	using namespace Metasound::Frontend;
-
-	FNodeHandle InitNodeHandle = GetNodeHandle();
-	const FGuid InitNodeID = GetNodeID();
-	const FMetasoundFrontendVersionNumber InitVersion = InitNodeHandle->GetClassMetadata().GetVersion();
-
-	FNodeHandle ReplacementNode = InitNodeHandle->ReplaceWithVersion(InNewVersion);
-	const FGuid ReplacementNodeID = ReplacementNode->GetID();
-
-	if (!ensure(InitNodeID != ReplacementNodeID))
-	{
-		return this;
-	}
-
-	const FMetasoundFrontendNodeStyle& Style = ReplacementNode->GetNodeStyle();
-	FVector2D NodeLocation = FVector2D::ZeroVector;
-	const FVector2D* StyleLocation = Style.Display.Locations.Find(NodeGuid);
-	if (ensure(StyleLocation))
-	{
-		NodeLocation = *StyleLocation;
-	}
-
-	UObject& MetaSound = GetMetasoundChecked();
-	UMetasoundEditorGraphExternalNode* ReplacementEdNode = FGraphBuilder::AddExternalNode(MetaSound, ReplacementNode, NodeLocation, false /* bInSelectNewNode */);
-
-	if (!ensure(ReplacementEdNode))
-	{
-		return this;
-	}
-
-	GetGraph()->RemoveNode(this);
-	MetaSound.Modify();
-
-	if (InitNodeID != ReplacementNodeID)
-	{
-		if (bInPropagateErrorMessages)
-		{
-			ReplacementEdNode->bHasCompilerMessage = bHasCompilerMessage;
-			ReplacementEdNode->ErrorMsg = ErrorMsg;
-			ReplacementEdNode->ErrorType = ErrorType;
-		}
-	}
-
-	return ReplacementEdNode;
-}
-
 bool UMetasoundEditorGraphExternalNode::Validate(Metasound::Editor::FGraphNodeValidationResult& OutResult)
 {
 	using namespace Metasound::Editor;
diff --git a/Engine/Plugins/Runtime/Metasound/Source/MetasoundEditor/Private/MetasoundEditorGraphSchema.cpp b/Engine/Plugins/Runtime/Metasound/Source/MetasoundEditor/Private/MetasoundEditorGraphSchema.cpp
--- a/Engine/Plugins/Runtime/Metasound/Source/MetasoundEditor/Private/MetasoundEditorGraphSchema.cpp
+++ b/Engine/Plugins/Runtime/Metasound/Source/MetasoundEditor/Private/MetasoundEditorGraphSchema.cpp
@@ -1046,6 +1046,7 @@ void UMetasoundEditorGraphSchema::GetGraphContextActions(FGraphContextMenuBuilde
 void UMetasoundEditorGraphSchema::GetContextMenuActions(class UToolMenu* Menu, class UGraphNodeContextMenuContext* Context) const
 {
 	using namespace Metasound::Editor;
+	using namespace Metasound::Frontend;
 
 	if (!Context->Pin && Context->Node && Context->Node->IsA<UMetasoundEditorGraphNode>())
 	{
@@ -1061,11 +1062,17 @@ void UMetasoundEditorGraphSchema::GetContextMenuActions(class UToolMenu* Menu, c
 		// and node registry is reporting a major update is available.
 		if (const UMetasoundEditorGraphExternalNode* ExternalNode = Cast<UMetasoundEditorGraphExternalNode>(Context->Node))
 		{
-			FMetasoundFrontendVersionNumber MajorUpdateVersion = ExternalNode->FindHighestVersionInRegistry();
+			FMetasoundFrontendVersionNumber HighestVersion = ExternalNode->FindHighestVersionInRegistry();
 			Metasound::Frontend::FConstNodeHandle NodeHandle = ExternalNode->GetConstNodeHandle();
-			if (MajorUpdateVersion.IsValid() && MajorUpdateVersion > NodeHandle->GetClassMetadata().GetVersion())
+			const FMetasoundFrontendClassMetadata& Metadata = NodeHandle->GetClassMetadata();
+			const bool bHasNewVersion = HighestVersion.IsValid() && HighestVersion > Metadata.GetVersion();
+
+			const FNodeRegistryKey RegistryKey = NodeRegistryKey::CreateKey(Metadata);
+			const bool bIsClassNative = FMetasoundFrontendRegistryContainer::Get()->IsNodeNative(RegistryKey);
+
+			if (bHasNewVersion || !bIsClassNative)
 			{
-				Section.AddMenuEntry(FEditorCommands::Get().UpdateNodes);
+				Section.AddMenuEntry(FEditorCommands::Get().UpdateNodeClass);
 			}
 		}
 	}
diff --git a/Engine/Plugins/Runtime/Metasound/Source/MetasoundEditor/Private/SMetasoundGraphPin.h b/Engine/Plugins/Runtime/Metasound/Source/MetasoundEditor/Private/SMetasoundGraphPin.h
--- a/Engine/Plugins/Runtime/Metasound/Source/MetasoundEditor/Private/SMetasoundGraphPin.h
+++ b/Engine/Plugins/Runtime/Metasound/Source/MetasoundEditor/Private/SMetasoundGraphPin.h
@@ -144,7 +144,9 @@ namespace Metasound
 								FConstInputHandle InputHandle = GetConstInputHandle();
 								if (const FMetasoundFrontendLiteral* Literal = InputHandle->GetLiteral())
 								{
-									if (Literal->GetType() != EMetasoundFrontendLiteralType::None)
+									const bool bIsDefaultConstructed = Literal->GetType() == EMetasoundFrontendLiteralType::None;
+									const bool bIsTriggerDataType = InputHandle->GetDataType() == GetMetasoundDataTypeName<FTrigger>();
+									if (!bIsDefaultConstructed && !bIsTriggerDataType)
 									{
 										return EVisibility::Visible;
 									}
diff --git a/Engine/Plugins/Runtime/Metasound/Source/MetasoundEditor/Public/MetasoundEditorCommands.h b/Engine/Plugins/Runtime/Metasound/Source/MetasoundEditor/Public/MetasoundEditorCommands.h
--- a/Engine/Plugins/Runtime/Metasound/Source/MetasoundEditor/Public/MetasoundEditorCommands.h
+++ b/Engine/Plugins/Runtime/Metasound/Source/MetasoundEditor/Public/MetasoundEditorCommands.h
@@ -52,8 +52,8 @@ namespace Metasound
 			/** Shows the Metasound's top-level general object settings in the Inspector */
 			TSharedPtr<FUICommandInfo> EditMetasoundSettings;
 
-			// Run updates on selected nodes if required by any subset of selected nodes
-			TSharedPtr<FUICommandInfo> UpdateNodes;
+			// Run updates on selected node(s) class(es) if required by any subset of selected nodes
+			TSharedPtr<FUICommandInfo> UpdateNodeClass;
 
 			// Converts a MetaSound from a restricted, preset edit state to a fully accessible MetaSound.
 			TSharedPtr<FUICommandInfo> ConvertFromPreset;
diff --git a/Engine/Plugins/Runtime/Metasound/Source/MetasoundEditor/Public/MetasoundEditorGraphNode.h b/Engine/Plugins/Runtime/Metasound/Source/MetasoundEditor/Public/MetasoundEditorGraphNode.h
--- a/Engine/Plugins/Runtime/Metasound/Source/MetasoundEditor/Public/MetasoundEditorGraphNode.h
+++ b/Engine/Plugins/Runtime/Metasound/Source/MetasoundEditor/Public/MetasoundEditorGraphNode.h
@@ -208,10 +208,6 @@ public:
 	FMetasoundFrontendVersionNumber FindHighestMinorVersionInRegistry() const;
 	bool CanAutoUpdate() const;
 
-	// Attempts to replace this node with a new one of the same class and given version number.
-	// If this node is already of the given version, returns itself. If update fails, returns this node.
-	UMetasoundEditorGraphExternalNode* UpdateToVersion(const FMetasoundFrontendVersionNumber& InNewVersion, bool bInPropagateErrorMessages);
-
 	// Validates node and returns whether or not the node is valid.
 	bool Validate(Metasound::Editor::FGraphNodeValidationResult& OutResult);
 
diff --git a/Engine/Plugins/Runtime/Metasound/Source/MetasoundFrontend/Private/MetasoundFrontendLiteral.cpp b/Engine/Plugins/Runtime/Metasound/Source/MetasoundFrontend/Private/MetasoundFrontendLiteral.cpp
--- a/Engine/Plugins/Runtime/Metasound/Source/MetasoundFrontend/Private/MetasoundFrontendLiteral.cpp
+++ b/Engine/Plugins/Runtime/Metasound/Source/MetasoundFrontend/Private/MetasoundFrontendLiteral.cpp
@@ -192,7 +192,7 @@ bool FMetasoundFrontendLiteral::IsArray() const
 	}
 }
 
-bool FMetasoundFrontendLiteral::IsEquivalent(const FMetasoundFrontendLiteral& InOther) const
+bool FMetasoundFrontendLiteral::IsEqual(const FMetasoundFrontendLiteral& InOther) const
 {
 	if (InOther.GetType() != GetType())
 	{
diff --git a/Engine/Plugins/Runtime/Metasound/Source/MetasoundFrontend/Private/MetasoundFrontendNodeController.cpp b/Engine/Plugins/Runtime/Metasound/Source/MetasoundFrontend/Private/MetasoundFrontendNodeController.cpp
--- a/Engine/Plugins/Runtime/Metasound/Source/MetasoundFrontend/Private/MetasoundFrontendNodeController.cpp
+++ b/Engine/Plugins/Runtime/Metasound/Source/MetasoundFrontend/Private/MetasoundFrontendNodeController.cpp
@@ -724,7 +724,7 @@ namespace Metasound
 
 		FNodeHandle FBaseNodeController::ReplaceWithVersion(const FMetasoundFrontendVersionNumber& InNewVersion)
 		{
-			const FMetasoundFrontendClassMetadata& Metadata = GetClassMetadata();
+			const FMetasoundFrontendClassMetadata Metadata = GetClassMetadata();
 			const TArray<FMetasoundFrontendClass> Versions = ISearchEngine::Get().FindClassesWithName(Metadata.GetClassName().ToNodeClassName(), false /* bInSortByVersion */);
 
 			auto IsClassOfNewVersion = [InNewVersion](const FMetasoundFrontendClass& RegisteredClass)
@@ -821,6 +821,7 @@ namespace Metasound
 				return this->AsShared();
 			}
 
+			Style.bMessageNodeUpdated = ReplacementNode->GetClassMetadata().GetVersion() > Metadata.GetVersion();
 			ReplacementNode->SetNodeStyle(Style);
 
 			ReplacementNode->IterateInputs([Connections = &InputConnections](FInputHandle InputHandle)
@@ -899,12 +900,17 @@ namespace Metasound
 			Algo::Transform(OutInterfaceUpdates.RegistryClass.Interface.Inputs, OutInterfaceUpdates.AddedInputs, [&](const FMetasoundFrontendClassInput& Input) { return &Input; });
 			for (const FMetasoundFrontendClassInput& Input : NodeClassInterface.Inputs)
 			{
-				auto IsFunctionalEquivalent = [NodeClassInput = &Input](const FMetasoundFrontendClassInput* Iter)
+				auto IsEquivalent = [NodeClassInput = &Input](const FMetasoundFrontendClassInput* Iter)
 				{
-					return FMetasoundFrontendClassVertex::IsFunctionalEquivalent(*NodeClassInput, *Iter);
+					const bool bDefaultEquivalent = Iter->DefaultLiteral.IsEqual(NodeClassInput->DefaultLiteral);
+					if (bDefaultEquivalent)
+					{
+						return FMetasoundFrontendClassVertex::IsFunctionalEquivalent(*NodeClassInput, *Iter);
+					}
+					return false;
 				};
 
-				const int32 Index = OutInterfaceUpdates.AddedInputs.FindLastByPredicate(IsFunctionalEquivalent);
+				const int32 Index = OutInterfaceUpdates.AddedInputs.FindLastByPredicate(IsEquivalent);
 				if (Index == INDEX_NONE)
 				{
 					OutInterfaceUpdates.RemovedInputs.Add(&Input);
diff --git a/Engine/Plugins/Runtime/Metasound/Source/MetasoundFrontend/Private/MetasoundFrontendTransform.cpp b/Engine/Plugins/Runtime/Metasound/Source/MetasoundFrontend/Private/MetasoundFrontendTransform.cpp
--- a/Engine/Plugins/Runtime/Metasound/Source/MetasoundFrontend/Private/MetasoundFrontendTransform.cpp
+++ b/Engine/Plugins/Runtime/Metasound/Source/MetasoundFrontend/Private/MetasoundFrontendTransform.cpp
@@ -540,9 +540,6 @@ namespace Metasound
 					FNodeHandle ExistingNode = Pair.Key;
 					FMetasoundFrontendVersionNumber InitialVersion = ExistingNode->GetClassMetadata().GetVersion();
 					FNodeHandle NewNode = ExistingNode->ReplaceWithVersion(Pair.Value);
-					FMetasoundFrontendNodeStyle Style = NewNode->GetNodeStyle();
-					Style.bMessageNodeUpdated = NewNode->GetClassMetadata().GetVersion() > InitialVersion;
-					NewNode->SetNodeStyle(Style);
 				}
 
 				InDocument->RemoveUnreferencedDependencies();
@@ -860,7 +857,7 @@ namespace Metasound
 							{
 								if (const FMetasoundFrontendLiteral* Literal = Input->GetClassDefaultLiteral())
 								{
-									if (!Literal->IsEquivalent(DefaultLiteral))
+									if (!Literal->IsEqual(DefaultLiteral))
 									{
 										Input->SetLiteral(DefaultLiteral);
 									}
diff --git a/Engine/Plugins/Runtime/Metasound/Source/MetasoundFrontend/Public/MetasoundFrontendDocument.h b/Engine/Plugins/Runtime/Metasound/Source/MetasoundFrontend/Public/MetasoundFrontendDocument.h
--- a/Engine/Plugins/Runtime/Metasound/Source/MetasoundFrontend/Public/MetasoundFrontendDocument.h
+++ b/Engine/Plugins/Runtime/Metasound/Source/MetasoundFrontend/Public/MetasoundFrontendDocument.h
@@ -924,15 +924,15 @@ struct METASOUNDFRONTEND_API FMetasoundFrontendClassName
 	FMetasoundFrontendClassName(const Metasound::FNodeClassName& InName);
 
 	// Namespace of class.
-	UPROPERTY(VisibleAnywhere, Category = General)
+	UPROPERTY(EditAnywhere, Category = General)
 	FName Namespace;
 
 	// Name of class.
-	UPROPERTY(VisibleAnywhere, Category = General)
+	UPROPERTY(EditAnywhere, Category = General)
 	FName Name;
 
 	// Variant of class. The Variant is used to describe an equivalent class which performs the same operation but on differing types.
-	UPROPERTY(VisibleAnywhere, Category = General)
+	UPROPERTY(EditAnywhere, Category = General)
 	FName Variant;
 
 	// Returns a full name of the class.
diff --git a/Engine/Plugins/Runtime/Metasound/Source/MetasoundFrontend/Public/MetasoundFrontendLiteral.h b/Engine/Plugins/Runtime/Metasound/Source/MetasoundFrontend/Public/MetasoundFrontendLiteral.h
--- a/Engine/Plugins/Runtime/Metasound/Source/MetasoundFrontend/Public/MetasoundFrontendLiteral.h
+++ b/Engine/Plugins/Runtime/Metasound/Source/MetasoundFrontend/Public/MetasoundFrontendLiteral.h
@@ -95,7 +95,7 @@ public:
 	bool IsArray() const;
 
 	// Returns whether the other literal is value equivalent
-	bool IsEquivalent(const FMetasoundFrontendLiteral& InOther) const;
+	bool IsEqual(const FMetasoundFrontendLiteral& InOther) const;
 
 	// Returns true if the literal is in a valid state (Type != EMetasoundFrontendLiteralType::Invalid)
 	bool IsValid() const;
