Fix for metasound waveplayer pop on PCM audio

diff --git a/Engine/Plugins/Runtime/Metasound/Source/MetasoundEngine/Private/MetasoundWaveProxyReader.cpp b/Engine/Plugins/Runtime/Metasound/Source/MetasoundEngine/Private/MetasoundWaveProxyReader.cpp
--- a/Engine/Plugins/Runtime/Metasound/Source/MetasoundEngine/Private/MetasoundWaveProxyReader.cpp
+++ b/Engine/Plugins/Runtime/Metasound/Source/MetasoundEngine/Private/MetasoundWaveProxyReader.cpp
@@ -298,7 +298,10 @@ namespace Metasound
 		TArrayView<float> OutBufferView{OutBuffer};
 		int32 NumSamplesUnset = OutBuffer.Num();
 		int32 NumSamplesCopied = 0;
-		bool bCanProduceMoreAudio = true;
+
+		bool bDecoderOutputHasMoreData = DecoderOutput.Num() > 0;
+		bool bDecoderCanDecodeMoreData = bIsDecoderValid && (EDecodeResult::MoreDataRemaining == DecodeResult);
+		bool bCanProduceMoreAudio = bDecoderOutputHasMoreData || bDecoderCanDecodeMoreData;
 
 		while ((NumSamplesUnset > 0) && bCanProduceMoreAudio)
 		{
@@ -308,18 +311,31 @@ namespace Metasound
 			NumSamplesCopied += NumSamplesCopiedThisLoop;
 			NumSamplesUnset -= NumSamplesCopiedThisLoop;
 			OutBufferView = OutBufferView.Slice(NumSamplesCopiedThisLoop, NumSamplesUnset);
+			
+			if (Settings.bIsLooping)
+			{
+				// Seek to loop start if we have used up all our decodable samples or 
+				// are at the loop boundary.
+				if (0 == DecoderOutput.Num())
+				{
+					if ((!bDecoderCanDecodeMoreData) || (CurrentFrameIndex >= LoopEndFrameIndex))
+					{
+						SeekToTime(Settings.LoopStartTimeInSeconds);
+						bDecoderCanDecodeMoreData = bIsDecoderValid && (EDecodeResult::MoreDataRemaining == DecodeResult);
+					}
+				}
+			}
 
 			// Determine if we can / should decode more data. 
-			if (bIsDecoderValid && (NumSamplesUnset > 0) && (DecoderOutput.Num() == 0))
+			if ((NumSamplesUnset > 0) && (DecoderOutput.Num() == 0) && bDecoderCanDecodeMoreData)
 			{
 				// Looping is handle within the FWaveProxyReader instead of 
 				// within the decoder.
 				DecodeResult = Decoder->Decode(false /* bIsLooping */); 
-
+				bDecoderCanDecodeMoreData = EDecodeResult::MoreDataRemaining == DecodeResult;
 			}
 
-			const bool bDecoderOutputHasMoreData = DecoderOutput.Num() > 0;
-			const bool bDecoderCanDecodeMoreData = bIsDecoderValid && (EDecodeResult::MoreDataRemaining == DecodeResult);
+			bDecoderOutputHasMoreData = DecoderOutput.Num() > 0;
 			bCanProduceMoreAudio = bDecoderOutputHasMoreData || bDecoderCanDecodeMoreData;
 		}
 
@@ -361,25 +377,17 @@ namespace Metasound
 				{
 					// Rewind sample counters if the loop boundary was overshot.
 					NumFramesCopied -= (CurrentFrameIndex - LoopEndFrameIndex);
+					CurrentFrameIndex = LoopEndFrameIndex;
 					// If Settings.bIsLooping info was altered, NumFramesCopied can end up 
 					// negative if the current frame index is past the loop end frame. 
 					NumFramesCopied = FMath::Max(0, NumFramesCopied); 
 					NumSamplesCopied = NumFramesCopied * NumChannels;
 
-					// Remove any remaining samples in the decoder because
+					// Remove any remaining samples in the decoder because they 
+					// are past the end of the loop.
 					DecoderOutput.Reset();
 				}
 
-				// Determine whether or not to seek to the loop start.
-				const bool bSeekToLoopStart = bDidOvershoot
-					|| ((DecodeResult == EDecodeResult::Finished) && (DecoderOutput.Num() == 0)) // Decoder cannot produce more audio
-					|| (CurrentFrameIndex >= NumFramesInWave); // Reached end of file.
-
-				if (bSeekToLoopStart)
-				{
-					// Start decoder at beginning of loop
-					SeekToTime(Settings.LoopStartTimeInSeconds);
-				}
 			}
 		}
 
@@ -418,8 +426,15 @@ namespace Metasound
 		if (!Decoder.IsValid())
 		{
 			UE_LOG(LogMetaSound, Error, TEXT("Failed to create decoder (format:%s) for wave (package:%s)"), *Format.ToString(), *WaveProxy->GetPackageName().ToString());
+			DecodeResult = EDecodeResult::Fail;
 			return false;
 		}
+		else
+		{
+			// The DecodeResult needs to be set to a valid state incase the prior
+			// decoder finished or failed. 
+			DecodeResult = EDecodeResult::MoreDataRemaining;
+		}
 
 		// For non-seekable streaming waves, use a fallback method to seek
 		// to the start time. 
@@ -486,7 +501,7 @@ namespace Metasound
 			const int32 MinLoopDurationInFrames = FMath::CeilToInt(MinLoopDurationInSeconds * SampleRate);
 			const float LoopEndTime = Settings.LoopStartTimeInSeconds + Settings.LoopDurationInSeconds;
 			const int32 MinLoopEndFrameIndex = LoopStartFrameIndex + MinLoopDurationInFrames;
-			const int32 MaxLoopEndFrameIndex = NumFramesInWave;
+			const int32 MaxLoopEndFrameIndex = NumFramesInWave + 1;
 
 			LoopStartFrameIndex = FMath::FloorToInt(Settings.LoopStartTimeInSeconds * SampleRate);
 			LoopEndFrameIndex = FMath::Clamp(LoopEndTime * SampleRate, MinLoopEndFrameIndex, MaxLoopEndFrameIndex);
@@ -494,7 +509,7 @@ namespace Metasound
 		else
 		{
 			LoopStartFrameIndex = 0;
-			LoopEndFrameIndex = NumFramesInWave;
+			LoopEndFrameIndex = NumFramesInWave + 1;
 		}
 	}
 }
