- Fix for superfluous adding a default literal set to false to copied triggers which causes a generated duplicate input when building the core graph if input is connected
- Fix for errors/displaying trigger simulate on output/variable nodes
#rnx

diff --git a/Engine/Plugins/Runtime/Metasound/Source/MetasoundEditor/Private/MetasoundEditor.cpp b/Engine/Plugins/Runtime/Metasound/Source/MetasoundEditor/Private/MetasoundEditor.cpp
--- a/Engine/Plugins/Runtime/Metasound/Source/MetasoundEditor/Private/MetasoundEditor.cpp
+++ b/Engine/Plugins/Runtime/Metasound/Source/MetasoundEditor/Private/MetasoundEditor.cpp
@@ -1890,7 +1890,7 @@ namespace Metasound
 					}
 
 					FInputHandle InputHandle = FGraphBuilder::GetInputHandleFromPin(Pin);
-					if (InputHandle->IsValid())
+					if (InputHandle->IsValid() && InputHandle->GetDataType() != GetMetasoundDataTypeName<FTrigger>())
 					{
 						FMetasoundFrontendLiteral LiteralValue;
 						if (FGraphBuilder::GetPinLiteral(*Pin, LiteralValue))
diff --git a/Engine/Plugins/Runtime/Metasound/Source/MetasoundEditor/Private/SMetasoundGraphNode.cpp b/Engine/Plugins/Runtime/Metasound/Source/MetasoundEditor/Private/SMetasoundGraphNode.cpp
--- a/Engine/Plugins/Runtime/Metasound/Source/MetasoundEditor/Private/SMetasoundGraphNode.cpp
+++ b/Engine/Plugins/Runtime/Metasound/Source/MetasoundEditor/Private/SMetasoundGraphNode.cpp
@@ -540,9 +540,9 @@ namespace Metasound
 			{
 				if (UMetasoundEditorGraphMemberNode* Node = Cast<UMetasoundEditorGraphMemberNode>(&GetMetaSoundNode()))
 				{
-					if (UMetasoundEditorGraphVertex* Vertex = Cast<UMetasoundEditorGraphVertex>(Node->GetMember()))
+					if (UMetasoundEditorGraphInput* Input = Cast<UMetasoundEditorGraphInput>(Node->GetMember()))
 					{
-						if (UMetasoundEditorGraphMemberDefaultLiteral* Literal = Vertex->GetLiteral())
+						if (UMetasoundEditorGraphMemberDefaultLiteral* Literal = Input->GetLiteral())
 						{
 							TAttribute<EVisibility> SimVisibility = GetSimulationVisibilityAttribute();
 							TAttribute<bool> SimEnablement = true;
