Compile out editor-only data from MetaSoundFrontend
- Move Author field to FString to avoid generated text collisions (Authors should always be culture invariant anyway) & add option to override default in MetaSound user settings

diff --git a/Engine/Plugins/Runtime/AudioModulation/Source/AudioModulation/Private/SoundModulatorAsset.cpp b/Engine/Plugins/Runtime/AudioModulation/Source/AudioModulation/Private/SoundModulatorAsset.cpp
--- a/Engine/Plugins/Runtime/AudioModulation/Source/AudioModulation/Private/SoundModulatorAsset.cpp
+++ b/Engine/Plugins/Runtime/AudioModulation/Source/AudioModulation/Private/SoundModulatorAsset.cpp
@@ -5,7 +5,7 @@
 
 namespace AudioModulation
 {
-	const FText PluginAuthor = NSLOCTEXT("AudioModulation", "PluginAuthor", "Epic Games, Inc.");
+	const FString PluginAuthor = TEXT("Epic Games, Inc.");
 	const FText PluginNodeMissingPrompt = NSLOCTEXT("AudioModulation", "DefaultMissingNodePrompt", "The node was likely removed, renamed, or the AudioModulation plugin is not loaded.");
 
 	FSoundModulatorAsset::FSoundModulatorAsset(const Audio::IProxyDataPtr& InProxyPtr)
diff --git a/Engine/Plugins/Runtime/AudioModulation/Source/AudioModulation/Public/SoundModulatorAsset.h b/Engine/Plugins/Runtime/AudioModulation/Source/AudioModulation/Public/SoundModulatorAsset.h
--- a/Engine/Plugins/Runtime/AudioModulation/Source/AudioModulation/Public/SoundModulatorAsset.h
+++ b/Engine/Plugins/Runtime/AudioModulation/Source/AudioModulation/Public/SoundModulatorAsset.h
@@ -9,7 +9,7 @@
 
 namespace AudioModulation
 {
-	extern const FText AUDIOMODULATION_API PluginAuthor;
+	extern const FString AUDIOMODULATION_API PluginAuthor;
 	extern const FText AUDIOMODULATION_API PluginNodeMissingPrompt;
 
 	class AUDIOMODULATION_API FSoundModulatorAsset
diff --git a/Engine/Plugins/Runtime/Metasound/Source/MetasoundEditor/Private/MetasoundEditorGraphBuilder.cpp b/Engine/Plugins/Runtime/Metasound/Source/MetasoundEditor/Private/MetasoundEditorGraphBuilder.cpp
--- a/Engine/Plugins/Runtime/Metasound/Source/MetasoundEditor/Private/MetasoundEditorGraphBuilder.cpp
+++ b/Engine/Plugins/Runtime/Metasound/Source/MetasoundEditor/Private/MetasoundEditorGraphBuilder.cpp
@@ -1121,7 +1121,7 @@ namespace Metasound
 			Metadata.SetVersion({ 1, 0 });
 			Metadata.SetDisplayName(FText::FromString(InMetaSound.GetName()));
 			Metadata.SetType(EMetasoundFrontendClassType::Graph);
-			Metadata.SetAuthor(FText::FromString(InAuthor));
+			Metadata.SetAuthor(InAuthor);
 
 			FMetasoundAssetBase* MetaSoundAsset = IMetasoundUObjectRegistry::Get().GetObjectAsAssetBase(&InMetaSound);
 			check(MetaSoundAsset);
diff --git a/Engine/Plugins/Runtime/Metasound/Source/MetasoundEditor/Private/MetasoundEditorGraphSchema.cpp b/Engine/Plugins/Runtime/Metasound/Source/MetasoundEditor/Private/MetasoundEditorGraphSchema.cpp
--- a/Engine/Plugins/Runtime/Metasound/Source/MetasoundEditor/Private/MetasoundEditorGraphSchema.cpp
+++ b/Engine/Plugins/Runtime/Metasound/Source/MetasoundEditor/Private/MetasoundEditorGraphSchema.cpp
@@ -1547,7 +1547,7 @@ void UMetasoundEditorGraphSchema::GetConversionActions(FGraphActionMenuBuilder&
 		FMetasoundFrontendClassMetadata Metadata = FrontendClass.Metadata;
 		const FText Tooltip = Metadata.GetAuthor().IsEmpty()
 			? Metadata.GetDescription()
-			: FText::Format(LOCTEXT("MetasoundTooltipAuthorFormat1", "{0}\nAuthor: {1}"), Metadata.GetDescription(), Metadata.GetAuthor());
+			: FText::Format(LOCTEXT("MetasoundTooltipAuthorFormat1", "{0}\nAuthor: {1}"), Metadata.GetDescription(), FText::FromString(Metadata.GetAuthor()));
 		if (!Metadata.GetCategoryHierarchy().IsEmpty() && !Metadata.GetCategoryHierarchy()[0].CompareTo(NodeCategories::Conversions))
 		{
 			FText KeywordsText = FText::Join(SchemaPrivate::KeywordDelim, Metadata.GetKeywords());
@@ -1678,7 +1678,7 @@ void UMetasoundEditorGraphSchema::GetFunctionActions(FGraphActionMenuBuilder& Ac
 		FMetasoundFrontendClassMetadata Metadata = FrontendClass.Metadata;
 		const FText Tooltip = Metadata.GetAuthor().IsEmpty()
 			? Metadata.GetDescription()
-			: FText::Format(LOCTEXT("MetasoundTooltipAuthorFormat2", "{0}\nAuthor: {1}"), Metadata.GetDescription(), Metadata.GetAuthor());
+			: FText::Format(LOCTEXT("MetasoundTooltipAuthorFormat2", "{0}\nAuthor: {1}"), Metadata.GetDescription(), FText::FromString(Metadata.GetAuthor()));
 
 		if (Metadata.GetCategoryHierarchy().IsEmpty() || Metadata.GetCategoryHierarchy()[0].CompareTo(Metasound::NodeCategories::Conversions))
 		{
diff --git a/Engine/Plugins/Runtime/Metasound/Source/MetasoundEditor/Private/MetasoundEditorSettings.h b/Engine/Plugins/Runtime/Metasound/Source/MetasoundEditor/Private/MetasoundEditorSettings.h
--- a/Engine/Plugins/Runtime/Metasound/Source/MetasoundEditor/Private/MetasoundEditorSettings.h
+++ b/Engine/Plugins/Runtime/Metasound/Source/MetasoundEditor/Private/MetasoundEditorSettings.h
@@ -22,6 +22,12 @@ class METASOUNDEDITOR_API UMetasoundEditorSettings : public UObject
 	GENERATED_UCLASS_BODY()
 
 public:
+	/** Default author title to use when authoring a new
+	  * MetaSound.  If empty, uses machine name by default.
+	  */
+	UPROPERTY(EditAnywhere, config, Category=General)
+	FString DefaultAuthor;
+
 	/** Default pin type color */
 	UPROPERTY(EditAnywhere, config, Category=PinColors)
 	FLinearColor DefaultPinTypeColor;
diff --git a/Engine/Plugins/Runtime/Metasound/Source/MetasoundEditor/Private/MetasoundFactory.cpp b/Engine/Plugins/Runtime/Metasound/Source/MetasoundEditor/Private/MetasoundFactory.cpp
--- a/Engine/Plugins/Runtime/Metasound/Source/MetasoundEditor/Private/MetasoundFactory.cpp
+++ b/Engine/Plugins/Runtime/Metasound/Source/MetasoundEditor/Private/MetasoundFactory.cpp
@@ -6,6 +6,7 @@
 #include "MetasoundEditorGraph.h"
 #include "MetasoundEditorGraphBuilder.h"
 #include "MetasoundEditorGraphSchema.h"
+#include "MetasoundEditorSettings.h"
 #include "MetasoundFrontendTransform.h"
 #include "MetasoundSource.h"
 
@@ -23,7 +24,16 @@ namespace Metasound
 			T* MetaSoundObject = NewObject<T>(InParent, InName, InFlags);
 			check(MetaSoundObject);
 
-			FGraphBuilder::InitMetaSound(*MetaSoundObject, UKismetSystemLibrary::GetPlatformUserName());
+			FString Author = UKismetSystemLibrary::GetPlatformUserName();
+			if (const UMetasoundEditorSettings* EditorSettings = GetDefault<UMetasoundEditorSettings>())
+			{
+				if (!EditorSettings->DefaultAuthor.IsEmpty())
+				{
+					Author = EditorSettings->DefaultAuthor;
+				}
+			}
+
+			FGraphBuilder::InitMetaSound(*MetaSoundObject, Author);
 
 			if (InReferencedMetaSoundObject)
 			{
diff --git a/Engine/Plugins/Runtime/Metasound/Source/MetasoundEngine/Private/MetasoundAssetSubsystem.cpp b/Engine/Plugins/Runtime/Metasound/Source/MetasoundEngine/Private/MetasoundAssetSubsystem.cpp
--- a/Engine/Plugins/Runtime/Metasound/Source/MetasoundEngine/Private/MetasoundAssetSubsystem.cpp
+++ b/Engine/Plugins/Runtime/Metasound/Source/MetasoundEngine/Private/MetasoundAssetSubsystem.cpp
@@ -417,7 +417,10 @@ void UMetaSoundAssetSubsystem::ResetAssetClassDisplayName(const FAssetData& InAs
 	FMetasoundAssetBase* MetaSoundAsset = Metasound::IMetasoundUObjectRegistry::Get().GetObjectAsAssetBase(Object);
 	check(MetaSoundAsset);
 	FMetasoundFrontendGraphClass& Class = MetaSoundAsset->GetDocumentChecked().RootGraph;
+
+#if WITH_EDITOR
 	Class.Metadata.SetDisplayName(FText());
+#endif // WITH_EDITOR
 }
 
 void UMetaSoundAssetSubsystem::SearchAndIterateDirectoryAssets(const TArray<FDirectoryPath>& InDirectories, TFunctionRef<void(const FAssetData&)> InFunction)
diff --git a/Engine/Plugins/Runtime/Metasound/Source/MetasoundEngine/Private/MetasoundEngineArchetypes.cpp b/Engine/Plugins/Runtime/Metasound/Source/MetasoundEngine/Private/MetasoundEngineArchetypes.cpp
--- a/Engine/Plugins/Runtime/Metasound/Source/MetasoundEngine/Private/MetasoundEngineArchetypes.cpp
+++ b/Engine/Plugins/Runtime/Metasound/Source/MetasoundEngine/Private/MetasoundEngineArchetypes.cpp
@@ -99,31 +99,38 @@ namespace Metasound
 				// Inputs
 				FMetasoundFrontendClassVertex OnPlayTrigger;
 				OnPlayTrigger.Name = ArchetypesPrivate::GetOnPlayInputName();
-
 				OnPlayTrigger.TypeName = GetMetasoundDataTypeName<FTrigger>();
+				OnPlayTrigger.VertexID = FGuid::NewGuid();
+
+#if WITH_EDITOR
 				OnPlayTrigger.Metadata.SetDisplayName(LOCTEXT("OnPlay", "On Play"));
 				OnPlayTrigger.Metadata.SetDescription(LOCTEXT("OnPlayTriggerToolTip", "Trigger executed when this source is played."));
-				OnPlayTrigger.VertexID = FGuid::NewGuid();
+#endif // WITH_EDITOR
 
 				Interface.Inputs.Add(OnPlayTrigger);
 
 				// Outputs 
 				FMetasoundFrontendClassVertex OnFinished;
 				OnFinished.Name = ArchetypesPrivate::GetIsFinishedOutputName();
-
 				OnFinished.TypeName = GetMetasoundDataTypeName<FTrigger>();
+				OnFinished.VertexID = FGuid::NewGuid();
+
+#if WITH_EDITOR
 				OnFinished.Metadata.SetDisplayName(LOCTEXT("OnFinished", "On Finished"));
 				OnFinished.Metadata.SetDescription(LOCTEXT("OnFinishedToolTip", "Trigger executed to initiate stopping the source."));
-				OnFinished.VertexID = FGuid::NewGuid();
+#endif // WITH_EDITOR
 
 				Interface.Outputs.Add(OnFinished);
 
 				FMetasoundFrontendClassVertex GeneratedAudio;
 				GeneratedAudio.Name = GetAudioOutputName();
 				GeneratedAudio.TypeName = GetMetasoundDataTypeName<FMonoAudioFormat>();
+				GeneratedAudio.VertexID = FGuid::NewGuid();
+
+#if WITH_EDITOR
 				GeneratedAudio.Metadata.SetDisplayName(LOCTEXT("GeneratedMono", "Audio"));
 				GeneratedAudio.Metadata.SetDescription(LOCTEXT("GeneratedAudioToolTip", "The resulting output audio from this source."));
-				GeneratedAudio.VertexID = FGuid::NewGuid();
+#endif // WITH_EDITOR
 
 				Interface.Outputs.Add(GeneratedAudio);
 
@@ -159,31 +166,38 @@ namespace Metasound
 				// Inputs
 				FMetasoundFrontendClassVertex OnPlayTrigger;
 				OnPlayTrigger.Name = ArchetypesPrivate::GetOnPlayInputName();
-
 				OnPlayTrigger.TypeName = GetMetasoundDataTypeName<FTrigger>();
+				OnPlayTrigger.VertexID = FGuid::NewGuid();
+
+#if WITH_EDITOR
 				OnPlayTrigger.Metadata.SetDisplayName(LOCTEXT("OnPlay", "On Play"));
 				OnPlayTrigger.Metadata.SetDescription(LOCTEXT("OnPlayTriggerToolTip", "Trigger executed when this source is played."));
-				OnPlayTrigger.VertexID = FGuid::NewGuid();
+#endif // WITH_EDITOR
 
 				Interface.Inputs.Add(OnPlayTrigger);
 
 				// Outputs
 				FMetasoundFrontendClassVertex OnFinished;
 				OnFinished.Name = ArchetypesPrivate::GetIsFinishedOutputName();
-
 				OnFinished.TypeName = GetMetasoundDataTypeName<FTrigger>();
+				OnFinished.VertexID = FGuid::NewGuid();
+
+#if WITH_EDITOR
 				OnFinished.Metadata.SetDisplayName(LOCTEXT("OnFinished", "On Finished"));
 				OnFinished.Metadata.SetDescription(LOCTEXT("OnFinishedToolTip", "Trigger executed to initiate stopping the source."));
-				OnFinished.VertexID = FGuid::NewGuid();
+#endif // WITH_EDITOR
 
 				Interface.Outputs.Add(OnFinished);
 
 				FMetasoundFrontendClassVertex GeneratedAudio;
 				GeneratedAudio.Name = GetAudioOutputName();
 				GeneratedAudio.TypeName = GetMetasoundDataTypeName<FStereoAudioFormat>();
+				GeneratedAudio.VertexID = FGuid::NewGuid();
+
+#if WITH_EDITOR
 				GeneratedAudio.Metadata.SetDisplayName(LOCTEXT("GeneratedStereo", "Audio"));
 				GeneratedAudio.Metadata.SetDescription(LOCTEXT("GeneratedAudioToolTip", "The resulting output audio from this source."));
-				GeneratedAudio.VertexID = FGuid::NewGuid();
+#endif // WITH_EDITOR
 
 				Interface.Outputs.Add(GeneratedAudio);
 
@@ -216,9 +230,12 @@ namespace Metasound
 				FMetasoundFrontendClassVertex GeneratedAudio;
 				GeneratedAudio.Name = GetAudioOutputName();
 				GeneratedAudio.TypeName = GetMetasoundDataTypeName<FAudioBuffer>();
+				GeneratedAudio.VertexID = FGuid::NewGuid();
+
+#if WITH_EDITOR
 				GeneratedAudio.Metadata.SetDisplayName(LOCTEXT("GeneratedMono", "Audio"));
 				GeneratedAudio.Metadata.SetDescription(LOCTEXT("GeneratedAudioToolTip", "The resulting output audio from this source."));
-				GeneratedAudio.VertexID = FGuid::NewGuid();
+#endif // WITH_EDITOR
 
 				return GeneratedAudio;
 			}
@@ -231,22 +248,26 @@ namespace Metasound
 				// Inputs
 				FMetasoundFrontendClassVertex OnPlayTrigger;
 				OnPlayTrigger.Name = ArchetypesPrivate::GetOnPlayInputName();
-
 				OnPlayTrigger.TypeName = GetMetasoundDataTypeName<FTrigger>();
+				OnPlayTrigger.VertexID = FGuid::NewGuid();
+
+#if WITH_EDITOR
 				OnPlayTrigger.Metadata.SetDisplayName(LOCTEXT("OnPlay", "On Play"));
 				OnPlayTrigger.Metadata.SetDescription(LOCTEXT("OnPlayTriggerToolTip", "Trigger executed when this source is played."));
-				OnPlayTrigger.VertexID = FGuid::NewGuid();
+#endif // WITH_EDITOR
 
 				Interface.Inputs.Add(OnPlayTrigger);
 
 				// Outputs
 				FMetasoundFrontendClassVertex OnFinished;
 				OnFinished.Name = ArchetypesPrivate::GetIsFinishedOutputName();
-
 				OnFinished.TypeName = GetMetasoundDataTypeName<FTrigger>();
+				OnFinished.VertexID = FGuid::NewGuid();
+
+#if WITH_EDITOR
 				OnFinished.Metadata.SetDisplayName(LOCTEXT("OnFinished", "On Finished"));
 				OnFinished.Metadata.SetDescription(LOCTEXT("OnFinishedToolTip", "Trigger executed to initiate stopping the source."));
-				OnFinished.VertexID = FGuid::NewGuid();
+#endif // WITH_EDITOR
 
 				Interface.Outputs.Add(OnFinished);
 
@@ -282,11 +303,13 @@ namespace Metasound
 					FOutputHandle OutputToReconnect = IOutputController::GetInvalidHandle();
 					if (MonoFormatOutput->IsValid())
 					{
+#if WITH_EDITOR
 						// Get the first location.
 						for (auto Location : MonoFormatOutput->GetNodeStyle().Display.Locations)
 						{
 							MonoFormatLocation = Location.Value;
 						}
+#endif // WITH_EDITOR
 
 						// Get connections
 						TArray<FInputHandle> Inputs = MonoFormatOutput->GetInputs();
@@ -302,9 +325,11 @@ namespace Metasound
 					FNodeHandle BufferOutput = Graph->AddOutputVertex(GetClassAudioOutput());
 					if (ensure(BufferOutput->IsValid()))
 					{
+#if WITH_EDITOR
 						FMetasoundFrontendNodeStyle Style = BufferOutput->GetNodeStyle();
 						Style.Display.Locations.Add(FGuid(), MonoFormatLocation);
 						BufferOutput->SetNodeStyle(Style);
+#endif // WITH_EDITOR
 
 						if (OutputToReconnect->IsValid())
 						{
@@ -352,9 +377,12 @@ namespace Metasound
 				FMetasoundFrontendClassVertex GeneratedLeftAudio;
 				GeneratedLeftAudio.Name = GetLeftAudioOutputName();
 				GeneratedLeftAudio.TypeName = GetMetasoundDataTypeName<FAudioBuffer>();
+				GeneratedLeftAudio.VertexID = FGuid::NewGuid();
+
+#if WITH_EDITOR
 				GeneratedLeftAudio.Metadata.SetDisplayName(LOCTEXT("GeneratedStereoLeft", "Left Audio"));
 				GeneratedLeftAudio.Metadata.SetDescription(LOCTEXT("GeneratedLeftAudioToolTip", "The resulting output audio from this source."));
-				GeneratedLeftAudio.VertexID = FGuid::NewGuid();
+#endif // WITH_EDITOR
 
 				return GeneratedLeftAudio;
 			}
@@ -364,8 +392,12 @@ namespace Metasound
 				FMetasoundFrontendClassVertex GeneratedRightAudio;
 				GeneratedRightAudio.Name = GetRightAudioOutputName();
 				GeneratedRightAudio.TypeName = GetMetasoundDataTypeName<FAudioBuffer>();
+
+#if WITH_EDITOR
 				GeneratedRightAudio.Metadata.SetDisplayName(LOCTEXT("GeneratedStereoRight", "Right Audio"));
 				GeneratedRightAudio.Metadata.SetDescription(LOCTEXT("GeneratedRightAudioToolTip", "The resulting output audio from this source."));
+#endif // WITH_EDITOR
+
 				GeneratedRightAudio.VertexID = FGuid::NewGuid();
 
 				return GeneratedRightAudio;
@@ -379,10 +411,13 @@ namespace Metasound
 				// Inputs
 				FMetasoundFrontendClassVertex OnPlayTrigger;
 				OnPlayTrigger.Name = ArchetypesPrivate::GetOnPlayInputName();
-
 				OnPlayTrigger.TypeName = GetMetasoundDataTypeName<FTrigger>();
+
+#if WITH_EDITOR
 				OnPlayTrigger.Metadata.SetDisplayName(LOCTEXT("OnPlay", "On Play"));
 				OnPlayTrigger.Metadata.SetDescription(LOCTEXT("OnPlayTriggerToolTip", "Trigger executed when this source is played."));
+#endif // WITH_EDITOR
+
 				OnPlayTrigger.VertexID = FGuid::NewGuid();
 
 				Interface.Inputs.Add(OnPlayTrigger);
@@ -392,8 +427,12 @@ namespace Metasound
 				OnFinished.Name = ArchetypesPrivate::GetIsFinishedOutputName();
 
 				OnFinished.TypeName = GetMetasoundDataTypeName<FTrigger>();
+
+#if WITH_EDITOR
 				OnFinished.Metadata.SetDisplayName(LOCTEXT("OnFinished", "On Finished"));
 				OnFinished.Metadata.SetDescription(LOCTEXT("OnFinishedToolTip", "Trigger executed to initiate stopping the source."));
+#endif // WITH_EDITOR
+
 				OnFinished.VertexID = FGuid::NewGuid();
 
 				Interface.Outputs.Add(OnFinished);
@@ -434,11 +473,13 @@ namespace Metasound
 
 					if (StereoFormatOutput->IsValid())
 					{
+#if WITH_EDITOR
 						// Get the first location.
 						for (auto Location : StereoFormatOutput->GetNodeStyle().Display.Locations)
 						{
 							StereoFormatLocation = Location.Value;
 						}
+#endif // WITH_EDITOR
 
 						FInputHandle LeftInput = StereoFormatOutput->GetInputWithVertexName("Left");
 						LeftOutputToReconnect = LeftInput->GetConnectedOutput();
@@ -452,9 +493,11 @@ namespace Metasound
 					FNodeHandle LeftBufferOutput = Graph->AddOutputVertex(GetClassLeftAudioOutput());
 					if (ensure(LeftBufferOutput->IsValid()))
 					{
+#if WITH_EDITOR
 						FMetasoundFrontendNodeStyle Style = LeftBufferOutput->GetNodeStyle();
 						Style.Display.Locations.Add(FGuid(), StereoFormatLocation);
 						LeftBufferOutput->SetNodeStyle(Style);
+#endif // WITH_EDITOR
 
 						if (LeftOutputToReconnect->IsValid())
 						{
@@ -471,9 +514,11 @@ namespace Metasound
 					{
 						if (RightOutputToReconnect->IsValid())
 						{
+#if WITH_EDITOR
 							FMetasoundFrontendNodeStyle Style = RightBufferOutput->GetNodeStyle();
 							Style.Display.Locations.Add(FGuid(), StereoFormatLocation + FVector2D{0, 100.f}); // TODO: How should this new output be placed?
 							RightBufferOutput->SetNodeStyle(Style);
+#endif // WITH_EDITOR
 
 							TArray<FInputHandle> Inputs = RightBufferOutput->GetInputs();
 							if (ensure(Inputs.Num() == 1))
diff --git a/Engine/Plugins/Runtime/Metasound/Source/MetasoundEngine/Private/MetasoundInterface.cpp b/Engine/Plugins/Runtime/Metasound/Source/MetasoundEngine/Private/MetasoundInterface.cpp
--- a/Engine/Plugins/Runtime/Metasound/Source/MetasoundEngine/Private/MetasoundInterface.cpp
+++ b/Engine/Plugins/Runtime/Metasound/Source/MetasoundEngine/Private/MetasoundInterface.cpp
@@ -138,12 +138,14 @@ namespace Metasound
 				ClassInput.Name = Input.InitValue.ParamName;
 				ClassInput.DefaultLiteral = FMetasoundFrontendLiteral(Input.InitValue);
 				ClassInput.TypeName = ResolveMemberDataType(Input.DataType, Input.InitValue.ParamType);
-				ClassInput.Metadata.SetDisplayName(Input.DisplayName);
-				ClassInput.Metadata.SetDescription(Input.Description);
 
+#if WITH_EDITOR
 				// Interfaces should never serialize text to avoid desync between
 				// copied versions serialized in assets and those defined in code.
 				ClassInput.Metadata.SetSerializeText(false);
+				ClassInput.Metadata.SetDisplayName(Input.DisplayName);
+				ClassInput.Metadata.SetDescription(Input.Description);
+#endif // WITH_EDITOR
 
 				ClassInput.VertexID = FGuid::NewGuid();
 
@@ -155,13 +157,16 @@ namespace Metasound
 				FMetasoundFrontendClassOutput ClassOutput;
 				ClassOutput.Name = Output.ParamName;
 				ClassOutput.TypeName = ResolveMemberDataType(Output.DataType, Output.ParamType);
-				ClassOutput.Metadata.SetDisplayName(Output.DisplayName);
-				ClassOutput.Metadata.SetDescription(Output.Description);
-				ClassOutput.VertexID = FGuid::NewGuid();
 
+#if WITH_EDITOR
 				// Interfaces should never serialize text to avoid desync between
 				// copied versions serialized in assets and those defined in code.
 				ClassOutput.Metadata.SetSerializeText(false);
+				ClassOutput.Metadata.SetDisplayName(Output.DisplayName);
+				ClassOutput.Metadata.SetDescription(Output.Description);
+#endif // WITH_EDITOR
+
+				ClassOutput.VertexID = FGuid::NewGuid();
 
 				return ClassOutput;
 			});
diff --git a/Engine/Plugins/Runtime/Metasound/Source/MetasoundFrontend/Private/MetasoundAssetBase.cpp b/Engine/Plugins/Runtime/Metasound/Source/MetasoundFrontend/Private/MetasoundAssetBase.cpp
--- a/Engine/Plugins/Runtime/Metasound/Source/MetasoundFrontend/Private/MetasoundAssetBase.cpp
+++ b/Engine/Plugins/Runtime/Metasound/Source/MetasoundFrontend/Private/MetasoundAssetBase.cpp
@@ -619,9 +619,12 @@ void FMetasoundAssetBase::ConvertFromPreset()
 {
 	using namespace Metasound::Frontend;
 	FGraphHandle GraphHandle = GetRootGraphHandle();
+
+#if WITH_EDITOR
 	FMetasoundFrontendGraphStyle Style = GraphHandle->GetGraphStyle();
 	Style.bIsGraphEditable = true;
 	GraphHandle->SetGraphStyle(Style);
+#endif // WITH_EDITOR
 
 	FMetasoundFrontendGraphClassPresetOptions PresetOptions = GraphHandle->GetGraphPresetOptions();
 	PresetOptions.bIsPreset = false;
@@ -663,10 +666,13 @@ Metasound::Frontend::FNodeHandle FMetasoundAssetBase::AddInputPinForSendAddress(
 
 	Description.Name = InSendInfo.Address.GetChannelName();
 	Description.TypeName = Metasound::GetMetasoundDataTypeName<Metasound::FSendAddress>();
-	Description.Metadata.SetDescription(FText::GetEmpty());
 	Description.VertexID = VertexID;
 	Description.DefaultLiteral.Set(InSendInfo.Address.GetChannelName().ToString());
 
+#if WITH_EDITOR
+	Description.Metadata.SetDescription(FText::GetEmpty());
+#endif // WITH_EDITOR
+
 	return InGraph->AddInputVertex(Description);
 }
 
diff --git a/Engine/Plugins/Runtime/Metasound/Source/MetasoundFrontend/Private/MetasoundFrontendDocument.cpp b/Engine/Plugins/Runtime/Metasound/Source/MetasoundFrontend/Private/MetasoundFrontendDocument.cpp
--- a/Engine/Plugins/Runtime/Metasound/Source/MetasoundFrontend/Private/MetasoundFrontendDocument.cpp
+++ b/Engine/Plugins/Runtime/Metasound/Source/MetasoundFrontend/Private/MetasoundFrontendDocument.cpp
@@ -192,7 +192,11 @@ FMetasoundFrontendClassInterface FMetasoundFrontendClassInterface::GenerateClass
 	// Copy over inputs
 	{
 		const FInputVertexInterface& InputInterface = InVertexInterface.GetInputInterface();
+
+#if WITH_EDITOR
 		FMetasoundFrontendInterfaceStyle InputStyle;
+#endif // WITH_EDITOR
+
 		for (const TPair<FVertexName, FInputDataVertex>& InputTuple : InputInterface)
 		{
 			FMetasoundFrontendClassInput ClassInput;
@@ -204,17 +208,11 @@ FMetasoundFrontendClassInterface FMetasoundFrontendClassInterface::GenerateClass
 
 			const FDataVertexMetadata& VertexMetadata = InputVertex.GetMetadata();
 
+#if WITH_EDITOR
 			ClassInput.Metadata.SetSerializeText(false);
 			ClassInput.Metadata.SetDescription(VertexMetadata.Description);
-
 			ClassInput.Metadata.bIsAdvancedDisplay = VertexMetadata.bIsAdvancedDisplay;
 
-			FLiteral DefaultLiteral = InputVertex.GetDefaultLiteral();
-			if (DefaultLiteral.GetType() != ELiteralType::Invalid)
-			{
-				ClassInput.DefaultLiteral.SetFromLiteral(DefaultLiteral);
-			}
-
 			// Advanced display items are pushed to bottom of sort order
 			int32 OrderIndex = InputInterface.GetOrderIndex(InputTuple.Key);
 			if (ClassInput.Metadata.bIsAdvancedDisplay)
@@ -222,21 +220,35 @@ FMetasoundFrontendClassInterface FMetasoundFrontendClassInterface::GenerateClass
 				OrderIndex += InputInterface.Num();
 			}
 			InputStyle.DefaultSortOrder.Add(OrderIndex);
+#endif // WITH_EDITOR
+
+			FLiteral DefaultLiteral = InputVertex.GetDefaultLiteral();
+			if (DefaultLiteral.GetType() != ELiteralType::Invalid)
+			{
+				ClassInput.DefaultLiteral.SetFromLiteral(DefaultLiteral);
+			}
+
 
 			ClassInterface.Inputs.Add(MoveTemp(ClassInput));
 		}
 
+#if WITH_EDITOR
 		// Must set via direct accessor to avoid updating the change GUID
 		// (All instances of this generation call should be done for code
 		// defined classes only, which do not currently create a persistent
 		// change hash between builds and leave the guid 0'ed).
 		ClassInterface.InputStyle = InputStyle;
+#endif // WITH_EDITOR
 	}
 
 	// Copy over outputs
 	{
 		const FOutputVertexInterface& OutputInterface = InVertexInterface.GetOutputInterface();
+
+#if WITH_EDITOR
 		FMetasoundFrontendInterfaceStyle OutputStyle;
+#endif // WITH_EDITOR
+
 		for (const TPair<FVertexName, FOutputDataVertex>& OutputTuple : OutputInterface)
 		{
 			FMetasoundFrontendClassOutput ClassOutput;
@@ -246,6 +258,8 @@ FMetasoundFrontendClassInterface FMetasoundFrontendClassInterface::GenerateClass
 			ClassOutput.VertexID = FGuid::NewGuid();
 
 			const FDataVertexMetadata& VertexMetadata = OutputTuple.Value.GetMetadata();
+
+#if WITH_EDITOR
 			ClassOutput.Metadata.SetSerializeText(false);
 			ClassOutput.Metadata.SetDescription(VertexMetadata.Description);
 			ClassOutput.Metadata.bIsAdvancedDisplay = VertexMetadata.bIsAdvancedDisplay;
@@ -257,15 +271,18 @@ FMetasoundFrontendClassInterface FMetasoundFrontendClassInterface::GenerateClass
 				OrderIndex += OutputInterface.Num();
 			}
 			OutputStyle.DefaultSortOrder.Add(OrderIndex);
+#endif // WITH_EDITOR
 
 			ClassInterface.Outputs.Add(MoveTemp(ClassOutput));
 		}
 
+#if WITH_EDITOR
 		// Must set via direct accessor to avoid updating the change GUID
 		// (All instances of this generation call should be done for code
 		// defined classes only, which do not currently create a persistent
 		// change hash between builds and leave the guid 0'ed).
 		ClassInterface.OutputStyle = MoveTemp(OutputStyle);
+#endif // WITH_EDITOR
 	}
 
 	for (auto& EnvTuple : InVertexInterface.GetEnvironmentInterface())
@@ -281,12 +298,12 @@ FMetasoundFrontendClassInterface FMetasoundFrontendClassInterface::GenerateClass
 	return ClassInterface;
 }
 
-void FMetasoundFrontendClassMetadata::SetAuthor(const FText& InAuthor)
+#if WITH_EDITOR
+void FMetasoundFrontendClassMetadata::SetAuthor(const FString& InAuthor)
 {
 	using namespace Metasound::DocumentPrivate;
 
-	FText& TextToSet = bSerializeText ? Author : AuthorTransient;
-	SetWithChangeID(InAuthor, TextToSet, ChangeID);
+	SetWithChangeID(InAuthor, Author, ChangeID);
 }
 
 void FMetasoundFrontendClassMetadata::SetCategoryHierarchy(const TArray<FText>& InCategoryHierarchy)
@@ -304,12 +321,6 @@ void FMetasoundFrontendClassMetadata::SetKeywords(const TArray<FText>& InKeyword
 	SetWithChangeID(InKeywords, TextToSet, ChangeID);
 }
 
-void FMetasoundFrontendClassMetadata::SetClassName(const FMetasoundFrontendClassName& InClassName)
-{
-	using namespace Metasound::DocumentPrivate;
-	SetWithChangeID(InClassName, ClassName, ChangeID);
-}
-
 void FMetasoundFrontendClassMetadata::SetDescription(const FText& InDescription)
 {
 	using namespace Metasound::DocumentPrivate;
@@ -344,11 +355,9 @@ void FMetasoundFrontendClassMetadata::SetSerializeText(bool bInSerializeText)
 	{
 		if (!bInSerializeText)
 		{
-			AuthorTransient = Author;
 			DescriptionTransient = Description;
 			DisplayNameTransient = DisplayName;
 
-			Author = { };
 			Description = { };
 			DisplayName = { };
 
@@ -360,11 +369,9 @@ void FMetasoundFrontendClassMetadata::SetSerializeText(bool bInSerializeText)
 	{
 		if (bInSerializeText)
 		{
-			Author = AuthorTransient;
 			Description = DescriptionTransient;
 			DisplayName = DisplayNameTransient;
 
-			AuthorTransient = { };
 			DescriptionTransient = { };
 			DisplayNameTransient = { };
 
@@ -381,19 +388,15 @@ void FMetasoundFrontendClassMetadata::SetVersion(const FMetasoundFrontendVersion
 	using namespace Metasound::DocumentPrivate;
 	SetWithChangeID(InVersion, Version, ChangeID);
 }
+#endif // WITH_EDITOR
 
-FMetasoundFrontendClassStyle FMetasoundFrontendClassStyle::GenerateClassStyle(const Metasound::FNodeDisplayStyle& InNodeDisplayStyle)
+void FMetasoundFrontendClassMetadata::SetClassName(const FMetasoundFrontendClassName& InClassName)
 {
-	FMetasoundFrontendClassStyle Style;
-	Style.Display.bShowName = InNodeDisplayStyle.bShowName;
-	Style.Display.bShowInputNames = InNodeDisplayStyle.bShowInputNames;
-	Style.Display.bShowOutputNames = InNodeDisplayStyle.bShowOutputNames;
-	Style.Display.ImageName = InNodeDisplayStyle.ImageName;
-
-	return Style;
+	using namespace Metasound::DocumentPrivate;
+	SetWithChangeID(InClassName, ClassName, ChangeID);
 }
 
-#if WITH_EDITORONLY_DATA
+#if WITH_EDITOR
 bool FMetasoundFrontendClass::CacheGraphDependencyMetadataFromRegistry(FMetasoundFrontendClass& InOutDependency)
 {
 	using namespace Metasound::Frontend;
@@ -445,7 +448,18 @@ bool FMetasoundFrontendClass::CacheGraphDependencyMetadataFromRegistry(FMetasoun
 
 	return false;
 }
-#endif // WITH_EDITORONLY_DATA
+#endif // WITH_EDITOR
+
+FMetasoundFrontendClassStyle FMetasoundFrontendClassStyle::GenerateClassStyle(const Metasound::FNodeDisplayStyle& InNodeDisplayStyle)
+{
+	FMetasoundFrontendClassStyle Style;
+	Style.Display.bShowName = InNodeDisplayStyle.bShowName;
+	Style.Display.bShowInputNames = InNodeDisplayStyle.bShowInputNames;
+	Style.Display.bShowOutputNames = InNodeDisplayStyle.bShowOutputNames;
+	Style.Display.ImageName = InNodeDisplayStyle.ImageName;
+
+	return Style;
+}
 
 FMetasoundFrontendClassMetadata FMetasoundFrontendClassMetadata::GenerateClassMetadata(const Metasound::FNodeClassMetadata& InNodeClassMetadata, EMetasoundFrontendClassType InType)
 {
@@ -462,6 +476,7 @@ FMetasoundFrontendClassMetadata FMetasoundFrontendClassMetadata::GenerateClassMe
 	NewMetadata.ClassName = InNodeClassMetadata.ClassName;
 	NewMetadata.Version = { InNodeClassMetadata.MajorVersion, InNodeClassMetadata.MinorVersion };
 
+#if WITH_EDITOR
 	NewMetadata.SetSerializeText(false);
 	NewMetadata.SetDisplayName(InNodeClassMetadata.DisplayName);
 	NewMetadata.SetDescription(InNodeClassMetadata.Description);
@@ -471,6 +486,7 @@ FMetasoundFrontendClassMetadata FMetasoundFrontendClassMetadata::GenerateClassMe
 	NewMetadata.SetCategoryHierarchy(InNodeClassMetadata.CategoryHierarchy);
 
 	NewMetadata.bIsDeprecated = InNodeClassMetadata.bDeprecated;
+#endif // WITH_EDITOR
 
 	return NewMetadata;
 }
diff --git a/Engine/Plugins/Runtime/Metasound/Source/MetasoundFrontend/Private/MetasoundFrontendDocumentController.cpp b/Engine/Plugins/Runtime/Metasound/Source/MetasoundFrontend/Private/MetasoundFrontendDocumentController.cpp
--- a/Engine/Plugins/Runtime/Metasound/Source/MetasoundFrontend/Private/MetasoundFrontendDocumentController.cpp
+++ b/Engine/Plugins/Runtime/Metasound/Source/MetasoundFrontend/Private/MetasoundFrontendDocumentController.cpp
@@ -357,13 +357,21 @@ namespace Metasound
 							}
 							else
 							{
+#if WITH_EDITOR
 								UE_LOG(LogMetaSound, Error,
 									TEXT("Cannot add external dependency. No Metasound class found with matching registry key [Key:%s, Name:%s, Version:%s]. Suggested solution \"%s\" by %s."),
 									*Key,
 									*InMetadata.GetClassName().GetFullName().ToString(),
 									*InMetadata.GetVersion().ToString(),
 									*InMetadata.GetPromptIfMissing().ToString(),
-									*InMetadata.GetAuthor().ToString());
+									*InMetadata.GetAuthor());
+#else
+								UE_LOG(LogMetaSound, Error,
+									TEXT("Cannot add external dependency. No Metasound class found with matching registry key [Key:%s, Name:%s, Version:%s]."),
+									*Key,
+									*InMetadata.GetClassName().GetFullName().ToString(),
+									*InMetadata.GetVersion().ToString());
+#endif // !WITH_EDITOR
 							}
 						}
 						break;
diff --git a/Engine/Plugins/Runtime/Metasound/Source/MetasoundFrontend/Private/MetasoundFrontendGraphController.cpp b/Engine/Plugins/Runtime/Metasound/Source/MetasoundFrontend/Private/MetasoundFrontendGraphController.cpp
--- a/Engine/Plugins/Runtime/Metasound/Source/MetasoundFrontend/Private/MetasoundFrontendGraphController.cpp
+++ b/Engine/Plugins/Runtime/Metasound/Source/MetasoundFrontend/Private/MetasoundFrontendGraphController.cpp
@@ -78,6 +78,7 @@ namespace Metasound
 			return Metasound::FrontendInvalidID;
 		}
 
+#if WITH_EDITOR
 		FText FGraphController::GetDisplayName() const
 		{
 			if (const FMetasoundFrontendGraphClass* GraphClass = GraphClassPtr.Get())
@@ -87,6 +88,7 @@ namespace Metasound
 
 			return Invalid::GetInvalidText();
 		}
+#endif // WITH_EDITOR
 
 		TArray<FVertexName> FGraphController::GetInputVertexNames() const
 		{
@@ -174,6 +176,7 @@ namespace Metasound
 			return GetNodeByPredicate(IsNodeWithSameID);
 		}
 
+#if WITH_EDITOR
 		const FMetasoundFrontendGraphStyle& FGraphController::GetGraphStyle() const
 		{
 			if (const FMetasoundFrontendGraphClass* GraphClass = GraphClassPtr.Get())
@@ -191,6 +194,7 @@ namespace Metasound
 				GraphClass->Graph.Style = InStyle;
 			}
 		}
+#endif // WITH_EDITOR
 
 		TArray<FNodeHandle> FGraphController::GetOutputNodes()
 		{
@@ -1149,6 +1153,7 @@ namespace Metasound
 			return false;
 		}
 
+#if WITH_EDITOR
 		const FText& FGraphController::GetInputDescription(const FVertexName& InName) const
 		{
 			if (const FMetasoundFrontendClassInput* Desc = FindInputDescriptionWithName(InName))
@@ -1200,6 +1205,7 @@ namespace Metasound
 				Desc->Metadata.SetDescription(InDescription);
 			}
 		}
+#endif // WITH_EDITOR
 
 		// This can be used to clear the current literal for a given input.
 		// @returns false if the input name couldn't be found.
@@ -1510,7 +1516,6 @@ namespace Metasound
 			return false;
 		}
 
-
 		void FGraphController::UpdateInterfaceChangeID()
 		{
 			if (FMetasoundFrontendGraphClass* GraphClass = GraphClassPtr.Get())
@@ -2082,4 +2087,3 @@ namespace Metasound
 	}
 }
 #undef LOCTEXT_NAMESPACE
-
diff --git a/Engine/Plugins/Runtime/Metasound/Source/MetasoundFrontend/Private/MetasoundFrontendGraphController.h b/Engine/Plugins/Runtime/Metasound/Source/MetasoundFrontend/Private/MetasoundFrontendGraphController.h
--- a/Engine/Plugins/Runtime/Metasound/Source/MetasoundFrontend/Private/MetasoundFrontendGraphController.h
+++ b/Engine/Plugins/Runtime/Metasound/Source/MetasoundFrontend/Private/MetasoundFrontendGraphController.h
@@ -51,7 +51,10 @@ namespace Metasound
 			bool IsValid() const override;
 
 			FGuid GetClassID() const override;
+
+#if WITH_EDITOR
 			FText GetDisplayName() const override;
+#endif // WITH_EDITOR
 
 			TArray<FVertexName> GetInputVertexNames() const override;
 			TArray<FVertexName> GetOutputVertexNames() const override;
@@ -93,8 +96,10 @@ namespace Metasound
 			TArray<FNodeHandle> GetOutputNodes() override;
 			TArray<FConstNodeHandle> GetConstOutputNodes() const override;
 
+#if WITH_EDITOR
 			const FMetasoundFrontendGraphStyle& GetGraphStyle() const override;
 			void SetGraphStyle(const FMetasoundFrontendGraphStyle& InStyle) override;
+#endif // WITH_EDITOR
 
 			bool ContainsInputVertex(const FVertexName& InName, const FName& InTypeName) const override;
 			bool ContainsInputVertexWithName(const FVertexName& InName) const override;
@@ -131,6 +136,7 @@ namespace Metasound
 			bool SetDefaultInput(const FGuid& InVertexID, const FMetasoundFrontendLiteral& InLiteral) override;
 			bool SetDefaultInputToDefaultLiteralOfType(const FGuid& InVertexID) override;
 
+#if WITH_EDITOR
 			// Get the description for the input with the given name.
 			const FText& GetInputDescription(const FVertexName& InName) const override;
 
@@ -148,6 +154,7 @@ namespace Metasound
 
 			// Set the display name for the output with the given name
 			void SetOutputDisplayName(const FVertexName& InName, const FText& InDisplayName) override;
+#endif // WITH_EDITOR
 
 			// This can be used to clear the current literal for a given input.
 			// @returns false if the input name couldn't be found.
diff --git a/Engine/Plugins/Runtime/Metasound/Source/MetasoundFrontend/Private/MetasoundFrontendInjectReceiveNodes.cpp b/Engine/Plugins/Runtime/Metasound/Source/MetasoundFrontend/Private/MetasoundFrontendInjectReceiveNodes.cpp
--- a/Engine/Plugins/Runtime/Metasound/Source/MetasoundFrontend/Private/MetasoundFrontendInjectReceiveNodes.cpp
+++ b/Engine/Plugins/Runtime/Metasound/Source/MetasoundFrontend/Private/MetasoundFrontendInjectReceiveNodes.cpp
@@ -97,7 +97,7 @@ namespace Metasound
 						0, // MinorVersion
 						FText::GetEmpty(), // DisplayName
 						FText::GetEmpty(), // Description
-						FText::GetEmpty(), // Author
+						FString(), // Author
 						FText::GetEmpty(), // Prompt If Missing
 						CreateVertexInterface(), // DefaultInterface
 						TArray<FText>(), // CategoryHierachy
diff --git a/Engine/Plugins/Runtime/Metasound/Source/MetasoundFrontend/Private/MetasoundFrontendInputController.cpp b/Engine/Plugins/Runtime/Metasound/Source/MetasoundFrontend/Private/MetasoundFrontendInputController.cpp
--- a/Engine/Plugins/Runtime/Metasound/Source/MetasoundFrontend/Private/MetasoundFrontendInputController.cpp
+++ b/Engine/Plugins/Runtime/Metasound/Source/MetasoundFrontend/Private/MetasoundFrontendInputController.cpp
@@ -57,6 +57,7 @@ namespace Metasound
 			return Invalid::GetInvalidName();
 		}
 
+#if WITH_EDITOR
 		FText FBaseInputController::GetDisplayName() const
 		{
 			if (const FMetasoundFrontendClassInput* ClassInput = ClassInputPtr.Get())
@@ -66,6 +67,7 @@ namespace Metasound
 			
 			return Invalid::GetInvalidText();
 		}
+#endif // WITH_EDITOR
 
 		bool FBaseInputController::ClearLiteral()
 		{
@@ -104,6 +106,7 @@ namespace Metasound
 			return nullptr;
 		}
 
+#if WITH_EDITOR
 		const FText& FBaseInputController::GetTooltip() const
 		{
 			if (const FMetasoundFrontendClassInput* ClassInput = ClassInputPtr.Get())
@@ -123,6 +126,7 @@ namespace Metasound
 
 			return Invalid::GetInvalidVertexMetadata();
 		}
+#endif // WITH_EDITOR
 
 		bool FBaseInputController::IsConnected() const 
 		{
@@ -399,6 +403,7 @@ namespace Metasound
 			return FBaseInputController::IsValid() && (nullptr != OwningGraphClassOutputPtr.Get());
 		}
 
+#if WITH_EDITOR
 		FText FOutputNodeInputController::GetDisplayName() const
 		{
 			if (const FMetasoundFrontendClassOutput* OwningOutput = OwningGraphClassOutputPtr.Get())
@@ -438,6 +443,7 @@ namespace Metasound
 
 			return Invalid::GetInvalidVertexMetadata();
 		}
+#endif // WITH_EDITOR
 
 		void FOutputNodeInputController::SetName(const FVertexName& InName)
 		{
@@ -480,6 +486,7 @@ namespace Metasound
 			return FBaseInputController::IsValid() && (nullptr != OwningGraphClassInputPtr.Get());
 		}
 
+#if WITH_EDITOR
 		FText FInputNodeInputController::GetDisplayName() const
 		{
 			if (const FMetasoundFrontendClassInput* OwningInput = OwningGraphClassInputPtr.Get())
@@ -509,6 +516,7 @@ namespace Metasound
 
 			return Invalid::GetInvalidVertexMetadata();
 		}
+#endif // WITH_EDITOR
 
 		void FInputNodeInputController::SetName(const FVertexName& InName)
 		{
diff --git a/Engine/Plugins/Runtime/Metasound/Source/MetasoundFrontend/Private/MetasoundFrontendInputController.h b/Engine/Plugins/Runtime/Metasound/Source/MetasoundFrontend/Private/MetasoundFrontendInputController.h
--- a/Engine/Plugins/Runtime/Metasound/Source/MetasoundFrontend/Private/MetasoundFrontendInputController.h
+++ b/Engine/Plugins/Runtime/Metasound/Source/MetasoundFrontend/Private/MetasoundFrontendInputController.h
@@ -49,10 +49,13 @@ namespace Metasound
 			// vertex names & should not be used for typical edit or runtime callsites.
 			virtual void SetName(const FVertexName& InName) override { checkNoEntry(); }
 
-			// Input metadata
+#if WITH_EDITOR
 			virtual FText GetDisplayName() const override;
 			virtual const FText& GetTooltip() const override;
+
+			// Input metadata
 			virtual const FMetasoundFrontendVertexMetadata& GetMetadata() const override;
+#endif // WITH_EDITOR
 
 			// Owning node info
 			virtual FGuid GetOwningNodeID() const override;
@@ -114,10 +117,14 @@ namespace Metasound
 
 			virtual bool IsValid() const override;
 
-			// Input metadata
+#if WITH_EDITOR
 			virtual FText GetDisplayName() const override;
 			virtual const FText& GetTooltip() const override;
+
+			// Input metadata
 			virtual const FMetasoundFrontendVertexMetadata& GetMetadata() const override;
+#endif // WITH_EDITOR
+
 			virtual void SetName(const FVertexName& InName) override;
 
 		protected:
@@ -157,10 +164,13 @@ namespace Metasound
 
 			virtual bool IsValid() const override;
 
-			// Input metadata
+#if WITH_EDITOR
 			virtual FText GetDisplayName() const override;
 			virtual const FText& GetTooltip() const override;
+
+			// Input metadata
 			virtual const FMetasoundFrontendVertexMetadata& GetMetadata() const override;
+#endif // WITH_EDITOR
 
 			virtual void SetName(const FVertexName& InName) override;
 
diff --git a/Engine/Plugins/Runtime/Metasound/Source/MetasoundFrontend/Private/MetasoundFrontendInvalidController.cpp b/Engine/Plugins/Runtime/Metasound/Source/MetasoundFrontend/Private/MetasoundFrontendInvalidController.cpp
--- a/Engine/Plugins/Runtime/Metasound/Source/MetasoundFrontend/Private/MetasoundFrontendInvalidController.cpp
+++ b/Engine/Plugins/Runtime/Metasound/Source/MetasoundFrontend/Private/MetasoundFrontendInvalidController.cpp
@@ -25,13 +25,6 @@ namespace Metasound
 				return Invalid;
 			}
 
-			const FMetasoundFrontendVertexMetadata& GetInvalidVertexMetadata()
-			{
-				static const FMetasoundFrontendVertexMetadata Invalid;
-				return Invalid;
-			}
-
-
 			const FMetasoundFrontendLiteral& GetInvalidLiteral()
 			{
 				auto CreateInvalidLiteral = []() ->FMetasoundFrontendLiteral
@@ -63,6 +56,13 @@ namespace Metasound
 				return Invalid;
 			}
 
+#if WITH_EDITOR
+			const FMetasoundFrontendVertexMetadata& GetInvalidVertexMetadata()
+			{
+				static const FMetasoundFrontendVertexMetadata Invalid;
+				return Invalid;
+			}
+
 			const FMetasoundFrontendInterfaceStyle& GetInvalidInterfaceStyle()
 			{
 				static const FMetasoundFrontendInterfaceStyle Invalid;
@@ -80,6 +80,7 @@ namespace Metasound
 				static const FMetasoundFrontendGraphStyle Invalid;
 				return Invalid;
 			}
+#endif // WITH_EDITOR
 
 			const FMetasoundFrontendGraphClass& GetInvalidGraphClass() 
 			{ 
diff --git a/Engine/Plugins/Runtime/Metasound/Source/MetasoundFrontend/Private/MetasoundFrontendInvalidController.h b/Engine/Plugins/Runtime/Metasound/Source/MetasoundFrontend/Private/MetasoundFrontendInvalidController.h
--- a/Engine/Plugins/Runtime/Metasound/Source/MetasoundFrontend/Private/MetasoundFrontendInvalidController.h
+++ b/Engine/Plugins/Runtime/Metasound/Source/MetasoundFrontend/Private/MetasoundFrontendInvalidController.h
@@ -24,14 +24,18 @@ namespace Metasound
 			const FText& GetInvalidText();
 			const FName& GetInvalidName();
 			const FString& GetInvalidString();
+
+#if WITH_EDITOR
 			const FMetasoundFrontendVertexMetadata& GetInvalidVertexMetadata();
+			const FMetasoundFrontendInterfaceStyle& GetInvalidInterfaceStyle();
+			const FMetasoundFrontendClassStyle& GetInvalidClassStyle();
+			const FMetasoundFrontendGraphStyle& GetInvalidGraphStyle();
+#endif // WITH_EDITOR
+
 			const FMetasoundFrontendLiteral& GetInvalidLiteral();
 			const FMetasoundFrontendClassInterface& GetInvalidClassInterface();
 			const FMetasoundFrontendClassMetadata& GetInvalidClassMetadata();
 			const FMetasoundFrontendGraphClassPresetOptions& GetInvalidGraphClassPresetOptions();
-			const FMetasoundFrontendInterfaceStyle& GetInvalidInterfaceStyle();
-			const FMetasoundFrontendClassStyle& GetInvalidClassStyle();
-			const FMetasoundFrontendGraphStyle& GetInvalidGraphStyle();
 			const FMetasoundFrontendGraphClass& GetInvalidGraphClass();
 			const TArray<FMetasoundFrontendClass>& GetInvalidClassArray();
 			const TArray<FMetasoundFrontendGraphClass>& GetInvalidGraphClassArray();
@@ -55,9 +59,13 @@ namespace Metasound
 			virtual FGuid GetID() const override { return Metasound::FrontendInvalidID; }
 			virtual const FName& GetDataType() const override { return Invalid::GetInvalidName(); }
 			virtual const FVertexName& GetName() const override { return Invalid::GetInvalidName(); }
+
+#if WITH_EDITOR
 			virtual FText GetDisplayName() const override { return FText::GetEmpty(); }
 			virtual const FText& GetTooltip() const override { return FText::GetEmpty(); }
 			virtual const FMetasoundFrontendVertexMetadata& GetMetadata() const override { return Invalid::GetInvalidVertexMetadata(); }
+#endif // WITH_EDITOR
+
 			virtual FGuid GetOwningNodeID() const override { return Metasound::FrontendInvalidID; }
 			virtual FNodeHandle GetOwningNode() override;
 			virtual FConstNodeHandle GetOwningNode() const override;
@@ -94,10 +102,14 @@ namespace Metasound
 			virtual bool IsConnected() const override { return false; }
 			virtual const FName& GetDataType() const override { return Invalid::GetInvalidName(); }
 			virtual const FVertexName& GetName() const override { return Invalid::GetInvalidName(); }
+
+#if WITH_EDITOR
 			virtual FText GetDisplayName() const override { return Invalid::GetInvalidText(); }
 			virtual const FText& GetTooltip() const override { return Invalid::GetInvalidText(); }
-			virtual bool ClearLiteral() override { return false; }
 			virtual const FMetasoundFrontendVertexMetadata& GetMetadata() const override { return Invalid::GetInvalidVertexMetadata(); }
+#endif // WITH_EDITOR
+
+			virtual bool ClearLiteral() override { return false; }
 			virtual const FMetasoundFrontendLiteral* GetLiteral() const override { return nullptr; }
 			virtual void SetLiteral(const FMetasoundFrontendLiteral& InLiteral) { };
 			virtual const FMetasoundFrontendLiteral* GetClassDefaultLiteral() const override { return nullptr; }
@@ -139,11 +151,13 @@ namespace Metasound
 			virtual const FName& GetName() const override { return Invalid::GetInvalidName(); }
 			virtual void SetName(const FName&) override { }
 
+#if WITH_EDITOR
 			/** Returns the human readable name associated with this output. */
 			virtual FText GetDisplayName() const override { return Invalid::GetInvalidText(); }
 			virtual void SetDisplayName(const FText&) override { }
 			virtual FText GetDescription() const override { return Invalid::GetInvalidText(); }
 			virtual void SetDescription(const FText&) override { }
+#endif // WITH_EDITOR
 
 			virtual FNodeHandle FindMutatorNode() override { return INodeController::GetInvalidHandle(); }
 			virtual FConstNodeHandle FindMutatorNode() const override { return INodeController::GetInvalidHandle(); }
@@ -196,8 +210,11 @@ namespace Metasound
 			virtual FConstInputHandle GetInputWithID(FGuid InVertexID) const override { return IInputController::GetInvalidHandle(); }
 			virtual FConstOutputHandle GetOutputWithID(FGuid InVertexID) const override { return IOutputController::GetInvalidHandle(); }
 
+#if WITH_EDITOR
 			virtual const FMetasoundFrontendNodeStyle& GetNodeStyle() const override { static const FMetasoundFrontendNodeStyle Invalid; return Invalid; }
 			virtual void SetNodeStyle(const FMetasoundFrontendNodeStyle& InNodeStyle) override { }
+#endif // WITH_EDITOR
+
 			virtual void SetNodeName(const FVertexName& InName) override { }
 
 			virtual FNodeHandle ReplaceWithVersion(const FMetasoundFrontendVersionNumber& InNewVersion) override { return INodeController::GetInvalidHandle(); }
@@ -216,11 +233,14 @@ namespace Metasound
 
 			virtual const FMetasoundFrontendClassInterface& GetClassInterface() const override { return Invalid::GetInvalidClassInterface(); }
 			virtual const FMetasoundFrontendClassMetadata& GetClassMetadata() const override { return Invalid::GetInvalidClassMetadata(); }
+
+#if WITH_EDITOR
 			virtual const FMetasoundFrontendInterfaceStyle& GetInputStyle() const override { return Invalid::GetInvalidInterfaceStyle(); }
 			virtual const FMetasoundFrontendInterfaceStyle& GetOutputStyle() const override { return Invalid::GetInvalidInterfaceStyle(); }
 			virtual const FMetasoundFrontendClassStyle& GetClassStyle() const override { return Invalid::GetInvalidClassStyle(); }
 
 			virtual const FText& GetDescription() const override { return Invalid::GetInvalidText(); }
+#endif // WITH_EDITOR
 
 			virtual const FMetasoundFrontendVersion& GetInterfaceVersion() const { return FMetasoundFrontendVersion::GetInvalid(); }
 			virtual bool IsInterfaceMember() const override { return false; }
@@ -250,10 +270,13 @@ namespace Metasound
 			virtual int32 GetNumOutputs() const override { return 0; }
 
 			virtual const FVertexName& GetNodeName() const override { return Invalid::GetInvalidName(); }
+
+#if WITH_EDITOR
 			virtual FText GetDisplayName() const override { return Invalid::GetInvalidText(); }
 			virtual const FText& GetDisplayTitle() const override { return Invalid::GetInvalidText(); }
 			virtual void SetDescription(const FText& InDescription) override { }
 			virtual void SetDisplayName(const FText& InText) override { }
+#endif // WITH_EDITOR
 
 		protected:
 			virtual FDocumentAccess ShareAccess() override { return FDocumentAccess(); }
@@ -272,7 +295,10 @@ namespace Metasound
 
 			virtual bool IsValid() const override { return false; }
 			virtual FGuid GetClassID() const override { return Metasound::FrontendInvalidID; }
+
+#if WITH_EDITOR
 			virtual FText GetDisplayName() const override { return Invalid::GetInvalidText(); }
+#endif // WITH_EDITOR
 
 			virtual TArray<FVertexName> GetInputVertexNames() const override { return TArray<FVertexName>(); }
 			virtual TArray<FVertexName> GetOutputVertexNames() const override { return TArray<FVertexName>(); }
@@ -304,9 +330,10 @@ namespace Metasound
 			virtual FNodeHandle AddVariableAccessorNode(const FGuid& InVariableID) override { return INodeController::GetInvalidHandle(); }
 			virtual FNodeHandle AddVariableDeferredAccessorNode(const FGuid& InVariableID) override{ return INodeController::GetInvalidHandle(); }
 
-
+#if WITH_EDITOR
 			virtual const FMetasoundFrontendGraphStyle& GetGraphStyle() const override { return Invalid::GetInvalidGraphStyle(); }
 			virtual void SetGraphStyle(const FMetasoundFrontendGraphStyle& InStyle) override { }
+#endif // WITH_EDITOR
 
 			virtual void ClearGraph() override { };
 
@@ -350,6 +377,7 @@ namespace Metasound
 			virtual bool SetDefaultInput(const FGuid& InVertexID, const FMetasoundFrontendLiteral& InLiteral) override { return false; }
 			virtual bool SetDefaultInputToDefaultLiteralOfType(const FGuid& InVertexID) override { return false; }
 
+#if WITH_EDITOR
 			virtual const FText& GetInputDescription(const FVertexName& InName) const override { return FText::GetEmpty(); }
 			virtual const FText& GetOutputDescription(const FVertexName& InName) const override { return FText::GetEmpty(); }
 
@@ -357,6 +385,7 @@ namespace Metasound
 			virtual void SetOutputDescription(const FVertexName& InName, const FText& InDescription) override { }
 			virtual void SetInputDisplayName(const FVertexName& InName, const FText& InDisplayName) override { }
 			virtual void SetOutputDisplayName(const FVertexName& InName, const FText& InDisplayName) override { }
+#endif // WITH_EDITOR
 
 			// This can be used to clear the current literal for a given input.
 			// @returns false if the input name couldn't be found.
diff --git a/Engine/Plugins/Runtime/Metasound/Source/MetasoundFrontend/Private/MetasoundFrontendNodeController.cpp b/Engine/Plugins/Runtime/Metasound/Source/MetasoundFrontend/Private/MetasoundFrontendNodeController.cpp
--- a/Engine/Plugins/Runtime/Metasound/Source/MetasoundFrontend/Private/MetasoundFrontendNodeController.cpp
+++ b/Engine/Plugins/Runtime/Metasound/Source/MetasoundFrontend/Private/MetasoundFrontendNodeController.cpp
@@ -165,6 +165,7 @@ namespace Metasound
 			return Invalid::GetInvalidClassMetadata();
 		}
 
+#if WITH_EDITOR
 		const FMetasoundFrontendInterfaceStyle& FBaseNodeController::GetInputStyle() const
 		{
 			if (const FMetasoundFrontendClass* Class = ClassPtr.Get())
@@ -223,6 +224,7 @@ namespace Metasound
 			}
 			return Invalid::GetInvalidText();
 		}
+#endif // WITH_EDITOR
 
 		const FVertexName& FBaseNodeController::GetNodeName() const
 		{
@@ -379,6 +381,7 @@ namespace Metasound
 			}
 		}
 
+#if WITH_EDITOR
 		const FText& FBaseNodeController::GetDisplayTitle() const
 		{
 			if (const FMetasoundFrontendClass* Class = ClassPtr.Get())
@@ -398,6 +401,7 @@ namespace Metasound
 
 			return Invalid::GetInvalidText();
 		}
+#endif // WITH_EDITOR
 
 		void FBaseNodeController::IterateConstInputs(TUniqueFunction<void(FConstInputHandle)> InFunction) const
 		{
@@ -738,7 +742,9 @@ namespace Metasound
 				return this->AsShared();
 			}
 
+#if WITH_EDITOR
 			FMetasoundFrontendNodeStyle Style = GetNodeStyle();
+#endif // WITH_EDITOR
 
 			using FConnectionKey = TPair<FVertexName, FName>;
 
@@ -821,8 +827,10 @@ namespace Metasound
 				return this->AsShared();
 			}
 
+#if WITH_EDITOR
 			Style.bMessageNodeUpdated = ReplacementNode->GetClassMetadata().GetVersion() > Metadata.GetVersion();
 			ReplacementNode->SetNodeStyle(Style);
+#endif // WITH_EDITOR
 
 			ReplacementNode->IterateInputs([Connections = &InputConnections](FInputHandle InputHandle)
 			{
@@ -1156,6 +1164,7 @@ namespace Metasound
 			return INodeController::GetInvalidHandle();
 		}
 
+#if WITH_EDITOR
 		const FText& FOutputNodeController::GetDescription() const
 		{
 			if (const FMetasoundFrontendClassOutput* OwningOutput = OwningGraphClassOutputPtr.Get())
@@ -1186,6 +1195,17 @@ namespace Metasound
 			}
 		}
 
+		void FOutputNodeController::SetDisplayName(const FText& InDisplayName)
+		{
+			// TODO: can we remove the const cast by constructing output nodes with a non-const access to class outputs?
+			if (FMetasoundFrontendClassOutput* ClassOutput = ConstCastAccessPtr<FClassOutputAccessPtr>(OwningGraphClassOutputPtr).Get())
+			{
+				ClassOutput->Metadata.SetDisplayName(InDisplayName);
+				OwningGraph->UpdateInterfaceChangeID();
+			}
+		}
+#endif // WITH_EDITOR
+
 		void FOutputNodeController::SetNodeName(const FVertexName& InName)
 		{
 			if (FMetasoundFrontendNode* Node = NodePtr.Get())
@@ -1211,16 +1231,6 @@ namespace Metasound
 			}
 		}
 
-		void FOutputNodeController::SetDisplayName(const FText& InDisplayName)
-		{
-			// TODO: can we remove the const cast by constructing output nodes with a non-const access to class outputs?
-			if (FMetasoundFrontendClassOutput* ClassOutput = ConstCastAccessPtr<FClassOutputAccessPtr>(OwningGraphClassOutputPtr).Get())
-			{
-				ClassOutput->Metadata.SetDisplayName(InDisplayName);
-				OwningGraph->UpdateInterfaceChangeID();
-			}
-		}
-
 		FConstNodeHandle FOutputNodeController::CreateConstOutputNodeHandle(const FOutputNodeController::FInitParams& InParams)
 		{
 			if (const FMetasoundFrontendNode* Node = InParams.NodePtr.Get())
@@ -1248,11 +1258,13 @@ namespace Metasound
 			return INodeController::GetInvalidHandle();
 		}
 
+#if WITH_EDITOR
 		const FText& FOutputNodeController::GetDisplayTitle() const
 		{
 			static FText OutputDisplayTitle = LOCTEXT("OutputNode_Title", "Output");
 			return OutputDisplayTitle;
 		}
+#endif // WITH_EDITOR
 
 		const FMetasoundFrontendVersion& FOutputNodeController::GetInterfaceVersion() const
 		{
@@ -1327,10 +1339,6 @@ namespace Metasound
 			return Access;
 		}
 
-
-		//
-		// FInputNodeController
-		//
 		FInputNodeController::FInputNodeController(EPrivateToken InToken, const FInputNodeController::FInitParams& InParams)
 		: FBaseNodeController({InParams.NodePtr, InParams.ClassPtr, InParams.OwningGraph})
 		, OwningGraphClassInputPtr(InParams.OwningGraphClassInputPtr)
@@ -1407,6 +1415,7 @@ namespace Metasound
 			return MakeShared<FInputNodeOutputController>(FInputNodeOutputController::FInitParams{InVertexID, InNodeVertexPtr, InClassOutputPtr, OwningGraphClassInputPtr, GraphPtr, InOwningNode});
 		}
 
+#if WITH_EDITOR
 		const FText& FInputNodeController::GetDescription() const
 		{
 			if (const FMetasoundFrontendClassInput* OwningInput = OwningGraphClassInputPtr.Get())
@@ -1432,6 +1441,7 @@ namespace Metasound
 			static FText InputDisplayTitle = LOCTEXT("InputNode_Title", "Input");
 			return InputDisplayTitle;
 		}
+#endif // WITH_EDITOR
 
 		const FMetasoundFrontendVersion& FInputNodeController::GetInterfaceVersion() const
 		{
@@ -1473,6 +1483,7 @@ namespace Metasound
 			return FMetasoundFrontendVersion::GetInvalid();
 		}
 
+#if WITH_EDITOR
 		void FInputNodeController::SetDescription(const FText& InDescription)
 		{
 			// TODO: can we remove these const casts by constructing FINputNodeController with non-const access to the class input?
@@ -1482,6 +1493,7 @@ namespace Metasound
 				OwningGraph->UpdateInterfaceChangeID();
 			}
 		}
+#endif // WITH_EDITOR
 
 		void FInputNodeController::SetNodeName(const FVertexName& InName)
 		{
@@ -1506,6 +1518,7 @@ namespace Metasound
 			}
 		}
 
+#if WITH_EDITOR
 		void FInputNodeController::SetDisplayName(const FText& InDisplayName)
 		{
 			// TODO: can we remove these const casts by constructing FINputNodeController with non-const access to the class input?
@@ -1515,6 +1528,7 @@ namespace Metasound
 				OwningGraph->UpdateInterfaceChangeID();
 			}
 		}
+#endif // WITH_EDITOR
 
 		FDocumentAccess FInputNodeController::ShareAccess()
 		{
diff --git a/Engine/Plugins/Runtime/Metasound/Source/MetasoundFrontend/Private/MetasoundFrontendNodeController.h b/Engine/Plugins/Runtime/Metasound/Source/MetasoundFrontend/Private/MetasoundFrontendNodeController.h
--- a/Engine/Plugins/Runtime/Metasound/Source/MetasoundFrontend/Private/MetasoundFrontendNodeController.h
+++ b/Engine/Plugins/Runtime/Metasound/Source/MetasoundFrontend/Private/MetasoundFrontendNodeController.h
@@ -46,15 +46,15 @@ namespace Metasound
 
 			virtual const FMetasoundFrontendClassInterface& GetClassInterface() const override;
 			virtual const FMetasoundFrontendClassMetadata& GetClassMetadata() const override;
+
+#if WITH_EDITOR
 			virtual const FMetasoundFrontendInterfaceStyle& GetInputStyle() const override;
 			virtual const FMetasoundFrontendInterfaceStyle& GetOutputStyle() const override;
 			virtual const FMetasoundFrontendClassStyle& GetClassStyle() const override;
 
-			/** Description of the given node. */
-			virtual const FText& GetDescription() const override;
-
 			virtual const FMetasoundFrontendNodeStyle& GetNodeStyle() const override;
 			virtual void SetNodeStyle(const FMetasoundFrontendNodeStyle& InStyle) override;
+#endif // WITH_EDITOR
 
 			virtual bool DiffAgainstRegistryInterface(FClassInterfaceUpdates& OutInterfaceUpdates, bool bInUseHighestMinorVersion) const override;
 
@@ -68,6 +68,10 @@ namespace Metasound
 			// This only exists to allow for transform fix-ups to easily cleanup input/output node names.
 			virtual void SetNodeName(const FVertexName& InName) override { checkNoEntry(); }
 
+#if WITH_EDITOR
+			/** Description of the given node. */
+			virtual const FText& GetDescription() const override;
+
 			/** Returns the readable display name of the given node (Used only within MetaSound
 			  * Editor context, and not guaranteed to be a unique identifier). */
 			virtual FText GetDisplayName() const override;
@@ -80,6 +84,7 @@ namespace Metasound
 
 			/** Returns the title of the given node (what to label in visual node). */
 			virtual const FText& GetDisplayTitle() const override;
+#endif // WITH_EDITOR
 
 			virtual bool CanAddInput(const FVertexName& InVertexName) const override;
 			virtual FInputHandle AddInput(const FVertexName& InVertexName, const FMetasoundFrontendLiteral* InDefault) override;
@@ -274,11 +279,15 @@ namespace Metasound
 			virtual ~FOutputNodeController() = default;
 
 			virtual bool IsValid() const override;
+
+#if WITH_EDITOR
 			virtual const FText& GetDescription() const override;
 			virtual FText GetDisplayName() const override;
 			virtual const FText& GetDisplayTitle() const override;
 			virtual void SetDescription(const FText& InDescription) override;
 			virtual void SetDisplayName(const FText& InText) override;
+#endif // WITH_EDITOR
+
 			virtual void SetNodeName(const FVertexName& InName) override;
 			virtual const FMetasoundFrontendVersion& GetInterfaceVersion() const override;
 
@@ -335,13 +344,16 @@ namespace Metasound
 
 			virtual ~FInputNodeController() = default;
 
+#if WITH_EDITOR
 			virtual const FText& GetDescription() const override;
 			virtual FText GetDisplayName() const override;
 			virtual const FText& GetDisplayTitle() const override;
-			virtual const FMetasoundFrontendVersion& GetInterfaceVersion() const override;
-			virtual bool IsValid() const override;
 			virtual void SetDescription(const FText& InDescription) override;
 			virtual void SetDisplayName(const FText& InText) override;
+#endif // WITH_EDITOR
+
+			virtual const FMetasoundFrontendVersion& GetInterfaceVersion() const override;
+			virtual bool IsValid() const override;
 			virtual void SetNodeName(const FVertexName& InName) override;
 
 			// No-ops as inputs do not handle literals the same way as other nodes
diff --git a/Engine/Plugins/Runtime/Metasound/Source/MetasoundFrontend/Private/MetasoundFrontendOutputController.cpp b/Engine/Plugins/Runtime/Metasound/Source/MetasoundFrontend/Private/MetasoundFrontendOutputController.cpp
--- a/Engine/Plugins/Runtime/Metasound/Source/MetasoundFrontend/Private/MetasoundFrontendOutputController.cpp
+++ b/Engine/Plugins/Runtime/Metasound/Source/MetasoundFrontend/Private/MetasoundFrontendOutputController.cpp
@@ -167,6 +167,7 @@ namespace Metasound
 			return TArray<FMetasoundFrontendEdge>();
 		}
 
+#if WITH_EDITOR
 		FText FBaseOutputController::GetDisplayName() const
 		{
 			if (const FMetasoundFrontendClassOutput* Output = ClassOutputPtr.Get())
@@ -196,6 +197,7 @@ namespace Metasound
 
 			return Invalid::GetInvalidVertexMetadata();
 		}
+#endif // WITH_EDITOR
 
 		FDocumentAccess FBaseOutputController::ShareAccess()
 		{
@@ -234,6 +236,7 @@ namespace Metasound
 			return FBaseOutputController::IsValid() && (nullptr != OwningGraphClassInputPtr.Get());
 		}
 
+#if WITH_EDITOR
 		FText FInputNodeOutputController::GetDisplayName() const
 		{
 			if (const FMetasoundFrontendClassInput* OwningInput = OwningGraphClassInputPtr.Get())
@@ -272,6 +275,7 @@ namespace Metasound
 
 			return Invalid::GetInvalidVertexMetadata();
 		}
+#endif // WITH_EDITOR
 
 		void FInputNodeOutputController::SetName(const FVertexName& InName)
 		{
@@ -313,6 +317,7 @@ namespace Metasound
 			return FBaseOutputController::IsValid() && (nullptr != OwningGraphClassOutputPtr.Get());
 		}
 
+#if WITH_EDITOR
 		FText FOutputNodeOutputController::GetDisplayName() const
 		{
 			if (const FMetasoundFrontendClassOutput* OwningOutput = OwningGraphClassOutputPtr.Get())
@@ -342,6 +347,7 @@ namespace Metasound
 
 			return Invalid::GetInvalidVertexMetadata();
 		}
+#endif // WITH_EDITOR
 
 		void FOutputNodeOutputController::SetName(const FVertexName& InName)
 		{
diff --git a/Engine/Plugins/Runtime/Metasound/Source/MetasoundFrontend/Private/MetasoundFrontendOutputController.h b/Engine/Plugins/Runtime/Metasound/Source/MetasoundFrontend/Private/MetasoundFrontendOutputController.h
--- a/Engine/Plugins/Runtime/Metasound/Source/MetasoundFrontend/Private/MetasoundFrontendOutputController.h
+++ b/Engine/Plugins/Runtime/Metasound/Source/MetasoundFrontend/Private/MetasoundFrontendOutputController.h
@@ -44,10 +44,13 @@ namespace Metasound
 			virtual const FName& GetDataType() const override;
 			virtual const FVertexName& GetName() const override;
 
-			// Output metadata
+#if WITH_EDITOR
 			virtual FText GetDisplayName() const override;
 			virtual const FText& GetTooltip() const override;
+
+			// Output metadata
 			virtual const FMetasoundFrontendVertexMetadata& GetMetadata() const override;
+#endif // WITH_EDITOR
 
 			// Return info on containing node.
 			virtual FGuid GetOwningNodeID() const override;
@@ -114,10 +117,13 @@ namespace Metasound
 
 			virtual bool IsValid() const override;
 
-			// Input metadata
+#if WITH_EDITOR
 			virtual FText GetDisplayName() const override;
 			virtual const FText& GetTooltip() const override;
+
+			// Input metadata
 			virtual const FMetasoundFrontendVertexMetadata& GetMetadata() const override;
+#endif // WITH_EDITOR
 
 			virtual void SetName(const FVertexName& InName) override;
 
@@ -161,10 +167,13 @@ namespace Metasound
 
 			virtual bool IsValid() const override;
 
-			// Output metadata
+#if WITH_EDITOR
 			virtual FText GetDisplayName() const override;
 			virtual const FText& GetTooltip() const override;
+
+			// Output metadata
 			virtual const FMetasoundFrontendVertexMetadata& GetMetadata() const override;
+#endif // WITH_EDITOR
 
 			virtual void SetName(const FVertexName& InName) override;
 
diff --git a/Engine/Plugins/Runtime/Metasound/Source/MetasoundFrontend/Private/MetasoundFrontendTransform.cpp b/Engine/Plugins/Runtime/Metasound/Source/MetasoundFrontend/Private/MetasoundFrontendTransform.cpp
--- a/Engine/Plugins/Runtime/Metasound/Source/MetasoundFrontend/Private/MetasoundFrontendTransform.cpp
+++ b/Engine/Plugins/Runtime/Metasound/Source/MetasoundFrontend/Private/MetasoundFrontendTransform.cpp
@@ -50,10 +50,12 @@ namespace Metasound
 			Init();
 		}
 
+#if WITH_EDITOR
 		void FModifyRootGraphInterfaces::SetDefaultNodeLocations(bool bInSetDefaultNodeLocations)
 		{
 			bSetDefaultNodeLocations = bInSetDefaultNodeLocations;
 		}
+#endif // WITH_EDITOR
 
 		void FModifyRootGraphInterfaces::SetNamePairingFunction(const TFunction<bool(FName, FName)>& InNamePairingFunction)
 		{
@@ -224,6 +226,7 @@ namespace Metasound
 				return bMatchesDataType;
 			};
 
+#if WITH_EDITOR
 			auto FindLowestNodeLocationOfClassType = [](EMetasoundFrontendClassType ClassType, FGraphHandle Graph, FName DataType, TFunctionRef<bool(FConstNodeHandle, FName)> NodeDataTypeFilter)
 			{
 				FVector2D LowestLocation;
@@ -243,6 +246,7 @@ namespace Metasound
 
 				return LowestLocation;
 			};
+#endif // WITH_EDITOR
 
 			// Add missing inputs
 			for (const FMetasoundFrontendClassInput& InputToAdd : InputsToAdd)
@@ -250,6 +254,7 @@ namespace Metasound
 				bDidEdit = true;
 				FNodeHandle NewInputNode = GraphHandle->AddInputVertex(InputToAdd);
 
+#if WITH_EDITOR
 				if (bSetDefaultNodeLocations)
 				{
 					FMetasoundFrontendNodeStyle Style = NewInputNode->GetNodeStyle();
@@ -257,6 +262,7 @@ namespace Metasound
 					Style.Display.Locations.Add(FGuid(), LastOutputLocation + DisplayStyle::NodeLayout::DefaultOffsetY);
 					NewInputNode->SetNodeStyle(Style);
 				}
+#endif // WITH_EDITOR
 			}
 
 			// Add missing outputs
@@ -265,6 +271,7 @@ namespace Metasound
 				bDidEdit = true;
 				FNodeHandle NewOutputNode = GraphHandle->AddOutputVertex(OutputToAdd);
 
+#if WITH_EDITOR
 				if (bSetDefaultNodeLocations)
 				{
 					FMetasoundFrontendNodeStyle Style = NewOutputNode->GetNodeStyle();
@@ -272,6 +279,7 @@ namespace Metasound
 					Style.Display.Locations.Add(FGuid(), LastOutputLocation + DisplayStyle::NodeLayout::DefaultOffsetY);
 					NewOutputNode->SetNodeStyle(Style);
 				}
+#endif // WITH_EDITOR
 			}
 
 			// Swap paired inputs.
@@ -292,7 +300,10 @@ namespace Metasound
 						NewVertex.DefaultLiteral = ClassInput->DefaultLiteral;
 						NewVertex.NodeID = ClassInput->NodeID;
 						FNodeHandle OriginalInputNode = GraphHandle->GetInputNodeWithName(OriginalVertex.Name);
+
+#if WITH_EDITOR
 						Locations = OriginalInputNode->GetNodeStyle().Display.Locations;
+#endif // WITH_EDITOR
 
 						FOutputHandle OriginalInputNodeOutput = OriginalInputNode->GetOutputWithVertexName(OriginalVertex.Name);
 						ConnectedInputs = OriginalInputNodeOutput->GetConnectedInputs();
@@ -301,6 +312,8 @@ namespace Metasound
 				}
 
 				FNodeHandle NewInputNode = GraphHandle->AddInputVertex(NewVertex);
+
+#if WITH_EDITOR
 				// Copy prior node locations
 				if (!Locations.IsEmpty())
 				{
@@ -309,6 +322,7 @@ namespace Metasound
 					Style.Display.Locations = Locations;
 					NewInputNode->SetNodeStyle(Style);
 				}
+#endif // WITH_EDITOR
 
 				// Copy prior node connections
 				FOutputHandle OutputHandle = NewInputNode->GetOutputWithVertexName(NewVertex.Name);
@@ -326,24 +340,32 @@ namespace Metasound
 				const FMetasoundFrontendClassVertex& OriginalVertex = OutputPair.Get<0>();
 				FMetasoundFrontendClassVertex NewVertex = OutputPair.Get<1>();
 
+#if WITH_EDITOR
 				// Cache off node locations to push to new node
-				TMap<FGuid, FVector2D> Locations;
-				FOutputHandle ConnectedOutput = IOutputController::GetInvalidHandle();
-
 				// Default add output node to origin.
+				TMap<FGuid, FVector2D> Locations;
 				Locations.Add(FGuid(), FVector2D{0.f, 0.f});
+#endif // WITH_EDITOR
+
+				FOutputHandle ConnectedOutput = IOutputController::GetInvalidHandle();
 				if (const FMetasoundFrontendClassOutput* ClassOutput = GraphHandle->FindClassOutputWithName(OriginalVertex.Name).Get())
 				{
 					if (FMetasoundFrontendVertex::IsFunctionalEquivalent(*ClassOutput, OriginalVertex))
 					{
 						NewVertex.NodeID = ClassOutput->NodeID;
 
+#if WITH_EDITOR
 						// Interface members do not serialize text to avoid localization
 						// mismatches between assets and interfaces defined in code.
 						NewVertex.Metadata.SetSerializeText(false);
+#endif // WITH_EDITOR
 
 						FNodeHandle OriginalOutputNode = GraphHandle->GetOutputNodeWithName(OriginalVertex.Name);
+
+#if WITH_EDITOR
 						Locations = OriginalOutputNode->GetNodeStyle().Display.Locations;
+#endif // WITH_EDITOR
+
 						FInputHandle Input = OriginalOutputNode->GetInputWithVertexName(OriginalVertex.Name);
 						ConnectedOutput = Input->GetConnectedOutput();
 						GraphHandle->RemoveOutputVertex(OriginalVertex.Name);
@@ -352,12 +374,14 @@ namespace Metasound
 
 				FNodeHandle NewOutputNode = GraphHandle->AddOutputVertex(NewVertex);
 
+#if WITH_EDITOR
 				if (Locations.Num() > 0)
 				{
 					FMetasoundFrontendNodeStyle Style = NewOutputNode->GetNodeStyle();
 					Style.Display.Locations = Locations;
 					NewOutputNode->SetNodeStyle(Style);
 				}
+#endif // WITH_EDITOR
 
 				// Copy prior node connections
 				FInputHandle InputHandle = NewOutputNode->GetInputWithVertexName(NewVertex.Name);
@@ -442,8 +466,13 @@ namespace Metasound
 
 			if (LastVersionUpdated)
 			{
+#if WITH_EDITOR
+				const FString AssetName = *InDocument->GetRootGraphClass().Metadata.GetDisplayName().ToString();
+#else
+				const FString AssetName = *InDocument->GetRootGraphClass().Metadata.GetClassName().ToString();
+#endif // !WITH_EDITOR
 				UE_LOG(LogMetaSound, Display, TEXT("Asset '%s' interface '%s' updated: '%s' --> '%s'"),
-					*InDocument->GetRootGraphClass().Metadata.GetDisplayName().ToString(),
+					*AssetName,
 					*InterfaceVersion.Name.ToString(),
 					*InterfaceVersion.Number.ToString(),
 					*LastVersionUpdated->ToString());
@@ -602,13 +631,16 @@ namespace Metasound
 			// Swap type on look-up as it will be referenced as an externally defined class relative to the new Preset asset
 			ReferencedClassMetadata.SetType(EMetasoundFrontendClassType::External);
 
-			// Set node location.
 			FNodeHandle ReferencedNodeHandle = PresetGraphHandle->AddNode(ReferencedClassMetadata, PresetNodeID);
+#if WITH_EDITOR
+
+			// Set node location.
 			FMetasoundFrontendNodeStyle RefNodeStyle;
 
 			// Offset to be to the right of input nodes
 			RefNodeStyle.Display.Locations.Add(FGuid::NewGuid(), DisplayStyle::NodeLayout::DefaultOffsetX);
 			ReferencedNodeHandle->SetNodeStyle(RefNodeStyle);
+#endif // WITH_EDITOR
 
 			// Connect parent graph to referenced graph
 			PresetGraphHandle->SetInputsInheritingDefault(MoveTemp(InputsInheritingDefault));
@@ -630,11 +662,13 @@ namespace Metasound
 
 				if (ensure(InputNode->IsValid()))
 				{
+#if WITH_EDITOR
 					// Set input node location
 					FMetasoundFrontendNodeStyle NodeStyle;
 					NodeStyle.Display.Locations.Add(FGuid::NewGuid(), InputNodeLocation);
 					InputNode->SetNodeStyle(NodeStyle);
 					InputNodeLocation += DisplayStyle::NodeLayout::DefaultOffsetY;
+#endif // WITH_EDITOR
 
 					// Connect input node to corresponding referencing node.
 					FOutputHandle OutputToConnect = InputNode->GetOutputWithVertexName(ClassInput.Name);
@@ -656,11 +690,13 @@ namespace Metasound
 				
 				if (ensure(OutputNode->IsValid()))
 				{
+#if WITH_EDITOR
 					// Set input node location
 					FMetasoundFrontendNodeStyle NodeStyle;
 					NodeStyle.Display.Locations.Add(FGuid::NewGuid(), OutputNodeLocation);
 					OutputNode->SetNodeStyle(NodeStyle);
 					OutputNodeLocation += DisplayStyle::NodeLayout::DefaultOffsetY;
+#endif // WITH_EDITOR
 
 					// Connect input node to corresponding referenced node. 
 					FInputHandle InputToConnect = OutputNode->GetInputWithVertexName(ClassOutput.Name);
@@ -689,8 +725,10 @@ namespace Metasound
 					ClassInput.Name = NodeName;
 					ClassInput.TypeName = Input->GetDataType();
 
+#if WITH_EDITOR
 					ClassInput.Metadata.SetDescription(InputNode->GetDescription());
 					ClassInput.Metadata.SetDisplayName(Input->GetMetadata().GetDisplayName());
+#endif // WITH_EDITOR
 
 					ClassInput.VertexID = FGuid::NewGuid();
 
@@ -762,8 +800,10 @@ namespace Metasound
 					ClassOutput.Name = NodeName;
 					ClassOutput.TypeName = Output->GetDataType();
 
+#if WITH_EDITOR
 					ClassOutput.Metadata.SetDescription(OutputNode->GetDescription());
 					ClassOutput.Metadata.SetDisplayName(Output->GetMetadata().GetDisplayName());
+#endif // WITH_EDITOR
 
 					ClassOutput.VertexID = FGuid::NewGuid();
 
@@ -818,7 +858,16 @@ namespace Metasound
 		/** Versions document from 1.0 to 1.1. */
 		class FVersionDocument_1_1 : public FVersionDocumentTransform
 		{
+			FName Name;
+			const FString& Path;
+
 		public:
+			FVersionDocument_1_1(FName InName, const FString& InPath)
+			: Name(InName)
+			, Path(InPath)
+			{
+			}
+
 			FMetasoundFrontendVersionNumber GetTargetVersion() const override
 			{
 				return { 1, 1 };
@@ -826,6 +875,7 @@ namespace Metasound
 
 			void TransformInternal(FDocumentHandle InDocument) const override
 			{
+#if WITH_EDITOR
 				FGraphHandle GraphHandle = InDocument->GetRootGraph();
 				TArray<FNodeHandle> FrontendNodes = GraphHandle->GetNodes();
 
@@ -871,6 +921,9 @@ namespace Metasound
 						GraphHandle->RemoveNode(*NodeHandle);
 					}
 				}
+#else
+				UE_LOG(LogMetaSound, Error, TEXT("Asset '%s' at '%s' must be saved with editor enabled in order to version document to target version '%s'."), *Name.ToString(), *Path, *GetTargetVersion().ToString());
+#endif // !WITH_EDITOR
 			}
 		};
 
@@ -895,12 +948,16 @@ namespace Metasound
 
 			void TransformInternal(FDocumentHandle InDocument) const override
 			{
+#if WITH_EDITOR
 				const FMetasoundFrontendGraphClass& GraphClass = InDocument->GetRootGraphClass();
 				FMetasoundFrontendClassMetadata Metadata = GraphClass.Metadata;
 
 				Metadata.SetClassName({ "GraphAsset", Name, *Path });
 				Metadata.SetDisplayName(FText::FromString(Name.ToString()));
 				InDocument->GetRootGraph()->SetGraphMetadata(Metadata);
+#else
+				UE_LOG(LogMetaSound, Error, TEXT("Asset '%s' at '%s' must be saved with editor enabled in order to version document to target version '%s'."), *Name.ToString(), *Path, *GetTargetVersion().ToString());
+#endif // !WITH_EDITOR
 			}
 		};
 
@@ -978,9 +1035,13 @@ namespace Metasound
 		/** Versions document from 1.4 to 1.5. */
 		class FVersionDocument_1_5 : public FVersionDocumentTransform
 		{
+			FName Name;
+			const FString& Path;
+
 		public:
-			FVersionDocument_1_5(FName InAssetName)
-				: AssetName(InAssetName)
+			FVersionDocument_1_5(FName InName, const FString& InPath)
+				: Name(InName)
+				, Path(InPath)
 			{
 			}
 
@@ -991,18 +1052,19 @@ namespace Metasound
 
 			void TransformInternal(FDocumentHandle InDocument) const override
 			{
+#if WITH_EDITOR
 				const FMetasoundFrontendClassMetadata& Metadata = InDocument->GetRootGraphClass().Metadata;
-				const FText NewAssetName = FText::FromString(AssetName.ToString());
+				const FText NewAssetName = FText::FromString(Name.ToString());
 				if (Metadata.GetDisplayName().CompareTo(NewAssetName) != 0)
 				{
 					FMetasoundFrontendClassMetadata NewMetadata = Metadata;
 					NewMetadata.SetDisplayName(NewAssetName);
 					InDocument->GetRootGraph()->SetGraphMetadata(NewMetadata);
 				}
+#else
+				UE_LOG(LogMetaSound, Error, TEXT("Asset '%s' at '%s' must be saved with editor enabled in order to version document to target version '%s'."), *Name.ToString(), *Path, *GetTargetVersion().ToString());
+#endif // !WITH_EDITOR
 			}
-
-		private:
-			FName AssetName;
 		};
 
 		/** Versions document from 1.5 to 1.6. */
@@ -1027,8 +1089,16 @@ namespace Metasound
 		/** Versions document from 1.6 to 1.7. */
 		class FVersionDocument_1_7 : public FVersionDocumentTransform
 		{
+			FName Name;
+			const FString& Path;
+
 		public:
-			FVersionDocument_1_7() = default;
+			FVersionDocument_1_7(FName InName, const FString& InPath)
+				: Name(InName)
+				, Path(InPath)
+			{
+			}
+
 
 			FMetasoundFrontendVersionNumber GetTargetVersion() const override
 			{
@@ -1037,6 +1107,7 @@ namespace Metasound
 
 			void TransformInternal(FDocumentHandle InDocument) const override
 			{
+#if WITH_EDITOR
 				auto RenameTransform = [](FNodeHandle NodeHandle)
 				{
 					// Required nodes are all (at the point of this transform) providing
@@ -1063,14 +1134,24 @@ namespace Metasound
 
 				InDocument->GetRootGraph()->IterateNodes(RenameTransform, EMetasoundFrontendClassType::Input);
 				InDocument->GetRootGraph()->IterateNodes(RenameTransform, EMetasoundFrontendClassType::Output);
+#else
+				UE_LOG(LogMetaSound, Error, TEXT("Asset '%s' at '%s' must be saved with editor enabled in order to version document to target version '%s'."), *Name.ToString(), *Path, *GetTargetVersion().ToString());
+#endif // !WITH_EDITOR
 			}
 		};
 
 		/** Versions document from 1.7 to 1.8. */
 		class FVersionDocument_1_8 : public FVersionDocumentTransform
 		{
+			FName Name;
+			const FString& Path;
+
 		public:
-			FVersionDocument_1_8() = default;
+			FVersionDocument_1_8(FName InName, const FString& InPath)
+				: Name(InName)
+				, Path(InPath)
+			{
+			}
 
 			FMetasoundFrontendVersionNumber GetTargetVersion() const override
 			{
@@ -1079,6 +1160,7 @@ namespace Metasound
 
 			void TransformInternal(FDocumentHandle InDocument) const override
 			{
+#if WITH_EDITOR
 				// Do not serialize MetaData text for dependencies as
 				// CacheRegistryData dynamically provides this.
 				InDocument->IterateDependencies([](FMetasoundFrontendClass& Dependency)
@@ -1138,14 +1220,24 @@ namespace Metasound
 				}
 
 				InDocument->SetRootGraphClass(MoveTemp(RootGraphClass));
+#else
+			UE_LOG(LogMetaSound, Error, TEXT("Asset '%s' at '%s' must be saved with editor enabled in order to version document to target version '%s'."), *Name.ToString(), *Path, *GetTargetVersion().ToString());
+#endif // !WITH_EDITOR
 			}
 		};
 
 		/** Versions document from 1.8 to 1.9. */
 		class FVersionDocument_1_9 : public FVersionDocumentTransform
 		{
+			FName Name;
+			const FString& Path;
+
 		public:
-			FVersionDocument_1_9() = default;
+			FVersionDocument_1_9(FName InName, const FString& InPath)
+				: Name(InName)
+				, Path(InPath)
+			{
+			}
 
 			FMetasoundFrontendVersionNumber GetTargetVersion() const override
 			{
@@ -1154,12 +1246,16 @@ namespace Metasound
 
 			void TransformInternal(FDocumentHandle InDocument) const override
 			{
+#if WITH_EDITOR
 				// Display name text is no longer copied at this versioning point for assets
 				// from the asset's FName to avoid FText warnings regarding generation from
 				// an FString.  It also avoids desync if asset gets moved.
 				FMetasoundFrontendGraphClass RootGraphClass = InDocument->GetRootGraphClass();
 				RootGraphClass.Metadata.SetDisplayName(FText());
 				InDocument->SetRootGraphClass(MoveTemp(RootGraphClass));
+#else
+				UE_LOG(LogMetaSound, Error, TEXT("Asset '%s' at '%s' must be saved with editor enabled in order to version document to target version '%s'."), *Name.ToString(), *Path, *GetTargetVersion().ToString());
+#endif // !WITH_EDITOR
 			}
 		};
 
@@ -1201,20 +1297,19 @@ namespace Metasound
 			const FMetasoundFrontendVersionNumber InitVersionNumber = InDocument->GetMetadata().Version.Number;
 
 			// Add additional transforms here after defining them above, example below.
-			bWasUpdated |= FVersionDocument_1_1().Transform(InDocument);
+			bWasUpdated |= FVersionDocument_1_1(Name, Path).Transform(InDocument);
 			bWasUpdated |= FVersionDocument_1_2(Name, Path).Transform(InDocument);
 			bWasUpdated |= FVersionDocument_1_3().Transform(InDocument);
 			bWasUpdated |= FVersionDocument_1_4().Transform(InDocument);
-			bWasUpdated |= FVersionDocument_1_5(Name).Transform(InDocument);
+			bWasUpdated |= FVersionDocument_1_5(Name, Path).Transform(InDocument);
 			bWasUpdated |= FVersionDocument_1_6().Transform(InDocument);
-			bWasUpdated |= FVersionDocument_1_7().Transform(InDocument);
-			bWasUpdated |= FVersionDocument_1_8().Transform(InDocument);
-			bWasUpdated |= FVersionDocument_1_9().Transform(InDocument);
+			bWasUpdated |= FVersionDocument_1_7(Name, Path).Transform(InDocument);
+			bWasUpdated |= FVersionDocument_1_8(Name, Path).Transform(InDocument);
+			bWasUpdated |= FVersionDocument_1_9(Name, Path).Transform(InDocument);
 			bWasUpdated |= FVersionDocument_1_10().Transform(InDocument);
 
 			if (bWasUpdated)
 			{
-				const FText& DisplayName = InDocument->GetRootGraph()->GetGraphMetadata().GetDisplayName();
 				const FMetasoundFrontendVersionNumber NewVersionNumber = InDocument->GetMetadata().Version.Number;
 				UE_LOG(LogMetaSound, Display, TEXT("MetaSound at '%s' Document Versioned: '%s' --> '%s'"), *Path, *InitVersionNumber.ToString(), *NewVersionNumber.ToString());
 			}
diff --git a/Engine/Plugins/Runtime/Metasound/Source/MetasoundFrontend/Private/MetasoundFrontendVariableController.cpp b/Engine/Plugins/Runtime/Metasound/Source/MetasoundFrontend/Private/MetasoundFrontendVariableController.cpp
--- a/Engine/Plugins/Runtime/Metasound/Source/MetasoundFrontend/Private/MetasoundFrontendVariableController.cpp
+++ b/Engine/Plugins/Runtime/Metasound/Source/MetasoundFrontend/Private/MetasoundFrontendVariableController.cpp
@@ -59,6 +59,7 @@ namespace Metasound
 			}
 		}
 		
+#if WITH_EDITOR
 		/** Returns the human readable name associated with this output. */
 		FText FVariableController::GetDisplayName() const
 		{
@@ -94,6 +95,7 @@ namespace Metasound
 				Variable->Description = InDescription;
 			}
 		}
+#endif // WITH_EDITOR
 
 		FNodeHandle FVariableController::FindMutatorNode()
 		{
diff --git a/Engine/Plugins/Runtime/Metasound/Source/MetasoundFrontend/Private/MetasoundFrontendVariableController.h b/Engine/Plugins/Runtime/Metasound/Source/MetasoundFrontend/Private/MetasoundFrontendVariableController.h
--- a/Engine/Plugins/Runtime/Metasound/Source/MetasoundFrontend/Private/MetasoundFrontendVariableController.h
+++ b/Engine/Plugins/Runtime/Metasound/Source/MetasoundFrontend/Private/MetasoundFrontendVariableController.h
@@ -38,6 +38,7 @@ namespace Metasound
 			/** Sets the name associated with this variable. */
 			virtual void SetName(const FName& InName) override;
 			
+#if WITH_EDITOR
 			/** Returns the human readable name associated with this variable. */
 			virtual FText GetDisplayName() const override;
 
@@ -49,6 +50,7 @@ namespace Metasound
 
 			/** Sets the human readable description associated with this variable. */
 			virtual void SetDescription(const FText& InDescription) override;
+#endif // WITH_EDITOR
 
 			/** Returns the mutator node associated with this variable. */
 			virtual FNodeHandle FindMutatorNode() override;
diff --git a/Engine/Plugins/Runtime/Metasound/Source/MetasoundFrontend/Public/MetasoundFrontendController.h b/Engine/Plugins/Runtime/Metasound/Source/MetasoundFrontend/Public/MetasoundFrontendController.h
--- a/Engine/Plugins/Runtime/Metasound/Source/MetasoundFrontend/Public/MetasoundFrontendController.h
+++ b/Engine/Plugins/Runtime/Metasound/Source/MetasoundFrontend/Public/MetasoundFrontendController.h
@@ -177,6 +177,7 @@ namespace Metasound
 			/** Returns the name associated with this output. */
 			virtual const FVertexName& GetName() const = 0;
 			
+#if WITH_EDITOR
 			/** Returns the human readable name associated with this output. */
 			virtual FText GetDisplayName() const = 0;
 			
@@ -185,6 +186,7 @@ namespace Metasound
 
 			/** Returns all metadata associated with this output. */
 			virtual const FMetasoundFrontendVertexMetadata& GetMetadata() const = 0;
+#endif // WITH_EDITOR
 
 			/** Returns the ID of the node which owns this output. */
 			virtual FGuid GetOwningNodeID() const = 0;
@@ -258,6 +260,7 @@ namespace Metasound
 			/** Sets the name associated with this variable. */
 			virtual void SetName(const FName& InName) = 0;
 			
+#if WITH_EDITOR
 			/** Returns the human readable name associated with this variable. */
 			virtual FText GetDisplayName() const = 0;
 
@@ -269,6 +272,7 @@ namespace Metasound
 
 			/** Sets the human readable description associated with this variable. */
 			virtual void SetDescription(const FText& InDescription) = 0;
+#endif // WITH_EDITOR
 
 			/** Returns the mutator node associated with this variable. */
 			virtual FNodeHandle FindMutatorNode() = 0;
@@ -324,8 +328,10 @@ namespace Metasound
 			/** Returns the data type name associated with this input. */
 			virtual const FVertexName& GetName() const = 0;
 
+#if WITH_EDITOR
 			/** Returns the data type name associated with this input. */
 			virtual FText GetDisplayName() const = 0;
+#endif // WITH_EDITOR
 
 			/** Clears the value for the given input instance if set. */
 			virtual bool ClearLiteral() = 0;
@@ -339,11 +345,13 @@ namespace Metasound
 			/** Returns the class default value of the given input. */
 			virtual const FMetasoundFrontendLiteral* GetClassDefaultLiteral() const = 0;
 
+#if WITH_EDITOR
 			/** Returns the data type name associated with this input. */
 			virtual const FText& GetTooltip() const = 0;
 
 			/** Returns all metadata associated with this input. */
 			virtual const FMetasoundFrontendVertexMetadata& GetMetadata() const = 0;
+#endif // WITH_EDITOR
 
 			/** Returns the ID of the node which owns this output. */
 			virtual FGuid GetOwningNodeID() const = 0;
@@ -413,19 +421,21 @@ namespace Metasound
 			/** Returns all node inputs. */
 			virtual TArray<FConstInputHandle> GetConstInputs() const = 0;
 
+#if WITH_EDITOR
 			/** Returns the display name of the given node (what to distinguish and label in visual arrays, such as context menus). */
 			virtual FText GetDisplayName() const = 0;
 
 			/** Sets the description of the node. */
 			virtual void SetDescription(const FText& InDescription) = 0;
 
-			virtual void SetNodeName(const FVertexName& InName) = 0;
-
 			/** Sets the display name of the node. */
 			virtual void SetDisplayName(const FText& InDisplayName) = 0;
 
 			/** Returns the title of the given node (what to label when displayed as visual node). */
 			virtual const FText& GetDisplayTitle() const = 0;
+#endif // WITH_EDITOR
+
+			virtual void SetNodeName(const FVertexName& InName) = 0;
 
 			/** Iterate over inputs */
 			virtual void IterateInputs(TUniqueFunction<void(FInputHandle)> InFunction) = 0;
@@ -513,12 +523,15 @@ namespace Metasound
 			// Clears an input's default literal
 			virtual bool ClearInputLiteral(FGuid InVertexID) = 0;
 
-			/** Returns associated node class data */
-			virtual const FMetasoundFrontendClassInterface& GetClassInterface() const = 0;
 			virtual const FMetasoundFrontendClassMetadata& GetClassMetadata() const = 0;
+			virtual const FMetasoundFrontendClassInterface& GetClassInterface() const = 0;
+
+#if WITH_EDITOR
+			/** Returns associated node class data */
 			virtual const FMetasoundFrontendInterfaceStyle& GetOutputStyle() const = 0;
 			virtual const FMetasoundFrontendInterfaceStyle& GetInputStyle() const = 0;
 			virtual const FMetasoundFrontendClassStyle& GetClassStyle() const = 0;
+#endif // WITH_EDITOR
 
 			/**
 			  * Fills out the provided ClassInterfaceUpdate struct with the differences between
@@ -537,8 +550,10 @@ namespace Metasound
 			  */
 			virtual bool CanAutoUpdate(FClassInterfaceUpdates& OutInterfaceUpdates) const = 0;
 
+#if WITH_EDITOR
 			/** Description of the given node. */
 			virtual const FText& GetDescription() const = 0;
+#endif // WITH_EDITOR
 
 			/** If the node is also a graph, this returns a graph handle.
 			 * If the node is not also a graph, it will return an invalid handle.
@@ -553,8 +568,10 @@ namespace Metasound
 			/** Returns the name of this node. */
 			virtual const FVertexName& GetNodeName() const = 0;
 
+#if WITH_EDITOR
 			virtual const FMetasoundFrontendNodeStyle& GetNodeStyle() const = 0;
 			virtual void SetNodeStyle(const FMetasoundFrontendNodeStyle& InStyle) = 0;
+#endif // WITH_EDITOR
 
 			/** Returns the ID associated with this node. */
 			virtual FGuid GetID() const = 0;
@@ -599,6 +616,7 @@ namespace Metasound
 			/** Sets the metadata for the current graph. */
 			virtual void SetGraphMetadata(const FMetasoundFrontendClassMetadata& InMetadata) = 0;
 
+#if WITH_EDITOR
 			// Returns graph style.
 			virtual const FMetasoundFrontendGraphStyle& GetGraphStyle() const = 0;
 
@@ -607,6 +625,7 @@ namespace Metasound
 
 			/** Return the display name of the graph. */
 			virtual FText GetDisplayName() const = 0;
+#endif // WITH_EDITOR
 
 			virtual TArray<FVertexName> GetInputVertexNames() const = 0;
 			virtual TArray<FVertexName> GetOutputVertexNames() const = 0;
@@ -837,6 +856,7 @@ namespace Metasound
 			 */
 			virtual bool SetDefaultInputToDefaultLiteralOfType(const FGuid& InVertexID) = 0;
 
+#if WITH_EDITOR
 			/** Set the display name for the input with the given name. */
 			virtual void SetInputDisplayName(const FVertexName& InName, const FText& InDisplayName) = 0;
 
@@ -854,6 +874,7 @@ namespace Metasound
 
 			/** Set the description for the output with the given name. */
 			virtual void SetOutputDescription(const FVertexName& InName, const FText& InDescription) = 0;
+#endif // WITH_EDITOR
 
 			/**
 			  * Updates the ChangeID for the class interface, which signals AutoUpdate to
diff --git a/Engine/Plugins/Runtime/Metasound/Source/MetasoundFrontend/Public/MetasoundFrontendDocument.h b/Engine/Plugins/Runtime/Metasound/Source/MetasoundFrontend/Public/MetasoundFrontendDocument.h
--- a/Engine/Plugins/Runtime/Metasound/Source/MetasoundFrontend/Public/MetasoundFrontendDocument.h
+++ b/Engine/Plugins/Runtime/Metasound/Source/MetasoundFrontend/Public/MetasoundFrontendDocument.h
@@ -497,7 +497,7 @@ struct FMetasoundFrontendStyleEdgeClass
 };
 
 // Styling for a class
-USTRUCT() 
+USTRUCT()
 struct FMetasoundFrontendGraphStyle 
 {
 	GENERATED_BODY()
@@ -534,11 +534,12 @@ struct FMetasoundFrontendGraph
 };
 
 // Metadata associated with a vertex.
-USTRUCT() 
+USTRUCT()
 struct FMetasoundFrontendVertexMetadata
 {
 	GENERATED_BODY()
 
+#if WITH_EDITORONLY_DATA
 private:
 	// Display name for a vertex
 	UPROPERTY(EditAnywhere, Category = Parameters, meta = (DisplayName = "Name"))
@@ -643,10 +644,10 @@ public:
 		
 		bSerializeText = bInSerializeText;
 	}
+#endif // WITH_EDITORONLY_DATA
 };
 
-
-USTRUCT() 
+USTRUCT()
 struct METASOUNDFRONTEND_API FMetasoundFrontendClassVertex : public FMetasoundFrontendVertex
 {
 	GENERATED_BODY()
@@ -656,9 +657,11 @@ struct METASOUNDFRONTEND_API FMetasoundFrontendClassVertex : public FMetasoundFr
 	UPROPERTY()
 	FGuid NodeID = Metasound::FrontendInvalidID;
 
+#if WITH_EDITORONLY_DATA
 	// Metadata associated with input.
 	UPROPERTY(EditAnywhere, Category = CustomView)
 	FMetasoundFrontendVertexMetadata Metadata;
+#endif // WITH_EDITORONLY_DATA
 
 	// Splits name into namespace & parameter name
 	void SplitName(FName& OutNamespace, FName& OutParameterName) const;
@@ -974,6 +977,7 @@ private:
 	UPROPERTY(VisibleAnywhere, Category = Metasound)
 	EMetasoundFrontendClassType Type = EMetasoundFrontendClassType::Invalid;
 
+#if WITH_EDITORONLY_DATA
 	UPROPERTY(EditAnywhere, Category = Metasound)
 	FText DisplayName;
 
@@ -994,10 +998,7 @@ private:
 	FText PromptIfMissingTransient;
 
 	UPROPERTY(EditAnywhere, Category = Metasound)
-	FText Author;
-
-	UPROPERTY(Transient)
-	FText AuthorTransient;
+	FString Author;
 
 	UPROPERTY(EditAnywhere, Category = Metasound)
 	TArray<FText> Keywords;
@@ -1011,6 +1012,8 @@ private:
 	UPROPERTY(Transient)
 	TArray<FText> CategoryHierarchyTransient;
 
+#endif // WITH_EDITORONLY_DATA
+
 	// If true, this node is deprecated and should not be used in new MetaSounds.
 	UPROPERTY(EditAnywhere, Category = Metasound)
 	bool bIsDeprecated = false;
@@ -1021,6 +1024,7 @@ private:
 	UPROPERTY()
 	bool bAutoUpdateManagesInterface = false;
 
+#if WITH_EDITORONLY_DATA
 	// Whether or not the given metadata text should be serialized
 	// or is procedurally maintained via auto-update & the referenced
 	// registry class (to avoid localization text desync).  Should be
@@ -1028,6 +1032,7 @@ private:
 	// or interfaces.
 	UPROPERTY()
 	bool bSerializeText = true;
+#endif // WITH_EDITORONLY_DATA
 
 	// ID used to identify if any of the above have been modified,
 	// to determine if the parent class should be auto-updated.
@@ -1035,6 +1040,7 @@ private:
 	FGuid ChangeID;
 
 public:
+#if WITH_EDITOR
 	static FName GetAuthorPropertyName()
 	{
 		return GET_MEMBER_NAME_CHECKED(FMetasoundFrontendClassMetadata, Author);
@@ -1074,6 +1080,7 @@ public:
 	{
 		return GET_MEMBER_NAME_CHECKED(FMetasoundFrontendClassMetadata, Version);
 	}
+#endif // WITH_EDITOR
 
 	const FMetasoundFrontendClassName& GetClassName() const
 	{
@@ -1092,6 +1099,7 @@ public:
 		return Version;
 	}
 
+#if WITH_EDITOR
 	const FText& GetDisplayName() const
 	{
 		return bSerializeText ? DisplayName : DisplayNameTransient;
@@ -1107,9 +1115,9 @@ public:
 		return PromptIfMissingTransient;
 	}
 
-	const FText& GetAuthor() const
+	const FString& GetAuthor() const
 	{
-		return bSerializeText ? Author : AuthorTransient;
+		return Author;
 	}
 
 	const TArray<FText>& GetKeywords() const
@@ -1122,17 +1130,7 @@ public:
 		return bSerializeText ? CategoryHierarchy : CategoryHierarchyTransient;
 	}
 
-	const FGuid& GetChangeID() const
-	{
-		return ChangeID;
-	}
-
-	bool GetIsDeprecated() const
-	{
-		return bIsDeprecated;
-	}
-
-	void SetAuthor(const FText& InAuthor);
+	void SetAuthor(const FString& InAuthor);
 	void SetCategoryHierarchy(const TArray<FText>& InCategoryHierarchy);
 	void SetDescription(const FText& InDescription);
 	void SetDisplayName(const FText& InDisplayName);
@@ -1142,6 +1140,17 @@ public:
 	void SetVersion(const FMetasoundFrontendVersionNumber& InVersion);
 
 	void SetSerializeText(bool bInSerializeText);
+#endif // WITH_EDITOR
+
+	const FGuid& GetChangeID() const
+	{
+		return ChangeID;
+	}
+
+	bool GetIsDeprecated() const
+	{
+		return bIsDeprecated;
+	}
 
 	void SetType(const EMetasoundFrontendClassType InType)
 	{
@@ -1209,7 +1218,7 @@ struct METASOUNDFRONTEND_API FMetasoundFrontendClass
 	UPROPERTY()
 	FMetasoundFrontendClassStyle Style;
 
-#if WITH_EDITORONLY_DATA
+#if WITH_EDITOR
 	/*
 	 * Caches transient style, class & vertex Metadata found in the registry
 	 * on a passed (presumed) dependency.  Only modifies properties that are
@@ -1218,7 +1227,7 @@ struct METASOUNDFRONTEND_API FMetasoundFrontendClass
 	 * @return - Whether class was found in the registry & data was cached successfully.
 	 */
 	static bool CacheGraphDependencyMetadataFromRegistry(FMetasoundFrontendClass& InOutDependency);
-#endif // WITH_EDITORONLY_DATA
+#endif // WITH_EDITOR
 };
 
 // Preset options related to a parent graph class.  A graph class with bIsPreset set to true
diff --git a/Engine/Plugins/Runtime/Metasound/Source/MetasoundFrontend/Public/MetasoundFrontendTransform.h b/Engine/Plugins/Runtime/Metasound/Source/MetasoundFrontend/Public/MetasoundFrontendTransform.h
--- a/Engine/Plugins/Runtime/Metasound/Source/MetasoundFrontend/Public/MetasoundFrontendTransform.h
+++ b/Engine/Plugins/Runtime/Metasound/Source/MetasoundFrontend/Public/MetasoundFrontendTransform.h
@@ -38,9 +38,11 @@ namespace Metasound
 			FModifyRootGraphInterfaces(const TArray<FMetasoundFrontendInterface>& InInterfacesToRemove, const TArray<FMetasoundFrontendInterface>& InInterfacesToAdd);
 			FModifyRootGraphInterfaces(const TArray<FMetasoundFrontendVersion>& InInterfaceVersionsToRemove, const TArray<FMetasoundFrontendVersion>& InInterfaceVersionsToAdd);
 
+#if WITH_EDITOR
 			// Whether or not to propagate node locations to new members. Setting to false
 			// results in members not having a default physical location in the editor graph.
 			void SetDefaultNodeLocations(bool bInSetDefaultNodeLocations);
+#endif // WITH_EDITOR
 
 			// Override function used to match removed members with added members, allowing
 			// transform to preserve connections made between removed interface members & new interface members
@@ -52,7 +54,9 @@ namespace Metasound
 		private:
 			void Init(const TFunction<bool(FName, FName)>* InNamePairingFunction = nullptr);
 
+#if WITH_EDITOR
 			bool bSetDefaultNodeLocations = true;
+#endif // WITH_EDITOR
 
 			TArray<FMetasoundFrontendInterface> InterfacesToRemove;
 			TArray<FMetasoundFrontendInterface> InterfacesToAdd;
diff --git a/Engine/Plugins/Runtime/Metasound/Source/MetasoundGraphCore/Private/MetasoundNodeInterface.cpp b/Engine/Plugins/Runtime/Metasound/Source/MetasoundGraphCore/Private/MetasoundNodeInterface.cpp
--- a/Engine/Plugins/Runtime/Metasound/Source/MetasoundGraphCore/Private/MetasoundNodeInterface.cpp
+++ b/Engine/Plugins/Runtime/Metasound/Source/MetasoundGraphCore/Private/MetasoundNodeInterface.cpp
@@ -4,7 +4,7 @@
 
 namespace Metasound
 {
-	const FText PluginAuthor = NSLOCTEXT("MetasoundGraphCore", "Metasound_DefaultAuthor", "Epic Games, Inc.");
+	const FString PluginAuthor = TEXT("Epic Games, Inc.");
 	const FText PluginNodeMissingPrompt = NSLOCTEXT("MetasoundGraphCore", "Metasound_DefaultMissingNodePrompt", "The node was likely removed, renamed, or the Metasound plugin is not loaded.");
 
 	FNodeClassName::FNodeClassName()
diff --git a/Engine/Plugins/Runtime/Metasound/Source/MetasoundGraphCore/Public/MetasoundNodeInterface.h b/Engine/Plugins/Runtime/Metasound/Source/MetasoundGraphCore/Public/MetasoundNodeInterface.h
--- a/Engine/Plugins/Runtime/Metasound/Source/MetasoundGraphCore/Public/MetasoundNodeInterface.h
+++ b/Engine/Plugins/Runtime/Metasound/Source/MetasoundGraphCore/Public/MetasoundNodeInterface.h
@@ -13,7 +13,7 @@
 
 namespace Metasound
 {
-	extern const FText METASOUNDGRAPHCORE_API PluginAuthor;
+	extern const FString METASOUNDGRAPHCORE_API PluginAuthor;
 	extern const FText METASOUNDGRAPHCORE_API PluginNodeMissingPrompt;
 
 	/**
@@ -98,7 +98,7 @@ namespace Metasound
 		FText Description;
 
 		/** Author information. */
-		FText Author;
+		FString Author;
 
 		/** Human readable prompt for acquiring plugin in case node is not loaded. */
 		FText PromptIfMissing;
