Don't copy filename to displayname by default when creating a MetaSound
#rnx

diff --git a/Engine/Plugins/Runtime/Metasound/Source/MetasoundEditor/Private/MetasoundEditorGraphBuilder.cpp b/Engine/Plugins/Runtime/Metasound/Source/MetasoundEditor/Private/MetasoundEditorGraphBuilder.cpp
--- a/Engine/Plugins/Runtime/Metasound/Source/MetasoundEditor/Private/MetasoundEditorGraphBuilder.cpp
+++ b/Engine/Plugins/Runtime/Metasound/Source/MetasoundEditor/Private/MetasoundEditorGraphBuilder.cpp
@@ -1119,7 +1119,6 @@ namespace Metasound
 			// 1. Set default class Metadata
 			Metadata.SetClassName(FMetasoundFrontendClassName(FName(), *FGuid::NewGuid().ToString(), FName()));
 			Metadata.SetVersion({ 1, 0 });
-			Metadata.SetDisplayName(FText::FromString(InMetaSound.GetName()));
 			Metadata.SetType(EMetasoundFrontendClassType::Graph);
 			Metadata.SetAuthor(InAuthor);
 
