Fix MetaSound Tests Part 1

diff --git a/Engine/Plugins/Runtime/Metasound/Source/MetasoundFrontend/Private/MetasoundAssetBase.cpp b/Engine/Plugins/Runtime/Metasound/Source/MetasoundFrontend/Private/MetasoundAssetBase.cpp
--- a/Engine/Plugins/Runtime/Metasound/Source/MetasoundFrontend/Private/MetasoundAssetBase.cpp
+++ b/Engine/Plugins/Runtime/Metasound/Source/MetasoundFrontend/Private/MetasoundAssetBase.cpp
@@ -112,28 +112,28 @@ void FMetasoundAssetBase::RegisterGraphWithFrontend(Metasound::Frontend::FMetaSo
 	{
 		FString OwningAssetName = GetOwningAssetName();
 		const bool bAutoUpdated = FAutoUpdateRootGraph(MoveTemp(OwningAssetName)).Transform(GetDocumentHandle());
-#if WITH_EDITORONLY_DATA
+#if WITH_EDITOR
 		if (bAutoUpdated || InRegistrationOptions.bForceViewSynchronization)
 		{
  			SetSynchronizationRequired();
 		}
-#endif // WITH_EDITORONLY_DATA
+#endif // WITH_EDITOR
 	}
 	else
 	{
-#if WITH_EDITORONLY_DATA
+#if WITH_EDITOR
 		if (InRegistrationOptions.bForceViewSynchronization)
 		{
 			SetSynchronizationRequired();
 		}
-#endif // WITH_EDITORONLY_DATA
+#endif // WITH_EDITOR
 	}
 
-#if WITH_EDITORONLY_DATA
+#if WITH_EDITOR
 	// Must be completed after auto-update to ensure all non-transient referenced dependency data is up-to-date (ex.
 	// class version), which is required for most accurately caching current registry metadata.
 	CacheRegistryMetadata();
-#endif // WITH_EDITORONLY_DATA
+#endif // WITH_EDITOR
 
 	// Registers node by copying document. Updates to document require re-registration.
 	class FNodeRegistryEntry : public INodeRegistryEntry
@@ -398,7 +398,7 @@ bool FMetasoundAssetBase::VersionAsset()
 	return bDidEdit;
 }
 
-#if WITH_EDITORONLY_DATA
+#if WITH_EDITOR
 void FMetasoundAssetBase::CacheRegistryMetadata()
 {
 	using namespace Metasound::Frontend;
@@ -513,7 +513,7 @@ void FMetasoundAssetBase::ResetSynchronizationState()
 	bSynchronizationUpdateDetails = false;
 	bSynchronizationRequired = false;
 }
-#endif // WITH_EDITORONLY_DATA
+#endif // WITH_EDITOR
 
 TSharedPtr<Metasound::IGraph, ESPMode::ThreadSafe> FMetasoundAssetBase::BuildMetasoundDocument() const
 {
@@ -676,7 +676,7 @@ Metasound::Frontend::FNodeHandle FMetasoundAssetBase::AddInputPinForSendAddress(
 	return InGraph->AddInputVertex(Description);
 }
 
-#if WITH_EDITORONLY_DATA
+#if WITH_EDITOR
 FText FMetasoundAssetBase::GetDisplayName(FString&& InTypeName) const
 {
 	using namespace Metasound::Frontend;
@@ -691,7 +691,7 @@ FText FMetasoundAssetBase::GetDisplayName(FString&& InTypeName) const
 
 	return FText::Format(LOCTEXT("PresetDisplayNameFormat", "{0} (Preset)"), FText::FromString(MoveTemp(InTypeName)));
 }
-#endif // WITH_EDITORONLY_DATA
+#endif // WITH_EDITOR
 
 bool FMetasoundAssetBase::MarkMetasoundDocumentDirty() const
 {
diff --git a/Engine/Plugins/Runtime/Metasound/Source/MetasoundFrontend/Public/MetasoundAssetBase.h b/Engine/Plugins/Runtime/Metasound/Source/MetasoundFrontend/Public/MetasoundAssetBase.h
--- a/Engine/Plugins/Runtime/Metasound/Source/MetasoundFrontend/Public/MetasoundAssetBase.h
+++ b/Engine/Plugins/Runtime/Metasound/Source/MetasoundFrontend/Public/MetasoundAssetBase.h
@@ -126,7 +126,7 @@ public:
 
 	bool VersionAsset();
 
-#if WITH_EDITORONLY_DATA
+#if WITH_EDITOR
 	/*
 	 * Caches transient metadata (class & vertex) found in the registry
 	 * that is not necessary for serialization or core graph generation.
@@ -143,7 +143,7 @@ public:
 	void ResetSynchronizationState();
 	void SetUpdateDetailsOnSynchronization();
 	void SetSynchronizationRequired();
-#endif // WITH_EDITORONLY_DATA
+#endif // WITH_EDITOR
 
 	// Calls the outermost package and marks it dirty.
 	bool MarkMetasoundDocumentDirty() const;
