Compile out FText from MetaSound Nodes in non editor builds

diff --git a/Engine/Plugins/Runtime/Metasound/Source/MetasoundEngine/Private/MetasoundWaveInfoNode.cpp b/Engine/Plugins/Runtime/Metasound/Source/MetasoundEngine/Private/MetasoundWaveInfoNode.cpp
--- a/Engine/Plugins/Runtime/Metasound/Source/MetasoundEngine/Private/MetasoundWaveInfoNode.cpp
+++ b/Engine/Plugins/Runtime/Metasound/Source/MetasoundEngine/Private/MetasoundWaveInfoNode.cpp
@@ -84,8 +84,8 @@ namespace Metasound
 			Info.ClassName = { Metasound::EngineNodes::Namespace, TEXT("Get Wave Duration"), TEXT(" ") };
 			Info.MajorVersion = 1;
 			Info.MinorVersion = 0;
-			Info.DisplayName = LOCTEXT("MetasoundGetWaveDuration_ClassNodeDisplayName", "Get Wave Duration");
-			Info.Description = LOCTEXT("GetWaveDuration_NodeDescription", "Returns the duration of the input Wave asset (in seconds)"),
+			Info.DisplayName = METASOUND_LOCTEXT("MetasoundGetWaveDuration_ClassNodeDisplayName", "Get Wave Duration");
+			Info.Description = METASOUND_LOCTEXT("GetWaveDuration_NodeDescription", "Returns the duration of the input Wave asset (in seconds)"),
 			Info.Author = PluginAuthor;
 			Info.PromptIfMissing = PluginNodeMissingPrompt;
 			Info.DefaultInterface = DeclareVertexInterface();
diff --git a/Engine/Plugins/Runtime/Metasound/Source/MetasoundEngine/Private/MetasoundWavePlayerNode.cpp b/Engine/Plugins/Runtime/Metasound/Source/MetasoundEngine/Private/MetasoundWavePlayerNode.cpp
--- a/Engine/Plugins/Runtime/Metasound/Source/MetasoundEngine/Private/MetasoundWavePlayerNode.cpp
+++ b/Engine/Plugins/Runtime/Metasound/Source/MetasoundEngine/Private/MetasoundWavePlayerNode.cpp
@@ -45,26 +45,26 @@ namespace Metasound
 		static const TCHAR* OutputAudioLeftName = TEXT("Out Left");
 		static const TCHAR* OutputAudioRightName = TEXT("Out Right");
 
-		static FText InputTriggerPlayTT = LOCTEXT("PlayTT", "Play the wave player.");
-		static FText InputTriggerStopTT = LOCTEXT("StopTT", "Stop the wave player.");
-		static FText InputWaveAssetTT = LOCTEXT("WaveTT", "The wave asset to be real-time decoded.");
-		static FText InputStartTimeTT = LOCTEXT("StartTimeTT", "Time into the wave asset to start (seek) the wave asset. For real-time decoding, the wave asset must be set to seekable!)");
-		static FText InputPitchShiftTT = LOCTEXT("PitchShiftTT", "The pitch shift to use for the wave asset in semitones.");
-		static FText InputLoopTT = LOCTEXT("LoopTT", "Whether or not to loop between the start and specified end times.");
-		static FText InputLoopStartTT = LOCTEXT("LoopStartTT", "When to start the loop.");
-		static FText InputLoopDurationTT = LOCTEXT("LoopDurationTT", "The duration of the loop when wave player is enabled for looping. A negative value will loop the whole wave asset.");
-
-		static FText OutputTriggerOnPlayTT = LOCTEXT("OnPlayTT", "Triggers when Play is triggered.");
-		static FText OutputTriggerOnDoneTT = LOCTEXT("OnDoneTT", "Triggers when the wave played has finished playing.");
-		static FText OutputTriggerOnNearlyDoneTT = LOCTEXT("OnNearlyDoneTT", "Triggers when the wave played has almost finished playing (the block before it finishes). Allows time for logic to trigger different variations to play seamlessly.");
-		static FText OutputTriggerOnLoopedTT = LOCTEXT("OnLoopedTT", "Triggers when the wave player has looped.");
-		static FText OutputTriggerOnCuePointTT = LOCTEXT("OnCuePointTT", "Triggers when a wave cue point was hit during playback.");
-		static FText OutputCuePointIDTT = LOCTEXT("CuePointIDTT", "The cue point ID that was triggered.");
-		static FText OutputCuePointLabelTT = LOCTEXT("CuePointLabelTT", "The cue point label that was triggered (if there was a label parsed in the imported .wav file).");
-		static FText OutputLoopPercentTT = LOCTEXT("LoopPercentTT", "Returns the current loop percent if looping is enabled.");
-		static FText OutputPlaybackLocationTT = LOCTEXT("PlaybackLocationTT", "Returns the absolute position of the wave playback as a precentage of wave duration.");
-		static FText OutputAudioLeftNameTT = LOCTEXT("AudioLeftTT", "The left channel audio output. Mono wave assets will be upmixed to dual stereo.");
-		static FText OutputAudioRightNameTT = LOCTEXT("AudioRightTT", "The right channel audio output. Mono wave assets will be upmixed to dual stereo.");
+		static FText InputTriggerPlayTT = METASOUND_LOCTEXT("PlayTT", "Play the wave player.");
+		static FText InputTriggerStopTT = METASOUND_LOCTEXT("StopTT", "Stop the wave player.");
+		static FText InputWaveAssetTT = METASOUND_LOCTEXT("WaveTT", "The wave asset to be real-time decoded.");
+		static FText InputStartTimeTT = METASOUND_LOCTEXT("StartTimeTT", "Time into the wave asset to start (seek) the wave asset. For real-time decoding, the wave asset must be set to seekable!)");
+		static FText InputPitchShiftTT = METASOUND_LOCTEXT("PitchShiftTT", "The pitch shift to use for the wave asset in semitones.");
+		static FText InputLoopTT = METASOUND_LOCTEXT("LoopTT", "Whether or not to loop between the start and specified end times.");
+		static FText InputLoopStartTT = METASOUND_LOCTEXT("LoopStartTT", "When to start the loop.");
+		static FText InputLoopDurationTT = METASOUND_LOCTEXT("LoopDurationTT", "The duration of the loop when wave player is enabled for looping. A negative value will loop the whole wave asset.");
+
+		static FText OutputTriggerOnPlayTT = METASOUND_LOCTEXT("OnPlayTT", "Triggers when Play is triggered.");
+		static FText OutputTriggerOnDoneTT = METASOUND_LOCTEXT("OnDoneTT", "Triggers when the wave played has finished playing.");
+		static FText OutputTriggerOnNearlyDoneTT = METASOUND_LOCTEXT("OnNearlyDoneTT", "Triggers when the wave played has almost finished playing (the block before it finishes). Allows time for logic to trigger different variations to play seamlessly.");
+		static FText OutputTriggerOnLoopedTT = METASOUND_LOCTEXT("OnLoopedTT", "Triggers when the wave player has looped.");
+		static FText OutputTriggerOnCuePointTT = METASOUND_LOCTEXT("OnCuePointTT", "Triggers when a wave cue point was hit during playback.");
+		static FText OutputCuePointIDTT = METASOUND_LOCTEXT("CuePointIDTT", "The cue point ID that was triggered.");
+		static FText OutputCuePointLabelTT = METASOUND_LOCTEXT("CuePointLabelTT", "The cue point label that was triggered (if there was a label parsed in the imported .wav file).");
+		static FText OutputLoopPercentTT = METASOUND_LOCTEXT("LoopPercentTT", "Returns the current loop percent if looping is enabled.");
+		static FText OutputPlaybackLocationTT = METASOUND_LOCTEXT("PlaybackLocationTT", "Returns the absolute position of the wave playback as a precentage of wave duration.");
+		static FText OutputAudioLeftNameTT = METASOUND_LOCTEXT("AudioLeftTT", "The left channel audio output. Mono wave assets will be upmixed to dual stereo.");
+		static FText OutputAudioRightNameTT = METASOUND_LOCTEXT("AudioRightTT", "The right channel audio output. Mono wave assets will be upmixed to dual stereo.");
 	}
 
 	class FWavePlayerNode : public FNode
@@ -1068,8 +1068,8 @@ namespace Metasound
 			Info.ClassName = { Metasound::EngineNodes::Namespace, TEXT("Wave Player"), Metasound::EngineNodes::StereoVariant };
 			Info.MajorVersion = 1;
 			Info.MinorVersion = 0;
-			Info.DisplayName = LOCTEXT("Metasound_WavePlayerNodeDisplayName", "Wave Player");
-			Info.Description = LOCTEXT("Metasound_WavePlayerNodeDescription", "Plays a wave asset.");
+			Info.DisplayName = METASOUND_LOCTEXT("Metasound_WavePlayerNodeDisplayName", "Wave Player");
+			Info.Description = METASOUND_LOCTEXT("Metasound_WavePlayerNodeDescription", "Plays a wave asset.");
 			Info.Author = PluginAuthor;
 			Info.PromptIfMissing = PluginNodeMissingPrompt;
 			Info.DefaultInterface = DeclareVertexInterface();
diff --git a/Engine/Plugins/Runtime/Metasound/Source/MetasoundFrontend/Private/MetasoundArrayNodes.cpp b/Engine/Plugins/Runtime/Metasound/Source/MetasoundFrontend/Private/MetasoundArrayNodes.cpp
--- a/Engine/Plugins/Runtime/Metasound/Source/MetasoundFrontend/Private/MetasoundArrayNodes.cpp
+++ b/Engine/Plugins/Runtime/Metasound/Source/MetasoundFrontend/Private/MetasoundArrayNodes.cpp
@@ -20,8 +20,8 @@ namespace Metasound
 				PluginAuthor,
 				PluginNodeMissingPrompt,
 				InDefaultInterface,
-				{ LOCTEXT("ArrayCategory", "Array") },
-				{ LOCTEXT("MetasoundArrayKeyword", "Array") },
+				{ METASOUND_LOCTEXT("ArrayCategory", "Array") },
+				{ METASOUND_LOCTEXT("MetasoundArrayKeyword", "Array") },
 				FNodeDisplayStyle{}
 			};
 
diff --git a/Engine/Plugins/Runtime/Metasound/Source/MetasoundFrontend/Public/MetasoundArrayNodes.h b/Engine/Plugins/Runtime/Metasound/Source/MetasoundFrontend/Public/MetasoundArrayNodes.h
--- a/Engine/Plugins/Runtime/Metasound/Source/MetasoundFrontend/Public/MetasoundArrayNodes.h
+++ b/Engine/Plugins/Runtime/Metasound/Source/MetasoundFrontend/Public/MetasoundArrayNodes.h
@@ -9,6 +9,7 @@
 #include "MetasoundFacade.h"
 #include "MetasoundLog.h"
 #include "MetasoundNodeInterface.h"
+#include "MetasoundNodeRegistrationMacro.h"
 #include "MetasoundOperatorInterface.h"
 #include "MetasoundPrimitives.h"
 #include "MetasoundTrigger.h"
@@ -75,10 +76,10 @@ namespace Metasound
 
 			static const FVertexInterface DefaultInterface(
 				FInputVertexInterface(
-					TInputDataVertexModel<ArrayType>(GetInputArrayName(), LOCTEXT("ArrayOpArrayNumInput", "Array to inspect."))
+					TInputDataVertexModel<ArrayType>(GetInputArrayName(), METASOUND_LOCTEXT("ArrayOpArrayNumInput", "Array to inspect."))
 				),
 				FOutputVertexInterface(
-					TOutputDataVertexModel<int32>(GetOutputNumName(), LOCTEXT("ArrayOpArrayNumOutput", "Number of elements in the array."))
+					TOutputDataVertexModel<int32>(GetOutputNumName(), METASOUND_LOCTEXT("ArrayOpArrayNumOutput", "Number of elements in the array."))
 				)
 			);
 
@@ -91,8 +92,8 @@ namespace Metasound
 			{
 				const FName DataTypeName = GetMetasoundDataTypeName<ArrayType>();
 				const FName OperatorName = TEXT("Num");
-				const FText NodeDisplayName = FText::Format(LOCTEXT("ArrayOpArrayNumDisplayNamePattern", "Num ({0})"), GetMetasoundDataTypeDisplayText<ArrayType>());
-				const FText NodeDescription = LOCTEXT("ArrayOpArrayNumDescription", "Number of elements in the array");
+				const FText NodeDisplayName = METASOUND_LOCTEXT_FORMAT("ArrayOpArrayNumDisplayNamePattern", "Num ({0})", GetMetasoundDataTypeDisplayText<ArrayType>());
+				const FText NodeDescription = METASOUND_LOCTEXT("ArrayOpArrayNumDescription", "Number of elements in the array");
 				const FVertexInterface NodeInterface = GetDefaultInterface();
 			
 				return MetasoundArrayNodesPrivate::CreateArrayNodeClassMetadata(DataTypeName, OperatorName, NodeDisplayName, NodeDescription, NodeInterface);
@@ -187,12 +188,12 @@ namespace Metasound
 			using namespace ArrayNodeVertexNames;
 			static const FVertexInterface DefaultInterface(
 				FInputVertexInterface(
-					TInputDataVertexModel<FTrigger>(GetInputTriggerName(), LOCTEXT("ArrayOpArrayGetTrigger", "Trigger to get value.")),
-					TInputDataVertexModel<ArrayType>(GetInputArrayName(), LOCTEXT("ArrayOpArrayGetInput", "Input Array.")),
-					TInputDataVertexModel<int32>(GetInputIndexName(), LOCTEXT("ArrayOpArrayGetIndex", "Index in Array."))
+					TInputDataVertexModel<FTrigger>(GetInputTriggerName(), METASOUND_LOCTEXT("ArrayOpArrayGetTrigger", "Trigger to get value.")),
+					TInputDataVertexModel<ArrayType>(GetInputArrayName(), METASOUND_LOCTEXT("ArrayOpArrayGetInput", "Input Array.")),
+					TInputDataVertexModel<int32>(GetInputIndexName(), METASOUND_LOCTEXT("ArrayOpArrayGetIndex", "Index in Array."))
 				),
 				FOutputVertexInterface(
-					TOutputDataVertexModel<ElementType>(GetOutputValueName(), LOCTEXT("ArrayOpArrayGetOutput", "Value of element at array index."))
+					TOutputDataVertexModel<ElementType>(GetOutputValueName(), METASOUND_LOCTEXT("ArrayOpArrayGetOutput", "Value of element at array index."))
 				)
 			);
 
@@ -205,8 +206,8 @@ namespace Metasound
 			{
 				const FName DataTypeName = GetMetasoundDataTypeName<ArrayType>();
 				const FName OperatorName = TEXT("Get"); 
-				const FText NodeDisplayName = FText::Format(LOCTEXT("ArrayOpArrayGetDisplayNamePattern", "Get ({0})"), GetMetasoundDataTypeDisplayText<ArrayType>());
-				const FText NodeDescription = LOCTEXT("ArrayOpArrayGetDescription", "Get element at index in array.");
+				const FText NodeDisplayName = METASOUND_LOCTEXT_FORMAT("ArrayOpArrayGetDisplayNamePattern", "Get ({0})", GetMetasoundDataTypeDisplayText<ArrayType>());
+				const FText NodeDescription = METASOUND_LOCTEXT("ArrayOpArrayGetDescription", "Get element at index in array.");
 				const FVertexInterface NodeInterface = GetDefaultInterface();
 			
 				return MetasoundArrayNodesPrivate::CreateArrayNodeClassMetadata(DataTypeName, OperatorName, NodeDisplayName, NodeDescription, NodeInterface);
@@ -324,13 +325,13 @@ namespace Metasound
 			using namespace ArrayNodeVertexNames;
 			static const FVertexInterface DefaultInterface(
 				FInputVertexInterface(
-					TInputDataVertexModel<FTrigger>(GetInputTriggerName(), LOCTEXT("ArrayOpArraySetTrigger", "Trigger to set value.")),
-					TInputDataVertexModel<ArrayType>(GetInputArrayName(), LOCTEXT("ArrayOpArraySetInput", "Input Array.")),
-					TInputDataVertexModel<int32>(GetInputIndexName(), LOCTEXT("ArrayOpArraySetIndex", "Index in Array.")),
-					TInputDataVertexModel<ElementType>(GetInputValueName(), LOCTEXT("ArrayOpArraySetElement", "Value to set"))
+					TInputDataVertexModel<FTrigger>(GetInputTriggerName(), METASOUND_LOCTEXT("ArrayOpArraySetTrigger", "Trigger to set value.")),
+					TInputDataVertexModel<ArrayType>(GetInputArrayName(), METASOUND_LOCTEXT("ArrayOpArraySetInput", "Input Array.")),
+					TInputDataVertexModel<int32>(GetInputIndexName(), METASOUND_LOCTEXT("ArrayOpArraySetIndex", "Index in Array.")),
+					TInputDataVertexModel<ElementType>(GetInputValueName(), METASOUND_LOCTEXT("ArrayOpArraySetElement", "Value to set"))
 				),
 				FOutputVertexInterface(
-					TOutputDataVertexModel<ArrayType>(GetOutputArrayName(), LOCTEXT("ArrayOpArraySetOutput", "Array after setting."))
+					TOutputDataVertexModel<ArrayType>(GetOutputArrayName(), METASOUND_LOCTEXT("ArrayOpArraySetOutput", "Array after setting."))
 				)
 			);
 
@@ -343,8 +344,8 @@ namespace Metasound
 			{
 				const FName DataTypeName = GetMetasoundDataTypeName<ArrayType>();
 				const FName OperatorName = TEXT("Set"); 
-				const FText NodeDisplayName = FText::Format(LOCTEXT("ArrayOpArraySetDisplayNamePattern", "Set ({0})"), GetMetasoundDataTypeDisplayText<ArrayType>());
-				const FText NodeDescription = LOCTEXT("ArrayOpArraySetDescription", "Set element at index in array.");
+				const FText NodeDisplayName = METASOUND_LOCTEXT_FORMAT("ArrayOpArraySetDisplayNamePattern", "Set ({0})", GetMetasoundDataTypeDisplayText<ArrayType>());
+				const FText NodeDescription = METASOUND_LOCTEXT("ArrayOpArraySetDescription", "Set element at index in array.");
 				const FVertexInterface NodeInterface = GetDefaultInterface();
 			
 				return MetasoundArrayNodesPrivate::CreateArrayNodeClassMetadata(DataTypeName, OperatorName, NodeDisplayName, NodeDescription, NodeInterface);
@@ -465,12 +466,12 @@ namespace Metasound
 
 			static const FVertexInterface DefaultInterface(
 				FInputVertexInterface(
-					TInputDataVertexModel<FTrigger>(GetInputTriggerName(), LOCTEXT("ArrayOpArrayConcatTrigger", "Trigger to set value.")),
-					TInputDataVertexModel<ArrayType>(GetInputLeftArrayName(), LOCTEXT("ArrayOpArrayConcatInputLeft", "Input Left Array.")),
-					TInputDataVertexModel<ArrayType>(GetInputRightArrayName(), LOCTEXT("ArrayOpArrayConcatInputRight", "Input Left Array."))
+					TInputDataVertexModel<FTrigger>(GetInputTriggerName(), METASOUND_LOCTEXT("ArrayOpArrayConcatTrigger", "Trigger to set value.")),
+					TInputDataVertexModel<ArrayType>(GetInputLeftArrayName(), METASOUND_LOCTEXT("ArrayOpArrayConcatInputLeft", "Input Left Array.")),
+					TInputDataVertexModel<ArrayType>(GetInputRightArrayName(), METASOUND_LOCTEXT("ArrayOpArrayConcatInputRight", "Input Left Array."))
 				),
 				FOutputVertexInterface(
-					TOutputDataVertexModel<ArrayType>(GetOutputArrayName(), LOCTEXT("ArrayOpArrayConcatOutput", "Array after concatenation."))
+					TOutputDataVertexModel<ArrayType>(GetOutputArrayName(), METASOUND_LOCTEXT("ArrayOpArrayConcatOutput", "Array after concatenation."))
 				)
 			);
 
@@ -483,8 +484,8 @@ namespace Metasound
 			{
 				const FName DataTypeName = GetMetasoundDataTypeName<ArrayType>();
 				const FName OperatorName = TEXT("Concat"); 
-				const FText NodeDisplayName = FText::Format(LOCTEXT("ArrayOpArrayConcatDisplayNamePattern", "Concatenate ({0})"), GetMetasoundDataTypeDisplayText<ArrayType>());
-				const FText NodeDescription = LOCTEXT("ArrayOpArrayConcatDescription", "Concatenates two arrays on trigger.");
+				const FText NodeDisplayName = METASOUND_LOCTEXT_FORMAT("ArrayOpArrayConcatDisplayNamePattern", "Concatenate ({0})", GetMetasoundDataTypeDisplayText<ArrayType>());
+				const FText NodeDescription = METASOUND_LOCTEXT("ArrayOpArrayConcatDescription", "Concatenates two arrays on trigger.");
 				const FVertexInterface NodeInterface = GetDefaultInterface();
 			
 				return MetasoundArrayNodesPrivate::CreateArrayNodeClassMetadata(DataTypeName, OperatorName, NodeDisplayName, NodeDescription, NodeInterface);
@@ -589,14 +590,14 @@ namespace Metasound
 
 			static const FVertexInterface DefaultInterface(
 				FInputVertexInterface(
-					TInputDataVertexModel<FTrigger>(GetInputTriggerName(), LOCTEXT("ArrayOpArraySubsetTrigger", "Trigger to set value.")),
-					TInputDataVertexModel<ArrayType>(GetInputArrayName(), LOCTEXT("ArrayOpArraySubsetInputLeft", "Input Array.")),
-					TInputDataVertexModel<int32>(GetInputStartIndexName(), LOCTEXT("ArrayOpArraySubsetStartIndex", "First index to include.")),
-					TInputDataVertexModel<int32>(GetInputEndIndexName(), LOCTEXT("ArrayOpArraySubsetEndIndex", "Last index to include."))
+					TInputDataVertexModel<FTrigger>(GetInputTriggerName(), METASOUND_LOCTEXT("ArrayOpArraySubsetTrigger", "Trigger to set value.")),
+					TInputDataVertexModel<ArrayType>(GetInputArrayName(), METASOUND_LOCTEXT("ArrayOpArraySubsetInputLeft", "Input Array.")),
+					TInputDataVertexModel<int32>(GetInputStartIndexName(), METASOUND_LOCTEXT("ArrayOpArraySubsetStartIndex", "First index to include.")),
+					TInputDataVertexModel<int32>(GetInputEndIndexName(), METASOUND_LOCTEXT("ArrayOpArraySubsetEndIndex", "Last index to include."))
 
 				),
 				FOutputVertexInterface(
-					TOutputDataVertexModel<ArrayType>(GetOutputArrayName(), LOCTEXT("ArrayOpArraySubsetOutput", "Subset of input array."))
+					TOutputDataVertexModel<ArrayType>(GetOutputArrayName(), METASOUND_LOCTEXT("ArrayOpArraySubsetOutput", "Subset of input array."))
 				)
 			);
 
@@ -609,8 +610,8 @@ namespace Metasound
 			{
 				const FName DataTypeName = GetMetasoundDataTypeName<ArrayType>();
 				const FName OperatorName = TEXT("Subset"); 
-				const FText NodeDisplayName = FText::Format(LOCTEXT("ArrayOpArraySubsetDisplayNamePattern", "Subset ({0})"), GetMetasoundDataTypeDisplayText<ArrayType>());
-				const FText NodeDescription = LOCTEXT("ArrayOpArraySubsetDescription", "Subset array on trigger.");
+				const FText NodeDisplayName = METASOUND_LOCTEXT_FORMAT("ArrayOpArraySubsetDisplayNamePattern", "Subset ({0})", GetMetasoundDataTypeDisplayText<ArrayType>());
+				const FText NodeDescription = METASOUND_LOCTEXT("ArrayOpArraySubsetDescription", "Subset array on trigger.");
 				const FVertexInterface NodeInterface = GetDefaultInterface();
 			
 				return MetasoundArrayNodesPrivate::CreateArrayNodeClassMetadata(DataTypeName, OperatorName, NodeDisplayName, NodeDescription, NodeInterface);
diff --git a/Engine/Plugins/Runtime/Metasound/Source/MetasoundFrontend/Public/MetasoundArrayRandomNode.h b/Engine/Plugins/Runtime/Metasound/Source/MetasoundFrontend/Public/MetasoundArrayRandomNode.h
--- a/Engine/Plugins/Runtime/Metasound/Source/MetasoundFrontend/Public/MetasoundArrayRandomNode.h
+++ b/Engine/Plugins/Runtime/Metasound/Source/MetasoundFrontend/Public/MetasoundArrayRandomNode.h
@@ -27,16 +27,16 @@ namespace Metasound
 		static const FName OutputTriggerOnResetName = "On Reset";
 		static const FName OutputValueName = "Value";
 
-		static const FText InputTriggerNextTooltip = LOCTEXT("RandomArrayGetNode_TriggerNextTooltip", "Trigger to get the next value in the randomized array.");
-		static const FText InputTriggerResetTooltip = LOCTEXT("RandomArrayGetNode_TriggerResetTooltip", "Trigger to reset the seed for the randomized array.");
-		static const FText InputRandomArrayTooltip = LOCTEXT("RandomArrayGetNode_RandomArrayTooltip", "Input array to randomized.");
-		static const FText InputWeightsTooltip = LOCTEXT("RandomArrayGetNode_WeightsTooltip", "Input array of weights to use for random selection. Will repeat if this array is shorter than the input array to select from.");
-		static const FText InputSeedTooltip = LOCTEXT("RandomArrayGetNode_SeedTooltip", "Seed to use for the random shuffle.");
-		static const FText InputNoRepeatOrderTooltip = LOCTEXT("RandomArrayGetNode_NoRepeatOrderTooltip", "The number of elements to track to avoid repeating in a row.");
-		static const FText InputEnableSharedStateTooltip = LOCTEXT("RandomArrayGetNode_EnableSharedStateTooltip", "Set to enabled to share state across instances of this MetaSound.");
-		static const FText OutputTriggerOnNextTooltip = LOCTEXT("RandomArrayGetNode_TriggerOnNextTooltip", "Triggers when the \"Next\" input is triggered.");
-		static const FText OutputTriggerOnResetTooltip = LOCTEXT("RandomArrayGetNode_TriggerOnResetTooltip", "Triggers when the \"Shuffle\" input is triggered or if the array is auto-shuffled.");
-		static const FText OutputValueTooltip = LOCTEXT("RandomArrayGetNode_ValueTooltip", "Value of the current shuffled element.");
+		static const FText InputTriggerNextTooltip = METASOUND_LOCTEXT("RandomArrayGetNode_TriggerNextTooltip", "Trigger to get the next value in the randomized array.");
+		static const FText InputTriggerResetTooltip = METASOUND_LOCTEXT("RandomArrayGetNode_TriggerResetTooltip", "Trigger to reset the seed for the randomized array.");
+		static const FText InputRandomArrayTooltip = METASOUND_LOCTEXT("RandomArrayGetNode_RandomArrayTooltip", "Input array to randomized.");
+		static const FText InputWeightsTooltip = METASOUND_LOCTEXT("RandomArrayGetNode_WeightsTooltip", "Input array of weights to use for random selection. Will repeat if this array is shorter than the input array to select from.");
+		static const FText InputSeedTooltip = METASOUND_LOCTEXT("RandomArrayGetNode_SeedTooltip", "Seed to use for the random shuffle.");
+		static const FText InputNoRepeatOrderTooltip = METASOUND_LOCTEXT("RandomArrayGetNode_NoRepeatOrderTooltip", "The number of elements to track to avoid repeating in a row.");
+		static const FText InputEnableSharedStateTooltip = METASOUND_LOCTEXT("RandomArrayGetNode_EnableSharedStateTooltip", "Set to enabled to share state across instances of this MetaSound.");
+		static const FText OutputTriggerOnNextTooltip = METASOUND_LOCTEXT("RandomArrayGetNode_TriggerOnNextTooltip", "Triggers when the \"Next\" input is triggered.");
+		static const FText OutputTriggerOnResetTooltip = METASOUND_LOCTEXT("RandomArrayGetNode_TriggerOnResetTooltip", "Triggers when the \"Shuffle\" input is triggered or if the array is auto-shuffled.");
+		static const FText OutputValueTooltip = METASOUND_LOCTEXT("RandomArrayGetNode_ValueTooltip", "Value of the current shuffled element.");
 	}
 
 	class METASOUNDFRONTEND_API FArrayRandomGet
@@ -140,8 +140,8 @@ namespace Metasound
 			{
 				FName DataTypeName = GetMetasoundDataTypeName<ArrayType>();
 				FName OperatorName = "Random Get";
-				FText NodeDisplayName = FText::Format(LOCTEXT("RandomArrayGetNode_OpDisplayNamePattern", "Random Get ({0})"), GetMetasoundDataTypeDisplayText<ArrayType>());
-				FText NodeDescription = LOCTEXT("RandomArrayGetNode_Description", "Randomly retrieve data from input array using the supplied weights.");
+				FText NodeDisplayName = METASOUND_LOCTEXT_FORMAT("RandomArrayGetNode_OpDisplayNamePattern", "Random Get ({0})", GetMetasoundDataTypeDisplayText<ArrayType>());
+				FText NodeDescription = METASOUND_LOCTEXT("RandomArrayGetNode_Description", "Randomly retrieve data from input array using the supplied weights.");
 				FVertexInterface NodeInterface = GetDefaultInterface();
 
 				return MetasoundArrayNodesPrivate::CreateArrayNodeClassMetadata(DataTypeName, OperatorName, NodeDisplayName, NodeDescription, NodeInterface);
diff --git a/Engine/Plugins/Runtime/Metasound/Source/MetasoundFrontend/Public/MetasoundArrayShuffleNode.h b/Engine/Plugins/Runtime/Metasound/Source/MetasoundFrontend/Public/MetasoundArrayShuffleNode.h
--- a/Engine/Plugins/Runtime/Metasound/Source/MetasoundFrontend/Public/MetasoundArrayShuffleNode.h
+++ b/Engine/Plugins/Runtime/Metasound/Source/MetasoundFrontend/Public/MetasoundArrayShuffleNode.h
@@ -147,19 +147,19 @@ namespace Metasound
 
 			static const FVertexInterface DefaultInterface(
 				FInputVertexInterface(
-					TInputDataVertexModel<FTrigger>(GetInputTriggerNextName(), LOCTEXT("ShuffleOpInputTriggerNextTT", "Trigger to get the next value in the shuffled array.")),
-					TInputDataVertexModel<FTrigger>(GetInputTriggerShuffleName(), LOCTEXT("ShuffleOpInputTriggerShuffleTT", "Trigger to shuffle the array manually.")),
-					TInputDataVertexModel<FTrigger>(GetInputTriggerResetName(), LOCTEXT("ShuffleOpInputTriggerResetTT", "Trigger to reset the random seed stream of the shuffle node.")),
-					TInputDataVertexModel<ArrayType>(GetInputShuffleArrayName(), LOCTEXT("ShuffleOpInputShuffleArrayTT", "Input Array.")),
-					TInputDataVertexModel<int32>(GetInputSeedName(), LOCTEXT("ShuffleOpInputSeedTT", "Seed to use for the the random shuffle."), -1),
-					TInputDataVertexModel<bool>(GetInputAutoShuffleName(), LOCTEXT("ShuffleOpInputAutoShuffleTT", "Set to true to automatically shuffle when the array has been read."), true),
-					TInputDataVertexModel<bool>(GetInputEnableSharedStateName(), LOCTEXT("ShuffleOpInputEnableSharedStatTT", "Set to enabled shared state across instances of this metasound."), false)
+					TInputDataVertexModel<FTrigger>(GetInputTriggerNextName(), METASOUND_LOCTEXT("ShuffleOpInputTriggerNextTT", "Trigger to get the next value in the shuffled array.")),
+					TInputDataVertexModel<FTrigger>(GetInputTriggerShuffleName(), METASOUND_LOCTEXT("ShuffleOpInputTriggerShuffleTT", "Trigger to shuffle the array manually.")),
+					TInputDataVertexModel<FTrigger>(GetInputTriggerResetName(), METASOUND_LOCTEXT("ShuffleOpInputTriggerResetTT", "Trigger to reset the random seed stream of the shuffle node.")),
+					TInputDataVertexModel<ArrayType>(GetInputShuffleArrayName(), METASOUND_LOCTEXT("ShuffleOpInputShuffleArrayTT", "Input Array.")),
+					TInputDataVertexModel<int32>(GetInputSeedName(), METASOUND_LOCTEXT("ShuffleOpInputSeedTT", "Seed to use for the the random shuffle."), -1),
+					TInputDataVertexModel<bool>(GetInputAutoShuffleName(), METASOUND_LOCTEXT("ShuffleOpInputAutoShuffleTT", "Set to true to automatically shuffle when the array has been read."), true),
+					TInputDataVertexModel<bool>(GetInputEnableSharedStateName(), METASOUND_LOCTEXT("ShuffleOpInputEnableSharedStatTT", "Set to enabled shared state across instances of this metasound."), false)
 					),
 				FOutputVertexInterface(
-					TOutputDataVertexModel<FTrigger>(GetOutputTriggerOnNextName(), LOCTEXT("ShuffleOpOutputTriggerOnNextNameTT", "Triggers when the \"Next\" input is triggered.")),
-					TOutputDataVertexModel<FTrigger>(GetOutputTriggerOnShuffleName(), LOCTEXT("ShuffleOpOutputTriggerOnShuffleNameTT", "Triggers when the \"Shuffle\" input is triggered or if the array is auto-shuffled.")),
-					TOutputDataVertexModel<FTrigger>(GetOutputTriggerOnResetName(), LOCTEXT("ShuffleOpOutputTriggerOnResetNameTT", "Triggers when the \"Reset Seed\" input is triggered.")),
-					TOutputDataVertexModel<ElementType>(GetOutputValueName(), LOCTEXT("ShuffleOpOutputValueTT", "Value of the current shuffled element."))
+					TOutputDataVertexModel<FTrigger>(GetOutputTriggerOnNextName(), METASOUND_LOCTEXT("ShuffleOpOutputTriggerOnNextNameTT", "Triggers when the \"Next\" input is triggered.")),
+					TOutputDataVertexModel<FTrigger>(GetOutputTriggerOnShuffleName(), METASOUND_LOCTEXT("ShuffleOpOutputTriggerOnShuffleNameTT", "Triggers when the \"Shuffle\" input is triggered or if the array is auto-shuffled.")),
+					TOutputDataVertexModel<FTrigger>(GetOutputTriggerOnResetName(), METASOUND_LOCTEXT("ShuffleOpOutputTriggerOnResetNameTT", "Triggers when the \"Reset Seed\" input is triggered.")),
+					TOutputDataVertexModel<ElementType>(GetOutputValueName(), METASOUND_LOCTEXT("ShuffleOpOutputValueTT", "Value of the current shuffled element."))
 				)
 			);
 
@@ -172,8 +172,8 @@ namespace Metasound
 			{
 				const FName DataTypeName = GetMetasoundDataTypeName<ArrayType>();
 				const FName OperatorName = "Shuffle";
-				const FText NodeDisplayName = FText::Format(LOCTEXT("ArrayOpArrayShuffleDisplayNamePattern", "Shuffle ({0})"), GetMetasoundDataTypeDisplayText<ArrayType>());
-				const FText NodeDescription = LOCTEXT("ArrayOpArrayShuffleDescription", "Output next element of a shuffled array on trigger.");
+				const FText NodeDisplayName = METASOUND_LOCTEXT_FORMAT("ArrayOpArrayShuffleDisplayNamePattern", "Shuffle ({0})", GetMetasoundDataTypeDisplayText<ArrayType>());
+				const FText NodeDescription = METASOUND_LOCTEXT("ArrayOpArrayShuffleDescription", "Output next element of a shuffled array on trigger.");
 				const FVertexInterface NodeInterface = GetDefaultInterface();
 
 				return MetasoundArrayNodesPrivate::CreateArrayNodeClassMetadata(DataTypeName, OperatorName, NodeDisplayName, NodeDescription, NodeInterface);
diff --git a/Engine/Plugins/Runtime/Metasound/Source/MetasoundFrontend/Public/MetasoundAutoConverterNode.h b/Engine/Plugins/Runtime/Metasound/Source/MetasoundFrontend/Public/MetasoundAutoConverterNode.h
--- a/Engine/Plugins/Runtime/Metasound/Source/MetasoundFrontend/Public/MetasoundAutoConverterNode.h
+++ b/Engine/Plugins/Runtime/Metasound/Source/MetasoundFrontend/Public/MetasoundAutoConverterNode.h
@@ -6,6 +6,7 @@
 #include "MetasoundBuildError.h"
 #include "MetasoundNode.h"
 #include "MetasoundNodeInterface.h"
+#include "MetasoundNodeRegistrationMacro.h"
 #include "MetasoundOperatorInterface.h"
 #include "MetasoundDataFactory.h"
 #include "MetasoundDataReference.h"
@@ -58,8 +59,8 @@ namespace Metasound
 
 		static FVertexInterface DeclareVertexInterface()
 		{
-			static const FText InputDesc = FText::Format(LOCTEXT("AutoConvDisplayNamePatternFrom", "Input {0} value."), GetMetasoundDataTypeDisplayText<FromDataType>());
-			static const FText OutputDesc = FText::Format(LOCTEXT("AutoConvDisplayNamePatternTo", "Output {0} value."), GetMetasoundDataTypeDisplayText<ToDataType>());
+			static const FText InputDesc = METASOUND_LOCTEXT_FORMAT("AutoConvDisplayNamePatternFrom", "Input {0} value.", GetMetasoundDataTypeDisplayText<FromDataType>());
+			static const FText OutputDesc = METASOUND_LOCTEXT_FORMAT("AutoConvDisplayNamePatternTo", "Output {0} value.", GetMetasoundDataTypeDisplayText<ToDataType>());
 
 			return FVertexInterface(
 				FInputVertexInterface(
@@ -88,8 +89,8 @@ namespace Metasound
 				Info.ClassName = { TEXT("Convert"), GetMetasoundDataTypeName<ToDataType>(), GetMetasoundDataTypeName<FromDataType>() };
 				Info.MajorVersion = 1;
 				Info.MinorVersion = 0;
-				Info.DisplayName = FText::Format(LOCTEXT("Metasound_AutoConverterNodeDisplayNameFormat", "{0} to {1}"), FromTypeText, ToTypeText);
-				Info.Description = FText::Format(LOCTEXT("Metasound_AutoConverterNodeDescriptionNameFormat", "Converts from {0} to {1}."), FromTypeText, ToTypeText);
+				Info.DisplayName = METASOUND_LOCTEXT_FORMAT("Metasound_AutoConverterNodeDisplayNameFormat", "{0} to {1}", FromTypeText, ToTypeText);
+				Info.Description = METASOUND_LOCTEXT_FORMAT("Metasound_AutoConverterNodeDescriptionNameFormat", "Converts from {0} to {1}.", FromTypeText, ToTypeText);
 				Info.Author = PluginAuthor;
 				Info.DisplayStyle = DisplayStyle;
 				Info.PromptIfMissing = PluginNodeMissingPrompt;
@@ -103,7 +104,7 @@ namespace Metasound
 				
 				Info.Keywords =
 				{
-					LOCTEXT("MetasoundConvertKeyword", "Convert"),
+					METASOUND_LOCTEXT("MetasoundConvertKeyword", "Convert"),
 					GetMetasoundDataTypeDisplayText<FromDataType>(),
 					GetMetasoundDataTypeDisplayText<ToDataType>()
 				};
diff --git a/Engine/Plugins/Runtime/Metasound/Source/MetasoundFrontend/Public/MetasoundEnum.h b/Engine/Plugins/Runtime/Metasound/Source/MetasoundFrontend/Public/MetasoundEnum.h
--- a/Engine/Plugins/Runtime/Metasound/Source/MetasoundFrontend/Public/MetasoundEnum.h
+++ b/Engine/Plugins/Runtime/Metasound/Source/MetasoundFrontend/Public/MetasoundEnum.h
@@ -24,9 +24,9 @@
 
 // 3. Define it using the BEGIN/ENTRY/END macros
 // DEFINE_METASOUND_ENUM_BEGIN(EMyOtherTestEnum)
-// 		DEFINE_METASOUND_ENUM_ENTRY(EMyOtherTestEnum::Alpha, LOCTEXT("AlphaDescription", "Alpha"), LOCTEXT("AlphaDescriptionTT", "Alpha tooltip")),
-// 		DEFINE_METASOUND_ENUM_ENTRY(EMyOtherTestEnum::Beta, LOCTEXT("BetaDescription", "Beta"), LOCTEXT("BetaDescriptioTT", "Beta tooltip")),
-// 		DEFINE_METASOUND_ENUM_ENTRY(EMyOtherTestEnum::Gamma, LOCTEXT("GammaDescription", "Gamma", LOCTEXT("GammaDescriptionTT", "Gamma tooltip"))
+// 		DEFINE_METASOUND_ENUM_ENTRY(EMyOtherTestEnum::Alpha, "AlphaDescription", "Alpha", "AlphaDescriptionTT", "Alpha tooltip"),
+// 		DEFINE_METASOUND_ENUM_ENTRY(EMyOtherTestEnum::Beta, "BetaDescription", "Beta", "BetaDescriptioTT", "Beta tooltip"),
+// 		DEFINE_METASOUND_ENUM_ENTRY(EMyOtherTestEnum::Gamma, "GammaDescription", "Gamma", "GammaDescriptionTT", "Gamma tooltip")
 // DEFINE_METASOUND_ENUM_END()
 
 namespace Metasound
diff --git a/Engine/Plugins/Runtime/Metasound/Source/MetasoundFrontend/Public/MetasoundEnumRegistrationMacro.h b/Engine/Plugins/Runtime/Metasound/Source/MetasoundFrontend/Public/MetasoundEnumRegistrationMacro.h
--- a/Engine/Plugins/Runtime/Metasound/Source/MetasoundFrontend/Public/MetasoundEnumRegistrationMacro.h
+++ b/Engine/Plugins/Runtime/Metasound/Source/MetasoundFrontend/Public/MetasoundEnumRegistrationMacro.h
@@ -89,10 +89,16 @@ namespace Metasound
 
 /** DEFINE_METASOUND_ENUM_ENTRY - defines a single Enum Entry
   * @param ENTRY - Fully Qualified Name of Entry of the Enum. (e.g. EMyType::One)
-  * @param DISPLAYNAME - FText Presented to User
-  * @param TOOLTIP - FText displayed as a tooltip
+  * @param DISPLAYNAME_KEY - Display Name loc key 
+  * @param DISPLAYNAME - Display Name text presented to User
+  * @param TOOLTIP_KEY - Tooltip loc key 
+  * @param TOOLTIP - Tooltip text 
   */
-	#define DEFINE_METASOUND_ENUM_ENTRY(ENTRY, DISPLAYNAME, TOOLTIP) { ENTRY, TEXT(#ENTRY), DISPLAYNAME, TOOLTIP }
+#if WITH_EDITOR
+	#define DEFINE_METASOUND_ENUM_ENTRY(ENTRY, DISPLAYNAME_KEY, DISPLAYNAME, TOOLTIP_KEY, TOOLTIP) { ENTRY, TEXT(#ENTRY), LOCTEXT(DISPLAYNAME_KEY, DISPLAYNAME), LOCTEXT(TOOLTIP_KEY, TOOLTIP) }
+#else
+	#define DEFINE_METASOUND_ENUM_ENTRY(ENTRY, DISPLAYNAME_KEY, DISPLAYNAME, TOOLTIP_KEY, TOOLTIP) { ENTRY, TEXT(#ENTRY), FText::GetEmpty(), FText::GetEmpty() }
+#endif // WITH_EDITOR
 
 /** DEFINE_METASOUND_ENUM_END - macro which ends the function body of the GetAllEntries function
    */
diff --git a/Engine/Plugins/Runtime/Metasound/Source/MetasoundFrontend/Public/MetasoundInputNode.h b/Engine/Plugins/Runtime/Metasound/Source/MetasoundFrontend/Public/MetasoundInputNode.h
--- a/Engine/Plugins/Runtime/Metasound/Source/MetasoundFrontend/Public/MetasoundInputNode.h
+++ b/Engine/Plugins/Runtime/Metasound/Source/MetasoundFrontend/Public/MetasoundInputNode.h
@@ -280,7 +280,7 @@ namespace Metasound
 				Info.ClassName = { "Input", GetMetasoundDataTypeName<DataType>(), FName() };
 				Info.MajorVersion = 1;
 				Info.MinorVersion = 0;
-				Info.Description = LOCTEXT("Metasound_InputNodeDescription", "Input into the parent Metasound graph.");
+				Info.Description = METASOUND_LOCTEXT("Metasound_InputNodeDescription", "Input into the parent Metasound graph.");
 				Info.Author = PluginAuthor;
 				Info.PromptIfMissing = PluginNodeMissingPrompt;
 				Info.DefaultInterface = DeclareVertexInterface(InVertexName);
diff --git a/Engine/Plugins/Runtime/Metasound/Source/MetasoundFrontend/Public/MetasoundNodeRegistrationMacro.h b/Engine/Plugins/Runtime/Metasound/Source/MetasoundFrontend/Public/MetasoundNodeRegistrationMacro.h
--- a/Engine/Plugins/Runtime/Metasound/Source/MetasoundFrontend/Public/MetasoundNodeRegistrationMacro.h
+++ b/Engine/Plugins/Runtime/Metasound/Source/MetasoundFrontend/Public/MetasoundNodeRegistrationMacro.h
@@ -124,3 +124,13 @@ namespace Metasound
 	 static_assert(::Metasound::ConstructorTakesNodeInitData<NodeClass>::Value, "In order to be registered as a Metasound Node, " #NodeClass " needs to implement the following public constructor: " #NodeClass "(const Metasound::FNodeInitData& InInitData);"); \
 	 static bool bSuccessfullyRegistered##NodeClass  = FMetasoundFrontendRegistryContainer::Get()->EnqueueInitCommand([](){ ::Metasound::RegisterNodeWithFrontend<NodeClass>(__VA_ARGS__); }); // This static bool is useful for debugging, but also is the only way the compiler will let us call this function outside of an expression.
 
+/*
+Macros to help define various FText node fields.
+*/
+#if WITH_EDITOR
+#define METASOUND_LOCTEXT(KEY, NAME_TEXT) LOCTEXT(KEY, NAME_TEXT)
+#define METASOUND_LOCTEXT_FORMAT(KEY, NAME_TEXT, ...) FText::Format(LOCTEXT(KEY, NAME_TEXT), __VA_ARGS__)
+#else 
+#define METASOUND_LOCTEXT(KEY, NAME_TEXT) FText::GetEmpty()
+#define METASOUND_LOCTEXT_FORMAT(KEY, NAME_TEXT, ...) FText::GetEmpty()
+#endif // WITH_EDITOR
diff --git a/Engine/Plugins/Runtime/Metasound/Source/MetasoundFrontend/Public/MetasoundOutputNode.h b/Engine/Plugins/Runtime/Metasound/Source/MetasoundFrontend/Public/MetasoundOutputNode.h
--- a/Engine/Plugins/Runtime/Metasound/Source/MetasoundFrontend/Public/MetasoundOutputNode.h
+++ b/Engine/Plugins/Runtime/Metasound/Source/MetasoundFrontend/Public/MetasoundOutputNode.h
@@ -24,7 +24,7 @@ namespace Metasound
 		FMissingOutputNodeInputReferenceError(const INode& InNode, const FText& InDataType)
 			: FBuildErrorBase(
 				"MetasoundMissingOutputDataReferenceError",
-				FText::Format(LOCTEXT("MissingOutputNodeInputReferenceError", "Missing required output node input reference for type {0}."), InDataType))
+				METASOUND_LOCTEXT_FORMAT("MissingOutputNodeInputReferenceError", "Missing required output node input reference for type {0}.", InDataType))
 		{
 			AddNode(InNode);
 		}
@@ -101,7 +101,7 @@ namespace Metasound
 
 		static FVertexInterface GetVertexInterface(const FVertexName& InVertexName)
 		{
-			static const FText VertexDescription = LOCTEXT("Metasound_OutputVertexDescription", "Output from the parent Metasound graph.");
+			static const FText VertexDescription = METASOUND_LOCTEXT("Metasound_OutputVertexDescription", "Output from the parent Metasound graph.");
 
 			return FVertexInterface(
 				FInputVertexInterface(
@@ -120,7 +120,7 @@ namespace Metasound
 			Info.ClassName = { "Output", GetMetasoundDataTypeName<DataType>(), FName() };
 			Info.MajorVersion = 1;
 			Info.MinorVersion = 0;
-			Info.Description = LOCTEXT("Metasound_OutputNodeDescription", "Output from the parent Metasound graph.");
+			Info.Description = METASOUND_LOCTEXT("Metasound_OutputNodeDescription", "Output from the parent Metasound graph.");
 			Info.Author = PluginAuthor;
 			Info.PromptIfMissing = PluginNodeMissingPrompt;
 			Info.DefaultInterface = GetVertexInterface(InVertexName);
diff --git a/Engine/Plugins/Runtime/Metasound/Source/MetasoundFrontend/Public/MetasoundReceiveNode.h b/Engine/Plugins/Runtime/Metasound/Source/MetasoundFrontend/Public/MetasoundReceiveNode.h
--- a/Engine/Plugins/Runtime/Metasound/Source/MetasoundFrontend/Public/MetasoundReceiveNode.h
+++ b/Engine/Plugins/Runtime/Metasound/Source/MetasoundFrontend/Public/MetasoundReceiveNode.h
@@ -4,6 +4,7 @@
 #include "MetasoundBuilderInterface.h"
 #include "MetasoundNode.h"
 #include "MetasoundNodeInterface.h"
+#include "MetasoundNodeRegistrationMacro.h"
 #include "MetasoundOperatorInterface.h"
 #include "MetasoundDataFactory.h"
 #include "MetasoundDataReference.h"
@@ -52,12 +53,12 @@ namespace Metasound
 				Info.ClassName = ReceiveNodeInfo::GetClassNameForDataType(GetMetasoundDataTypeName<TDataType>());
 				Info.MajorVersion = ReceiveNodeInfo::GetCurrentMajorVersion();
 				Info.MinorVersion = ReceiveNodeInfo::GetCurrentMinorVersion();
-				Info.DisplayName = FText::Format(LOCTEXT("Metasound_ReceiveNodeDisplayNameFormat", "Receive {0}"), GetMetasoundDataTypeDisplayText<TDataType>());
-				Info.Description = LOCTEXT("Metasound_ReceiveNodeDescription", "Receives data from a send node with the same name.");
+				Info.DisplayName = METASOUND_LOCTEXT_FORMAT("Metasound_ReceiveNodeDisplayNameFormat", "Receive {0}", GetMetasoundDataTypeDisplayText<TDataType>());
+				Info.Description = METASOUND_LOCTEXT("Metasound_ReceiveNodeDescription", "Receives data from a send node with the same name.");
 				Info.Author = PluginAuthor;
 				Info.PromptIfMissing = PluginNodeMissingPrompt;
 				Info.DefaultInterface = DeclareVertexInterface();
-				Info.CategoryHierarchy = { LOCTEXT("Metasound_TransmissionNodeCategory", "Transmission") };
+				Info.CategoryHierarchy = { METASOUND_LOCTEXT("Metasound_TransmissionNodeCategory", "Transmission") };
 				Info.Keywords = { };
 
 				// Then send & receive nodes do not work as expected, particularly 
diff --git a/Engine/Plugins/Runtime/Metasound/Source/MetasoundFrontend/Public/MetasoundSendNode.h b/Engine/Plugins/Runtime/Metasound/Source/MetasoundFrontend/Public/MetasoundSendNode.h
--- a/Engine/Plugins/Runtime/Metasound/Source/MetasoundFrontend/Public/MetasoundSendNode.h
+++ b/Engine/Plugins/Runtime/Metasound/Source/MetasoundFrontend/Public/MetasoundSendNode.h
@@ -55,12 +55,12 @@ namespace Metasound
 				Info.ClassName = { "Send", GetMetasoundDataTypeName<TDataType>(), FName() };
 				Info.MajorVersion = 1;
 				Info.MinorVersion = 0;
-				Info.DisplayName = FText::Format(LOCTEXT("Metasound_SendNodeDisplayNameFormat", "Send {0}"), GetMetasoundDataTypeDisplayText<TDataType>());
-				Info.Description = LOCTEXT("Metasound_SendNodeDescription", "Sends data from a send node with the same name.");
+				Info.DisplayName = METASOUND_LOCTEXT_FORMAT("Metasound_SendNodeDisplayNameFormat", "Send {0}", GetMetasoundDataTypeDisplayText<TDataType>());
+				Info.Description = METASOUND_LOCTEXT("Metasound_SendNodeDescription", "Sends data from a send node with the same name.");
 				Info.Author = PluginAuthor;
 				Info.PromptIfMissing = PluginNodeMissingPrompt;
 				Info.DefaultInterface = DeclareVertexInterface();
-				Info.CategoryHierarchy = { LOCTEXT("Metasound_TransmissionNodeCategory", "Transmission") };
+				Info.CategoryHierarchy = { METASOUND_LOCTEXT("Metasound_TransmissionNodeCategory", "Transmission") };
 				Info.Keywords = { };
 
 				// Then send & receive nodes do not work as expected, particularly 
diff --git a/Engine/Plugins/Runtime/Metasound/Source/MetasoundGraphCore/Private/MetasoundNodeInterface.cpp b/Engine/Plugins/Runtime/Metasound/Source/MetasoundGraphCore/Private/MetasoundNodeInterface.cpp
--- a/Engine/Plugins/Runtime/Metasound/Source/MetasoundGraphCore/Private/MetasoundNodeInterface.cpp
+++ b/Engine/Plugins/Runtime/Metasound/Source/MetasoundGraphCore/Private/MetasoundNodeInterface.cpp
@@ -5,7 +5,11 @@
 namespace Metasound
 {
 	const FString PluginAuthor = TEXT("Epic Games, Inc.");
+#if WITH_EDITOR
 	const FText PluginNodeMissingPrompt = NSLOCTEXT("MetasoundGraphCore", "Metasound_DefaultMissingNodePrompt", "The node was likely removed, renamed, or the Metasound plugin is not loaded.");
+#else 
+	const FText PluginNodeMissingPrompt = FText::GetEmpty();
+#endif // WITH_EDITOR
 
 	FNodeClassName::FNodeClassName()
 	{
diff --git a/Engine/Plugins/Runtime/Metasound/Source/MetasoundGraphCore/Public/MetasoundLiteralNode.h b/Engine/Plugins/Runtime/Metasound/Source/MetasoundGraphCore/Public/MetasoundLiteralNode.h
--- a/Engine/Plugins/Runtime/Metasound/Source/MetasoundGraphCore/Public/MetasoundLiteralNode.h
+++ b/Engine/Plugins/Runtime/Metasound/Source/MetasoundGraphCore/Public/MetasoundLiteralNode.h
@@ -116,8 +116,10 @@ namespace Metasound
 			Info.ClassName = {"Literal", GetMetasoundDataTypeName<DataType>(), ""};
 			Info.MajorVersion = 1;
 			Info.MinorVersion = 0;
+#if WITH_EDITOR
 			Info.DisplayName = FText::Format(LOCTEXT("Metasound_LiteralNodeDisplayNameFormat", "Literal {0}"), FText::FromName(GetMetasoundDataTypeName<DataType>()));
 			Info.Description = LOCTEXT("Metasound_LiteralNodeDescription", "Literal accessible within a parent Metasound graph.");
+#endif // WITH_EDITOR
 			Info.Author = PluginAuthor;
 			Info.PromptIfMissing = PluginNodeMissingPrompt;
 			Info.DefaultInterface = DeclareVertexInterface();
diff --git a/Engine/Plugins/Runtime/Metasound/Source/MetasoundGraphCore/Public/MetasoundVariableNodes.h b/Engine/Plugins/Runtime/Metasound/Source/MetasoundGraphCore/Public/MetasoundVariableNodes.h
--- a/Engine/Plugins/Runtime/Metasound/Source/MetasoundGraphCore/Public/MetasoundVariableNodes.h
+++ b/Engine/Plugins/Runtime/Metasound/Source/MetasoundGraphCore/Public/MetasoundVariableNodes.h
@@ -181,11 +181,13 @@ namespace Metasound
 			Info.ClassName = VariableNames::GetVariableNodeClassName<DataType>();
 			Info.MajorVersion = 1;
 			Info.MinorVersion = 0;
+#if WITH_EDITOR
 			Info.Description = LOCTEXT("Metasound_InitVariableNodeDescription", "Initialize a variable of a MetaSound graph.");
+			Info.CategoryHierarchy = { LOCTEXT("Metasound_VariableCategory", "Variable") };
+#endif // WITH_EDITOR
 			Info.Author = PluginAuthor;
 			Info.PromptIfMissing = PluginNodeMissingPrompt;
 			Info.DefaultInterface = DeclareVertexInterface();
-			Info.CategoryHierarchy = { LOCTEXT("Metasound_VariableCategory", "Variable") };
 
 			return Info;
 		}
@@ -327,12 +329,15 @@ namespace Metasound
 			Info.ClassName = VariableNames::GetVariableMutatorNodeClassName<DataType>();
 			Info.MajorVersion = 1;
 			Info.MinorVersion = 0;
+
+#if WITH_EDITOR
 			Info.DisplayName = LOCTEXT("Metasound_VariableMutatorNodeDisplayName", "Set");
 			Info.Description = LOCTEXT("Metasound_VariableMutatorNodeDescription", "Set variable on MetaSound graph.");
+			Info.CategoryHierarchy = { LOCTEXT("Metasound_VariableCategory", "Variable") };
+#endif // WITH_EDITOR
 			Info.Author = PluginAuthor;
 			Info.PromptIfMissing = PluginNodeMissingPrompt;
 			Info.DefaultInterface = DeclareVertexInterface();
-			Info.CategoryHierarchy = { LOCTEXT("Metasound_VariableCategory", "Variable") };
 
 			return Info;
 		};
@@ -435,12 +440,14 @@ namespace Metasound
 			Info.ClassName = VariableNames::GetVariableDeferredAccessorNodeClassName<DataType>();
 			Info.MajorVersion = 1;
 			Info.MinorVersion = 0;
+#if WITH_EDITOR
 			Info.DisplayName = LOCTEXT("Metasound_VariableDeferredAccessorNodeDisplayName", "Get Delayed");
 			Info.Description = LOCTEXT("Metasound_VariableDeferredAccessorNodeDescription", "Get a delayed variable on MetaSound graph.");
+			Info.CategoryHierarchy = { LOCTEXT("Metasound_VariableCategory", "Variable") };
+#endif // WITH_EDITOR
 			Info.Author = PluginAuthor;
 			Info.PromptIfMissing = PluginNodeMissingPrompt;
 			Info.DefaultInterface = DeclareVertexInterface();
-			Info.CategoryHierarchy = { LOCTEXT("Metasound_VariableCategory", "Variable") };
 
 			return Info;
 		};
@@ -542,13 +549,15 @@ namespace Metasound
 			Info.ClassName = VariableNames::GetVariableAccessorNodeClassName<DataType>();
 			Info.MajorVersion = 1;
 			Info.MinorVersion = 0;
+#if WITH_EDITOR
 			Info.DisplayName = LOCTEXT("Metasound_VariableAccessorNodeDisplayName", "Get");
 			Info.Description = LOCTEXT("Metasound_VariableAccessorNodeDescription", "Get variable on MetaSound graph.");
+#endif // WITH_EDITOR
+
 			Info.Author = PluginAuthor;
 			Info.PromptIfMissing = PluginNodeMissingPrompt;
 			Info.DefaultInterface = DeclareVertexInterface();
 
-			Info.CategoryHierarchy = { LOCTEXT("Metasound_VariableCategory", "Variable") };
 			return Info;
 		};
 
diff --git a/Engine/Plugins/Runtime/Metasound/Source/MetasoundStandardNodes/Private/MetasoundADEnvelopeNode.cpp b/Engine/Plugins/Runtime/Metasound/Source/MetasoundStandardNodes/Private/MetasoundADEnvelopeNode.cpp
--- a/Engine/Plugins/Runtime/Metasound/Source/MetasoundStandardNodes/Private/MetasoundADEnvelopeNode.cpp
+++ b/Engine/Plugins/Runtime/Metasound/Source/MetasoundStandardNodes/Private/MetasoundADEnvelopeNode.cpp
@@ -229,8 +229,8 @@ namespace Metasound
 			{
 				const FName DataTypeName = GetMetasoundDataTypeName<ValueType>();
 				const FName OperatorName = "AD Envelope";
-				const FText NodeDisplayName = FText::Format(LOCTEXT("ADEnvelopeDisplayNamePattern", "AD Envelope ({0})"), GetMetasoundDataTypeDisplayText<ValueType>());
-				const FText NodeDescription = LOCTEXT("ADEnevelopeDesc", "Generates an attack-decay envelope value output when triggered.");
+				const FText NodeDisplayName = METASOUND_LOCTEXT_FORMAT("ADEnvelopeDisplayNamePattern", "AD Envelope ({0})", GetMetasoundDataTypeDisplayText<ValueType>());
+				const FText NodeDescription = METASOUND_LOCTEXT("ADEnevelopeDesc", "Generates an attack-decay envelope value output when triggered.");
 				const FVertexInterface NodeInterface = GetDefaultInterface();
 
 				const FNodeClassMetadata Metadata
diff --git a/Engine/Plugins/Runtime/Metasound/Source/MetasoundStandardNodes/Private/MetasoundADSREnvelopeNode.cpp b/Engine/Plugins/Runtime/Metasound/Source/MetasoundStandardNodes/Private/MetasoundADSREnvelopeNode.cpp
--- a/Engine/Plugins/Runtime/Metasound/Source/MetasoundStandardNodes/Private/MetasoundADSREnvelopeNode.cpp
+++ b/Engine/Plugins/Runtime/Metasound/Source/MetasoundStandardNodes/Private/MetasoundADSREnvelopeNode.cpp
@@ -359,8 +359,8 @@ namespace Metasound
 			{
 				const FName DataTypeName = GetMetasoundDataTypeName<ValueType>();
 				const FName OperatorName = "ADSR Envelope";
-				const FText NodeDisplayName = FText::Format(LOCTEXT("ADSREnvelopeDisplayNamePattern", "ADSR Envelope ({0})"), GetMetasoundDataTypeDisplayText<ValueType>());
-				const FText NodeDescription = LOCTEXT("ADSREnevelopeDesc", "Generates an attack-decay-sustain-release envelope value output when triggered.");
+				const FText NodeDisplayName = METASOUND_LOCTEXT_FORMAT("ADSREnvelopeDisplayNamePattern", "ADSR Envelope ({0})", GetMetasoundDataTypeDisplayText<ValueType>());
+				const FText NodeDescription = METASOUND_LOCTEXT("ADSREnevelopeDesc", "Generates an attack-decay-sustain-release envelope value output when triggered.");
 				const FVertexInterface NodeInterface = GetDefaultInterface();
 
 				const FNodeClassMetadata Metadata
diff --git a/Engine/Plugins/Runtime/Metasound/Source/MetasoundStandardNodes/Private/MetasoundAdditiveSynthNode.cpp b/Engine/Plugins/Runtime/Metasound/Source/MetasoundStandardNodes/Private/MetasoundAdditiveSynthNode.cpp
--- a/Engine/Plugins/Runtime/Metasound/Source/MetasoundStandardNodes/Private/MetasoundAdditiveSynthNode.cpp
+++ b/Engine/Plugins/Runtime/Metasound/Source/MetasoundStandardNodes/Private/MetasoundAdditiveSynthNode.cpp
@@ -38,8 +38,8 @@ namespace Metasound
 			auto InitNodeInfo = []() -> FNodeClassMetadata
 			{
 				const FName OperatorName = "Additive Synth";
-				const FText NodeDisplayName = LOCTEXT("Metasound_AdditiveSynthNodeDisplayName", "Additive Synth");
-				const FText NodeDescription = LOCTEXT("Metasound_AdditiveSynthNodeDescription", "Synthesizes audio output given input array of sinusoids.");
+				const FText NodeDisplayName = METASOUND_LOCTEXT("Metasound_AdditiveSynthNodeDisplayName", "Additive Synth");
+				const FText NodeDescription = METASOUND_LOCTEXT("Metasound_AdditiveSynthNodeDescription", "Synthesizes audio output given input array of sinusoids.");
 
 				FNodeClassMetadata Info;
 				Info.ClassName = { Metasound::StandardNodes::Namespace, OperatorName, FName() };
diff --git a/Engine/Plugins/Runtime/Metasound/Source/MetasoundStandardNodes/Private/MetasoundAudioFormats.cpp b/Engine/Plugins/Runtime/Metasound/Source/MetasoundStandardNodes/Private/MetasoundAudioFormats.cpp
--- a/Engine/Plugins/Runtime/Metasound/Source/MetasoundStandardNodes/Private/MetasoundAudioFormats.cpp
+++ b/Engine/Plugins/Runtime/Metasound/Source/MetasoundStandardNodes/Private/MetasoundAudioFormats.cpp
@@ -170,11 +170,11 @@ namespace Metasound
 		{
 			return FVertexInterface(
 				FInputVertexInterface(
-					TInputDataVertexModel<FAudioBuffer>(StereoAudioFormatVertexKeys::LeftChannelVertexKey, LOCTEXT("Metasound_LeftStereoOutputVertexDescription", "Left channel audio output.")),
-					TInputDataVertexModel<FAudioBuffer>(StereoAudioFormatVertexKeys::RightChannelVertexKey, LOCTEXT("Metasound_RightStereoOutputVertexDescription", "Right channel audio output."))
+					TInputDataVertexModel<FAudioBuffer>(StereoAudioFormatVertexKeys::LeftChannelVertexKey, METASOUND_LOCTEXT("Metasound_LeftStereoOutputVertexDescription", "Left channel audio output.")),
+					TInputDataVertexModel<FAudioBuffer>(StereoAudioFormatVertexKeys::RightChannelVertexKey, METASOUND_LOCTEXT("Metasound_RightStereoOutputVertexDescription", "Right channel audio output."))
 				),
 				FOutputVertexInterface(
-					TOutputDataVertexModel<FStereoAudioFormat>(InVertexName, LOCTEXT("Metasound_StereoOutputVertexDescription", "Stereo Output."))
+					TOutputDataVertexModel<FStereoAudioFormat>(InVertexName, METASOUND_LOCTEXT("Metasound_StereoOutputVertexDescription", "Stereo Output."))
 				)
 			);
 		}
@@ -186,8 +186,8 @@ namespace Metasound
 			Info.ClassName = { "Output", GetMetasoundDataTypeName<FStereoAudioFormat>(), "" };
 			Info.MajorVersion = 1;
 			Info.MinorVersion = 0;
-			Info.DisplayName = FText::Format(LOCTEXT("Metasound_OutputNodeDisplayNameFormat", "Output {0}"), GetMetasoundDataTypeDisplayText<FStereoAudioFormat>());
-			Info.Description = LOCTEXT("Metasound_OutputNodeDescription", "Output from the parent Metasound graph.");
+			Info.DisplayName = METASOUND_LOCTEXT_FORMAT("Metasound_OutputNodeDisplayNameFormat", "Output {0}", GetMetasoundDataTypeDisplayText<FStereoAudioFormat>());
+			Info.Description = METASOUND_LOCTEXT("Metasound_OutputNodeDescription", "Output from the parent Metasound graph.");
 			Info.Author = PluginAuthor;
 			Info.PromptIfMissing = PluginNodeMissingPrompt;
 			Info.DefaultInterface = GetVertexInterface(InOutputName);
@@ -309,11 +309,11 @@ namespace Metasound
 		{
 			return FVertexInterface(
 				FInputVertexInterface(
-					TInputDataVertexModel<FStereoAudioFormat>(InVertexName, LOCTEXT("Metasound_StereoInputVertexDescription", "Stereo Input."))
+					TInputDataVertexModel<FStereoAudioFormat>(InVertexName, METASOUND_LOCTEXT("Metasound_StereoInputVertexDescription", "Stereo Input."))
 				),
 				FOutputVertexInterface(
-					TOutputDataVertexModel<FAudioBuffer>(StereoAudioFormatVertexKeys::LeftChannelVertexKey, LOCTEXT("Metasound_LeftStereoInputVertexDescription", "Left channel audio output.")),
-					TOutputDataVertexModel<FAudioBuffer>(StereoAudioFormatVertexKeys::RightChannelVertexKey, LOCTEXT("Metasound_RightStereoInputVertexDescription", "Right channel audio output."))
+					TOutputDataVertexModel<FAudioBuffer>(StereoAudioFormatVertexKeys::LeftChannelVertexKey, METASOUND_LOCTEXT("Metasound_LeftStereoInputVertexDescription", "Left channel audio output.")),
+					TOutputDataVertexModel<FAudioBuffer>(StereoAudioFormatVertexKeys::RightChannelVertexKey, METASOUND_LOCTEXT("Metasound_RightStereoInputVertexDescription", "Right channel audio output."))
 				)
 			);
 		}
@@ -325,8 +325,8 @@ namespace Metasound
 			Info.ClassName = { "Input", GetMetasoundDataTypeName<FStereoAudioFormat>(), "" };
 			Info.MajorVersion = 1;
 			Info.MinorVersion = 0;
-			Info.DisplayName = FText::Format(LOCTEXT("Metasound_InputNodeDisplayNameFormat", "Input {0}"), GetMetasoundDataTypeDisplayText<FStereoAudioFormat>());
-			Info.Description = LOCTEXT("Metasound_InputNodeDescription", "Input from the parent Metasound graph.");
+			Info.DisplayName = METASOUND_LOCTEXT_FORMAT("Metasound_InputNodeDisplayNameFormat", "Input {0}", GetMetasoundDataTypeDisplayText<FStereoAudioFormat>());
+			Info.Description = METASOUND_LOCTEXT("Metasound_InputNodeDescription", "Input from the parent Metasound graph.");
 			Info.Author = PluginAuthor;
 			Info.PromptIfMissing = PluginNodeMissingPrompt;
 			Info.DefaultInterface = GetVertexInterface(InInputName);
@@ -444,10 +444,10 @@ namespace Metasound
 		{
 			return FVertexInterface(
 				FInputVertexInterface(
-					TInputDataVertexModel<FAudioBuffer>(InVertexName, LOCTEXT("Metasound_CenterMonoOutputVertexDescription", "Center channel audio output."))
+					TInputDataVertexModel<FAudioBuffer>(InVertexName, METASOUND_LOCTEXT("Metasound_CenterMonoOutputVertexDescription", "Center channel audio output."))
 				),
 				FOutputVertexInterface(
-					TOutputDataVertexModel<FMonoAudioFormat>(InVertexName, LOCTEXT("Metasound_MonoOutputVertexDescription", "Mono Output."))
+					TOutputDataVertexModel<FMonoAudioFormat>(InVertexName, METASOUND_LOCTEXT("Metasound_MonoOutputVertexDescription", "Mono Output."))
 				)
 			);
 		}
@@ -459,8 +459,8 @@ namespace Metasound
 			Info.ClassName = { "Output", GetMetasoundDataTypeName<FMonoAudioFormat>(), "" };
 			Info.MajorVersion = 1;
 			Info.MinorVersion = 0;
-			Info.DisplayName = FText::Format(LOCTEXT("Metasound_OutputNodeDisplayNameFormat", "Output {0}"), GetMetasoundDataTypeDisplayText<FMonoAudioFormat>());
-			Info.Description = LOCTEXT("Metasound_OutputNodeDescription", "Output from the parent Metasound graph.");
+			Info.DisplayName = METASOUND_LOCTEXT_FORMAT("Metasound_OutputNodeDisplayNameFormat", "Output {0}", GetMetasoundDataTypeDisplayText<FMonoAudioFormat>());
+			Info.Description = METASOUND_LOCTEXT("Metasound_OutputNodeDescription", "Output from the parent Metasound graph.");
 			Info.Author = PluginAuthor;
 			Info.PromptIfMissing = PluginNodeMissingPrompt;
 			Info.DefaultInterface = GetVertexInterface(InOutputName);
@@ -577,10 +577,10 @@ namespace Metasound
 		{
 			return FVertexInterface(
 				FInputVertexInterface(
-					TInputDataVertexModel<FMonoAudioFormat>(InVertexName, LOCTEXT("Metasound_MonoInputVertexDescription", "Mono Input."))
+					TInputDataVertexModel<FMonoAudioFormat>(InVertexName, METASOUND_LOCTEXT("Metasound_MonoInputVertexDescription", "Mono Input."))
 				),
 				FOutputVertexInterface(
-					TOutputDataVertexModel<FAudioBuffer>(InVertexName, LOCTEXT("Metasound_CenterMonoInputVertexDescription", "Center channel audio output."))
+					TOutputDataVertexModel<FAudioBuffer>(InVertexName, METASOUND_LOCTEXT("Metasound_CenterMonoInputVertexDescription", "Center channel audio output."))
 				)
 			);
 		}
@@ -592,8 +592,8 @@ namespace Metasound
 			Info.ClassName = { "Input", GetMetasoundDataTypeName<FMonoAudioFormat>(), "" };
 			Info.MajorVersion = 1;
 			Info.MinorVersion = 0;
-			Info.DisplayName = FText::Format(LOCTEXT("Metasound_InputNodeDisplayNameFormat", "Input {0}"), GetMetasoundDataTypeDisplayText<FMonoAudioFormat>());
-			Info.Description = LOCTEXT("Metasound_InputNodeDescription", "Input from the parent Metasound graph.");
+			Info.DisplayName = METASOUND_LOCTEXT_FORMAT("Metasound_InputNodeDisplayNameFormat", "Input {0}", GetMetasoundDataTypeDisplayText<FMonoAudioFormat>());
+			Info.Description = METASOUND_LOCTEXT("Metasound_InputNodeDescription", "Input from the parent Metasound graph.");
 			Info.Author = PluginAuthor;
 			Info.PromptIfMissing = PluginNodeMissingPrompt;
 			Info.DefaultInterface = GetVertexInterface(InInputName);
diff --git a/Engine/Plugins/Runtime/Metasound/Source/MetasoundStandardNodes/Private/MetasoundBPMToSecondsNode.cpp b/Engine/Plugins/Runtime/Metasound/Source/MetasoundStandardNodes/Private/MetasoundBPMToSecondsNode.cpp
--- a/Engine/Plugins/Runtime/Metasound/Source/MetasoundStandardNodes/Private/MetasoundBPMToSecondsNode.cpp
+++ b/Engine/Plugins/Runtime/Metasound/Source/MetasoundStandardNodes/Private/MetasoundBPMToSecondsNode.cpp
@@ -134,12 +134,12 @@ namespace Metasound
 
 		static const FVertexInterface Interface(
 			FInputVertexInterface(
-				TInputDataVertexModel<float>(GetInputBPMTempoName(), LOCTEXT("BPMToSecondsNode_BPMTT", "Beats Per Minute."), 90.0f),
-				TInputDataVertexModel<float>(GetInputBeatMultiplierName(), LOCTEXT("BPMToSecondsNode_BeatMultiplierTT", "The multiplier of the BPM."), 1.0f),
-				TInputDataVertexModel<float>(GetInputDivisionsOfWholeNoteName(), LOCTEXT("BPMToSecondsNode_DivOfWholeNoteTT", "Divisions of a whole note."), 4.0f)
+				TInputDataVertexModel<float>(GetInputBPMTempoName(), METASOUND_LOCTEXT("BPMToSecondsNode_BPMTT", "Beats Per Minute."), 90.0f),
+				TInputDataVertexModel<float>(GetInputBeatMultiplierName(), METASOUND_LOCTEXT("BPMToSecondsNode_BeatMultiplierTT", "The multiplier of the BPM."), 1.0f),
+				TInputDataVertexModel<float>(GetInputDivisionsOfWholeNoteName(), METASOUND_LOCTEXT("BPMToSecondsNode_DivOfWholeNoteTT", "Divisions of a whole note."), 4.0f)
 			),
 			FOutputVertexInterface(
-				TOutputDataVertexModel<FTime>(GetOutputTimeSecondsName(), LOCTEXT("BPMToSecondsNode_OutputTimeTT", "The output time in seconds."))
+				TOutputDataVertexModel<FTime>(GetOutputTimeSecondsName(), METASOUND_LOCTEXT("BPMToSecondsNode_OutputTimeTT", "The output time in seconds."))
 			)
 		);
 
@@ -154,8 +154,8 @@ namespace Metasound
 			Info.ClassName = { StandardNodes::Namespace, "BPMToSeconds", "" };
 			Info.MajorVersion = 1;
 			Info.MinorVersion = 1;
-			Info.DisplayName = LOCTEXT("BPMToSecondsNode_DisplayName", "BPM To Seconds");
-			Info.Description = LOCTEXT("BPMToSecondsNode_Desc", "Calculates a beat time in seconds from the given BPM, beat multiplier and divisions of a whole note.");
+			Info.DisplayName = METASOUND_LOCTEXT("BPMToSecondsNode_DisplayName", "BPM To Seconds");
+			Info.Description = METASOUND_LOCTEXT("BPMToSecondsNode_Desc", "Calculates a beat time in seconds from the given BPM, beat multiplier and divisions of a whole note.");
 			Info.Author = PluginAuthor;
 			Info.PromptIfMissing = PluginNodeMissingPrompt;
 			Info.DefaultInterface = GetVertexInterface();
diff --git a/Engine/Plugins/Runtime/Metasound/Source/MetasoundStandardNodes/Private/MetasoundBandsplitterNode.cpp b/Engine/Plugins/Runtime/Metasound/Source/MetasoundStandardNodes/Private/MetasoundBandsplitterNode.cpp
--- a/Engine/Plugins/Runtime/Metasound/Source/MetasoundStandardNodes/Private/MetasoundBandsplitterNode.cpp
+++ b/Engine/Plugins/Runtime/Metasound/Source/MetasoundStandardNodes/Private/MetasoundBandsplitterNode.cpp
@@ -35,10 +35,10 @@ namespace Metasound
 		FEnumBandSplitterFilterOrder, FEnumBandSplitterFilterOrderInfo, FEnumBandSplitterFilterOrderReadRef, FEnumBandSplitterFilterOrderWriteRef);
 
 	DEFINE_METASOUND_ENUM_BEGIN(EBandSplitterFilterOrder, FEnumBandSplitterFilterOrder, "Filter Order")
-		DEFINE_METASOUND_ENUM_ENTRY(EBandSplitterFilterOrder::TwoPole, LOCTEXT("BandSplitterFilterOrderTwoPoleDescription", "Two Pole"), LOCTEXT("EnvelopePeakModeMSDescriptionTT", "Envelope follows a running Mean Squared of the audio signal.")),
-		DEFINE_METASOUND_ENUM_ENTRY(EBandSplitterFilterOrder::FourPole, LOCTEXT("BandSplitterFilterOrderFourPoleDescription", "Four Pole"), LOCTEXT("EnvelopePeakModeRMSDescriptionTT", "Envelope follows a running Root Mean Squared of the audio signal.")),
-		DEFINE_METASOUND_ENUM_ENTRY(EBandSplitterFilterOrder::SixPole, LOCTEXT("BandSplitterFilterOrderSixPoleDescription", "Six Pole"), LOCTEXT("EnvelopePeakModePeakDescriptionTT", "Envelope follows the peaks in the audio signal.")),
-		DEFINE_METASOUND_ENUM_ENTRY(EBandSplitterFilterOrder::EightPole, LOCTEXT("BandSplitterFilterOrderEightPoleDescription", "Eight Pole"), LOCTEXT("EnvelopePeakModePeakDescriptionTT", "Envelope follows the peaks in the audio signal.")),
+		DEFINE_METASOUND_ENUM_ENTRY(EBandSplitterFilterOrder::TwoPole, "BandSplitterFilterOrderTwoPoleDescription", "Two Pole", "EnvelopePeakModeMSDescriptionTT", "Envelope follows a running Mean Squared of the audio signal."),
+		DEFINE_METASOUND_ENUM_ENTRY(EBandSplitterFilterOrder::FourPole, "BandSplitterFilterOrderFourPoleDescription", "Four Pole", "EnvelopePeakModeRMSDescriptionTT", "Envelope follows a running Root Mean Squared of the audio signal."),
+		DEFINE_METASOUND_ENUM_ENTRY(EBandSplitterFilterOrder::SixPole, "BandSplitterFilterOrderSixPoleDescription", "Six Pole", "EnvelopePeakModePeakDescriptionTT", "Envelope follows the peaks in the audio signal."),
+		DEFINE_METASOUND_ENUM_ENTRY(EBandSplitterFilterOrder::EightPole, "BandSplitterFilterOrderEightPoleDescription", "Eight Pole", "EnvelopePeakModePeakDescriptionTT", "Envelope follows the peaks in the audio signal."),
 	DEFINE_METASOUND_ENUM_END()
 
 	template<uint32 NumBands, uint32 NumChannels>
@@ -89,8 +89,8 @@ namespace Metasound
 			auto CreateNodeClassMetadataMono = []() -> FNodeClassMetadata
 			{
 				FName OperatorName = *FString::Printf(TEXT("Band Splitter (Mono, %d)"), NumBands);
-				FText NodeDisplayName = FText::Format(LOCTEXT("BandSplitterDisplayNamePattern1", "Mono Band Splitter ({0})"), NumBands);
-				FText NodeDescription = LOCTEXT("BandSplitterDescription1", "Will split incoming audio into separate frequency bands.");
+				FText NodeDisplayName = METASOUND_LOCTEXT_FORMAT("BandSplitterDisplayNamePattern1", "Mono Band Splitter ({0})", NumBands);
+				const FText NodeDescription = METASOUND_LOCTEXT("BandSplitterDescription1", "Will split incoming audio into separate frequency bands.");
 				FVertexInterface NodeInterface = GetDefaultInterface();
 
 				return CreateNodeClassMetadata(OperatorName, NodeDisplayName, NodeDescription, NodeInterface);
@@ -100,8 +100,8 @@ namespace Metasound
 			auto CreateNodeClassMetadataStereo = []() -> FNodeClassMetadata
 			{
 				FName OperatorName = *FString::Printf(TEXT("Band Splitter (Stereo, %d)"), NumBands);
-				FText NodeDisplayName = FText::Format(LOCTEXT("BandSplitterDisplayNamePattern2", "Stereo Band Splitter ({0})"), NumBands);
-				FText NodeDescription = LOCTEXT("BandSplitterDescription2", "Will split incoming audio into separate frequency bands.");
+				FText NodeDisplayName = METASOUND_LOCTEXT_FORMAT("BandSplitterDisplayNamePattern2", "Stereo Band Splitter ({0})", NumBands);
+				const FText NodeDescription = METASOUND_LOCTEXT("BandSplitterDescription2", "Will split incoming audio into separate frequency bands.");
 				FVertexInterface NodeInterface = GetDefaultInterface();
 
 				return  CreateNodeClassMetadata(OperatorName, NodeDisplayName, NodeDescription, NodeInterface);
@@ -111,8 +111,8 @@ namespace Metasound
 			auto CreateNodeClassMetadataMultiChan = []() -> FNodeClassMetadata
 			{
 				FName OperatorName = *FString::Printf(TEXT("Band Splitter (%d-Channel, %d)"), NumChannels, NumBands);
-				FText NodeDisplayName = FText::Format(LOCTEXT("BandSplitterDisplayNamePattern3", "{0}-channel Band Splitter ({1})"), NumChannels, NumBands);
-				FText NodeDescription = LOCTEXT("BandSplitterDescription3", "Will split incoming audio into separate frequency bands.");
+				FText NodeDisplayName = METASOUND_LOCTEXT_FORMAT("BandSplitterDisplayNamePattern3", "{0}-channel Band Splitter ({1})", NumChannels, NumBands);
+				const FText NodeDescription = METASOUND_LOCTEXT("BandSplitterDescription3", "Will split incoming audio into separate frequency bands.");
 				FVertexInterface NodeInterface = GetDefaultInterface();
 
 				return  CreateNodeClassMetadata(OperatorName, NodeDisplayName, NodeDescription, NodeInterface);
@@ -336,7 +336,7 @@ namespace Metasound
 
 		static const FText GetAudioInputDescription(uint32 ChannelIndex)
 		{
-			return FText::Format(LOCTEXT("BandSplitterAudioInputDescription", "Audio Channel: {0}"), ChannelIndex);
+			return METASOUND_LOCTEXT_FORMAT("BandSplitterAudioInputDescription", "Audio Channel: {0}", ChannelIndex);
 		}
 
 		static const FVertexName GetCrossoverInputName(uint32 InputIndex)
@@ -346,7 +346,7 @@ namespace Metasound
 
 		static const FText GetCrossoverInputDescription(uint32 InputIndex)
 		{
-			return FText::Format(LOCTEXT("BandSplitterCrossoverInputDescription", "Crossover Input #: {0}"), InputIndex);
+			return METASOUND_LOCTEXT_FORMAT("BandSplitterCrossoverInputDescription", "Crossover Input #: {0}", InputIndex);
 		}
 
 		static const FVertexName GetAudioOutputName(uint32 OutputIndex, uint32 ChannelIndex)
@@ -365,7 +365,7 @@ namespace Metasound
 
 		static const FText GetAudioOutputDescription(uint32 OutputIndex, uint32 ChannelIndex)
 		{
-			return FText::Format(LOCTEXT("BandSplitterAudioOutputDescription", "Channel {0} Output for Band {1}"), ChannelIndex, OutputIndex);
+			return METASOUND_LOCTEXT_FORMAT("BandSplitterAudioOutputDescription", "Channel {0} Output for Band {1}", ChannelIndex, OutputIndex);
 		}
 
 		FORCEINLINE Audio::EFilterOrder GetFilterOrder(EBandSplitterFilterOrder InFilterOrder) const
diff --git a/Engine/Plugins/Runtime/Metasound/Source/MetasoundStandardNodes/Private/MetasoundBasicFilters.cpp b/Engine/Plugins/Runtime/Metasound/Source/MetasoundStandardNodes/Private/MetasoundBasicFilters.cpp
--- a/Engine/Plugins/Runtime/Metasound/Source/MetasoundStandardNodes/Private/MetasoundBasicFilters.cpp
+++ b/Engine/Plugins/Runtime/Metasound/Source/MetasoundStandardNodes/Private/MetasoundBasicFilters.cpp
@@ -60,16 +60,16 @@ namespace Metasound
 	METASOUNDSTANDARDNODES_API, FEnumEBiquadFilterType, FEnumBiQuadFilterTypeInfo, FEnumBiQuadFilterReadRef, FEnumBiQuadFilterWriteRef);
 
 	DEFINE_METASOUND_ENUM_BEGIN(Audio::EBiquadFilter::Type, FEnumEBiquadFilterType, "BiquadFilterType")
-		DEFINE_METASOUND_ENUM_ENTRY(Audio::EBiquadFilter::Lowpass, LOCTEXT("LpDescription", "Low Pass"), LOCTEXT("LpDescriptionTT", "Low pass Biquad filter.")),
-		DEFINE_METASOUND_ENUM_ENTRY(Audio::EBiquadFilter::Highpass, LOCTEXT("HpDescription", "High Pass"), LOCTEXT("HpDescriptionTT", "High pass Biquad filter.")),
-		DEFINE_METASOUND_ENUM_ENTRY(Audio::EBiquadFilter::Bandpass, LOCTEXT("BpDescription", "Band Pass"), LOCTEXT("BpDescriptionTT", "Band pass Biquad filter.")),
-		DEFINE_METASOUND_ENUM_ENTRY(Audio::EBiquadFilter::Notch, LOCTEXT("NotchDescription", "Notch "), LOCTEXT("NotchDescriptionTT", "Notch biquad filter.")),
-		DEFINE_METASOUND_ENUM_ENTRY(Audio::EBiquadFilter::ParametricEQ, LOCTEXT("ParaEqDescription", "Parametric EQ"), LOCTEXT("ParaEqDescriptionTT", "Parametric EQ biquad filter.")),
-		DEFINE_METASOUND_ENUM_ENTRY(Audio::EBiquadFilter::LowShelf, LOCTEXT("LowShelfDescription", "Low Shelf"), LOCTEXT("LowShelfDescriptionTT", "Low shelf biquad filter.")),
-		DEFINE_METASOUND_ENUM_ENTRY(Audio::EBiquadFilter::HighShelf, LOCTEXT("HighShelfDescription", "High Shelf"), LOCTEXT("HighShelfDescriptionTT", "High shelf biquad filter.")),
-		DEFINE_METASOUND_ENUM_ENTRY(Audio::EBiquadFilter::AllPass, LOCTEXT("AllPassDescription", "All Pass"), LOCTEXT("AllPassDescriptionTT", "All pass biquad Filter.")),
-		DEFINE_METASOUND_ENUM_ENTRY(Audio::EBiquadFilter::ButterworthLowPass, LOCTEXT("LowPassButterDescription", "Butterworth Low Pass"), LOCTEXT("LowPassButterDescriptionTT", "Butterworth Low Pass Biquad Filter.")),
-		DEFINE_METASOUND_ENUM_ENTRY(Audio::EBiquadFilter::ButterworthHighPass, LOCTEXT("HighPassButterDescription", "Butterworth High Pass"), LOCTEXT("HighPassButterDescriptionTT", "Butterworth High Pass Biquad Filter."))
+		DEFINE_METASOUND_ENUM_ENTRY(Audio::EBiquadFilter::Lowpass, "LpDescription", "Low Pass", "LpDescriptionTT", "Low pass Biquad filter."),
+		DEFINE_METASOUND_ENUM_ENTRY(Audio::EBiquadFilter::Highpass, "HpDescription", "High Pass", "HpDescriptionTT", "High pass Biquad filter."),
+		DEFINE_METASOUND_ENUM_ENTRY(Audio::EBiquadFilter::Bandpass, "BpDescription", "Band Pass", "BpDescriptionTT", "Band pass Biquad filter."),
+		DEFINE_METASOUND_ENUM_ENTRY(Audio::EBiquadFilter::Notch, "NotchDescription", "Notch ", "NotchDescriptionTT", "Notch biquad filter."),
+		DEFINE_METASOUND_ENUM_ENTRY(Audio::EBiquadFilter::ParametricEQ, "ParaEqDescription", "Parametric EQ", "ParaEqDescriptionTT", "Parametric EQ biquad filter."),
+		DEFINE_METASOUND_ENUM_ENTRY(Audio::EBiquadFilter::LowShelf, "LowShelfDescription", "Low Shelf", "LowShelfDescriptionTT", "Low shelf biquad filter."),
+		DEFINE_METASOUND_ENUM_ENTRY(Audio::EBiquadFilter::HighShelf, "HighShelfDescription", "High Shelf", "HighShelfDescriptionTT", "High shelf biquad filter."),
+		DEFINE_METASOUND_ENUM_ENTRY(Audio::EBiquadFilter::AllPass, "AllPassDescription", "All Pass", "AllPassDescriptionTT", "All pass biquad Filter."),
+		DEFINE_METASOUND_ENUM_ENTRY(Audio::EBiquadFilter::ButterworthLowPass, "LowPassButterDescription", "Butterworth Low Pass", "LowPassButterDescriptionTT", "Butterworth Low Pass Biquad Filter."),
+		DEFINE_METASOUND_ENUM_ENTRY(Audio::EBiquadFilter::ButterworthHighPass, "HighPassButterDescription", "Butterworth High Pass", "HighPassButterDescriptionTT", "Butterworth High Pass Biquad Filter.")
 		DEFINE_METASOUND_ENUM_END()
 #pragma endregion
 
@@ -163,8 +163,8 @@ namespace Metasound
 			Info.ClassName = { StandardNodes::Namespace, TEXT("Ladder Filter"), StandardNodes::AudioVariant };
 			Info.MajorVersion = 1;
 			Info.MinorVersion = 0;
-			Info.DisplayName = LOCTEXT("Metasound_LadderFilterNodeDisplayName", "Ladder Filter");
-			Info.Description = LOCTEXT("Ladder_Filter_NodeDescription", "Ladder filter"),
+			Info.DisplayName = METASOUND_LOCTEXT("Metasound_LadderFilterNodeDisplayName", "Ladder Filter");
+			Info.Description = METASOUND_LOCTEXT("Ladder_Filter_NodeDescription", "Ladder filter");
 			Info.Author = PluginAuthor;
 			Info.PromptIfMissing = PluginNodeMissingPrompt;
 			Info.DefaultInterface = DeclareVertexInterface();
@@ -345,9 +345,9 @@ namespace Metasound
 			Info.ClassName = { StandardNodes::Namespace, TEXT("State Variable Filter"), StandardNodes::AudioVariant };
 			Info.MajorVersion = 1;
 			Info.MinorVersion = 0;
-			Info.DisplayName = LOCTEXT("Metasound_StateVariableFilterNodeDisplayName", "State Variable Filter");
-			Info.Description = LOCTEXT("State_Variable_Filter_NodeDescription", "State Variable filter"),
-				Info.Author = PluginAuthor;
+			Info.DisplayName = METASOUND_LOCTEXT("Metasound_StateVariableFilterNodeDisplayName", "State Variable Filter");
+			Info.Description = METASOUND_LOCTEXT("State_Variable_Filter_NodeDescription", "State Variable filter");
+			Info.Author = PluginAuthor;
 			Info.PromptIfMissing = PluginNodeMissingPrompt;
 			Info.DefaultInterface = DeclareVertexInterface();
 			Info.CategoryHierarchy.Emplace(NodeCategories::Filters);
@@ -517,9 +517,9 @@ namespace Metasound
 			Info.ClassName = { StandardNodes::Namespace, TEXT("One-Pole Low Pass Filter"), StandardNodes::AudioVariant };
 			Info.MajorVersion = 1;
 			Info.MinorVersion = 0;
-			Info.DisplayName = LOCTEXT("Metasound_OnePoleLpfNodeDisplayName", "One-Pole Low Pass Filter");
-			Info.Description = LOCTEXT("One_Pole_Low_Pass_Filter_NodeDescription", "One-Pole Low Pass Filter"),
-				Info.Author = PluginAuthor;
+			Info.DisplayName = METASOUND_LOCTEXT("Metasound_OnePoleLpfNodeDisplayName", "One-Pole Low Pass Filter");
+			Info.Description = METASOUND_LOCTEXT("One_Pole_Low_Pass_Filter_NodeDescription", "One-Pole Low Pass Filter");
+			Info.Author = PluginAuthor;
 			Info.PromptIfMissing = PluginNodeMissingPrompt;
 			Info.DefaultInterface = DeclareVertexInterface();
 			Info.CategoryHierarchy.Emplace(NodeCategories::Filters);
@@ -654,9 +654,9 @@ namespace Metasound
 			Info.ClassName = { StandardNodes::Namespace, TEXT("One-Pole High Pass Filter"), StandardNodes::AudioVariant };
 			Info.MajorVersion = 1;
 			Info.MinorVersion = 0;
-			Info.DisplayName = LOCTEXT("Metasound_OnePoleHpfNodeDisplayName", "One-Pole High Pass Filter");
-			Info.Description = LOCTEXT("One_Pole_High_Pass_Filter_NodeDescription", "One-Pole High Pass Filter"),
-				Info.Author = PluginAuthor;
+			Info.DisplayName = METASOUND_LOCTEXT("Metasound_OnePoleHpfNodeDisplayName", "One-Pole High Pass Filter");
+			Info.Description = METASOUND_LOCTEXT("One_Pole_High_Pass_Filter_NodeDescription", "One-Pole High Pass Filter");
+			Info.Author = PluginAuthor;
 			Info.PromptIfMissing = PluginNodeMissingPrompt;
 			Info.DefaultInterface = DeclareVertexInterface();
 			Info.CategoryHierarchy.Emplace(NodeCategories::Filters);
@@ -810,9 +810,9 @@ namespace Metasound
 			Info.ClassName = { StandardNodes::Namespace, TEXT("Biquad Filter"), StandardNodes::AudioVariant };
 			Info.MajorVersion = 1;
 			Info.MinorVersion = 0;
-			Info.DisplayName = LOCTEXT("Metasound_BiquadFilterNodeDisplayName", "Biquad Filter");
-			Info.Description = LOCTEXT("Biquad_Filter_NodeDescription", "Biquad filter"),
-				Info.Author = PluginAuthor;
+			Info.DisplayName = METASOUND_LOCTEXT("Metasound_BiquadFilterNodeDisplayName", "Biquad Filter");
+			Info.Description = METASOUND_LOCTEXT("Biquad_Filter_NodeDescription", "Biquad filter");
+			Info.Author = PluginAuthor;
 			Info.PromptIfMissing = PluginNodeMissingPrompt;
 			Info.DefaultInterface = DeclareVertexInterface();
 			Info.CategoryHierarchy.Emplace(NodeCategories::Filters);
diff --git a/Engine/Plugins/Runtime/Metasound/Source/MetasoundStandardNodes/Private/MetasoundBitcrusherNode.cpp b/Engine/Plugins/Runtime/Metasound/Source/MetasoundStandardNodes/Private/MetasoundBitcrusherNode.cpp
--- a/Engine/Plugins/Runtime/Metasound/Source/MetasoundStandardNodes/Private/MetasoundBitcrusherNode.cpp
+++ b/Engine/Plugins/Runtime/Metasound/Source/MetasoundStandardNodes/Private/MetasoundBitcrusherNode.cpp
@@ -59,8 +59,8 @@ namespace Metasound
 					FNodeClassName { StandardNodes::Namespace, "Bitcrusher", StandardNodes::AudioVariant },
 					1, // Major Version
 					0, // Minor Version
-					LOCTEXT("BitcrusherDisplayName", "Bitcrusher"),
-					LOCTEXT("BitcrusherDesc", "Downsamples and lowers the bit-depth of an incoming audio signal."),
+					METASOUND_LOCTEXT("BitcrusherDisplayName", "Bitcrusher"),
+					METASOUND_LOCTEXT("BitcrusherDesc", "Downsamples and lowers the bit-depth of an incoming audio signal."),
 					PluginAuthor,
 					PluginNodeMissingPrompt,
 					NodeInterface,
diff --git a/Engine/Plugins/Runtime/Metasound/Source/MetasoundStandardNodes/Private/MetasoundClampNode.cpp b/Engine/Plugins/Runtime/Metasound/Source/MetasoundStandardNodes/Private/MetasoundClampNode.cpp
--- a/Engine/Plugins/Runtime/Metasound/Source/MetasoundStandardNodes/Private/MetasoundClampNode.cpp
+++ b/Engine/Plugins/Runtime/Metasound/Source/MetasoundStandardNodes/Private/MetasoundClampNode.cpp
@@ -35,8 +35,8 @@ namespace Metasound
 				PluginAuthor,
 				PluginNodeMissingPrompt,
 				InDefaultInterface,
-				{ LOCTEXT("ClampCategory", "Math") },
-				{ LOCTEXT("ClampKeyword", "Clamp") },
+				{ METASOUND_LOCTEXT("ClampCategory", "Math") },
+				{ METASOUND_LOCTEXT("ClampKeyword", "Clamp") },
 				FNodeDisplayStyle()
 			};
 
@@ -175,8 +175,8 @@ namespace Metasound
 			{
 				const FName DataTypeName = GetMetasoundDataTypeName<ValueType>();
 				const FName OperatorName = "Clamp";
-				const FText NodeDisplayName = FText::Format(LOCTEXT("ClampDisplayNamePattern", "Clamp ({0})"), GetMetasoundDataTypeDisplayText<ValueType>());
-				const FText NodeDescription = LOCTEXT("ClampDesc", "Returns the clamped value of the input within the given value range.");
+				const FText NodeDisplayName = METASOUND_LOCTEXT_FORMAT("ClampDisplayNamePattern", "Clamp ({0})", GetMetasoundDataTypeDisplayText<ValueType>());
+				const FText NodeDescription = METASOUND_LOCTEXT("ClampDesc", "Returns the clamped value of the input within the given value range.");
 				const FVertexInterface NodeInterface = GetDefaultInterface();
 
 				return MetasoundClampNodePrivate::CreateNodeClassMetadata(DataTypeName, OperatorName, NodeDisplayName, NodeDescription, NodeInterface);
diff --git a/Engine/Plugins/Runtime/Metasound/Source/MetasoundStandardNodes/Private/MetasoundCompressorNode.cpp b/Engine/Plugins/Runtime/Metasound/Source/MetasoundStandardNodes/Private/MetasoundCompressorNode.cpp
--- a/Engine/Plugins/Runtime/Metasound/Source/MetasoundStandardNodes/Private/MetasoundCompressorNode.cpp
+++ b/Engine/Plugins/Runtime/Metasound/Source/MetasoundStandardNodes/Private/MetasoundCompressorNode.cpp
@@ -122,8 +122,8 @@ namespace Metasound
 					FNodeClassName { StandardNodes::Namespace, "Compressor", StandardNodes::AudioVariant },
 					1, // Major Version
 					0, // Minor Version
-					LOCTEXT("CompressorDisplayName", "Compressor"),
-					LOCTEXT("CompressorDesc", "Lowers the dynamic range of a signal."),
+					METASOUND_LOCTEXT("CompressorDisplayName", "Compressor"),
+					METASOUND_LOCTEXT("CompressorDesc", "Lowers the dynamic range of a signal."),
 					PluginAuthor,
 					PluginNodeMissingPrompt,
 					NodeInterface,
diff --git a/Engine/Plugins/Runtime/Metasound/Source/MetasoundStandardNodes/Private/MetasoundConversionNodes.cpp b/Engine/Plugins/Runtime/Metasound/Source/MetasoundStandardNodes/Private/MetasoundConversionNodes.cpp
--- a/Engine/Plugins/Runtime/Metasound/Source/MetasoundStandardNodes/Private/MetasoundConversionNodes.cpp
+++ b/Engine/Plugins/Runtime/Metasound/Source/MetasoundStandardNodes/Private/MetasoundConversionNodes.cpp
@@ -152,8 +152,8 @@ namespace Metasound
 		{
 			using namespace ConversionNodeVertexNames;
 
-			static FText InputDesc = FText::Format(LOCTEXT("ConvDisplayNamePatternFrom", "Input {0} value."), GetMetasoundDataTypeDisplayText<FromType>());
-			static FText OutputDesc = FText::Format(LOCTEXT("ConvDisplayNamePatternTo", "Output {0} value."), GetMetasoundDataTypeDisplayText<ToType>());
+			static FText InputDesc = METASOUND_LOCTEXT_FORMAT("ConvDisplayNamePatternFrom", "Input {0} value.", GetMetasoundDataTypeDisplayText<FromType>());
+			static FText OutputDesc = METASOUND_LOCTEXT_FORMAT("ConvDisplayNamePatternTo", "Output {0} value.", GetMetasoundDataTypeDisplayText<ToType>());
 
 			static const FVertexInterface DefaultInterface(
 				FInputVertexInterface(
@@ -184,7 +184,7 @@ namespace Metasound
 				const FString& ToTypeString = GetMetasoundDataTypeString<ToType>();
 
 				const FName ClassName = *FString::Format(TEXT("Conversion{0}To{1}"), { FromTypeString, ToTypeString });
-				const FText NodeDisplayName = FText::Format(LOCTEXT("ConverterNodeDisplayName", "{0} To {1}"), FromTypeText, ToTypeText);
+				const FText NodeDisplayName = METASOUND_LOCTEXT_FORMAT("ConverterNodeDisplayName", "{0} To {1}", FromTypeText, ToTypeText);
 
 				FNodeClassMetadata Info;
 				Info.ClassName = { StandardNodes::Namespace, ClassName, "" };
@@ -227,15 +227,15 @@ namespace Metasound
 		{
 			if constexpr (std::is_same<FromType, float>::value && std::is_same<ToType, FAudioBuffer>::value)
 			{
-				return LOCTEXT("FloatToAudioConverterDescription", "Converts from float to audio buffer with each sample set to the given float value.");
+				return METASOUND_LOCTEXT("FloatToAudioConverterDescription", "Converts from float to audio buffer with each sample set to the given float value.");
 			}
 			else if constexpr (std::is_same<FromType, FAudioBuffer>::value && std::is_same<ToType, float>::value)
 			{
-				return LOCTEXT("AudioToFloatConverterDescription", "Converts from audio buffer to float by averaging sample values over the buffer. ");
+				return METASOUND_LOCTEXT("AudioToFloatConverterDescription", "Converts from audio buffer to float by averaging sample values over the buffer. ");
 			}
 			else
 			{
-				return FText::Format(LOCTEXT("ConverterNodeDesc", "Converts from {0} to {1}."), GetMetasoundDataTypeDisplayText<FromType>(), GetMetasoundDataTypeDisplayText<ToType>());
+				return METASOUND_LOCTEXT_FORMAT("ConverterNodeDesc", "Converts from {0} to {1}.", GetMetasoundDataTypeDisplayText<FromType>(), GetMetasoundDataTypeDisplayText<ToType>());
 			}
 		}
 
diff --git a/Engine/Plugins/Runtime/Metasound/Source/MetasoundStandardNodes/Private/MetasoundCrossfadeNode.cpp b/Engine/Plugins/Runtime/Metasound/Source/MetasoundStandardNodes/Private/MetasoundCrossfadeNode.cpp
--- a/Engine/Plugins/Runtime/Metasound/Source/MetasoundStandardNodes/Private/MetasoundCrossfadeNode.cpp
+++ b/Engine/Plugins/Runtime/Metasound/Source/MetasoundStandardNodes/Private/MetasoundCrossfadeNode.cpp
@@ -33,7 +33,7 @@ namespace Metasound
 
 		const FText GetInputCrossfadeValueDescription()
 		{
-			static const FText Desc = LOCTEXT("InputCrossfadeValueDesc", "Crossfade value to crossfade across inputs. Output will be the float value between adjacent whole number values.");
+			static const FText Desc = METASOUND_LOCTEXT("InputCrossfadeValueDesc", "Crossfade value to crossfade across inputs. Output will be the float value between adjacent whole number values.");
 			return Desc;
 		}
 
@@ -44,7 +44,7 @@ namespace Metasound
 
 		const FText GetInputDescription(uint32 InIndex)
 		{
-			return FText::Format(LOCTEXT("CrossfadeInputDesc", "Cross fade {0} input."), InIndex);
+			return METASOUND_LOCTEXT_FORMAT("CrossfadeInputDesc", "Cross fade {0} input.", InIndex);
 		}
 
 		const FVertexName& GetOutputName()
@@ -55,7 +55,7 @@ namespace Metasound
 
 		const FText& GetOutputDescription()
 		{
-			static const FText Desc = LOCTEXT("TriggerAccumulateOutputTriggerDesc", "Triggered when all input triggers have been triggered. Call Reset to reset the state or use \"Auto Reset\"");
+			static const FText Desc = METASOUND_LOCTEXT("TriggerAccumulateOutputTriggerDesc", "Triggered when all input triggers have been triggered. Call Reset to reset the state or use \"Auto Reset\"");
 			return Desc;
 		}
 	}
@@ -185,8 +185,8 @@ namespace Metasound
 			{
 				FName DataTypeName = GetMetasoundDataTypeName<ValueType>();
 				FName OperatorName = *FString::Printf(TEXT("Trigger Route (%s, %d)"), *DataTypeName.ToString(), NumInputs);
-				FText NodeDisplayName = FText::Format(LOCTEXT("CrossfadeDisplayNamePattern", "Crossfade ({0}, {1})"), GetMetasoundDataTypeDisplayText<ValueType>(), NumInputs);
-				FText NodeDescription = LOCTEXT("CrossfadeDescription", "Crossfades inputs to outputs.");
+				FText NodeDisplayName = METASOUND_LOCTEXT_FORMAT("CrossfadeDisplayNamePattern", "Crossfade ({0}, {1})", GetMetasoundDataTypeDisplayText<ValueType>(), NumInputs);
+				const FText NodeDescription = METASOUND_LOCTEXT("CrossfadeDescription", "Crossfades inputs to outputs.");
 				FVertexInterface NodeInterface = GetVertexInterface();
 
 				FNodeClassMetadata Metadata
diff --git a/Engine/Plugins/Runtime/Metasound/Source/MetasoundStandardNodes/Private/MetasoundDebugLogNode.cpp b/Engine/Plugins/Runtime/Metasound/Source/MetasoundStandardNodes/Private/MetasoundDebugLogNode.cpp
--- a/Engine/Plugins/Runtime/Metasound/Source/MetasoundStandardNodes/Private/MetasoundDebugLogNode.cpp
+++ b/Engine/Plugins/Runtime/Metasound/Source/MetasoundStandardNodes/Private/MetasoundDebugLogNode.cpp
@@ -76,9 +76,9 @@ namespace Metasound
 				static const FVertexInterface DefaultInterface(
 					FInputVertexInterface(
 
-						TInputDataVertexModel<FTrigger>(PrintLogVertexNames::GetInputTriggerName(), LOCTEXT("PrintLogTrigger", "Trigger to write the set value to the log.")),
-						TInputDataVertexModel<FString>(PrintLogVertexNames::GetLabelPrintLogName(), LOCTEXT("PrintLogLabel", "The label to attach to the value that will be logged")),
-						TInputDataVertexModel<PrintLogType>(PrintLogVertexNames::GetToLogPrintLogName(), LOCTEXT("PrintLogValueToLog", "The value to record to the log when triggered"))
+						TInputDataVertexModel<FTrigger>(PrintLogVertexNames::GetInputTriggerName(), METASOUND_LOCTEXT("PrintLogTrigger", "Trigger to write the set value to the log.")),
+						TInputDataVertexModel<FString>(PrintLogVertexNames::GetLabelPrintLogName(), METASOUND_LOCTEXT("PrintLogLabel", "The label to attach to the value that will be logged")),
+						TInputDataVertexModel<PrintLogType>(PrintLogVertexNames::GetToLogPrintLogName(), METASOUND_LOCTEXT("PrintLogValueToLog", "The value to record to the log when triggered"))
 					),
 					FOutputVertexInterface(
 					)
@@ -93,8 +93,8 @@ namespace Metasound
 				{
 					FName DataTypeName = GetMetasoundDataTypeName<PrintLogType>();
 					FName OperatorName = TEXT("Print Log");
-					FText NodeDisplayName = FText::Format(LOCTEXT("PrintLogDisplayNamePattern", "Print Log ({0})"), GetMetasoundDataTypeDisplayText<PrintLogType>());
-					FText NodeDescription = LOCTEXT("PrintLogOpDescription", "Used to record values to the log, on trigger");
+					FText NodeDisplayName = METASOUND_LOCTEXT_FORMAT("PrintLogDisplayNamePattern", "Print Log ({0})", GetMetasoundDataTypeDisplayText<PrintLogType>());
+					const FText NodeDescription = METASOUND_LOCTEXT("PrintLogOpDescription", "Used to record values to the log, on trigger");
 					FVertexInterface NodeInterface = GetDefaultInterface();
 
 					return MetasoundPrintLogNodePrivate::CreateNodeClassMetadata(DataTypeName, OperatorName, NodeDisplayName, NodeDescription, NodeInterface);
diff --git a/Engine/Plugins/Runtime/Metasound/Source/MetasoundStandardNodes/Private/MetasoundDecibelsToLinearGain.cpp b/Engine/Plugins/Runtime/Metasound/Source/MetasoundStandardNodes/Private/MetasoundDecibelsToLinearGain.cpp
--- a/Engine/Plugins/Runtime/Metasound/Source/MetasoundStandardNodes/Private/MetasoundDecibelsToLinearGain.cpp
+++ b/Engine/Plugins/Runtime/Metasound/Source/MetasoundStandardNodes/Private/MetasoundDecibelsToLinearGain.cpp
@@ -98,8 +98,8 @@ namespace Metasound
 		{
 			const FName DataTypeName = GetMetasoundDataTypeName<float>();
 			const FName OperatorName = TEXT("Decibels to Linear Gain");
-			const FText NodeDisplayName = LOCTEXT("Metasound_DecibelsToLinearGainName", "Decibels to Linear Gain");
-			const FText NodeDescription = LOCTEXT("Metasound_DecibelsToLinearGainDescription", "Converts a logarithmic (dB) gain value to a linear gain value.");
+			const FText NodeDisplayName = METASOUND_LOCTEXT("Metasound_DecibelsToLinearGainName", "Decibels to Linear Gain");
+			const FText NodeDescription = METASOUND_LOCTEXT("Metasound_DecibelsToLinearGainDescription", "Converts a logarithmic (dB) gain value to a linear gain value.");
 
 			FNodeClassMetadata Info;
 			Info.ClassName = { StandardNodes::Namespace, OperatorName, DataTypeName };
diff --git a/Engine/Plugins/Runtime/Metasound/Source/MetasoundStandardNodes/Private/MetasoundDelayNode.cpp b/Engine/Plugins/Runtime/Metasound/Source/MetasoundStandardNodes/Private/MetasoundDelayNode.cpp
--- a/Engine/Plugins/Runtime/Metasound/Source/MetasoundStandardNodes/Private/MetasoundDelayNode.cpp
+++ b/Engine/Plugins/Runtime/Metasound/Source/MetasoundStandardNodes/Private/MetasoundDelayNode.cpp
@@ -170,14 +170,14 @@ namespace Metasound
 	{
 		static const FVertexInterface Interface(
 			FInputVertexInterface(
-				TInputDataVertexModel<FAudioBuffer>(Delay::InParamNameAudioInput, LOCTEXT("DelayNode_AudioInputTooltip", "Audio input.")),
-				TInputDataVertexModel<FTime>(Delay::InParamNameDelayTime, LOCTEXT("DelayNode_DelayTimeTooltip", "The amount of time to delay the audio, in seconds."), 1.0f),
-				TInputDataVertexModel<float>(Delay::InParamNameDryLevel, LOCTEXT("DelayNode_DryLevelTooltip", "The dry level of the delay."), 0.0f),
-				TInputDataVertexModel<float>(Delay::InParamNameWetLevel, LOCTEXT("DelayNode_WetlevelTooltip", "The wet level of the delay."), 1.0f),
-				TInputDataVertexModel<float>(Delay::InParamNameFeedbackAmount, LOCTEXT("DelayNode_FeedbackTooltip", "Feedback amount."), 0.0f)
+				TInputDataVertexModel<FAudioBuffer>(Delay::InParamNameAudioInput, METASOUND_LOCTEXT("DelayNode_AudioInputTooltip", "Audio input.")),
+				TInputDataVertexModel<FTime>(Delay::InParamNameDelayTime, METASOUND_LOCTEXT("DelayNode_DelayTimeTooltip", "The amount of time to delay the audio, in seconds."), 1.0f),
+				TInputDataVertexModel<float>(Delay::InParamNameDryLevel, METASOUND_LOCTEXT("DelayNode_DryLevelTooltip", "The dry level of the delay."), 0.0f),
+				TInputDataVertexModel<float>(Delay::InParamNameWetLevel, METASOUND_LOCTEXT("DelayNode_WetlevelTooltip", "The wet level of the delay."), 1.0f),
+				TInputDataVertexModel<float>(Delay::InParamNameFeedbackAmount, METASOUND_LOCTEXT("DelayNode_FeedbackTooltip", "Feedback amount."), 0.0f)
 			),
 			FOutputVertexInterface(
-				TOutputDataVertexModel<FAudioBuffer>(Delay::OutParamNameAudio, LOCTEXT("DelayNode_DelayOutputTooltip", "Audio output."))
+				TOutputDataVertexModel<FAudioBuffer>(Delay::OutParamNameAudio, METASOUND_LOCTEXT("DelayNode_DelayOutputTooltip", "Audio output."))
 			)
 		);
 
@@ -192,8 +192,8 @@ namespace Metasound
 			Info.ClassName = { StandardNodes::Namespace, "Delay", StandardNodes::AudioVariant };
 			Info.MajorVersion = 1;
 			Info.MinorVersion = 1;
-			Info.DisplayName = LOCTEXT("DelayNode_DisplayName", "Delay");
-			Info.Description = LOCTEXT("DelayNode_Description", "Delays an audio buffer by the specified amount.");
+			Info.DisplayName = METASOUND_LOCTEXT("DelayNode_DisplayName", "Delay");
+			Info.Description = METASOUND_LOCTEXT("DelayNode_Description", "Delays an audio buffer by the specified amount.");
 			Info.Author = PluginAuthor;
 			Info.PromptIfMissing = PluginNodeMissingPrompt;
 			Info.DefaultInterface = GetVertexInterface();
diff --git a/Engine/Plugins/Runtime/Metasound/Source/MetasoundStandardNodes/Private/MetasoundDiffuserNode.cpp b/Engine/Plugins/Runtime/Metasound/Source/MetasoundStandardNodes/Private/MetasoundDiffuserNode.cpp
--- a/Engine/Plugins/Runtime/Metasound/Source/MetasoundStandardNodes/Private/MetasoundDiffuserNode.cpp
+++ b/Engine/Plugins/Runtime/Metasound/Source/MetasoundStandardNodes/Private/MetasoundDiffuserNode.cpp
@@ -42,8 +42,8 @@ namespace Metasound
 				Info.ClassName = { StandardNodes::Namespace, TEXT("Diffuser"), StandardNodes::AudioVariant };
 				Info.MajorVersion = 1;
 				Info.MinorVersion = 0;
-				Info.DisplayName = LOCTEXT("Metasound_DiffuserDisplayName", "Diffuser");
-				Info.Description = LOCTEXT("Metasound_DiffuserNodeDescription", "Applies diffusion to incoming audio.");
+				Info.DisplayName = METASOUND_LOCTEXT("Metasound_DiffuserDisplayName", "Diffuser");
+				Info.Description = METASOUND_LOCTEXT("Metasound_DiffuserNodeDescription", "Applies diffusion to incoming audio.");
 				Info.Author = PluginAuthor;
 				Info.PromptIfMissing = PluginNodeMissingPrompt;
 				Info.DefaultInterface = GetVertexInterface();
diff --git a/Engine/Plugins/Runtime/Metasound/Source/MetasoundStandardNodes/Private/MetasoundDynamicFilterNode.cpp b/Engine/Plugins/Runtime/Metasound/Source/MetasoundStandardNodes/Private/MetasoundDynamicFilterNode.cpp
--- a/Engine/Plugins/Runtime/Metasound/Source/MetasoundStandardNodes/Private/MetasoundDynamicFilterNode.cpp
+++ b/Engine/Plugins/Runtime/Metasound/Source/MetasoundStandardNodes/Private/MetasoundDynamicFilterNode.cpp
@@ -48,9 +48,9 @@ namespace Metasound
 	);
 
 	DEFINE_METASOUND_ENUM_BEGIN(Audio::EDynamicFilterType, FEnumEDynamicFilterType, "DynamicFilterType")
-		DEFINE_METASOUND_ENUM_ENTRY(Audio::EDynamicFilterType::Bell, LOCTEXT("BellDescription", "Bell"), LOCTEXT("BellTT", "Bell Filter")),
-		DEFINE_METASOUND_ENUM_ENTRY(Audio::EDynamicFilterType::LowShelf, LOCTEXT("LowShelfDescription", "Low Shelf"), LOCTEXT("LowShelfTT", "Low Shelf Filter")),
-		DEFINE_METASOUND_ENUM_ENTRY(Audio::EDynamicFilterType::HighShelf, LOCTEXT("HighShelfDescription", "High Shelf"), LOCTEXT("HighShelfTT", "High Shelf Filter")),
+		DEFINE_METASOUND_ENUM_ENTRY(Audio::EDynamicFilterType::Bell, "BellDescription", "Bell", "BellTT", "Bell Filter"),
+		DEFINE_METASOUND_ENUM_ENTRY(Audio::EDynamicFilterType::LowShelf, "LowShelfDescription", "Low Shelf", "LowShelfTT", "Low Shelf Filter"),
+		DEFINE_METASOUND_ENUM_ENTRY(Audio::EDynamicFilterType::HighShelf, "HighShelfDescription", "High Shelf", "HighShelfTT", "High Shelf Filter"),
 	DEFINE_METASOUND_ENUM_END()
 
 	class FDynamicFilterOperator : public TExecutableOperator<FDynamicFilterOperator>
@@ -106,8 +106,8 @@ namespace Metasound
 					FNodeClassName { StandardNodes::Namespace, "DynamicFilter", StandardNodes::AudioVariant },
 					1, // Major Version
 					0, // Minor Version
-					LOCTEXT("DynamicFilterDisplayName", "Dynamic Filter"),
-					LOCTEXT("DynamicFilterDesc", "Filters a band of audio based on the strength of the input signal."),
+					METASOUND_LOCTEXT("DynamicFilterDisplayName", "Dynamic Filter"),
+					METASOUND_LOCTEXT("DynamicFilterDesc", "Filters a band of audio based on the strength of the input signal."),
 					PluginAuthor,
 					PluginNodeMissingPrompt,
 					NodeInterface,
diff --git a/Engine/Plugins/Runtime/Metasound/Source/MetasoundStandardNodes/Private/MetasoundEnvelopeFollowerNode.cpp b/Engine/Plugins/Runtime/Metasound/Source/MetasoundStandardNodes/Private/MetasoundEnvelopeFollowerNode.cpp
--- a/Engine/Plugins/Runtime/Metasound/Source/MetasoundStandardNodes/Private/MetasoundEnvelopeFollowerNode.cpp
+++ b/Engine/Plugins/Runtime/Metasound/Source/MetasoundStandardNodes/Private/MetasoundEnvelopeFollowerNode.cpp
@@ -172,13 +172,13 @@ namespace Metasound
 	{
 		static const FVertexInterface Interface(
 			FInputVertexInterface(
-				TInputDataVertexModel<FAudioBuffer>(EnvelopeFollower::InParamNameAudioInput, LOCTEXT("AudioInputToolTT", "Audio input.")),
-				TInputDataVertexModel<FTime>(EnvelopeFollower::InParamNameAttackTime, LOCTEXT("AttackTimeTT", "The attack time of the envelope follower."), 0.01f),
-				TInputDataVertexModel<FTime>(EnvelopeFollower::InParamNameReleaseTime, LOCTEXT("ReleaseTimeTT", "The release time of the envelope follower."), 0.1f),
-				TInputDataVertexModel<FEnumEnvelopePeakMode>(EnvelopeFollower::InParamNameFollowMode, LOCTEXT("FollowModeTT", "The following-method of the envelope follower."))
+				TInputDataVertexModel<FAudioBuffer>(EnvelopeFollower::InParamNameAudioInput, METASOUND_LOCTEXT("AudioInputToolTT", "Audio input.")),
+				TInputDataVertexModel<FTime>(EnvelopeFollower::InParamNameAttackTime, METASOUND_LOCTEXT("AttackTimeTT", "The attack time of the envelope follower."), 0.01f),
+				TInputDataVertexModel<FTime>(EnvelopeFollower::InParamNameReleaseTime, METASOUND_LOCTEXT("ReleaseTimeTT", "The release time of the envelope follower."), 0.1f),
+				TInputDataVertexModel<FEnumEnvelopePeakMode>(EnvelopeFollower::InParamNameFollowMode, METASOUND_LOCTEXT("FollowModeTT", "The following-method of the envelope follower."))
 			),
 			FOutputVertexInterface(
-				TOutputDataVertexModel<float>(EnvelopeFollower::OutParamNameEnvelope, LOCTEXT("EnvelopeFollowerOutputTT", "The output envelope value of the audio signal.")),
+				TOutputDataVertexModel<float>(EnvelopeFollower::OutParamNameEnvelope, METASOUND_LOCTEXT("EnvelopeFollowerOutputTT", "The output envelope value of the audio signal.")),
 				TOutputDataVertexModel<FAudioBuffer>(METASOUND_GET_PARAM_NAME_AND_TT(OutputAudioEnvelope))
 			)
 		);
@@ -194,8 +194,8 @@ namespace Metasound
 			Info.ClassName = { StandardNodes::Namespace, TEXT("Envelope Follower"), TEXT("") };
 			Info.MajorVersion = 1;
 			Info.MinorVersion = 3;
-			Info.DisplayName = LOCTEXT("Metasound_EnvelopeFollowerDisplayName", "Envelope Follower");
-			Info.Description = LOCTEXT("Metasound_EnvelopeFollowerDescription", "Outputs an envelope from an input audio signal.");
+			Info.DisplayName = METASOUND_LOCTEXT("Metasound_EnvelopeFollowerDisplayName", "Envelope Follower");
+			Info.Description = METASOUND_LOCTEXT("Metasound_EnvelopeFollowerDescription", "Outputs an envelope from an input audio signal.");
 			Info.Author = PluginAuthor;
 			Info.PromptIfMissing = PluginNodeMissingPrompt;
 			Info.DefaultInterface = GetVertexInterface();
diff --git a/Engine/Plugins/Runtime/Metasound/Source/MetasoundStandardNodes/Private/MetasoundEnvelopeFollowerTypes.cpp b/Engine/Plugins/Runtime/Metasound/Source/MetasoundStandardNodes/Private/MetasoundEnvelopeFollowerTypes.cpp
--- a/Engine/Plugins/Runtime/Metasound/Source/MetasoundStandardNodes/Private/MetasoundEnvelopeFollowerTypes.cpp
+++ b/Engine/Plugins/Runtime/Metasound/Source/MetasoundStandardNodes/Private/MetasoundEnvelopeFollowerTypes.cpp
@@ -8,9 +8,9 @@
 namespace Metasound
 {
 	DEFINE_METASOUND_ENUM_BEGIN(EEnvelopePeakMode, FEnumEnvelopePeakMode, "EnvelopePeakMode")
-		DEFINE_METASOUND_ENUM_ENTRY(EEnvelopePeakMode::MeanSquared, LOCTEXT("EnvelopePeakModeMSDescription", "MS"), LOCTEXT("EnvelopePeakModeMSDescriptionTT", "Envelope follows a running Mean Squared of the audio signal.")),
-		DEFINE_METASOUND_ENUM_ENTRY(EEnvelopePeakMode::RootMeanSquared, LOCTEXT("EnvelopePeakModeRMSDescription", "RMS"), LOCTEXT("EnvelopePeakModeRMSDescriptionTT", "Envelope follows a running Root Mean Squared of the audio signal.")),
-		DEFINE_METASOUND_ENUM_ENTRY(EEnvelopePeakMode::Peak, LOCTEXT("EnvelopePeakModePeakDescription", "Peak"), LOCTEXT("EnvelopePeakModePeakDescriptionTT", "Envelope follows the peaks in the audio signal.")),
+		DEFINE_METASOUND_ENUM_ENTRY(EEnvelopePeakMode::MeanSquared, "EnvelopePeakModeMSDescription", "MS", "EnvelopePeakModeMSDescriptionTT", "Envelope follows a running Mean Squared of the audio signal."),
+		DEFINE_METASOUND_ENUM_ENTRY(EEnvelopePeakMode::RootMeanSquared, "EnvelopePeakModeRMSDescription", "RMS", "EnvelopePeakModeRMSDescriptionTT", "Envelope follows a running Root Mean Squared of the audio signal."),
+		DEFINE_METASOUND_ENUM_ENTRY(EEnvelopePeakMode::Peak, "EnvelopePeakModePeakDescription", "Peak", "EnvelopePeakModePeakDescriptionTT", "Envelope follows the peaks in the audio signal."),
 		DEFINE_METASOUND_ENUM_END()
 }
 
diff --git a/Engine/Plugins/Runtime/Metasound/Source/MetasoundStandardNodes/Private/MetasoundFlangerNode.cpp b/Engine/Plugins/Runtime/Metasound/Source/MetasoundStandardNodes/Private/MetasoundFlangerNode.cpp
--- a/Engine/Plugins/Runtime/Metasound/Source/MetasoundStandardNodes/Private/MetasoundFlangerNode.cpp
+++ b/Engine/Plugins/Runtime/Metasound/Source/MetasoundStandardNodes/Private/MetasoundFlangerNode.cpp
@@ -35,8 +35,8 @@ namespace Metasound
 			auto InitNodeInfo = []() -> FNodeClassMetadata
 			{
 				const FName OperatorName = TEXT("Flanger");
-				const FText NodeDisplayName = LOCTEXT("Metasound_FlangerNodeDisplayName", "Flanger");
-				const FText NodeDescription = LOCTEXT("Metasound_FlangerNodeDescription", "Applies a flanger effect to input audio.");
+				const FText NodeDisplayName = METASOUND_LOCTEXT("Metasound_FlangerNodeDisplayName", "Flanger");
+				const FText NodeDescription = METASOUND_LOCTEXT("Metasound_FlangerNodeDescription", "Applies a flanger effect to input audio.");
 
 				FNodeClassMetadata Info;
 				Info.ClassName = { StandardNodes::Namespace, OperatorName, TEXT("") };
diff --git a/Engine/Plugins/Runtime/Metasound/Source/MetasoundStandardNodes/Private/MetasoundFreqMultiplierToSemitoneNode.cpp b/Engine/Plugins/Runtime/Metasound/Source/MetasoundStandardNodes/Private/MetasoundFreqMultiplierToSemitoneNode.cpp
--- a/Engine/Plugins/Runtime/Metasound/Source/MetasoundStandardNodes/Private/MetasoundFreqMultiplierToSemitoneNode.cpp
+++ b/Engine/Plugins/Runtime/Metasound/Source/MetasoundStandardNodes/Private/MetasoundFreqMultiplierToSemitoneNode.cpp
@@ -107,8 +107,8 @@ namespace Metasound
 		{
 			const FName DataTypeName = GetMetasoundDataTypeName<float>();
 			const FName OperatorName = TEXT("Frequency Multiplier to Semitone");
-			const FText NodeDisplayName = LOCTEXT("Metasound_FrequencyMultiplierToSemitoneName", "Frequency Multiplier to Semitone");
-			const FText NodeDescription = LOCTEXT("Metasound_FrequencyMultiplierToSemitoneDescription", "Converts a frequency multiplier to the corresponding difference in semitones.");
+			const FText NodeDisplayName = METASOUND_LOCTEXT("Metasound_FrequencyMultiplierToSemitoneName", "Frequency Multiplier to Semitone");
+			const FText NodeDescription = METASOUND_LOCTEXT("Metasound_FrequencyMultiplierToSemitoneDescription", "Converts a frequency multiplier to the corresponding difference in semitones.");
 
 			FNodeClassMetadata Info;
 			Info.ClassName = { StandardNodes::Namespace, OperatorName, DataTypeName };
diff --git a/Engine/Plugins/Runtime/Metasound/Source/MetasoundStandardNodes/Private/MetasoundFreqToMIDINode.cpp b/Engine/Plugins/Runtime/Metasound/Source/MetasoundStandardNodes/Private/MetasoundFreqToMIDINode.cpp
--- a/Engine/Plugins/Runtime/Metasound/Source/MetasoundStandardNodes/Private/MetasoundFreqToMIDINode.cpp
+++ b/Engine/Plugins/Runtime/Metasound/Source/MetasoundStandardNodes/Private/MetasoundFreqToMIDINode.cpp
@@ -108,8 +108,8 @@ namespace Metasound
 		{
 			const FName DataTypeName = GetMetasoundDataTypeName<float>();
 			const FName OperatorName = TEXT("Frequency to MIDI");
-			const FText NodeDisplayName = LOCTEXT("Metasound_FreqToMidiNodeName", "Frequency To MIDI");
-			const FText NodeDescription = LOCTEXT("Metasound_FreqToMidiNodeDescription", "Converts a frequency (Hz) value to a Midi Note value.");
+			const FText NodeDisplayName = METASOUND_LOCTEXT("Metasound_FreqToMidiNodeName", "Frequency To MIDI");
+			const FText NodeDescription = METASOUND_LOCTEXT("Metasound_FreqToMidiNodeDescription", "Converts a frequency (Hz) value to a Midi Note value.");
 
 			FNodeClassMetadata Info;
 			Info.ClassName = { StandardNodes::Namespace, OperatorName, DataTypeName };
diff --git a/Engine/Plugins/Runtime/Metasound/Source/MetasoundStandardNodes/Private/MetasoundITDPannerNode.cpp b/Engine/Plugins/Runtime/Metasound/Source/MetasoundStandardNodes/Private/MetasoundITDPannerNode.cpp
--- a/Engine/Plugins/Runtime/Metasound/Source/MetasoundStandardNodes/Private/MetasoundITDPannerNode.cpp
+++ b/Engine/Plugins/Runtime/Metasound/Source/MetasoundStandardNodes/Private/MetasoundITDPannerNode.cpp
@@ -30,7 +30,7 @@ namespace Metasound
 
 		const FText& GetInputAudioDescription()
 		{
-			static FText Desc = LOCTEXT("ITDPannerNodeInDesc", "The input audio to spatialize.");
+			static FText Desc = METASOUND_LOCTEXT("ITDPannerNodeInDesc", "The input audio to spatialize.");
 			return Desc;
 		}
 
@@ -42,7 +42,7 @@ namespace Metasound
 
 		const FText& GetInputPanAngleDescription()
 		{
-			static FText Desc = LOCTEXT("ITDPannerNodePanAngleDesc", "The sound source angle in degrees. 90 degrees is in front, 0 degrees is to the right, 270 degrees is behind, 180 degrees is to the left.");
+			static FText Desc = METASOUND_LOCTEXT("ITDPannerNodePanAngleDesc", "The sound source angle in degrees. 90 degrees is in front, 0 degrees is to the right, 270 degrees is behind, 180 degrees is to the left.");
 			return Desc;
 		}
 
@@ -54,7 +54,7 @@ namespace Metasound
 
 		const FText& GetInputDistanceFactorDescription()
 		{
-			static FText Desc = LOCTEXT("ITDPannerNodeDistanceFactorDesc", "The normalized distance factor (0.0 to 1.0) to use for ILD (Inter-aural level difference) calculations. 0.0 is near, 1.0 is far. The further away something is the less there is a difference in levels (gain) between the ears.");
+			static FText Desc = METASOUND_LOCTEXT("ITDPannerNodeDistanceFactorDesc", "The normalized distance factor (0.0 to 1.0) to use for ILD (Inter-aural level difference) calculations. 0.0 is near, 1.0 is far. The further away something is the less there is a difference in levels (gain) between the ears.");
 			return Desc;
 		}
 
@@ -66,7 +66,7 @@ namespace Metasound
 
 		const FText& GetInputHeadWidthDescription()
 		{
-			static FText Desc = LOCTEXT("ITDPannerNodeHeadWidthDesc", "The width of the listener head to use for ITD calculations in centimeters.");
+			static FText Desc = METASOUND_LOCTEXT("ITDPannerNodeHeadWidthDesc", "The width of the listener head to use for ITD calculations in centimeters.");
 			return Desc;
 		}
 
@@ -78,7 +78,7 @@ namespace Metasound
 
 		const FText& GetOutputAudioLeftDescription()
 		{
-			static FText Desc = LOCTEXT("ITDPannerNodeOutputLeftDescription", "Left channel audio output.");
+			static FText Desc = METASOUND_LOCTEXT("ITDPannerNodeOutputLeftDescription", "Left channel audio output.");
 			return Desc;
 		}
 
@@ -90,7 +90,7 @@ namespace Metasound
 
 		const FText& GetOutputAudioRightDescription()
 		{
-			static FText Desc = LOCTEXT("ITDPannerNodeOutputRightDescription", "Right channel audio output.");
+			static FText Desc = METASOUND_LOCTEXT("ITDPannerNodeOutputRightDescription", "Right channel audio output.");
 			return Desc;
 		}
 	}
@@ -314,8 +314,8 @@ namespace Metasound
 			Info.ClassName = { StandardNodes::Namespace, TEXT("ITD Panner"), TEXT("") };
 			Info.MajorVersion = 1;
 			Info.MinorVersion = 0;
-			Info.DisplayName = LOCTEXT("Metasound_ITDPannerDisplayName", "ITD Panner");
-			Info.Description = LOCTEXT("Metasound_ITDPannerNodeDescription", "Pans an input audio signal using an inter-aural time delay method.");
+			Info.DisplayName = METASOUND_LOCTEXT("Metasound_ITDPannerDisplayName", "ITD Panner");
+			Info.Description = METASOUND_LOCTEXT("Metasound_ITDPannerNodeDescription", "Pans an input audio signal using an inter-aural time delay method.");
 			Info.Author = PluginAuthor;
 			Info.PromptIfMissing = PluginNodeMissingPrompt;
 			Info.DefaultInterface = GetVertexInterface();
diff --git a/Engine/Plugins/Runtime/Metasound/Source/MetasoundStandardNodes/Private/MetasoundInterpToNode.cpp b/Engine/Plugins/Runtime/Metasound/Source/MetasoundStandardNodes/Private/MetasoundInterpToNode.cpp
--- a/Engine/Plugins/Runtime/Metasound/Source/MetasoundStandardNodes/Private/MetasoundInterpToNode.cpp
+++ b/Engine/Plugins/Runtime/Metasound/Source/MetasoundStandardNodes/Private/MetasoundInterpToNode.cpp
@@ -126,11 +126,11 @@ namespace Metasound
 	{
 		static const FVertexInterface Interface(
 			FInputVertexInterface(
-				TInputDataVertexModel<FTime>(InterpTo::InParamNameInterpTime, LOCTEXT("InterpTimeTooltip", "The time to interpolate from the current value to the target value."), 0.1f),
-				TInputDataVertexModel<float>(InterpTo::InParamNameTarget, LOCTEXT("TargetValueTooltip", "Target value."), 1.0f)
+				TInputDataVertexModel<FTime>(InterpTo::InParamNameInterpTime, METASOUND_LOCTEXT("InterpTimeTooltip", "The time to interpolate from the current value to the target value."), 0.1f),
+				TInputDataVertexModel<float>(InterpTo::InParamNameTarget, METASOUND_LOCTEXT("TargetValueTooltip", "Target value."), 1.0f)
 			),
 			FOutputVertexInterface(
-				TOutputDataVertexModel<float>(InterpTo::OutParamNameValue, LOCTEXT("ValueTooltip", "The current value."))
+				TOutputDataVertexModel<float>(InterpTo::OutParamNameValue, METASOUND_LOCTEXT("ValueTooltip", "The current value."))
 			)
 		);
 
@@ -145,8 +145,8 @@ namespace Metasound
 			Info.ClassName = { StandardNodes::Namespace, TEXT("InterpTo"), StandardNodes::AudioVariant };
 			Info.MajorVersion = 1;
 			Info.MinorVersion = 0;
-			Info.DisplayName = LOCTEXT("Metasound_InterpDisplayName", "InterpTo");
-			Info.Description = LOCTEXT("Metasound_InterpNodeDescription", "Interpolates between the current value and a target value over the specified time.");
+			Info.DisplayName = METASOUND_LOCTEXT("Metasound_InterpDisplayName", "InterpTo");
+			Info.Description = METASOUND_LOCTEXT("Metasound_InterpNodeDescription", "Interpolates between the current value and a target value over the specified time.");
 			Info.Author = PluginAuthor;
 			Info.PromptIfMissing = PluginNodeMissingPrompt;
 			Info.DefaultInterface = GetVertexInterface();
diff --git a/Engine/Plugins/Runtime/Metasound/Source/MetasoundStandardNodes/Private/MetasoundLimiterNode.cpp b/Engine/Plugins/Runtime/Metasound/Source/MetasoundStandardNodes/Private/MetasoundLimiterNode.cpp
--- a/Engine/Plugins/Runtime/Metasound/Source/MetasoundStandardNodes/Private/MetasoundLimiterNode.cpp
+++ b/Engine/Plugins/Runtime/Metasound/Source/MetasoundStandardNodes/Private/MetasoundLimiterNode.cpp
@@ -39,8 +39,8 @@ namespace Metasound
 	FEnumKneeMode, FEnumKneeModeInfo, FKneeModeReadRef, FEnumKneeModeWriteRef);
 
 	DEFINE_METASOUND_ENUM_BEGIN(EKneeMode, FEnumKneeMode, "KneeMode")
-		DEFINE_METASOUND_ENUM_ENTRY(EKneeMode::Hard, LOCTEXT("KneeModeHardDescription", "Hard"), LOCTEXT("KneeModeHardDescriptionTT", "Only audio strictly above the threshold is affected by the limiter.")),
-		DEFINE_METASOUND_ENUM_ENTRY(EKneeMode::Soft, LOCTEXT("KneeModeSoftDescription", "Soft"), LOCTEXT("KneeModeSoftDescriptionTT", "Limiter activates more smoothly near the threshold.")),
+		DEFINE_METASOUND_ENUM_ENTRY(EKneeMode::Hard, "KneeModeHardDescription", "Hard", "KneeModeHardDescriptionTT", "Only audio strictly above the threshold is affected by the limiter."),
+		DEFINE_METASOUND_ENUM_ENTRY(EKneeMode::Soft, "KneeModeSoftDescription", "Soft", "KneeModeSoftDescriptionTT", "Limiter activates more smoothly near the threshold."),
 		DEFINE_METASOUND_ENUM_END()
 
 	// Operator Class
@@ -100,8 +100,8 @@ namespace Metasound
 					FNodeClassName{StandardNodes::Namespace, TEXT("Limiter"), StandardNodes::AudioVariant},
 					1, // Major Version
 					0, // Minor Version
-					LOCTEXT("LimiterDisplayName", "Limiter"),
-					LOCTEXT("LimiterDesc", "Prevents a signal from going above a given threshold."),
+					METASOUND_LOCTEXT("LimiterDisplayName", "Limiter"),
+					METASOUND_LOCTEXT("LimiterDesc", "Prevents a signal from going above a given threshold."),
 					PluginAuthor,
 					PluginNodeMissingPrompt,
 					NodeInterface,
diff --git a/Engine/Plugins/Runtime/Metasound/Source/MetasoundStandardNodes/Private/MetasoundLinearGainToDecibels.cpp b/Engine/Plugins/Runtime/Metasound/Source/MetasoundStandardNodes/Private/MetasoundLinearGainToDecibels.cpp
--- a/Engine/Plugins/Runtime/Metasound/Source/MetasoundStandardNodes/Private/MetasoundLinearGainToDecibels.cpp
+++ b/Engine/Plugins/Runtime/Metasound/Source/MetasoundStandardNodes/Private/MetasoundLinearGainToDecibels.cpp
@@ -103,8 +103,8 @@ namespace Metasound
 		{
 			const FName DataTypeName = GetMetasoundDataTypeName<float>();
 			const FName OperatorName = TEXT("Linear Gain to Decibels");
-			const FText NodeDisplayName = LOCTEXT("Metasound_LinearGainToDecibelsName", "Linear Gain to Decibels");
-			const FText NodeDescription = LOCTEXT("Metasound_LinearGainToDecibelsDescription", "Converts a linear gain value to a logarithmic (dB) gain value.");
+			const FText NodeDisplayName = METASOUND_LOCTEXT("Metasound_LinearGainToDecibelsName", "Linear Gain to Decibels");
+			const FText NodeDescription = METASOUND_LOCTEXT("Metasound_LinearGainToDecibelsDescription", "Converts a linear gain value to a logarithmic (dB) gain value.");
 
 			FNodeClassMetadata Info;
 			Info.ClassName = { StandardNodes::Namespace, OperatorName, DataTypeName };
diff --git a/Engine/Plugins/Runtime/Metasound/Source/MetasoundStandardNodes/Private/MetasoundLinearToLogFrequencyNode.cpp b/Engine/Plugins/Runtime/Metasound/Source/MetasoundStandardNodes/Private/MetasoundLinearToLogFrequencyNode.cpp
--- a/Engine/Plugins/Runtime/Metasound/Source/MetasoundStandardNodes/Private/MetasoundLinearToLogFrequencyNode.cpp
+++ b/Engine/Plugins/Runtime/Metasound/Source/MetasoundStandardNodes/Private/MetasoundLinearToLogFrequencyNode.cpp
@@ -33,8 +33,8 @@ namespace Metasound
 			auto InitNodeInfo = []() -> FNodeClassMetadata
 			{
 				FName OperatorName = TEXT("Linear To Log Frequency");
-				FText NodeDisplayName = LOCTEXT("LinearToLogFrequencyName", "Linear To Log Frequency");
-				FText NodeDescription = LOCTEXT("LinearToLogFrequencyDescription", "Converts a linear space input value to log-frequency space.");
+				FText NodeDisplayName = METASOUND_LOCTEXT("LinearToLogFrequencyName", "Linear To Log Frequency");
+				const FText NodeDescription = METASOUND_LOCTEXT("LinearToLogFrequencyDescription", "Converts a linear space input value to log-frequency space.");
 
 				FNodeClassMetadata Info;
 				Info.ClassName = { StandardNodes::Namespace, OperatorName, TEXT("") };
diff --git a/Engine/Plugins/Runtime/Metasound/Source/MetasoundStandardNodes/Private/MetasoundMIDIToFreqNode.cpp b/Engine/Plugins/Runtime/Metasound/Source/MetasoundStandardNodes/Private/MetasoundMIDIToFreqNode.cpp
--- a/Engine/Plugins/Runtime/Metasound/Source/MetasoundStandardNodes/Private/MetasoundMIDIToFreqNode.cpp
+++ b/Engine/Plugins/Runtime/Metasound/Source/MetasoundStandardNodes/Private/MetasoundMIDIToFreqNode.cpp
@@ -150,8 +150,8 @@ namespace Metasound
 		{
 			const FName DataTypeName = GetMetasoundDataTypeName<ValueType>();
 			const FName OperatorName = TEXT("MIDI To Frequency");
-			const FText NodeDisplayName = FText::Format(LOCTEXT("RandomGetArrayOpDisplayNamePattern", "MIDI To Frequency ({0})"), GetMetasoundDataTypeDisplayText<ValueType>());
-			const FText NodeDescription = LOCTEXT("Metasound_MidiToFreqNodeDescription", "Converts a Midi note value to a frequency (hz) value.");
+			const FText NodeDisplayName = METASOUND_LOCTEXT_FORMAT("RandomGetArrayOpDisplayNamePattern", "MIDI To Frequency ({0})", GetMetasoundDataTypeDisplayText<ValueType>());
+			const FText NodeDescription = METASOUND_LOCTEXT("Metasound_MidiToFreqNodeDescription", "Converts a Midi note value to a frequency (hz) value.");
 
 			FNodeClassMetadata Info;
 			Info.ClassName = { StandardNodes::Namespace, OperatorName, DataTypeName };
diff --git a/Engine/Plugins/Runtime/Metasound/Source/MetasoundStandardNodes/Private/MetasoundMapRangeNode.cpp b/Engine/Plugins/Runtime/Metasound/Source/MetasoundStandardNodes/Private/MetasoundMapRangeNode.cpp
--- a/Engine/Plugins/Runtime/Metasound/Source/MetasoundStandardNodes/Private/MetasoundMapRangeNode.cpp
+++ b/Engine/Plugins/Runtime/Metasound/Source/MetasoundStandardNodes/Private/MetasoundMapRangeNode.cpp
@@ -42,7 +42,7 @@ namespace Metasound
 
 			static FText GetNodeName()
 			{
-				return LOCTEXT("MapRange_Int32Name", "Map Range (Int32)");
+				return METASOUND_LOCTEXT("MapRange_Int32Name", "Map Range (Int32)");
 			}
 
 			static FVertexInterface GetVertexInterface()
@@ -126,7 +126,7 @@ namespace Metasound
 
 			static FText GetNodeName()
 			{
-				return LOCTEXT("MapRange_FloatName", "Map Range (Float)");
+				return METASOUND_LOCTEXT("MapRange_FloatName", "Map Range (Float)");
 			}
 
 			static FVertexInterface GetVertexInterface()
@@ -210,7 +210,7 @@ namespace Metasound
 
 			static FText GetNodeName()
 			{
-				return LOCTEXT("MapRange_AudioName", "Map Range (Audio)");
+				return METASOUND_LOCTEXT("MapRange_AudioName", "Map Range (Audio)");
 			}
 
 			static FVertexInterface GetVertexInterface()
@@ -301,7 +301,7 @@ namespace Metasound
 			auto CreateNodeClassMetadata = []() -> FNodeClassMetadata
 			{
 				FName OperatorName = TEXT("MapRange");
-				FText NodeDescription = LOCTEXT("MapRangeDescription", "Maps an input value in the given input range to the given output range.");
+				const FText NodeDescription = METASOUND_LOCTEXT("MapRangeDescription", "Maps an input value in the given input range to the given output range.");
 				FVertexInterface NodeInterface = FMappingClass::GetVertexInterface();
 
 				FNodeClassMetadata Metadata
diff --git a/Engine/Plugins/Runtime/Metasound/Source/MetasoundStandardNodes/Private/MetasoundMathNodes.cpp b/Engine/Plugins/Runtime/Metasound/Source/MetasoundStandardNodes/Private/MetasoundMathNodes.cpp
--- a/Engine/Plugins/Runtime/Metasound/Source/MetasoundStandardNodes/Private/MetasoundMathNodes.cpp
+++ b/Engine/Plugins/Runtime/Metasound/Source/MetasoundStandardNodes/Private/MetasoundMathNodes.cpp
@@ -56,25 +56,25 @@ namespace Metasound
 
 		static const FVertexName AdditionalOperandsName = TEXT("AdditionalOperands");
 
-		static const TArray<FText> AddKeywords = { LOCTEXT("AddMathKeyword", "+") };
-		static const TArray<FText> SubtractKeywords = { LOCTEXT("SubtractMathKeyword", "-") };
-		static const TArray<FText> MultiplyKeywords = { LOCTEXT("MultiplyMathKeyword", "*") };
-		static const TArray<FText> DivideKeywords = { LOCTEXT("DivideMathKeyword", "/") };
-		static const TArray<FText> PowerKeywords = { LOCTEXT("PowerMathKeyword", "^") };
-		static const TArray<FText> ModuloKeywords = { LOCTEXT("ModuloMathKeyword", "%") };
+		static const TArray<FText> AddKeywords = { METASOUND_LOCTEXT("AddMathKeyword", "+") };
+		static const TArray<FText> SubtractKeywords = { METASOUND_LOCTEXT("SubtractMathKeyword", "-") };
+		static const TArray<FText> MultiplyKeywords = { METASOUND_LOCTEXT("MultiplyMathKeyword", "*") };
+		static const TArray<FText> DivideKeywords = { METASOUND_LOCTEXT("DivideMathKeyword", "/") };
+		static const TArray<FText> PowerKeywords = { METASOUND_LOCTEXT("PowerMathKeyword", "^") };
+		static const TArray<FText> ModuloKeywords = { METASOUND_LOCTEXT("ModuloMathKeyword", "%") };
 
 		template<typename DataType>
 		const FText GetAddDisplayName()
 		{
-			static const FText DisplayName = FText::Format(LOCTEXT("AddNodeDisplayNamePattern1", "Add ({0})"), GetMetasoundDataTypeDisplayText<DataType>());
+			static const FText DisplayName = METASOUND_LOCTEXT_FORMAT("AddNodeDisplayNamePattern1", "Add ({0})", GetMetasoundDataTypeDisplayText<DataType>());
 			return DisplayName;
 		}
 
 		template<typename DataType, typename OperandDataType>
 		const FText GetAddDisplayName()
 		{
-			static const FText DisplayName = FText::Format(
-				LOCTEXT("AddNodeDisplayNamePattern2", "Add ({0} to {1})"),
+			static const FText DisplayName = 
+				METASOUND_LOCTEXT_FORMAT("AddNodeDisplayNamePattern2", "Add ({0} to {1})",
 				GetMetasoundDataTypeDisplayText<OperandDataType>(),
 				GetMetasoundDataTypeDisplayText<DataType>());
 			return DisplayName;
@@ -83,15 +83,15 @@ namespace Metasound
 		template<typename DataType>
 		const FText GetSubtractDisplayName()
 		{
-			static const FText DisplayName = FText::Format(LOCTEXT("SubtractNodeDisplayNamePattern", "Subtract ({0})"), GetMetasoundDataTypeDisplayText<DataType>());
+			static const FText DisplayName = METASOUND_LOCTEXT_FORMAT("SubtractNodeDisplayNamePattern", "Subtract ({0})", GetMetasoundDataTypeDisplayText<DataType>());
 			return DisplayName;
 		}
 
 		template<typename DataType, typename OperandDataType>
 		const FText GetSubtractDisplayName()
 		{
-			static const FText DisplayName = FText::Format(
-				LOCTEXT("SubtractNodeOperandTypedDisplayNamePattern", "Subtract ({0} from {1})"),
+			static const FText DisplayName = 
+				METASOUND_LOCTEXT_FORMAT("SubtractNodeOperandTypedDisplayNamePattern", "Subtract ({0} from {1})",
 				GetMetasoundDataTypeDisplayText<OperandDataType>(),
 				GetMetasoundDataTypeDisplayText<DataType>());
 			return DisplayName;
@@ -100,15 +100,15 @@ namespace Metasound
 		template<typename DataType>
 		const FText GetMultiplyDisplayName()
 		{
-			static const FText DisplayName = FText::Format(LOCTEXT("MultiplyNodeDisplayNamePattern", "Multiply ({0})"), GetMetasoundDataTypeDisplayText<DataType>());
+			static const FText DisplayName = METASOUND_LOCTEXT_FORMAT("MultiplyNodeDisplayNamePattern", "Multiply ({0})", GetMetasoundDataTypeDisplayText<DataType>());
 			return DisplayName;
 		}
 
 		template<typename DataType, typename OperandDataType>
 		const FText GetMultiplyDisplayName()
 		{
-			static const FText DisplayName = FText::Format(
-				LOCTEXT("MultiplyNodeOperandTypedDisplayNamePattern", "Multiply ({0} by {1})"),
+			static const FText DisplayName = 
+				METASOUND_LOCTEXT_FORMAT("MultiplyNodeOperandTypedDisplayNamePattern", "Multiply ({0} by {1})",
 				GetMetasoundDataTypeDisplayText<DataType>(),
 				GetMetasoundDataTypeDisplayText<OperandDataType>());
 			return DisplayName;
@@ -117,15 +117,15 @@ namespace Metasound
 		template<typename DataType>
 		const FText GetDivideDisplayName()
 		{
-			static const FText DisplayName = FText::Format(LOCTEXT("DivideNodeDisplayNamePattern", "Divide ({0})"), GetMetasoundDataTypeDisplayText<DataType>());
+			static const FText DisplayName = METASOUND_LOCTEXT_FORMAT("DivideNodeDisplayNamePattern", "Divide ({0})", GetMetasoundDataTypeDisplayText<DataType>());
 			return DisplayName;
 		}
 
 		template<typename DataType, typename OperandDataType>
 		const FText GetDivideDisplayName()
 		{
-			static const FText DisplayName = FText::Format(
-				LOCTEXT("DivideNodeOperandTypedDisplayNamePattern", "Divide ({0} by {1})"),
+			static const FText DisplayName = 
+				METASOUND_LOCTEXT_FORMAT("DivideNodeOperandTypedDisplayNamePattern", "Divide ({0} by {1})",
 				GetMetasoundDataTypeDisplayText<DataType>(),
 				GetMetasoundDataTypeDisplayText<OperandDataType>());
 			return DisplayName;
@@ -134,15 +134,15 @@ namespace Metasound
 		template<typename DataType>
 		const FText GetModuloDisplayName()
 		{
-			static const FText DisplayName = FText::Format(LOCTEXT("ModuloNodeDisplayNamePattern", "Modulo ({0})"), GetMetasoundDataTypeDisplayText<DataType>());
+			static const FText DisplayName = METASOUND_LOCTEXT_FORMAT("ModuloNodeDisplayNamePattern", "Modulo ({0})", GetMetasoundDataTypeDisplayText<DataType>());
 			return DisplayName;
 		}
 
 		template<typename DataType, typename OperandDataType>
 		const FText GetModuloDisplayName()
 		{
-			static const FText DisplayName = FText::Format(
-				LOCTEXT("ModuloNodeOperandTypedDisplayNamePattern", "Modulo ({0} by {1})"),
+			static const FText DisplayName = 
+				METASOUND_LOCTEXT_FORMAT("ModuloNodeOperandTypedDisplayNamePattern", "Modulo ({0} by {1})",
 				GetMetasoundDataTypeDisplayText<DataType>(),
 				GetMetasoundDataTypeDisplayText<OperandDataType>());
 			return DisplayName;
@@ -151,15 +151,15 @@ namespace Metasound
 		template<typename DataType>
 		const FText GetPowerDisplayName()
 		{
-			static const FText DisplayName = FText::Format(LOCTEXT("PowerNodeDisplayNamePattern", "Power ({0})"), GetMetasoundDataTypeDisplayText<DataType>());
+			static const FText DisplayName = METASOUND_LOCTEXT_FORMAT("PowerNodeDisplayNamePattern", "Power ({0})", GetMetasoundDataTypeDisplayText<DataType>());
 			return DisplayName;
 		}
 
 		template<typename DataType, typename OperandDataType>
 		const FText GetPowerDisplayName()
 		{
-			static const FText DisplayName = FText::Format(
-				LOCTEXT("PowerNodeOperandTypedDisplayNamePattern", "Power ({0} to the power of {1})"),
+			static const FText DisplayName = 
+				METASOUND_LOCTEXT_FORMAT("PowerNodeOperandTypedDisplayNamePattern", "Power ({0} to the power of {1})",
 				GetMetasoundDataTypeDisplayText<DataType>(),
 				GetMetasoundDataTypeDisplayText<OperandDataType>());
 			return DisplayName;
@@ -168,15 +168,15 @@ namespace Metasound
 		template<typename DataType>
 		const FText GetLogarithmDisplayName()
 		{
-			static const FText DisplayName = FText::Format(LOCTEXT("LogNodeDisplayNamePattern", "Log ({0})"), GetMetasoundDataTypeDisplayText<DataType>());
+			static const FText DisplayName = METASOUND_LOCTEXT_FORMAT("LogNodeDisplayNamePattern", "Log ({0})", GetMetasoundDataTypeDisplayText<DataType>());
 			return DisplayName;
 		}
 
 		template<typename DataType, typename OperandDataType>
 		const FText GetLogarithmDisplayName()
 		{
-			static const FText DisplayName = FText::Format(
-				LOCTEXT("LogarithmNodeOperandTypedDisplayNamePattern", "Log ({0}-Base logarithm of {1})"),
+			static const FText DisplayName = 
+				METASOUND_LOCTEXT_FORMAT("LogarithmNodeOperandTypedDisplayNamePattern", "Log ({0}-Base logarithm of {1})",
 				GetMetasoundDataTypeDisplayText<DataType>(),
 				GetMetasoundDataTypeDisplayText<OperandDataType>());
 			return DisplayName;
@@ -309,11 +309,11 @@ namespace Metasound
 		{
 			static const FVertexInterface Interface(
 				FInputVertexInterface(
-					TInputDataVertexModel<TDataClass>(MathOpNames::PrimaryOperandName, LOCTEXT("MathOpAddNode_InitialTooltip", "Initial addend."), static_cast<TDataClass>(0)),
-					TInputDataVertexModel<TOperandDataClass>(MathOpNames::AdditionalOperandsName, LOCTEXT("MathOpAddNode_AddendsTooltip", "Additional addend(s)."), static_cast<TOperandDataClass>(0))
+					TInputDataVertexModel<TDataClass>(MathOpNames::PrimaryOperandName, METASOUND_LOCTEXT("MathOpAddNode_InitialTooltip", "Initial addend."), static_cast<TDataClass>(0)),
+					TInputDataVertexModel<TOperandDataClass>(MathOpNames::AdditionalOperandsName, METASOUND_LOCTEXT("MathOpAddNode_AddendsTooltip", "Additional addend(s)."), static_cast<TOperandDataClass>(0))
 				),
 				FOutputVertexInterface(
-					TOutputDataVertexModel<TDataClass>(TEXT("Out"), LOCTEXT("MathOpAddNode_OutTooltip", "Math operation result"))
+					TOutputDataVertexModel<TDataClass>(TEXT("Out"), METASOUND_LOCTEXT("MathOpAddNode_OutTooltip", "Math operation result"))
 				)
 			);
 
@@ -358,11 +358,11 @@ namespace Metasound
 		{
 			static const FVertexInterface Interface(
 				FInputVertexInterface(
-					TInputDataVertexModel<TDataClass>(MathOpNames::PrimaryOperandName, LOCTEXT("MathOpSubtractNode_MinuendTooltip", "Minuend."), static_cast<TDataClass>(0)),
-					TInputDataVertexModel<TOperandDataClass>(MathOpNames::AdditionalOperandsName, LOCTEXT("MathOpSubtractNode_SubtrahendsTooltip", "Subtrahend(s)."), static_cast<TOperandDataClass>(0))
+					TInputDataVertexModel<TDataClass>(MathOpNames::PrimaryOperandName, METASOUND_LOCTEXT("MathOpSubtractNode_MinuendTooltip", "Minuend."), static_cast<TDataClass>(0)),
+					TInputDataVertexModel<TOperandDataClass>(MathOpNames::AdditionalOperandsName, METASOUND_LOCTEXT("MathOpSubtractNode_SubtrahendsTooltip", "Subtrahend(s)."), static_cast<TOperandDataClass>(0))
 				),
 				FOutputVertexInterface(
-					TOutputDataVertexModel<TDataClass>(TEXT("Out"), LOCTEXT("MathOpSubtractNode_OutTooltip", "Subtraction result"))
+					TOutputDataVertexModel<TDataClass>(TEXT("Out"), METASOUND_LOCTEXT("MathOpSubtractNode_OutTooltip", "Subtraction result"))
 				)
 			);
 
@@ -407,11 +407,11 @@ namespace Metasound
 		{
 			static const FVertexInterface Interface(
 				FInputVertexInterface(
-					TInputDataVertexModel<TDataClass>(MathOpNames::PrimaryOperandName, LOCTEXT("MathOpMultiplyNode_InitMultiplicandTooltip", "Initial multiplicand."), static_cast<TDataClass>(1)),
-					TInputDataVertexModel<TOperandDataClass>(MathOpNames::AdditionalOperandsName, LOCTEXT("MathOpMultiplyNode_MultiplicandsTooltip", "Additional multiplicand(s)."), static_cast<TOperandDataClass>(1))
+					TInputDataVertexModel<TDataClass>(MathOpNames::PrimaryOperandName, METASOUND_LOCTEXT("MathOpMultiplyNode_InitMultiplicandTooltip", "Initial multiplicand."), static_cast<TDataClass>(1)),
+					TInputDataVertexModel<TOperandDataClass>(MathOpNames::AdditionalOperandsName, METASOUND_LOCTEXT("MathOpMultiplyNode_MultiplicandsTooltip", "Additional multiplicand(s)."), static_cast<TOperandDataClass>(1))
 				),
 				FOutputVertexInterface(
-					TOutputDataVertexModel<TDataClass>(TEXT("Out"), LOCTEXT("MathOpMultiplyNode_ResultTooltip", "Multiplication result"))
+					TOutputDataVertexModel<TDataClass>(TEXT("Out"), METASOUND_LOCTEXT("MathOpMultiplyNode_ResultTooltip", "Multiplication result"))
 				)
 			);
 
@@ -456,11 +456,11 @@ namespace Metasound
 		{
 			static const FVertexInterface Interface(
 				FInputVertexInterface(
-					TInputDataVertexModel<TDataClass>(MathOpNames::PrimaryOperandName, LOCTEXT("MathOpDivideNode_DividendTooltip", "Dividend."), static_cast<TDataClass>(1)),
-					TInputDataVertexModel<TOperandDataClass>(MathOpNames::AdditionalOperandsName, LOCTEXT("MathOpDivideNode_DivisorsTooltip", "Divisor(s)."), static_cast<TOperandDataClass>(1))
+					TInputDataVertexModel<TDataClass>(MathOpNames::PrimaryOperandName, METASOUND_LOCTEXT("MathOpDivideNode_DividendTooltip", "Dividend."), static_cast<TDataClass>(1)),
+					TInputDataVertexModel<TOperandDataClass>(MathOpNames::AdditionalOperandsName, METASOUND_LOCTEXT("MathOpDivideNode_DivisorsTooltip", "Divisor(s)."), static_cast<TOperandDataClass>(1))
 				),
 				FOutputVertexInterface(
-					TOutputDataVertexModel<TDataClass>(TEXT("Out"), LOCTEXT("MathOpDivideNode_OutTooltip", "Division result"))
+					TOutputDataVertexModel<TDataClass>(TEXT("Out"), METASOUND_LOCTEXT("MathOpDivideNode_OutTooltip", "Division result"))
 				)
 			);
 
@@ -512,11 +512,11 @@ namespace Metasound
 		{
 			static const FVertexInterface Interface(
 				FInputVertexInterface(
-					TInputDataVertexModel<TDataClass>(MathOpNames::PrimaryOperandName, LOCTEXT("MathOpModuloNode_DividendTooltip", "Dividend."), static_cast<TDataClass>(1)),
-					TInputDataVertexModel<TOperandDataClass>(MathOpNames::AdditionalOperandsName, LOCTEXT("MathOpModuloNode_DivisorsTooltip", "Divisor(s)."), static_cast<TOperandDataClass>(1))
+					TInputDataVertexModel<TDataClass>(MathOpNames::PrimaryOperandName, METASOUND_LOCTEXT("MathOpModuloNode_DividendTooltip", "Dividend."), static_cast<TDataClass>(1)),
+					TInputDataVertexModel<TOperandDataClass>(MathOpNames::AdditionalOperandsName, METASOUND_LOCTEXT("MathOpModuloNode_DivisorsTooltip", "Divisor(s)."), static_cast<TOperandDataClass>(1))
 				),
 				FOutputVertexInterface(
-					TOutputDataVertexModel<TDataClass>(TEXT("Out"), LOCTEXT("MathOpModuloNode_OutTooltip", "Modulo result"))
+					TOutputDataVertexModel<TDataClass>(TEXT("Out"), METASOUND_LOCTEXT("MathOpModuloNode_OutTooltip", "Modulo result"))
 				)
 			);
 
@@ -570,11 +570,11 @@ namespace Metasound
 		{
 			static const FVertexInterface DefaultInterface(
 				FInputVertexInterface(
-					TInputDataVertexModel<float>(MathOpNames::PrimaryOperandName, LOCTEXT("MathOpPowerNode_Base", "The base of the power")),
-					TInputDataVertexModel<float>(MathOpNames::AdditionalOperandsName, LOCTEXT("MathOpPowerNode_Exponent", "The exponent to take the base to the power of"))
+					TInputDataVertexModel<float>(MathOpNames::PrimaryOperandName, METASOUND_LOCTEXT("MathOpPowerNode_Base", "The base of the power")),
+					TInputDataVertexModel<float>(MathOpNames::AdditionalOperandsName, METASOUND_LOCTEXT("MathOpPowerNode_Exponent", "The exponent to take the base to the power of"))
 				),
 				FOutputVertexInterface(
-					TOutputDataVertexModel<float>(TEXT("Out"), LOCTEXT("MathOpPowerNode_Result", "Returns Base to the Exponent power"))
+					TOutputDataVertexModel<float>(TEXT("Out"), METASOUND_LOCTEXT("MathOpPowerNode_Result", "Returns Base to the Exponent power"))
 				)
 			);
 
@@ -622,11 +622,11 @@ namespace Metasound
 		{
 			static const FVertexInterface DefaultInterface(
 				FInputVertexInterface(
-					TInputDataVertexModel<float>(MathOpNames::PrimaryOperandName, LOCTEXT("MathOpLogarithmNode_Base", "The base of the logarithm")),
-					TInputDataVertexModel<float>(MathOpNames::AdditionalOperandsName, LOCTEXT("MathOpLogarithmNode_Value", "The value to find the logarithm of"))
+					TInputDataVertexModel<float>(MathOpNames::PrimaryOperandName, METASOUND_LOCTEXT("MathOpLogarithmNode_Base", "The base of the logarithm")),
+					TInputDataVertexModel<float>(MathOpNames::AdditionalOperandsName, METASOUND_LOCTEXT("MathOpLogarithmNode_Value", "The value to find the logarithm of"))
 				),
 				FOutputVertexInterface(
-					TOutputDataVertexModel<float>(TEXT("Out"), LOCTEXT("MathOpLogarithmNode_Result", "The logarithm of the inputted value"))
+					TOutputDataVertexModel<float>(TEXT("Out"), METASOUND_LOCTEXT("MathOpLogarithmNode_Result", "The logarithm of the inputted value"))
 				)
 			);
 
@@ -671,11 +671,11 @@ namespace Metasound
 		{
 			static const FVertexInterface Interface(
 				FInputVertexInterface(
-					TInputDataVertexModel<FAudioBuffer>(MathOpNames::PrimaryOperandName, LOCTEXT("MathOpAddendTooltip1", "First addend.")),
-					TInputDataVertexModel<FAudioBuffer>(MathOpNames::AdditionalOperandsName, LOCTEXT("MathOpAddendAdditionalTooltip1", "Additional addends."))
+					TInputDataVertexModel<FAudioBuffer>(MathOpNames::PrimaryOperandName, METASOUND_LOCTEXT("MathOpAddendTooltip1", "First addend.")),
+					TInputDataVertexModel<FAudioBuffer>(MathOpNames::AdditionalOperandsName, METASOUND_LOCTEXT("MathOpAddendAdditionalTooltip1", "Additional addends."))
 				),
 				FOutputVertexInterface(
-					TOutputDataVertexModel<FAudioBuffer>(TEXT("Out"), LOCTEXT("MathOpOutTooltip", "Math operation result"))
+					TOutputDataVertexModel<FAudioBuffer>(TEXT("Out"), METASOUND_LOCTEXT("MathOpOutTooltip", "Math operation result"))
 				)
 			);
 
@@ -748,11 +748,11 @@ namespace Metasound
 		{
 			static const FVertexInterface Interface(
 				FInputVertexInterface(
-					TInputDataVertexModel<FTime>(MathOpNames::PrimaryOperandName, LOCTEXT("MathOpAddendTooltip2", "First addend."), 0.0f),
-					TInputDataVertexModel<FTime>(MathOpNames::AdditionalOperandsName, LOCTEXT("MathOpAddendAdditionalTooltip2", "Additional addends."), 0.0f)
+					TInputDataVertexModel<FTime>(MathOpNames::PrimaryOperandName, METASOUND_LOCTEXT("MathOpAddendTooltip2", "First addend."), 0.0f),
+					TInputDataVertexModel<FTime>(MathOpNames::AdditionalOperandsName, METASOUND_LOCTEXT("MathOpAddendAdditionalTooltip2", "Additional addends."), 0.0f)
 				),
 				FOutputVertexInterface(
-					TOutputDataVertexModel<FTime>(TEXT("Out"), LOCTEXT("MathOpOutTooltip", "Math operation result"))
+					TOutputDataVertexModel<FTime>(TEXT("Out"), METASOUND_LOCTEXT("MathOpOutTooltip", "Math operation result"))
 				)
 			);
 
@@ -792,11 +792,11 @@ namespace Metasound
 		{
 			static const FVertexInterface Interface(
 				FInputVertexInterface(
-					TInputDataVertexModel<FAudioBuffer>(MathOpNames::PrimaryOperandName, LOCTEXT("MathOpAddendAudioTooltip", "Audio Buffer to add offset(s) to.")),
-					TInputDataVertexModel<float>(MathOpNames::AdditionalOperandsName, LOCTEXT("MathOpAddendAdditionalTooltip2", "Float addends of which to offset buffer samples."), 0.0f)
+					TInputDataVertexModel<FAudioBuffer>(MathOpNames::PrimaryOperandName, METASOUND_LOCTEXT("MathOpAddendAudioTooltip", "Audio Buffer to add offset(s) to.")),
+					TInputDataVertexModel<float>(MathOpNames::AdditionalOperandsName, METASOUND_LOCTEXT("MathOpAddendAdditionalTooltip2", "Float addends of which to offset buffer samples."), 0.0f)
 				),
 				FOutputVertexInterface(
-					TOutputDataVertexModel<FAudioBuffer>(TEXT("Out"), LOCTEXT("MathOpAudioFloatAddOutTooltip", "Resulting buffer"))
+					TOutputDataVertexModel<FAudioBuffer>(TEXT("Out"), METASOUND_LOCTEXT("MathOpAudioFloatAddOutTooltip", "Resulting buffer"))
 				)
 			);
 
@@ -840,11 +840,11 @@ namespace Metasound
 		{
 			static const FVertexInterface Interface(
 				FInputVertexInterface(
-					TInputDataVertexModel<FAudioBuffer>(MathOpNames::PrimaryOperandName, LOCTEXT("MathOpBuffersMinuendTooltip", "Initial buffer to act as minuend.")),
-					TInputDataVertexModel<FAudioBuffer>(MathOpNames::AdditionalOperandsName, LOCTEXT("MathOpSubtractBuffersSubtrahendsTooltip", "Additional buffers to act as subtrahend(s)."))
+					TInputDataVertexModel<FAudioBuffer>(MathOpNames::PrimaryOperandName, METASOUND_LOCTEXT("MathOpBuffersMinuendTooltip", "Initial buffer to act as minuend.")),
+					TInputDataVertexModel<FAudioBuffer>(MathOpNames::AdditionalOperandsName, METASOUND_LOCTEXT("MathOpSubtractBuffersSubtrahendsTooltip", "Additional buffers to act as subtrahend(s)."))
 				),
 				FOutputVertexInterface(
-					TOutputDataVertexModel<FAudioBuffer>(TEXT("Out"), LOCTEXT("MathOpSubtractBuffersOutTooltip", "Resulting buffer"))
+					TOutputDataVertexModel<FAudioBuffer>(TEXT("Out"), METASOUND_LOCTEXT("MathOpSubtractBuffersOutTooltip", "Resulting buffer"))
 				)
 			);
 
@@ -915,11 +915,11 @@ namespace Metasound
 		{
 			static const FVertexInterface Interface(
 				FInputVertexInterface(
-					TInputDataVertexModel<FTime>(MathOpNames::PrimaryOperandName, LOCTEXT("MathOpTimeNode_MinuendTooltip", "Time minuend."), 0.0f),
-					TInputDataVertexModel<FTime>(MathOpNames::AdditionalOperandsName, LOCTEXT("MathOpTimeNode_SubtrahendsTooltip", "Time subtrahends."), 0.0f)
+					TInputDataVertexModel<FTime>(MathOpNames::PrimaryOperandName, METASOUND_LOCTEXT("MathOpTimeNode_MinuendTooltip", "Time minuend."), 0.0f),
+					TInputDataVertexModel<FTime>(MathOpNames::AdditionalOperandsName, METASOUND_LOCTEXT("MathOpTimeNode_SubtrahendsTooltip", "Time subtrahends."), 0.0f)
 				),
 				FOutputVertexInterface(
-					TOutputDataVertexModel<FTime>(TEXT("Out"), LOCTEXT("MathOpTimeNode_SubtractOutTooltip", "Resulting time value"))
+					TOutputDataVertexModel<FTime>(TEXT("Out"), METASOUND_LOCTEXT("MathOpTimeNode_SubtractOutTooltip", "Resulting time value"))
 				)
 			);
 
@@ -959,11 +959,11 @@ namespace Metasound
 		{
 			static const FVertexInterface Interface(
 				FInputVertexInterface(
-					TInputDataVertexModel<FAudioBuffer>(MathOpNames::PrimaryOperandName, LOCTEXT("MathOpAudioMultiplyNode_InitMultiplicandTooltip", "Initial audio to multiply.")),
-					TInputDataVertexModel<FAudioBuffer>(MathOpNames::AdditionalOperandsName, LOCTEXT("MathOpAudioMultiplyNode_MultiplicandsTooltip", "Additional audio to multiply sample-by-sample."))
+					TInputDataVertexModel<FAudioBuffer>(MathOpNames::PrimaryOperandName, METASOUND_LOCTEXT("MathOpAudioMultiplyNode_InitMultiplicandTooltip", "Initial audio to multiply.")),
+					TInputDataVertexModel<FAudioBuffer>(MathOpNames::AdditionalOperandsName, METASOUND_LOCTEXT("MathOpAudioMultiplyNode_MultiplicandsTooltip", "Additional audio to multiply sample-by-sample."))
 				),
 				FOutputVertexInterface(
-					TOutputDataVertexModel<FAudioBuffer>(TEXT("Out"), LOCTEXT("MathOpAudioMultiplyNode_OutTooltip", "Resulting buffer"))
+					TOutputDataVertexModel<FAudioBuffer>(TEXT("Out"), METASOUND_LOCTEXT("MathOpAudioMultiplyNode_OutTooltip", "Resulting buffer"))
 				)
 			);
 
@@ -1037,11 +1037,11 @@ namespace Metasound
 		{
 			static const FVertexInterface Interface(
 				FInputVertexInterface(
-					TInputDataVertexModel<FAudioBuffer>(MathOpNames::PrimaryOperandName, LOCTEXT("MathOpAudioFloatMultiplyNode_FloatTooltip", "Audio multiplicand.")),
-					TInputDataVertexModel<float>(MathOpNames::AdditionalOperandsName, LOCTEXT("MathOpAudioFloatMultiplyNode_MultiplicandTooltip", "Float multiplicand to apply sample-by-sample to audio. Interpolates over buffer size on value change."), 1.0f)
+					TInputDataVertexModel<FAudioBuffer>(MathOpNames::PrimaryOperandName, METASOUND_LOCTEXT("MathOpAudioFloatMultiplyNode_FloatTooltip", "Audio multiplicand.")),
+					TInputDataVertexModel<float>(MathOpNames::AdditionalOperandsName, METASOUND_LOCTEXT("MathOpAudioFloatMultiplyNode_MultiplicandTooltip", "Float multiplicand to apply sample-by-sample to audio. Interpolates over buffer size on value change."), 1.0f)
 				),
 				FOutputVertexInterface(
-					TOutputDataVertexModel<FAudioBuffer>(TEXT("Out"), LOCTEXT("MathOpAudioFloatMultiplyNode_OutTooltip", "Resulting buffer"))
+					TOutputDataVertexModel<FAudioBuffer>(TEXT("Out"), METASOUND_LOCTEXT("MathOpAudioFloatMultiplyNode_OutTooltip", "Resulting buffer"))
 				)
 			);
 
@@ -1101,11 +1101,11 @@ namespace Metasound
 		{
 			static const FVertexInterface Interface(
 				FInputVertexInterface(
-					TInputDataVertexModel<FTime>(MathOpNames::PrimaryOperandName, LOCTEXT("MathOpTimeFloatMultiplyNode_MultiplyFloatTooltip", "Time multiplicand."), 1.0f),
-					TInputDataVertexModel<float>(MathOpNames::AdditionalOperandsName, LOCTEXT("MathOpTimeFloatMultiplyNode_MultiplicandsTooltip", "Float multiplicand(s)."), 1.0f)
+					TInputDataVertexModel<FTime>(MathOpNames::PrimaryOperandName, METASOUND_LOCTEXT("MathOpTimeFloatMultiplyNode_MultiplyFloatTooltip", "Time multiplicand."), 1.0f),
+					TInputDataVertexModel<float>(MathOpNames::AdditionalOperandsName, METASOUND_LOCTEXT("MathOpTimeFloatMultiplyNode_MultiplicandsTooltip", "Float multiplicand(s)."), 1.0f)
 				),
 				FOutputVertexInterface(
-					TOutputDataVertexModel<FTime>(TEXT("Out"), LOCTEXT("MathOpTimeFloatMultiplyNode_OutTooltip", "Time float multiplication output"))
+					TOutputDataVertexModel<FTime>(TEXT("Out"), METASOUND_LOCTEXT("MathOpTimeFloatMultiplyNode_OutTooltip", "Time float multiplication output"))
 				)
 			);
 
@@ -1147,11 +1147,11 @@ namespace Metasound
 		{
 			static const FVertexInterface Interface(
 				FInputVertexInterface(
-					TInputDataVertexModel<FTime>(MathOpNames::PrimaryOperandName, LOCTEXT("MathOpTimeFloatDivideNode_DividendFloatTooltip", "Time dividend."), 1.0f),
-					TInputDataVertexModel<float>(MathOpNames::AdditionalOperandsName, LOCTEXT("MathOpTimeFloatDivideNode_DivisorsTooltip", "Float divisor(s)."), 1.0f)
+					TInputDataVertexModel<FTime>(MathOpNames::PrimaryOperandName, METASOUND_LOCTEXT("MathOpTimeFloatDivideNode_DividendFloatTooltip", "Time dividend."), 1.0f),
+					TInputDataVertexModel<float>(MathOpNames::AdditionalOperandsName, METASOUND_LOCTEXT("MathOpTimeFloatDivideNode_DivisorsTooltip", "Float divisor(s)."), 1.0f)
 				),
 				FOutputVertexInterface(
-					TOutputDataVertexModel<FTime>(TEXT("Out"), LOCTEXT("MathOpTimeFloatDivideNode_OutTooltip", "Time divide-by-float output"))
+					TOutputDataVertexModel<FTime>(TEXT("Out"), METASOUND_LOCTEXT("MathOpTimeFloatDivideNode_OutTooltip", "Time divide-by-float output"))
 				)
 			);
 
@@ -1192,35 +1192,35 @@ namespace Metasound
 	};
 
 	// Definitions
-	DEFINE_METASOUND_MATHOP(Add, Float, float, LOCTEXT("MathAddFloatNode_Description", "Adds floats."), MathOpNames::AddKeywords)
-	DEFINE_METASOUND_MATHOP(Add, Int32, int32, LOCTEXT("MathAddInt32Node_Description", "Adds int32s."), MathOpNames::AddKeywords)
-	DEFINE_METASOUND_MATHOP(Add, Audio, FAudioBuffer, LOCTEXT("MathAddBufferNode_Description", "Adds buffers together by sample."), MathOpNames::AddKeywords)
-	DEFINE_METASOUND_MATHOP(Add, Time, FTime, LOCTEXT("MathAddTimeNode_Description", "Adds time values."), MathOpNames::AddKeywords)
+	DEFINE_METASOUND_MATHOP(Add, Float, float, METASOUND_LOCTEXT("MathAddFloatNode_Description", "Adds floats."), MathOpNames::AddKeywords)
+	DEFINE_METASOUND_MATHOP(Add, Int32, int32, METASOUND_LOCTEXT("MathAddInt32Node_Description", "Adds int32s."), MathOpNames::AddKeywords)
+	DEFINE_METASOUND_MATHOP(Add, Audio, FAudioBuffer, METASOUND_LOCTEXT("MathAddBufferNode_Description", "Adds buffers together by sample."), MathOpNames::AddKeywords)
+	DEFINE_METASOUND_MATHOP(Add, Time, FTime, METASOUND_LOCTEXT("MathAddTimeNode_Description", "Adds time values."), MathOpNames::AddKeywords)
 
-	DEFINE_METASOUND_OPERAND_TYPED_MATHOP(Add, Audio, FAudioBuffer, Float, float, LOCTEXT("MathAddAudioFloatNode_Description", "Add floats to buffer sample-by-sample."), MathOpNames::AddKeywords)
+	DEFINE_METASOUND_OPERAND_TYPED_MATHOP(Add, Audio, FAudioBuffer, Float, float, METASOUND_LOCTEXT("MathAddAudioFloatNode_Description", "Add floats to buffer sample-by-sample."), MathOpNames::AddKeywords)
 
-	DEFINE_METASOUND_MATHOP(Subtract, Float, float, LOCTEXT("MathSubractFloatNode_Description", "Subtracts floats."), MathOpNames::SubtractKeywords)
-	DEFINE_METASOUND_MATHOP(Subtract, Int32, int32, LOCTEXT("MathSubractInt32Node_Description", "Subtracts int32s."), MathOpNames::SubtractKeywords)
-	DEFINE_METASOUND_MATHOP(Subtract, Audio, FAudioBuffer, LOCTEXT("MathSubtractBufferNode_Description", "Subtracts buffers sample-by-sample."), MathOpNames::SubtractKeywords)
-	DEFINE_METASOUND_MATHOP(Subtract, Time, FTime, LOCTEXT("MathSubractTimeNode_Description", "Subtracts time values."), MathOpNames::SubtractKeywords)
+	DEFINE_METASOUND_MATHOP(Subtract, Float, float, METASOUND_LOCTEXT("MathSubractFloatNode_Description", "Subtracts floats."), MathOpNames::SubtractKeywords)
+	DEFINE_METASOUND_MATHOP(Subtract, Int32, int32, METASOUND_LOCTEXT("MathSubractInt32Node_Description", "Subtracts int32s."), MathOpNames::SubtractKeywords)
+	DEFINE_METASOUND_MATHOP(Subtract, Audio, FAudioBuffer, METASOUND_LOCTEXT("MathSubtractBufferNode_Description", "Subtracts buffers sample-by-sample."), MathOpNames::SubtractKeywords)
+	DEFINE_METASOUND_MATHOP(Subtract, Time, FTime, METASOUND_LOCTEXT("MathSubractTimeNode_Description", "Subtracts time values."), MathOpNames::SubtractKeywords)
 
-	DEFINE_METASOUND_MATHOP(Multiply, Float, float, LOCTEXT("MathMultiplyFloatNode_Description", "Multiplies floats."), MathOpNames::MultiplyKeywords)
-	DEFINE_METASOUND_MATHOP(Multiply, Int32, int32, LOCTEXT("MathMultiplyInt32Node_Description", "Multiplies int32s."), MathOpNames::MultiplyKeywords)
-	DEFINE_METASOUND_MATHOP(Multiply, Audio, FAudioBuffer, LOCTEXT("MathMultiplyBufferNode_Description", "Multiplies buffers together sample-by-sample."), MathOpNames::MultiplyKeywords)
+	DEFINE_METASOUND_MATHOP(Multiply, Float, float, METASOUND_LOCTEXT("MathMultiplyFloatNode_Description", "Multiplies floats."), MathOpNames::MultiplyKeywords)
+	DEFINE_METASOUND_MATHOP(Multiply, Int32, int32, METASOUND_LOCTEXT("MathMultiplyInt32Node_Description", "Multiplies int32s."), MathOpNames::MultiplyKeywords)
+	DEFINE_METASOUND_MATHOP(Multiply, Audio, FAudioBuffer, METASOUND_LOCTEXT("MathMultiplyBufferNode_Description", "Multiplies buffers together sample-by-sample."), MathOpNames::MultiplyKeywords)
 
-	DEFINE_METASOUND_OPERAND_TYPED_MATHOP(Multiply, Audio, FAudioBuffer, Float, float, LOCTEXT("MathMultiplyAudioByFloatDescription", "Multiplies buffer by float scalars."), MathOpNames::MultiplyKeywords)
-	DEFINE_METASOUND_OPERAND_TYPED_MATHOP(Multiply, Time, FTime, Float, float, LOCTEXT("MathMultiplyTimeNode_Description", "Scales time by floats."), MathOpNames::MultiplyKeywords)
+	DEFINE_METASOUND_OPERAND_TYPED_MATHOP(Multiply, Audio, FAudioBuffer, Float, float, METASOUND_LOCTEXT("MathMultiplyAudioByFloatDescription", "Multiplies buffer by float scalars."), MathOpNames::MultiplyKeywords)
+	DEFINE_METASOUND_OPERAND_TYPED_MATHOP(Multiply, Time, FTime, Float, float, METASOUND_LOCTEXT("MathMultiplyTimeNode_Description", "Scales time by floats."), MathOpNames::MultiplyKeywords)
 
-	DEFINE_METASOUND_MATHOP(Divide, Float, float, LOCTEXT("MathDivideFloatNode_Description", "Divide float by another float."), MathOpNames::DivideKeywords)
-	DEFINE_METASOUND_MATHOP(Divide, Int32, int32, LOCTEXT("MathDivideInt32Node_Description", "Divide int32 by another int32."), MathOpNames::DivideKeywords)
+	DEFINE_METASOUND_MATHOP(Divide, Float, float, METASOUND_LOCTEXT("MathDivideFloatNode_Description", "Divide float by another float."), MathOpNames::DivideKeywords)
+	DEFINE_METASOUND_MATHOP(Divide, Int32, int32, METASOUND_LOCTEXT("MathDivideInt32Node_Description", "Divide int32 by another int32."), MathOpNames::DivideKeywords)
 
-	DEFINE_METASOUND_OPERAND_TYPED_MATHOP(Divide, Time, FTime, Float, float, LOCTEXT("MathDivideTimeNode_Description", "Divides time by floats."), MathOpNames::DivideKeywords)
+	DEFINE_METASOUND_OPERAND_TYPED_MATHOP(Divide, Time, FTime, Float, float, METASOUND_LOCTEXT("MathDivideTimeNode_Description", "Divides time by floats."), MathOpNames::DivideKeywords)
 
-	DEFINE_METASOUND_MATHOP(Modulo, Int32, int32, LOCTEXT("MathModulusInt32Node_Description", "Modulo int32 by another int32."), MathOpNames::ModuloKeywords)
+	DEFINE_METASOUND_MATHOP(Modulo, Int32, int32, METASOUND_LOCTEXT("MathModulusInt32Node_Description", "Modulo int32 by another int32."), MathOpNames::ModuloKeywords)
 
-	DEFINE_METASOUND_MATHOP(Power, Float, float, LOCTEXT("MathPowerFloatNode_Description", "Raise float to the power of another float."), MathOpNames::PowerKeywords)
+	DEFINE_METASOUND_MATHOP(Power, Float, float, METASOUND_LOCTEXT("MathPowerFloatNode_Description", "Raise float to the power of another float."), MathOpNames::PowerKeywords)
 
-	DEFINE_METASOUND_MATHOP(Logarithm, Float, float, LOCTEXT("MathLogarithmFloatNode_Description", "Calculate float-base logarithm of another float."), TArray<FText>())
+	DEFINE_METASOUND_MATHOP(Logarithm, Float, float, METASOUND_LOCTEXT("MathLogarithmFloatNode_Description", "Calculate float-base logarithm of another float."), TArray<FText>())
 
 	METASOUND_REGISTER_NODE(FAddFloatNode)
 	METASOUND_REGISTER_NODE(FAddInt32Node)
diff --git a/Engine/Plugins/Runtime/Metasound/Source/MetasoundStandardNodes/Private/MetasoundMaxNode.cpp b/Engine/Plugins/Runtime/Metasound/Source/MetasoundStandardNodes/Private/MetasoundMaxNode.cpp
--- a/Engine/Plugins/Runtime/Metasound/Source/MetasoundStandardNodes/Private/MetasoundMaxNode.cpp
+++ b/Engine/Plugins/Runtime/Metasound/Source/MetasoundStandardNodes/Private/MetasoundMaxNode.cpp
@@ -156,8 +156,8 @@ namespace Metasound
 			{
 				const FName DataTypeName = GetMetasoundDataTypeName<ValueType>();
 				const FName OperatorName = TEXT("Max");
-				const FText NodeDisplayName = FText::Format(LOCTEXT("MaxDisplayNamePattern", "Max ({0})"), GetMetasoundDataTypeDisplayText<ValueType>());
-				const FText NodeDescription = LOCTEXT("MaxDesc", "Returns the max of A and B.");
+				const FText NodeDisplayName = METASOUND_LOCTEXT_FORMAT("MaxDisplayNamePattern", "Max ({0})", GetMetasoundDataTypeDisplayText<ValueType>());
+				const FText NodeDescription = METASOUND_LOCTEXT("MaxDesc", "Returns the max of A and B.");
 				const FVertexInterface NodeInterface = GetDefaultInterface();
 
 				return MetasoundMaxNodePrivate::CreateNodeClassMetadata(DataTypeName, OperatorName, NodeDisplayName, NodeDescription, NodeInterface);
diff --git a/Engine/Plugins/Runtime/Metasound/Source/MetasoundStandardNodes/Private/MetasoundMidSideNodes.cpp b/Engine/Plugins/Runtime/Metasound/Source/MetasoundStandardNodes/Private/MetasoundMidSideNodes.cpp
--- a/Engine/Plugins/Runtime/Metasound/Source/MetasoundStandardNodes/Private/MetasoundMidSideNodes.cpp
+++ b/Engine/Plugins/Runtime/Metasound/Source/MetasoundStandardNodes/Private/MetasoundMidSideNodes.cpp
@@ -64,8 +64,8 @@ namespace Metasound
 					FNodeClassName { StandardNodes::Namespace, "Mid-Side Decode", StandardNodes::AudioVariant },
 					1, // Major Version
 					0, // Minor Version
-					LOCTEXT("MidSideEncodeDisplayName", "Mid-Side Encode"),
-					LOCTEXT("MidSideEncodeDesc", "Converts a stereo audio signal from Left and Right to Mid and Side channels."),
+					METASOUND_LOCTEXT("MidSideEncodeDisplayName", "Mid-Side Encode"),
+					METASOUND_LOCTEXT("MidSideEncodeDesc", "Converts a stereo audio signal from Left and Right to Mid and Side channels."),
 					PluginAuthor,
 					PluginNodeMissingPrompt,
 					NodeInterface,
@@ -288,8 +288,8 @@ namespace Metasound
 					FNodeClassName { StandardNodes::Namespace, "Mid-Side Encode", StandardNodes::AudioVariant },
 					1, // Major Version
 					0, // Minor Version
-					LOCTEXT("MidSideDecodeDisplayName", "Mid-Side Decode"),
-					LOCTEXT("MidSideDecodeDesc", "Decodes a stereo signal with Mid and Side channels to Left and Right channels."),
+					METASOUND_LOCTEXT("MidSideDecodeDisplayName", "Mid-Side Decode"),
+					METASOUND_LOCTEXT("MidSideDecodeDesc", "Decodes a stereo signal with Mid and Side channels to Left and Right channels."),
 					PluginAuthor,
 					PluginNodeMissingPrompt,
 					NodeInterface,
diff --git a/Engine/Plugins/Runtime/Metasound/Source/MetasoundStandardNodes/Private/MetasoundMidiNoteQuantizerNode.cpp b/Engine/Plugins/Runtime/Metasound/Source/MetasoundStandardNodes/Private/MetasoundMidiNoteQuantizerNode.cpp
--- a/Engine/Plugins/Runtime/Metasound/Source/MetasoundStandardNodes/Private/MetasoundMidiNoteQuantizerNode.cpp
+++ b/Engine/Plugins/Runtime/Metasound/Source/MetasoundStandardNodes/Private/MetasoundMidiNoteQuantizerNode.cpp
@@ -113,8 +113,8 @@ namespace Metasound
 			Info.ClassName = { StandardNodes::Namespace, TEXT("MIDI Note Quantizer"), StandardNodes::AudioVariant };
 			Info.MajorVersion = 1;
 			Info.MinorVersion = 0;
-			Info.DisplayName = LOCTEXT("Metasound_MIDI_Note_Quantizer_NodeDisplayName", "MIDI Note Quantizer");
-			Info.Description = LOCTEXT("MidiNoteQuantizer_NodeDescription", "Quantizes a MIDI note to the nearset note that matches provided criteria"),
+			Info.DisplayName = METASOUND_LOCTEXT("Metasound_MIDI_Note_Quantizer_NodeDisplayName", "MIDI Note Quantizer");
+			Info.Description = METASOUND_LOCTEXT("MidiNoteQuantizer_NodeDescription", "Quantizes a MIDI note to the nearset note that matches provided criteria");
 			Info.Author = PluginAuthor;
 			Info.PromptIfMissing = PluginNodeMissingPrompt;
 			Info.DefaultInterface = DeclareVertexInterface();
diff --git a/Engine/Plugins/Runtime/Metasound/Source/MetasoundStandardNodes/Private/MetasoundMidiScaleToArrayNode.cpp b/Engine/Plugins/Runtime/Metasound/Source/MetasoundStandardNodes/Private/MetasoundMidiScaleToArrayNode.cpp
--- a/Engine/Plugins/Runtime/Metasound/Source/MetasoundStandardNodes/Private/MetasoundMidiScaleToArrayNode.cpp
+++ b/Engine/Plugins/Runtime/Metasound/Source/MetasoundStandardNodes/Private/MetasoundMidiScaleToArrayNode.cpp
@@ -107,8 +107,8 @@ namespace Metasound
 			Info.ClassName = { StandardNodes::Namespace, TEXT("Musical Scale To Note Array"), StandardNodes::AudioVariant };
 			Info.MajorVersion = 1;
 			Info.MinorVersion = 0;
-			Info.DisplayName = LOCTEXT("Metasound_ScaleToNoteArray_NodeDisplayName", "Scale to Note Array");
-			Info.Description = LOCTEXT("MusicalScaleToNoteArray_NodeDescription", "Returns an array represeting the selected scale"),
+			Info.DisplayName = METASOUND_LOCTEXT("Metasound_ScaleToNoteArray_NodeDisplayName", "Scale to Note Array");
+			Info.Description = METASOUND_LOCTEXT("MusicalScaleToNoteArray_NodeDescription", "Returns an array represeting the selected scale");
 			Info.Author = PluginAuthor;
 			Info.PromptIfMissing = PluginNodeMissingPrompt;
 			Info.DefaultInterface = DeclareVertexInterface();
diff --git a/Engine/Plugins/Runtime/Metasound/Source/MetasoundStandardNodes/Private/MetasoundMinNode.cpp b/Engine/Plugins/Runtime/Metasound/Source/MetasoundStandardNodes/Private/MetasoundMinNode.cpp
--- a/Engine/Plugins/Runtime/Metasound/Source/MetasoundStandardNodes/Private/MetasoundMinNode.cpp
+++ b/Engine/Plugins/Runtime/Metasound/Source/MetasoundStandardNodes/Private/MetasoundMinNode.cpp
@@ -155,8 +155,8 @@ namespace Metasound
 			{
 				const FName DataTypeName = GetMetasoundDataTypeName<ValueType>();
 				const FName OperatorName = TEXT("Min");
-				const FText NodeDisplayName = FText::Format(LOCTEXT("MinDisplayNamePattern", "Min ({0})"), GetMetasoundDataTypeDisplayText<ValueType>());
-				const FText NodeDescription = LOCTEXT("MinDesc", "Returns the min of A and B.");
+				const FText NodeDisplayName = METASOUND_LOCTEXT_FORMAT("MinDisplayNamePattern", "Min ({0})", GetMetasoundDataTypeDisplayText<ValueType>());
+				const FText NodeDescription = METASOUND_LOCTEXT("MinDesc", "Returns the min of A and B.");
 				const FVertexInterface NodeInterface = GetDefaultInterface();
 
 				return MetasoundMinNodePrivate::CreateNodeClassMetadata(DataTypeName, OperatorName, NodeDisplayName, NodeDescription, NodeInterface);
diff --git a/Engine/Plugins/Runtime/Metasound/Source/MetasoundStandardNodes/Private/MetasoundMixerNode.cpp b/Engine/Plugins/Runtime/Metasound/Source/MetasoundStandardNodes/Private/MetasoundMixerNode.cpp
--- a/Engine/Plugins/Runtime/Metasound/Source/MetasoundStandardNodes/Private/MetasoundMixerNode.cpp
+++ b/Engine/Plugins/Runtime/Metasound/Source/MetasoundStandardNodes/Private/MetasoundMixerNode.cpp
@@ -85,8 +85,8 @@ namespace Metasound
 			auto CreateNodeClassMetadataMono = []() -> FNodeClassMetadata
 			{
 				FName OperatorName = *FString::Printf(TEXT("Audio Mixer (Mono, %d)"), NumInputs);
-				FText NodeDisplayName = FText::Format(LOCTEXT("MonoMixer", "Mono Mixer ({0})"), NumInputs);
-				FText NodeDescription = LOCTEXT("MixerDescription1", "Will scale input channels by their corresponding gain value and sum them together.");
+				FText NodeDisplayName = METASOUND_LOCTEXT_FORMAT("MonoMixer", "Mono Mixer ({0})", NumInputs);
+				const FText NodeDescription = METASOUND_LOCTEXT("MixerDescription1", "Will scale input channels by their corresponding gain value and sum them together.");
 				FVertexInterface NodeInterface = GetDefaultInterface();
 
 				return CreateNodeClassMetadata(OperatorName, NodeDisplayName, NodeDescription, NodeInterface);
@@ -96,8 +96,8 @@ namespace Metasound
 			auto CreateNodeClassMetadataStereo = []() -> FNodeClassMetadata
 			{
 				FName OperatorName = *FString::Printf(TEXT("Audio Mixer (Stereo, %d)"), NumInputs);
-				FText NodeDisplayName = FText::Format(LOCTEXT("StereoMixer", "Stereo Mixer ({0})"), NumInputs);
-				FText NodeDescription = LOCTEXT("MixerDescription2", "Will scale input channels by their corresponding gain value and sum them together.");
+				FText NodeDisplayName = METASOUND_LOCTEXT_FORMAT("StereoMixer", "Stereo Mixer ({0})", NumInputs);
+				const FText NodeDescription = METASOUND_LOCTEXT("MixerDescription2", "Will scale input channels by their corresponding gain value and sum them together.");
 				FVertexInterface NodeInterface = GetDefaultInterface();
 
 				return  CreateNodeClassMetadata(OperatorName, NodeDisplayName, NodeDescription, NodeInterface);
@@ -107,8 +107,8 @@ namespace Metasound
 			auto CreateNodeClassMetadataMultiChan = []() -> FNodeClassMetadata
 			{
 				FName OperatorName = *FString::Printf(TEXT("Audio Mixer (%d-Channel, %d)"), NumChannels, NumInputs);
-				FText NodeDisplayName = FText::Format(LOCTEXT("NChannelMixer", "{0}-channel Mixer ({1})"), NumChannels, NumInputs);
-				FText NodeDescription = LOCTEXT("MixerDescription3", "Will scale input audio by their corresponding gain value and sum them together.");
+				FText NodeDisplayName = METASOUND_LOCTEXT_FORMAT("NChannelMixer", "{0}-channel Mixer ({1})", NumChannels, NumInputs);
+				const FText NodeDescription = METASOUND_LOCTEXT("MixerDescription3", "Will scale input audio by their corresponding gain value and sum them together.");
 				FVertexInterface NodeInterface = GetDefaultInterface();
 
 				return  CreateNodeClassMetadata(OperatorName, NodeDisplayName, NodeDescription, NodeInterface);
@@ -218,7 +218,7 @@ namespace Metasound
 				PluginNodeMissingPrompt,
 				InDefaultInterface,
 				{ NodeCategories::Mix },
-				{ LOCTEXT("Metasound_AudioMixerKeyword", "Mixer") },
+				{ METASOUND_LOCTEXT("Metasound_AudioMixerKeyword", "Mixer") },
 				FNodeDisplayStyle{}
 			};
 
@@ -242,7 +242,7 @@ namespace Metasound
 
 		static const FText GetAudioInputDescription(uint32 InputIndex, uint32 ChannelIndex)
 		{
-			return FText::Format(LOCTEXT("AudioMixerAudioInputDescription", "Audio Input #: {0}, Channel: {1}"), InputIndex, ChannelIndex);
+			return METASOUND_LOCTEXT_FORMAT("AudioMixerAudioInputDescription", "Audio Input #: {0}, Channel: {1}", InputIndex, ChannelIndex);
 		}
 
 		static const FVertexName GetGainInputName(uint32 InputIndex)
@@ -252,7 +252,7 @@ namespace Metasound
 
 		static const FText GetGainInputDescription(uint32 InputIndex)
 		{
-			return FText::Format(LOCTEXT("AudioMixerGainInputDescription", "Gain Input #: {0}"), InputIndex);
+			return METASOUND_LOCTEXT_FORMAT("AudioMixerGainInputDescription", "Gain Input #: {0}", InputIndex);
 		}
 
 		static const FVertexName GetAudioOutputName(uint32 ChannelIndex)
@@ -271,7 +271,7 @@ namespace Metasound
 
 		static const FText GetAudioOutputDescription(uint32 ChannelIndex)
 		{
-			return FText::Format(LOCTEXT("AudioMixerAudioOutputDescription", "Summed output for channel: {0}"), ChannelIndex);
+			return METASOUND_LOCTEXT_FORMAT("AudioMixerAudioOutputDescription", "Summed output for channel: {0}", ChannelIndex);
 		}
 #pragma endregion
 	}; // class TAudioMixerNodeOperator
diff --git a/Engine/Plugins/Runtime/Metasound/Source/MetasoundStandardNodes/Private/MetasoundNoiseGenerator.cpp b/Engine/Plugins/Runtime/Metasound/Source/MetasoundStandardNodes/Private/MetasoundNoiseGenerator.cpp
--- a/Engine/Plugins/Runtime/Metasound/Source/MetasoundStandardNodes/Private/MetasoundNoiseGenerator.cpp
+++ b/Engine/Plugins/Runtime/Metasound/Source/MetasoundStandardNodes/Private/MetasoundNoiseGenerator.cpp
@@ -29,8 +29,8 @@ namespace Metasound
 		FEnumNoiseType, FEnumNoiseTypeInfo, FEnumNoiseTypeReadRef, FEnumNoiseTypeWriteRef);
 
 	DEFINE_METASOUND_ENUM_BEGIN(ENoiseType, FEnumNoiseType, "NoiseType")
-		DEFINE_METASOUND_ENUM_ENTRY(ENoiseType::Pink, LOCTEXT("PinkDescription", "Pink Noise"), LOCTEXT("PinkDescriptionTT", "Pink noise or ​1⁄f noise, spectral density is inversely proportional to the frequency of the signal")),
-		DEFINE_METASOUND_ENUM_ENTRY(ENoiseType::White, LOCTEXT("WhiteDescription", "White Noise"), LOCTEXT("WhiteDescriptionTT", "A random signal having equal intensity at different frequencies")),
+		DEFINE_METASOUND_ENUM_ENTRY(ENoiseType::Pink, "PinkDescription", "Pink Noise", "PinkDescriptionTT", "Pink noise or ​1⁄f noise, spectral density is inversely proportional to the frequency of the signal"),
+		DEFINE_METASOUND_ENUM_ENTRY(ENoiseType::White, "WhiteDescription", "White Noise", "WhiteDescriptionTT", "A random signal having equal intensity at different frequencies"),
 	DEFINE_METASOUND_ENUM_END()
 
 	class FNoiseOperator : public IOperator
@@ -162,11 +162,11 @@ namespace Metasound
 	{
 		static const FVertexInterface Interface(
 			FInputVertexInterface(
-				TInputDataVertexModel<int32>(SeedPinName, LOCTEXT("SeedPinDescription", "Seed for seeding the Random Number Generator, -1 (default) will use current time"), FNoiseOperator::DefaultSeed),
-				TInputDataVertexModel<FEnumNoiseType>(TriggerTypePinName, LOCTEXT("NoiseTypeDescription", "Type of Noise to Generate"))
+				TInputDataVertexModel<int32>(SeedPinName, METASOUND_LOCTEXT("SeedPinDescription", "Seed for seeding the Random Number Generator, -1 (default) will use current time"), FNoiseOperator::DefaultSeed),
+				TInputDataVertexModel<FEnumNoiseType>(TriggerTypePinName, METASOUND_LOCTEXT("NoiseTypeDescription", "Type of Noise to Generate"))
 			),
 			FOutputVertexInterface(
-				TOutputDataVertexModel<FAudioBuffer>(OutPinName, LOCTEXT("AudioOutDescription", "Audio output"))
+				TOutputDataVertexModel<FAudioBuffer>(OutPinName, METASOUND_LOCTEXT("AudioOutDescription", "Audio output"))
 			)
 		);
 		return Interface;
@@ -180,8 +180,8 @@ namespace Metasound
 			Info.ClassName = { StandardNodes::Namespace, TEXT("Noise"), StandardNodes::AudioVariant };
 			Info.MajorVersion = 1;
 			Info.MinorVersion = 0;
-			Info.DisplayName = LOCTEXT("Metasound_NoiseNodeDisplayNameX", "Noise");
-			Info.Description = LOCTEXT("Metasound_NoiseNodeDescription", "Noise Generator that produces different types of noise");
+			Info.DisplayName = METASOUND_LOCTEXT("Metasound_NoiseNodeDisplayNameX", "Noise");
+			Info.Description = METASOUND_LOCTEXT("Metasound_NoiseNodeDescription", "Noise Generator that produces different types of noise");
 			Info.Author = PluginAuthor;
 			Info.PromptIfMissing = PluginNodeMissingPrompt;
 			Info.DefaultInterface = DeclareVertexInterface();
diff --git a/Engine/Plugins/Runtime/Metasound/Source/MetasoundStandardNodes/Private/MetasoundOscillatorNodes.cpp b/Engine/Plugins/Runtime/Metasound/Source/MetasoundStandardNodes/Private/MetasoundOscillatorNodes.cpp
--- a/Engine/Plugins/Runtime/Metasound/Source/MetasoundStandardNodes/Private/MetasoundOscillatorNodes.cpp
+++ b/Engine/Plugins/Runtime/Metasound/Source/MetasoundStandardNodes/Private/MetasoundOscillatorNodes.cpp
@@ -533,16 +533,16 @@ namespace Metasound
 			static const FVertexInterface Interface
 			{
 				FInputVertexInterface{
-					TInputDataVertexModel<bool>(EnabledPinName, LOCTEXT("OscActivateDescription", "Enable the oscillator."), true),
-					TInputDataVertexModel<bool>(BiPolarPinName, LOCTEXT("OscBaseBiPolarDescription", "If the output is Bi-Polar (-1..1) or Uni-Polar (0..1)"), true),
-					TInputDataVertexModel<float>(BaseFrequencyPinName, LOCTEXT("OscBaseFrequencyDescription", "Base Frequency of Oscillator in HZ"), 440.f),
-					TInputDataVertexModel<FAudioBuffer>(FrequencyModPinName, LOCTEXT("OscModFrequencyDescription", "Modulation Frequency Input (for doing FM)")),
-					TInputDataVertexModel<FTrigger>(PhaseResetPinName, LOCTEXT("OscPhaseResetDescription", "Phase Reset")),
-					TInputDataVertexModel<float>(PhaseOffsetPinName, LOCTEXT("OscPhaseOffsetDescription", "Phase Offset In Degrees (0..360)"), 0.f),
-					TInputDataVertexModel<float>(GlideFactorPinName, LOCTEXT("OscGlideFactorDescription", "The amount of glide to use when changing frequencies. 0.0 = no glide, 1.0 = lots of glide."), 0.f)
+					TInputDataVertexModel<bool>(EnabledPinName, METASOUND_LOCTEXT("OscActivateDescription", "Enable the oscillator."), true),
+					TInputDataVertexModel<bool>(BiPolarPinName, METASOUND_LOCTEXT("OscBaseBiPolarDescription", "If the output is Bi-Polar (-1..1) or Uni-Polar (0..1)"), true),
+					TInputDataVertexModel<float>(BaseFrequencyPinName, METASOUND_LOCTEXT("OscBaseFrequencyDescription", "Base Frequency of Oscillator in HZ"), 440.f),
+					TInputDataVertexModel<FAudioBuffer>(FrequencyModPinName, METASOUND_LOCTEXT("OscModFrequencyDescription", "Modulation Frequency Input (for doing FM)")),
+					TInputDataVertexModel<FTrigger>(PhaseResetPinName, METASOUND_LOCTEXT("OscPhaseResetDescription", "Phase Reset")),
+					TInputDataVertexModel<float>(PhaseOffsetPinName, METASOUND_LOCTEXT("OscPhaseOffsetDescription", "Phase Offset In Degrees (0..360)"), 0.f),
+					TInputDataVertexModel<float>(GlideFactorPinName, METASOUND_LOCTEXT("OscGlideFactorDescription", "The amount of glide to use when changing frequencies. 0.0 = no glide, 1.0 = lots of glide."), 0.f)
 				},
 				FOutputVertexInterface{
-					TOutputDataVertexModel<FAudioBuffer>(TEXT("Audio"), LOCTEXT("AudioTooltip", "The output audio"))
+					TOutputDataVertexModel<FAudioBuffer>(TEXT("Audio"), METASOUND_LOCTEXT("AudioTooltip", "The output audio"))
 				}
 			};
 			return Interface;
@@ -740,11 +740,11 @@ namespace Metasound
 		FEnumSineGenerationType, FEnumSineGenerationTypeInfo, FEnumSineGenerationTypeReadRef, FEnumSineGenerationTypeWriteRef);
 
 	DEFINE_METASOUND_ENUM_BEGIN(ESineGenerationType, FEnumSineGenerationType, "SineGenerationType")
-		DEFINE_METASOUND_ENUM_ENTRY(ESineGenerationType::Rotation, LOCTEXT("RotationDescription", "2D Rotation"), LOCTEXT("RotationDescriptionTT", "Rotates around the unit circle generate the sine. Note: Glide and audio rate FM modulation is not supported with the 2D rotator.")),
-		DEFINE_METASOUND_ENUM_ENTRY(ESineGenerationType::Sinf, LOCTEXT("SinfDescription", "Pure Math"), LOCTEXT("SinfDescriptionTT", "Uses the standard math library (Sinf) to generate the sine (most expensive)")),
-		DEFINE_METASOUND_ENUM_ENTRY(ESineGenerationType::Bhaskara, LOCTEXT("BhaskaraDescription", "Bhaskara"), LOCTEXT("BhaskaraDescriptionTT", "Sine approximation using Bhaskara technique discovered in 7th century")),
-		DEFINE_METASOUND_ENUM_ENTRY(ESineGenerationType::Wavetable, LOCTEXT("WavetableDescription", "Wavetable"), LOCTEXT("WavetableDescriptionTT", "Uses a wavetable to generate the sine")),
-		//DEFINE_METASOUND_ENUM_ENTRY(ESineGenerationType::FilterResonanceFeedback, LOCTEXT("FRFDescription", "Filter resonance"), LOCTEXT("FRFDescriptionTT", "Filter resonance/feedback")),
+		DEFINE_METASOUND_ENUM_ENTRY(ESineGenerationType::Rotation, "RotationDescription", "2D Rotation", "RotationDescriptionTT", "Rotates around the unit circle generate the sine. Note: Glide and audio rate FM modulation is not supported with the 2D rotator."),
+		DEFINE_METASOUND_ENUM_ENTRY(ESineGenerationType::Sinf, "SinfDescription", "Pure Math", "SinfDescriptionTT", "Uses the standard math library (Sinf) to generate the sine (most expensive)"),
+		DEFINE_METASOUND_ENUM_ENTRY(ESineGenerationType::Bhaskara, "BhaskaraDescription", "Bhaskara", "BhaskaraDescriptionTT", "Sine approximation using Bhaskara technique discovered in 7th century"),
+		DEFINE_METASOUND_ENUM_ENTRY(ESineGenerationType::Wavetable, "WavetableDescription", "Wavetable", "WavetableDescriptionTT", "Uses a wavetable to generate the sine"),
+		//DEFINE_METASOUND_ENUM_ENTRY(ESineGenerationType::FilterResonanceFeedback, "FRFDescription", "Filter resonance", "FRFDescriptionTT", "Filter resonance/feedback"),
 	DEFINE_METASOUND_ENUM_END()
 
 	class FSineOscilatorNode::FFactory : public FOscilatorFactoryBase
@@ -762,8 +762,8 @@ namespace Metasound
 				Info.ClassName = { StandardNodes::Namespace, TEXT("Sine"), StandardNodes::AudioVariant };
 				Info.MajorVersion = 1;
 				Info.MinorVersion = 1;
-				Info.DisplayName = LOCTEXT("Metasound_SineNodeDisplayName", "Sine");
-				Info.Description = LOCTEXT("Metasound_SineNodeDescription", "Emits an audio signal of a sinusoid.");
+				Info.DisplayName = METASOUND_LOCTEXT("Metasound_SineNodeDisplayName", "Sine");
+				Info.Description = METASOUND_LOCTEXT("Metasound_SineNodeDescription", "Emits an audio signal of a sinusoid.");
 				Info.Author = PluginAuthor;
 				Info.PromptIfMissing = PluginNodeMissingPrompt;
 				Info.DefaultInterface = GetVertexInterface();
@@ -779,7 +779,7 @@ namespace Metasound
 			{
 				FVertexInterface Interface = GetCommmonVertexInterface();
 				Interface.GetInputInterface().Add(
-					TInputDataVertexModel<FEnumSineGenerationType>(SineTypePin, LOCTEXT("SineTypeDescription", "Type of the Sinewave Generator"), static_cast<int32>(ESineGenerationType::Wavetable))
+					TInputDataVertexModel<FEnumSineGenerationType>(SineTypePin, METASOUND_LOCTEXT("SineTypeDescription", "Type of the Sinewave Generator"), static_cast<int32>(ESineGenerationType::Wavetable))
 				);
 				return Interface;
 			};
@@ -863,9 +863,9 @@ namespace Metasound
 		FEnumSawGenerationType, FEnumSawGenerationTypeInfo, FSawGenerationTypeReadRef, FEnumSawGenerationTypeWriteRef);
 
 	DEFINE_METASOUND_ENUM_BEGIN(ESawGenerationType, FEnumSawGenerationType, "SawGenerationType")
-		DEFINE_METASOUND_ENUM_ENTRY(ESawGenerationType::PolySmooth, LOCTEXT("SawPolySmoothDescription", "Poly Smooth"), LOCTEXT("PolySmoothDescriptionTT", "PolySmooth (i.e. BLEP)")),
-		DEFINE_METASOUND_ENUM_ENTRY(ESawGenerationType::Trivial, LOCTEXT("SawTrivialDescription", "Trivial"), LOCTEXT("TrivialDescriptionTT", "The most basic raw implementation")),
-		//DEFINE_METASOUND_ENUM_ENTRY(ESawGenerationType::Wavetable, LOCTEXT("SawWavetableDescription", "Wavetable"), LOCTEXT("SawWavetableDescriptionTT", "Use a Wavetable iterpolation to generate the Waveform"))
+		DEFINE_METASOUND_ENUM_ENTRY(ESawGenerationType::PolySmooth, "SawPolySmoothDescription", "Poly Smooth", "PolySmoothDescriptionTT", "PolySmooth (i.e. BLEP)"),
+		DEFINE_METASOUND_ENUM_ENTRY(ESawGenerationType::Trivial, "SawTrivialDescription", "Trivial", "TrivialDescriptionTT", "The most basic raw implementation"),
+		//DEFINE_METASOUND_ENUM_ENTRY(ESawGenerationType::Wavetable, "SawWavetableDescription", "Wavetable", "SawWavetableDescriptionTT", "Use a Wavetable iterpolation to generate the Waveform")
 	DEFINE_METASOUND_ENUM_END()
 
 	class FSawOscilatorNode::FFactory : public FOscilatorFactoryBase
@@ -885,8 +885,8 @@ namespace Metasound
 				Info.ClassName = { StandardNodes::Namespace, TEXT("Saw"), StandardNodes::AudioVariant };
 				Info.MajorVersion = 1;
 				Info.MinorVersion = 0;
-				Info.DisplayName = LOCTEXT("Metasound_SawNodeDisplayName", "Saw");
-				Info.Description = LOCTEXT("Metasound_SawNodeDescription", "Emits a Saw wave");
+				Info.DisplayName = METASOUND_LOCTEXT("Metasound_SawNodeDisplayName", "Saw");
+				Info.Description = METASOUND_LOCTEXT("Metasound_SawNodeDescription", "Emits a Saw wave");
 				Info.Author = PluginAuthor;
 				Info.PromptIfMissing = PluginNodeMissingPrompt;
 				Info.DefaultInterface = GetVertexInterface();
@@ -902,7 +902,7 @@ namespace Metasound
 			{
 				FVertexInterface Interface = GetCommmonVertexInterface();
 				Interface.GetInputInterface().Add(
-					TInputDataVertexModel<FEnumSawGenerationType>(SawTypePin, LOCTEXT("Saw Type", "Which type of Saw Generator to Create"))
+					TInputDataVertexModel<FEnumSawGenerationType>(SawTypePin, METASOUND_LOCTEXT("Saw Type", "Which type of Saw Generator to Create"))
 				);
 				return Interface;
 			};
@@ -982,9 +982,9 @@ namespace Metasound
 	FEnumSquareGenerationType, FEnumSquareGenerationTypeInfo, FSquareGenerationTypeReadRef, FEnumSquareGenerationTypeWriteRef);
 
 	DEFINE_METASOUND_ENUM_BEGIN(ESquareGenerationType, FEnumSquareGenerationType, "SquareGenerationType")
-		DEFINE_METASOUND_ENUM_ENTRY(ESquareGenerationType::PolySmooth, LOCTEXT("SquarePolySmoothDescription", "Poly Smooth"), LOCTEXT("PolySmoothDescriptionTT", "PolySmooth (i.e. BLEP)")),
-		DEFINE_METASOUND_ENUM_ENTRY(ESquareGenerationType::Trivial, LOCTEXT("SquareTrivialDescription", "Trivial"), LOCTEXT("SquareTrivialDescriptionTT", "The most basic raw implementation. Does not obey pulse width.")),
-		//DEFINE_METASOUND_ENUM_ENTRY(ESquareGenerationType::Wavetable, LOCTEXT("SquareWavetableDescription", "Wavetable"), LOCTEXT("SquareWavetableDescriptionTT", "Use a Wavetable interpolation to generate the Waveform"))
+		DEFINE_METASOUND_ENUM_ENTRY(ESquareGenerationType::PolySmooth, "SquarePolySmoothDescription", "Poly Smooth", "PolySmoothDescriptionTT", "PolySmooth (i.e. BLEP)"),
+		DEFINE_METASOUND_ENUM_ENTRY(ESquareGenerationType::Trivial, "SquareTrivialDescription", "Trivial", "SquareTrivialDescriptionTT", "The most basic raw implementation. Does not obey pulse width."),
+		//DEFINE_METASOUND_ENUM_ENTRY(ESquareGenerationType::Wavetable, "SquareWavetableDescription", "Wavetable", "SquareWavetableDescriptionTT", "Use a Wavetable interpolation to generate the Waveform")
 	DEFINE_METASOUND_ENUM_END()
 	
 	template<typename GeneratorPolicy>
@@ -1116,8 +1116,8 @@ namespace Metasound
 				Info.ClassName = { StandardNodes::Namespace, TEXT("Square"), StandardNodes::AudioVariant };
 				Info.MajorVersion = 1;
 				Info.MinorVersion = 0;
-				Info.DisplayName = LOCTEXT("Metasound_SquareNodeDisplayName", "Square");
-				Info.Description = LOCTEXT("Metasound_SquareNodeDescription", "Emits a Square wave");
+				Info.DisplayName = METASOUND_LOCTEXT("Metasound_SquareNodeDisplayName", "Square");
+				Info.Description = METASOUND_LOCTEXT("Metasound_SquareNodeDescription", "Emits a Square wave");
 				Info.Author = PluginAuthor;
 				Info.PromptIfMissing = PluginNodeMissingPrompt;
 				Info.DefaultInterface = GetVertexInterface();
@@ -1132,8 +1132,8 @@ namespace Metasound
 			auto MakeInterface = []() -> FVertexInterface
 			{
 				FVertexInterface Interface = GetCommmonVertexInterface();
-				Interface.GetInputInterface().Add(TInputDataVertexModel<FEnumSquareGenerationType>(SquareTypePinName, LOCTEXT("SquareTypeDescription", "The generator type to make the squarewave")));
-				Interface.GetInputInterface().Add(TInputDataVertexModel<float>(PulseWidthPinName, LOCTEXT("PhaseWidthDescription", "The Width of the square part of the wave"), 0.5f));
+				Interface.GetInputInterface().Add(TInputDataVertexModel<FEnumSquareGenerationType>(SquareTypePinName, METASOUND_LOCTEXT("SquareTypeDescription", "The generator type to make the squarewave")));
+				Interface.GetInputInterface().Add(TInputDataVertexModel<float>(PulseWidthPinName, METASOUND_LOCTEXT("PhaseWidthDescription", "The Width of the square part of the wave"), 0.5f));
 				return Interface;
 			};
 			static const FVertexInterface Interface = MakeInterface();
@@ -1166,9 +1166,9 @@ namespace Metasound
 	FEnumTriangleGenerationType, FEnumTriangleGenerationTypeInfo, FTriangleGenerationTypeReadRef, FEnumTriangleGenerationTypeWriteRef);
 
 	DEFINE_METASOUND_ENUM_BEGIN(ETriangleGenerationType, FEnumTriangleGenerationType, "TriangleGenerationType")
-		DEFINE_METASOUND_ENUM_ENTRY(ETriangleGenerationType::PolySmooth, LOCTEXT("TrianglePolySmoothDescription", "Poly Smooth"), LOCTEXT("PolySmoothDescriptionTT", "PolySmooth (i.e. BLEP)")),
-		DEFINE_METASOUND_ENUM_ENTRY(ETriangleGenerationType::Trivial, LOCTEXT("TriangleTrivialDescription", "Trivial"), LOCTEXT("TriangleTrivialDescriptionTT", "The most basic raw implementation")),
-		//DEFINE_METASOUND_ENUM_ENTRY(ETriangleGenerationType::Wavetable, LOCTEXT("TriangleWavetableDescription", "Wavetable"), LOCTEXT("TriangleWavetableDescriptionTT", "Use a Wavetable iterpolation to generate the Waveform"))
+		DEFINE_METASOUND_ENUM_ENTRY(ETriangleGenerationType::PolySmooth, "TrianglePolySmoothDescription", "Poly Smooth", "PolySmoothDescriptionTT", "PolySmooth (i.e. BLEP)"),
+		DEFINE_METASOUND_ENUM_ENTRY(ETriangleGenerationType::Trivial, "TriangleTrivialDescription", "Trivial", "TriangleTrivialDescriptionTT", "The most basic raw implementation"),
+		//DEFINE_METASOUND_ENUM_ENTRY(ETriangleGenerationType::Wavetable, "TriangleWavetableDescription", "Wavetable", "TriangleWavetableDescriptionTT", "Use a Wavetable iterpolation to generate the Waveform")
 	DEFINE_METASOUND_ENUM_END()
 
 	class FTriangleOscilatorNode::FFactory : public FOscilatorFactoryBase
@@ -1230,8 +1230,8 @@ namespace Metasound
 				Info.ClassName = { StandardNodes::Namespace, TEXT("Triangle"), StandardNodes::AudioVariant };
 				Info.MajorVersion = 1;
 				Info.MinorVersion = 0;
-				Info.DisplayName = LOCTEXT("Metasound_TriangleNodeDisplayName", "Triangle");
-				Info.Description = LOCTEXT("Metasound_TriangleNodeDescription", "Emits a Triangle wave");
+				Info.DisplayName = METASOUND_LOCTEXT("Metasound_TriangleNodeDisplayName", "Triangle");
+				Info.Description = METASOUND_LOCTEXT("Metasound_TriangleNodeDescription", "Emits a Triangle wave");
 				Info.Author = PluginAuthor;
 				Info.PromptIfMissing = PluginNodeMissingPrompt;
 				Info.DefaultInterface = GetVertexInterface();
@@ -1247,7 +1247,7 @@ namespace Metasound
 			{
 				FVertexInterface Interface = GetCommmonVertexInterface();
 				Interface.GetInputInterface().Add(
-					TInputDataVertexModel<FEnumTriangleGenerationType>(TriangeTypePinName, LOCTEXT("TriangleTypeDescription", "The generator type to make the triangle wave"))
+					TInputDataVertexModel<FEnumTriangleGenerationType>(TriangeTypePinName, METASOUND_LOCTEXT("TriangleTypeDescription", "The generator type to make the triangle wave"))
 				);
 				return Interface;
 			};
@@ -1284,10 +1284,10 @@ namespace Metasound
 		FEnumLfoWaveshapeType, FEnumLfoWaveshapeTypeInfo, FEnumLfoWaveshapeTypeReadRef, FEnumLfoWaveshapeTypeWriteRef);
 
 	DEFINE_METASOUND_ENUM_BEGIN(ELfoWaveshapeType, FEnumLfoWaveshapeType, "LfoWaveshapeType")
-		DEFINE_METASOUND_ENUM_ENTRY(ELfoWaveshapeType::Sine, LOCTEXT("LfoWaveShapeSineDescription", "Sine"), LOCTEXT("LfoWaveShapeSineDescriptionTT", "Sinewave Low Frequency Oscillator")),
-		DEFINE_METASOUND_ENUM_ENTRY(ELfoWaveshapeType::Saw, LOCTEXT("LfoWaveShapeSawDescription", "Saw"), LOCTEXT("LfoWaveShapeSawDescriptionTT", "Sawtooth Low Frequency Oscillator")),
-		DEFINE_METASOUND_ENUM_ENTRY(ELfoWaveshapeType::Triangle, LOCTEXT("LfoWaveShapeTriangleDescription", "Triangle"), LOCTEXT("LfoWaveShapeTriangleDescriptionTT", "Triangle shape Frequency Oscillator")),
-		DEFINE_METASOUND_ENUM_ENTRY(ELfoWaveshapeType::Square, LOCTEXT("LfoWaveShapeSquareDescription", "Square"), LOCTEXT("LfoWaveShapeSquareDescriptionTT", "Square shape Low Frequency Oscillator"))
+		DEFINE_METASOUND_ENUM_ENTRY(ELfoWaveshapeType::Sine, "LfoWaveShapeSineDescription", "Sine", "LfoWaveShapeSineDescriptionTT", "Sinewave Low Frequency Oscillator"),
+		DEFINE_METASOUND_ENUM_ENTRY(ELfoWaveshapeType::Saw, "LfoWaveShapeSawDescription", "Saw", "LfoWaveShapeSawDescriptionTT", "Sawtooth Low Frequency Oscillator"),
+		DEFINE_METASOUND_ENUM_ENTRY(ELfoWaveshapeType::Triangle, "LfoWaveShapeTriangleDescription", "Triangle", "LfoWaveShapeTriangleDescriptionTT", "Triangle shape Frequency Oscillator"),
+		DEFINE_METASOUND_ENUM_ENTRY(ELfoWaveshapeType::Square, "LfoWaveShapeSquareDescription", "Square", "LfoWaveShapeSquareDescriptionTT", "Square shape Low Frequency Oscillator")
 	DEFINE_METASOUND_ENUM_END()
 
 	// Blockrate All-Purpose Oscillator
@@ -1303,16 +1303,16 @@ namespace Metasound
 			static const FVertexInterface Interface
 			{
 				FInputVertexInterface{
-					TInputDataVertexModel<float>(FOscilatorFactoryBase::BaseFrequencyPinName, LOCTEXT("LfoFrequencyDescription", "Frequency of LFO (Hz), clamped at blockrate"), 5.f),
-					TInputDataVertexModel<FEnumLfoWaveshapeType>(WaveshapePinName, LOCTEXT("LfoShapeDescription", "Waveshape of the LFO")),
-					TInputDataVertexModel<float>(FOscilatorFactoryBase::MinOutputValuePinName, LOCTEXT("LfoMinValueDescription", "The minimum output value."), -1.0f),
-					TInputDataVertexModel<float>(FOscilatorFactoryBase::MaxOutputValuePinName, LOCTEXT("LfoMaxValueDescription", "The maximum output value."), 1.0f),
-					TInputDataVertexModel<FTrigger>(FOscilatorFactoryBase::PhaseResetPinName, LOCTEXT("LfoPhaseResetDescription", "Phase Reset (block rate only)")),
-					TInputDataVertexModel<float>(FOscilatorFactoryBase::PhaseOffsetPinName, LOCTEXT("LfoPhaseOffsetDescription", "Phase Offset In Degrees (0..360)"), 0.f),
-					TInputDataVertexModel<float>(FSquareOscilatorNode::FFactory::PulseWidthPinName, LOCTEXT("LfoPulseWidthDescription", "Pulse Width (0..1)"), 0.5f)
+					TInputDataVertexModel<float>(FOscilatorFactoryBase::BaseFrequencyPinName, METASOUND_LOCTEXT("LfoFrequencyDescription", "Frequency of LFO (Hz), clamped at blockrate"), 5.f),
+					TInputDataVertexModel<FEnumLfoWaveshapeType>(WaveshapePinName, METASOUND_LOCTEXT("LfoShapeDescription", "Waveshape of the LFO")),
+					TInputDataVertexModel<float>(FOscilatorFactoryBase::MinOutputValuePinName, METASOUND_LOCTEXT("LfoMinValueDescription", "The minimum output value."), -1.0f),
+					TInputDataVertexModel<float>(FOscilatorFactoryBase::MaxOutputValuePinName, METASOUND_LOCTEXT("LfoMaxValueDescription", "The maximum output value."), 1.0f),
+					TInputDataVertexModel<FTrigger>(FOscilatorFactoryBase::PhaseResetPinName, METASOUND_LOCTEXT("LfoPhaseResetDescription", "Phase Reset (block rate only)")),
+					TInputDataVertexModel<float>(FOscilatorFactoryBase::PhaseOffsetPinName, METASOUND_LOCTEXT("LfoPhaseOffsetDescription", "Phase Offset In Degrees (0..360)"), 0.f),
+					TInputDataVertexModel<float>(FSquareOscilatorNode::FFactory::PulseWidthPinName, METASOUND_LOCTEXT("LfoPulseWidthDescription", "Pulse Width (0..1)"), 0.5f)
 				},
 				FOutputVertexInterface{
-					TOutputDataVertexModel<float>(LfoOutPinName, LOCTEXT("LfoOutputDescription", "Output of the LFO (blockrate)"))
+					TOutputDataVertexModel<float>(LfoOutPinName, METASOUND_LOCTEXT("LfoOutputDescription", "Output of the LFO (blockrate)"))
 				}
 			};
 			return Interface;
@@ -1326,8 +1326,8 @@ namespace Metasound
 				Info.ClassName = { StandardNodes::Namespace, TEXT("LFO"), StandardNodes::AudioVariant };
 				Info.MajorVersion = 1;
 				Info.MinorVersion = 0;
-				Info.DisplayName = LOCTEXT("Metasound_LfoNodeDisplayName", "LFO");
-				Info.Description = LOCTEXT("Metasound_LfoNodeDescription", "Low frequency oscillator < blockrate");
+				Info.DisplayName = METASOUND_LOCTEXT("Metasound_LfoNodeDisplayName", "LFO");
+				Info.Description = METASOUND_LOCTEXT("Metasound_LfoNodeDescription", "Low frequency oscillator < blockrate");
 				Info.Author = PluginAuthor;
 				Info.PromptIfMissing = PluginNodeMissingPrompt;
 				Info.DefaultInterface = GetVertexInterface();
diff --git a/Engine/Plugins/Runtime/Metasound/Source/MetasoundStandardNodes/Private/MetasoundRandomNode.cpp b/Engine/Plugins/Runtime/Metasound/Source/MetasoundStandardNodes/Private/MetasoundRandomNode.cpp
--- a/Engine/Plugins/Runtime/Metasound/Source/MetasoundStandardNodes/Private/MetasoundRandomNode.cpp
+++ b/Engine/Plugins/Runtime/Metasound/Source/MetasoundStandardNodes/Private/MetasoundRandomNode.cpp
@@ -61,42 +61,42 @@ namespace Metasound
 
 		static FText GetNextTriggerDescription()
 		{
-			return LOCTEXT("RandomNode_NextTT", "Trigger to generate the next random integer.");
+			return METASOUND_LOCTEXT("RandomNode_NextTT", "Trigger to generate the next random integer.");
 		}
 
 		static FText GetResetDescription()
 		{
-			return LOCTEXT("RandomNode_ResetTT", "Trigger to reset the random sequence with the supplied seed. Useful to get randomized repetition.");
+			return METASOUND_LOCTEXT("RandomNode_ResetTT", "Trigger to reset the random sequence with the supplied seed. Useful to get randomized repetition.");
 		}
 
 		static FText GetSeedDescription()
 		{
-			return LOCTEXT("RandomNode_SeedTT", "The seed value to use for the random node. Set to -1 to use a random seed.");
+			return METASOUND_LOCTEXT("RandomNode_SeedTT", "The seed value to use for the random node. Set to -1 to use a random seed.");
 		}
 
 		static FText GetMinDescription()
 		{
-			return LOCTEXT("RandomNode_MinTT", "Min random value.");
+			return METASOUND_LOCTEXT("RandomNode_MinTT", "Min random value.");
 		}
 
 		static FText GetMaxDescription()
 		{
-			return LOCTEXT("RandomNode_MaxTT", "Max random value.");
+			return METASOUND_LOCTEXT("RandomNode_MaxTT", "Max random value.");
 		}
 
 		static FText GetOutputDescription()
 		{
-			return LOCTEXT("RandomNode_OutputTT", "The randomly generated value.");
+			return METASOUND_LOCTEXT("RandomNode_OutputTT", "The randomly generated value.");
 		}
 
 		static FText GetOutputOnNextDescription()
 		{
-			return LOCTEXT("RandomNode_OutputNextTT", "Triggers when next is triggered.");
+			return METASOUND_LOCTEXT("RandomNode_OutputNextTT", "Triggers when next is triggered.");
 		}
 
 		static FText GetOutputOnResetDescription()
 		{
-			return LOCTEXT("RandomNode_OutputResetTT", "Triggers when reset is triggered.");
+			return METASOUND_LOCTEXT("RandomNode_OutputResetTT", "Triggers when reset is triggered.");
 		}
 	}
 
diff --git a/Engine/Plugins/Runtime/Metasound/Source/MetasoundStandardNodes/Private/MetasoundSampleAndHoldNode.cpp b/Engine/Plugins/Runtime/Metasound/Source/MetasoundStandardNodes/Private/MetasoundSampleAndHoldNode.cpp
--- a/Engine/Plugins/Runtime/Metasound/Source/MetasoundStandardNodes/Private/MetasoundSampleAndHoldNode.cpp
+++ b/Engine/Plugins/Runtime/Metasound/Source/MetasoundStandardNodes/Private/MetasoundSampleAndHoldNode.cpp
@@ -41,8 +41,8 @@ namespace Metasound
 					FNodeClassName { "SampleAndHold", "Sample And Hold", StandardNodes::AudioVariant },
 					1, // Major Version
 					0, // Minor Version
-					LOCTEXT("SampleAndHoldDisplayName", "Sample And Hold"),
-					LOCTEXT("SampleAndHoldDesc", "Will output a single value of the input audio signal when triggered."),
+					METASOUND_LOCTEXT("SampleAndHoldDisplayName", "Sample And Hold"),
+					METASOUND_LOCTEXT("SampleAndHoldDesc", "Will output a single value of the input audio signal when triggered."),
 					PluginAuthor,
 					PluginNodeMissingPrompt,
 					NodeInterface,
diff --git a/Engine/Plugins/Runtime/Metasound/Source/MetasoundStandardNodes/Private/MetasoundSemitoneToFreqMultiplierNode.cpp b/Engine/Plugins/Runtime/Metasound/Source/MetasoundStandardNodes/Private/MetasoundSemitoneToFreqMultiplierNode.cpp
--- a/Engine/Plugins/Runtime/Metasound/Source/MetasoundStandardNodes/Private/MetasoundSemitoneToFreqMultiplierNode.cpp
+++ b/Engine/Plugins/Runtime/Metasound/Source/MetasoundStandardNodes/Private/MetasoundSemitoneToFreqMultiplierNode.cpp
@@ -109,8 +109,8 @@ namespace Metasound
 		{
 			const FName DataTypeName = GetMetasoundDataTypeName<float>();
 			const FName OperatorName = TEXT("Semitone to Frequency Multiplier");
-			const FText NodeDisplayName = LOCTEXT("Metasound_SemitoneToFrequencyMultiplierName", "Semitone to Frequency Multiplier");
-			const FText NodeDescription = LOCTEXT("Metasound_SemitoneToFrequencyMultiplierDescription", "Converts a number of semitones to the corresponding frequency multiplier.");
+			const FText NodeDisplayName = METASOUND_LOCTEXT("Metasound_SemitoneToFrequencyMultiplierName", "Semitone to Frequency Multiplier");
+			const FText NodeDescription = METASOUND_LOCTEXT("Metasound_SemitoneToFrequencyMultiplierDescription", "Converts a number of semitones to the corresponding frequency multiplier.");
 
 			FNodeClassMetadata Info;
 			Info.ClassName = { StandardNodes::Namespace, OperatorName, DataTypeName };
diff --git a/Engine/Plugins/Runtime/Metasound/Source/MetasoundStandardNodes/Private/MetasoundStandardNodesCategories.cpp b/Engine/Plugins/Runtime/Metasound/Source/MetasoundStandardNodes/Private/MetasoundStandardNodesCategories.cpp
--- a/Engine/Plugins/Runtime/Metasound/Source/MetasoundStandardNodes/Private/MetasoundStandardNodesCategories.cpp
+++ b/Engine/Plugins/Runtime/Metasound/Source/MetasoundStandardNodes/Private/MetasoundStandardNodesCategories.cpp
@@ -8,19 +8,19 @@ namespace Metasound
 {
 	namespace NodeCategories
 	{
-		const FText Debug = { LOCTEXT("Metasound_DebugCategory", "Debug") };
-		const FText Delays = { LOCTEXT("Metasound_EffectsCategory", "Delays") };
-		const FText Dynamics = { LOCTEXT("Metasound_DynamicsCategory", "Dynamics") };
-		const FText Envelopes = { LOCTEXT("Metasound_EnvelopesCategory", "Envelopes") };
-		const FText Filters = { LOCTEXT("Metasound_FiltersCategory", "Filters") };
-		const FText Generators = { LOCTEXT("Metasound_GeneratorsCategory", "Generators") };
-		const FText Io = { LOCTEXT("Metasound_IoCategory", "External IO") };
-		const FText Math = { LOCTEXT("Metasound_MathCategory", "Math") };
-		const FText Mix = { LOCTEXT("Metasound_MixCategory", "Mix") };
-		const FText Music = { LOCTEXT("Metasound_MusicCategory", "Music") };
-		const FText RandomUtils = { LOCTEXT("Metasound_RandomCategory", "Random") };
-		const FText Spatialization = { LOCTEXT("Metasound_SpatializationCategory", "Spatialization") };
-		const FText Trigger = { LOCTEXT("Metasound_TriggerCategory", "Triggers") };
+		const FText Debug = { METASOUND_LOCTEXT("Metasound_DebugCategory", "Debug") };
+		const FText Delays = { METASOUND_LOCTEXT("Metasound_EffectsCategory", "Delays") };
+		const FText Dynamics = { METASOUND_LOCTEXT("Metasound_DynamicsCategory", "Dynamics") };
+		const FText Envelopes = { METASOUND_LOCTEXT("Metasound_EnvelopesCategory", "Envelopes") };
+		const FText Filters = { METASOUND_LOCTEXT("Metasound_FiltersCategory", "Filters") };
+		const FText Generators = { METASOUND_LOCTEXT("Metasound_GeneratorsCategory", "Generators") };
+		const FText Io = { METASOUND_LOCTEXT("Metasound_IoCategory", "External IO") };
+		const FText Math = { METASOUND_LOCTEXT("Metasound_MathCategory", "Math") };
+		const FText Mix = { METASOUND_LOCTEXT("Metasound_MixCategory", "Mix") };
+		const FText Music = { METASOUND_LOCTEXT("Metasound_MusicCategory", "Music") };
+		const FText RandomUtils = { METASOUND_LOCTEXT("Metasound_RandomCategory", "Random") };
+		const FText Spatialization = { METASOUND_LOCTEXT("Metasound_SpatializationCategory", "Spatialization") };
+		const FText Trigger = { METASOUND_LOCTEXT("Metasound_TriggerCategory", "Triggers") };
 	}
 }
 
diff --git a/Engine/Plugins/Runtime/Metasound/Source/MetasoundStandardNodes/Private/MetasoundStereoDelayNode.cpp b/Engine/Plugins/Runtime/Metasound/Source/MetasoundStandardNodes/Private/MetasoundStereoDelayNode.cpp
--- a/Engine/Plugins/Runtime/Metasound/Source/MetasoundStandardNodes/Private/MetasoundStereoDelayNode.cpp
+++ b/Engine/Plugins/Runtime/Metasound/Source/MetasoundStandardNodes/Private/MetasoundStereoDelayNode.cpp
@@ -45,9 +45,9 @@ namespace Metasound
 		FEnumStereoDelayMode, FEnumStereoDelayModeInfo, FStereoDelayModeReadRef, FEnumStereoDelayModeWriteRef);
 
 	DEFINE_METASOUND_ENUM_BEGIN(EStereoDelayMode, FEnumStereoDelayMode, "StereoDelayMode")
-		DEFINE_METASOUND_ENUM_ENTRY(EStereoDelayMode::Normal, LOCTEXT("StereoDelayModeNormalDescription", "Normal"), LOCTEXT("StereoDelayModeNormalDescriptionTT", "Left input mixes with left delay output and feeds to left output.")),
-		DEFINE_METASOUND_ENUM_ENTRY(EStereoDelayMode::Cross, LOCTEXT("StereoDelayModeCrossDescription", "Cross"), LOCTEXT("StereoDelayModeCrossDescriptionTT", "Left input mixes with right delay output and feeds to right output.")),
-		DEFINE_METASOUND_ENUM_ENTRY(EStereoDelayMode::PingPong, LOCTEXT("StereoDelayModePingPongDescription", "Ping Pong"), LOCTEXT("StereoDelayModePingPongDescriptionTT", "Left input mixes with left delay output and feeds to right output.")),
+		DEFINE_METASOUND_ENUM_ENTRY(EStereoDelayMode::Normal, "StereoDelayModeNormalDescription", "Normal", "StereoDelayModeNormalDescriptionTT", "Left input mixes with left delay output and feeds to left output."),
+		DEFINE_METASOUND_ENUM_ENTRY(EStereoDelayMode::Cross, "StereoDelayModeCrossDescription", "Cross", "StereoDelayModeCrossDescriptionTT", "Left input mixes with right delay output and feeds to right output."),
+		DEFINE_METASOUND_ENUM_ENTRY(EStereoDelayMode::PingPong, "StereoDelayModePingPongDescription", "Ping Pong", "StereoDelayModePingPongDescriptionTT", "Left input mixes with left delay output and feeds to right output."),
 		DEFINE_METASOUND_ENUM_END()
 
 	class FStereoDelayOperator : public TExecutableOperator<FStereoDelayOperator>
@@ -284,18 +284,18 @@ namespace Metasound
 	{
 		static const FVertexInterface Interface(
 			FInputVertexInterface(
-				TInputDataVertexModel<FAudioBuffer>(StereoDelay::InParamNameAudioInputLeft, LOCTEXT("LeftAudioInputTooltip", "Left channel audio input.")),
-				TInputDataVertexModel<FAudioBuffer>(StereoDelay::InParamNameAudioInputRight, LOCTEXT("RightAudioInputTooltip", "Right channel audio input.")),
-				TInputDataVertexModel<FEnumStereoDelayMode>(StereoDelay::InParamNameDelayMode, LOCTEXT("DelayModeTooltip", "Delay mode.")),
-				TInputDataVertexModel<FTime>(StereoDelay::InParamNameDelayTime, LOCTEXT("DelayTimeTooltip", "The amount of time to delay the audio."), 1.0f),
-				TInputDataVertexModel<float>(StereoDelay::InParamNameDelayRatio, LOCTEXT("DelayRatioTooltip", "Delay spread for left and right channels. Allows left and right channels to have differential delay amounts. Useful for stereo channel decorrelation"), 0.0f),
-				TInputDataVertexModel<float>(StereoDelay::InParamNameDryLevel, LOCTEXT("DryLevelTooltip", "The dry level of the delay."), 0.0f),
-				TInputDataVertexModel<float>(StereoDelay::InParamNameWetLevel, LOCTEXT("WetLevelTooltip", "The wet level of the delay."), 1.0f),
-				TInputDataVertexModel<float>(StereoDelay::InParamNameFeedbackAmount, LOCTEXT("FeedbackTooltip", "Feedback amount."), 0.0f)
+				TInputDataVertexModel<FAudioBuffer>(StereoDelay::InParamNameAudioInputLeft, METASOUND_LOCTEXT("LeftAudioInputTooltip", "Left channel audio input.")),
+				TInputDataVertexModel<FAudioBuffer>(StereoDelay::InParamNameAudioInputRight, METASOUND_LOCTEXT("RightAudioInputTooltip", "Right channel audio input.")),
+				TInputDataVertexModel<FEnumStereoDelayMode>(StereoDelay::InParamNameDelayMode, METASOUND_LOCTEXT("DelayModeTooltip", "Delay mode.")),
+				TInputDataVertexModel<FTime>(StereoDelay::InParamNameDelayTime, METASOUND_LOCTEXT("DelayTimeTooltip", "The amount of time to delay the audio."), 1.0f),
+				TInputDataVertexModel<float>(StereoDelay::InParamNameDelayRatio, METASOUND_LOCTEXT("DelayRatioTooltip", "Delay spread for left and right channels. Allows left and right channels to have differential delay amounts. Useful for stereo channel decorrelation"), 0.0f),
+				TInputDataVertexModel<float>(StereoDelay::InParamNameDryLevel, METASOUND_LOCTEXT("DryLevelTooltip", "The dry level of the delay."), 0.0f),
+				TInputDataVertexModel<float>(StereoDelay::InParamNameWetLevel, METASOUND_LOCTEXT("WetLevelTooltip", "The wet level of the delay."), 1.0f),
+				TInputDataVertexModel<float>(StereoDelay::InParamNameFeedbackAmount, METASOUND_LOCTEXT("FeedbackTooltip", "Feedback amount."), 0.0f)
 			),
 			FOutputVertexInterface(
-				TOutputDataVertexModel<FAudioBuffer>(StereoDelay::OutParamNameAudioLeft, LOCTEXT("LeftDelayOutputTooltip", "Left channel audio output.")),
-				TOutputDataVertexModel<FAudioBuffer>(StereoDelay::OutParamNameAudioRight, LOCTEXT("RightDelayOutputTooltip", "Right channel audio output."))
+				TOutputDataVertexModel<FAudioBuffer>(StereoDelay::OutParamNameAudioLeft, METASOUND_LOCTEXT("LeftDelayOutputTooltip", "Left channel audio output.")),
+				TOutputDataVertexModel<FAudioBuffer>(StereoDelay::OutParamNameAudioRight, METASOUND_LOCTEXT("RightDelayOutputTooltip", "Right channel audio output."))
 			)
 		);
 
@@ -310,8 +310,8 @@ namespace Metasound
 			Info.ClassName = { StandardNodes::Namespace, TEXT("Stereo Delay"), StandardNodes::AudioVariant };
 			Info.MajorVersion = 1;
 			Info.MinorVersion = 0;
-			Info.DisplayName = LOCTEXT("Metasound_StereoDelayDisplayName", "Stereo Delay");
-			Info.Description = LOCTEXT("Metasound_StereoDelayNodeDescription", "Delays a stereo audio buffer by the specified amount.");
+			Info.DisplayName = METASOUND_LOCTEXT("Metasound_StereoDelayDisplayName", "Stereo Delay");
+			Info.Description = METASOUND_LOCTEXT("Metasound_StereoDelayNodeDescription", "Delays a stereo audio buffer by the specified amount.");
 			Info.Author = PluginAuthor;
 			Info.PromptIfMissing = PluginNodeMissingPrompt;
 			Info.DefaultInterface = GetVertexInterface();
diff --git a/Engine/Plugins/Runtime/Metasound/Source/MetasoundStandardNodes/Private/MetasoundStereoPannerNode.cpp b/Engine/Plugins/Runtime/Metasound/Source/MetasoundStandardNodes/Private/MetasoundStereoPannerNode.cpp
--- a/Engine/Plugins/Runtime/Metasound/Source/MetasoundStandardNodes/Private/MetasoundStereoPannerNode.cpp
+++ b/Engine/Plugins/Runtime/Metasound/Source/MetasoundStandardNodes/Private/MetasoundStereoPannerNode.cpp
@@ -29,7 +29,7 @@ namespace Metasound
 
 		const FText& GetInputAudioDescription()
 		{
-			static FText Desc = LOCTEXT("StereoPannerNodeInDesc", "The input audio to pan.");
+			static FText Desc = METASOUND_LOCTEXT("StereoPannerNodeInDesc", "The input audio to pan.");
 			return Desc;
 		}
 
@@ -41,7 +41,7 @@ namespace Metasound
 
 		const FText& GetInputPanAmountDescription()
 		{
-			static FText Desc = LOCTEXT("StereoPannerNodePanAmountDesc", "The amount of pan. -1.0 is full left, 1.0 is full right.");
+			static FText Desc = METASOUND_LOCTEXT("StereoPannerNodePanAmountDesc", "The amount of pan. -1.0 is full left, 1.0 is full right.");
 			return Desc;
 		}
 
@@ -53,7 +53,7 @@ namespace Metasound
 
 		const FText& GetInputPanningLawDescription()
 		{
-			static FText Desc = LOCTEXT("StereoPannerNodePanningLawDescription", "Which panning law should be used for the stereo panner.");
+			static FText Desc = METASOUND_LOCTEXT("StereoPannerNodePanningLawDescription", "Which panning law should be used for the stereo panner.");
 			return Desc;
 		}
 
@@ -65,7 +65,7 @@ namespace Metasound
 
 		const FText& GetOutputAudioLeftDescription()
 		{
-			static FText Desc = LOCTEXT("StereoPannerNodeOutputLeftDescription", "Left channel audio output.");
+			static FText Desc = METASOUND_LOCTEXT("StereoPannerNodeOutputLeftDescription", "Left channel audio output.");
 			return Desc;
 		}
 
@@ -77,7 +77,7 @@ namespace Metasound
 
 		const FText& GetOutputAudioRightDescription()
 		{
-			static FText Desc = LOCTEXT("StereoPannerNodeOutputRightDescription", "Right channel audio output.");
+			static FText Desc = METASOUND_LOCTEXT("StereoPannerNodeOutputRightDescription", "Right channel audio output.");
 			return Desc;
 		}
 	}
@@ -92,8 +92,8 @@ namespace Metasound
 	FEnumPanningLaw, FEnumPanningLawInfo, FPanningLawReadRef, FPanningLawWriteRef);
 
 	DEFINE_METASOUND_ENUM_BEGIN(EPanningLaw, FEnumPanningLaw, "PanningLaw")
-		DEFINE_METASOUND_ENUM_ENTRY(EPanningLaw::EqualPower, LOCTEXT("PanningLawEqualPowerName", "Equal Power"), LOCTEXT("PanningLawEqualPowerTT", "The power of the audio signal is constant while panning.")),
-		DEFINE_METASOUND_ENUM_ENTRY(EPanningLaw::Linear, LOCTEXT("PanningLawLinearName", "Linear"), LOCTEXT("PanningLawLinearTT", "The amplitude of the audio signal is constant while panning.")),
+		DEFINE_METASOUND_ENUM_ENTRY(EPanningLaw::EqualPower, "PanningLawEqualPowerName", "Equal Power", "PanningLawEqualPowerTT", "The power of the audio signal is constant while panning."),
+		DEFINE_METASOUND_ENUM_ENTRY(EPanningLaw::Linear, "PanningLawLinearName", "Linear", "PanningLawLinearTT", "The amplitude of the audio signal is constant while panning."),
 	DEFINE_METASOUND_ENUM_END()
 
 	class FStereoPannerOperator : public TExecutableOperator<FStereoPannerOperator>
@@ -254,8 +254,8 @@ namespace Metasound
 			Info.ClassName = { StandardNodes::Namespace, TEXT("Stereo Panner"), TEXT("") };
 			Info.MajorVersion = 1;
 			Info.MinorVersion = 0;
-			Info.DisplayName = LOCTEXT("Metasound_StereoPannerDisplayName", "Stereo Panner");
-			Info.Description = LOCTEXT("Metasound_StereoPannerNodeDescription", "Pans an input audio signal to left and right outputs.");
+			Info.DisplayName = METASOUND_LOCTEXT("Metasound_StereoPannerDisplayName", "Stereo Panner");
+			Info.Description = METASOUND_LOCTEXT("Metasound_StereoPannerNodeDescription", "Pans an input audio signal to left and right outputs.");
 			Info.Author = PluginAuthor;
 			Info.PromptIfMissing = PluginNodeMissingPrompt;
 			Info.DefaultInterface = GetVertexInterface();
diff --git a/Engine/Plugins/Runtime/Metasound/Source/MetasoundStandardNodes/Private/MetasoundTriggerAnyNode.cpp b/Engine/Plugins/Runtime/Metasound/Source/MetasoundStandardNodes/Private/MetasoundTriggerAnyNode.cpp
--- a/Engine/Plugins/Runtime/Metasound/Source/MetasoundStandardNodes/Private/MetasoundTriggerAnyNode.cpp
+++ b/Engine/Plugins/Runtime/Metasound/Source/MetasoundStandardNodes/Private/MetasoundTriggerAnyNode.cpp
@@ -50,7 +50,7 @@ namespace Metasound
 
 		METASOUNDSTANDARDNODES_API const FText GetInputTriggerDescription(uint32 InIndex)
 		{
-			return FText::Format(LOCTEXT("TriggerAnyInputTriggerDesc", "Trigger {0} input. The output trigger is hit when any of the input triggers are hit."), InIndex);
+			return METASOUND_LOCTEXT_FORMAT("TriggerAnyInputTriggerDesc", "Trigger {0} input. The output trigger is hit when any of the input triggers are hit.", InIndex);
 		}
 
 		METASOUNDSTANDARDNODES_API const FVertexName& GetOutputTriggerName()
@@ -61,7 +61,7 @@ namespace Metasound
 
 		METASOUNDSTANDARDNODES_API const FText& GetOutputTriggerDescription()
 		{
-			static const FText Desc = LOCTEXT("TriggerAnyOutputTriggerDesc", "Triggered when any of the input triggers have been triggered. ");
+			static const FText Desc = METASOUND_LOCTEXT("TriggerAnyOutputTriggerDesc", "Triggered when any of the input triggers have been triggered. ");
 			return Desc;
 		}
 	}
@@ -98,8 +98,8 @@ namespace Metasound
 			auto CreateNodeClassMetadata = []() -> FNodeClassMetadata
 			{
 				FName OperatorName = *FString::Printf(TEXT("Trigger Any (%d)"), NumInputs);
-				FText NodeDisplayName = FText::Format(LOCTEXT("TriggerAnyDisplayNamePattern", "Trigger Any ({0})"), NumInputs);
-				FText NodeDescription = LOCTEXT("TriggerAnyDescription", "Will trigger output on any of the input triggers.");
+				FText NodeDisplayName = METASOUND_LOCTEXT_FORMAT("TriggerAnyDisplayNamePattern", "Trigger Any ({0})", NumInputs);
+				const FText NodeDescription = METASOUND_LOCTEXT("TriggerAnyDescription", "Will trigger output on any of the input triggers.");
 				FVertexInterface NodeInterface = GetDefaultInterface();
 
 				return MetasoundTriggerAnyNodePrivate::CreateNodeClassMetadata(OperatorName, NodeDisplayName, NodeDescription, NodeInterface);
diff --git a/Engine/Plugins/Runtime/Metasound/Source/MetasoundStandardNodes/Private/MetasoundTriggerCoinNode.cpp b/Engine/Plugins/Runtime/Metasound/Source/MetasoundStandardNodes/Private/MetasoundTriggerCoinNode.cpp
--- a/Engine/Plugins/Runtime/Metasound/Source/MetasoundStandardNodes/Private/MetasoundTriggerCoinNode.cpp
+++ b/Engine/Plugins/Runtime/Metasound/Source/MetasoundStandardNodes/Private/MetasoundTriggerCoinNode.cpp
@@ -199,8 +199,8 @@ namespace Metasound
 			Info.ClassName = { StandardNodes::Namespace, TEXT("Trigger Filter"), TEXT("") };
 			Info.MajorVersion = 1;
 			Info.MinorVersion = 0;
-			Info.DisplayName = LOCTEXT("Metasound_TriggerCoinNodeDisplayName", "Trigger Filter");
-			Info.Description = LOCTEXT("Metasound_TriggerCoinNodeDescription", "When triggered, randomly sends one of two output triggers based on a given probability.");
+			Info.DisplayName = METASOUND_LOCTEXT("Metasound_TriggerCoinNodeDisplayName", "Trigger Filter");
+			Info.Description = METASOUND_LOCTEXT("Metasound_TriggerCoinNodeDescription", "When triggered, randomly sends one of two output triggers based on a given probability.");
 			Info.Author = PluginAuthor;
 			Info.PromptIfMissing = PluginNodeMissingPrompt;
 			Info.DefaultInterface = GetVertexInterface();
diff --git a/Engine/Plugins/Runtime/Metasound/Source/MetasoundStandardNodes/Private/MetasoundTriggerControlNode.cpp b/Engine/Plugins/Runtime/Metasound/Source/MetasoundStandardNodes/Private/MetasoundTriggerControlNode.cpp
--- a/Engine/Plugins/Runtime/Metasound/Source/MetasoundStandardNodes/Private/MetasoundTriggerControlNode.cpp
+++ b/Engine/Plugins/Runtime/Metasound/Source/MetasoundStandardNodes/Private/MetasoundTriggerControlNode.cpp
@@ -204,8 +204,8 @@ namespace Metasound
 			Info.ClassName = { Metasound::StandardNodes::Namespace, "Trigger Control", FName() };
 			Info.MajorVersion = 1;
 			Info.MinorVersion = 0;
-			Info.DisplayName = LOCTEXT("Metasound_TriggerControlNodeDisplayName", "Trigger Control");
-			Info.Description = LOCTEXT("Metasound_TriggerControlNodeDescription", "Control whether input triggers are passed through.");
+			Info.DisplayName = METASOUND_LOCTEXT("Metasound_TriggerControlNodeDisplayName", "Trigger Control");
+			Info.Description = METASOUND_LOCTEXT("Metasound_TriggerControlNodeDescription", "Control whether input triggers are passed through.");
 			Info.Author = PluginAuthor;
 			Info.PromptIfMissing = PluginNodeMissingPrompt;
 			Info.DefaultInterface = GetVertexInterface();
diff --git a/Engine/Plugins/Runtime/Metasound/Source/MetasoundStandardNodes/Private/MetasoundTriggerCounterNode.cpp b/Engine/Plugins/Runtime/Metasound/Source/MetasoundStandardNodes/Private/MetasoundTriggerCounterNode.cpp
--- a/Engine/Plugins/Runtime/Metasound/Source/MetasoundStandardNodes/Private/MetasoundTriggerCounterNode.cpp
+++ b/Engine/Plugins/Runtime/Metasound/Source/MetasoundStandardNodes/Private/MetasoundTriggerCounterNode.cpp
@@ -207,8 +207,8 @@ namespace Metasound
 			Info.ClassName = {StandardNodes::Namespace, TEXT("Trigger Counter"), TEXT("")};
 			Info.MajorVersion = 1;
 			Info.MinorVersion = 0;
-			Info.DisplayName = LOCTEXT("Metasound_TriggerCounterNodeDisplayName", "Trigger Counter");
-			Info.Description = LOCTEXT("Metasound_TriggerCounterNodeDescription", "Counts the trigger inputs. Supports a start count value, counting by a step size, and auto resetting back to the start count.");
+			Info.DisplayName = METASOUND_LOCTEXT("Metasound_TriggerCounterNodeDisplayName", "Trigger Counter");
+			Info.Description = METASOUND_LOCTEXT("Metasound_TriggerCounterNodeDescription", "Counts the trigger inputs. Supports a start count value, counting by a step size, and auto resetting back to the start count.");
 			Info.Author = PluginAuthor;
 			Info.PromptIfMissing = PluginNodeMissingPrompt;
 			Info.DefaultInterface = GetVertexInterface();
diff --git a/Engine/Plugins/Runtime/Metasound/Source/MetasoundStandardNodes/Private/MetasoundTriggerDelayNode.cpp b/Engine/Plugins/Runtime/Metasound/Source/MetasoundStandardNodes/Private/MetasoundTriggerDelayNode.cpp
--- a/Engine/Plugins/Runtime/Metasound/Source/MetasoundStandardNodes/Private/MetasoundTriggerDelayNode.cpp
+++ b/Engine/Plugins/Runtime/Metasound/Source/MetasoundStandardNodes/Private/MetasoundTriggerDelayNode.cpp
@@ -175,8 +175,8 @@ namespace Metasound
 			Info.ClassName = {StandardNodes::Namespace, TEXT("Trigger Delay"), TEXT("")};
 			Info.MajorVersion = 1;
 			Info.MinorVersion = 0;
-			Info.DisplayName = LOCTEXT("Metasound_TriggerDelayNodeDisplayName", "Trigger Delay");
-			Info.Description = LOCTEXT("Metasound_TriggerDelayNodeDescription", "Executes output trigger after the given delay time from the most recent execution of the input trigger .");
+			Info.DisplayName = METASOUND_LOCTEXT("Metasound_TriggerDelayNodeDisplayName", "Trigger Delay");
+			Info.Description = METASOUND_LOCTEXT("Metasound_TriggerDelayNodeDescription", "Executes output trigger after the given delay time from the most recent execution of the input trigger .");
 			Info.Author = PluginAuthor;
 			Info.PromptIfMissing = PluginNodeMissingPrompt;
 			Info.DefaultInterface = GetVertexInterface();
diff --git a/Engine/Plugins/Runtime/Metasound/Source/MetasoundStandardNodes/Private/MetasoundTriggerOnThresholdNode.cpp b/Engine/Plugins/Runtime/Metasound/Source/MetasoundStandardNodes/Private/MetasoundTriggerOnThresholdNode.cpp
--- a/Engine/Plugins/Runtime/Metasound/Source/MetasoundStandardNodes/Private/MetasoundTriggerOnThresholdNode.cpp
+++ b/Engine/Plugins/Runtime/Metasound/Source/MetasoundStandardNodes/Private/MetasoundTriggerOnThresholdNode.cpp
@@ -29,9 +29,9 @@ namespace Metasound
 	FEnumBufferTriggerType, FEnumBufferTriggerTypeInfo, FBufferTriggerTypeReadRef, FEnumBufferTriggerTypeWriteRef);
 
 	DEFINE_METASOUND_ENUM_BEGIN(EBufferTriggerType, FEnumBufferTriggerType, "BufferTriggerType")
-		DEFINE_METASOUND_ENUM_ENTRY(EBufferTriggerType::RisingEdge, LOCTEXT("RisingEdgeDescription", "Rising Edge"), LOCTEXT("RisingEdgeDescriptionTT", "")),
-		DEFINE_METASOUND_ENUM_ENTRY(EBufferTriggerType::FallingEdge, LOCTEXT("FallingEdgeDescription", "Falling Edge"), LOCTEXT("FallingEdgeDescriptionTT", "")),
-		DEFINE_METASOUND_ENUM_ENTRY(EBufferTriggerType::AbsThreshold, LOCTEXT("AbsThresholdDescription", "Abs Threshold"), LOCTEXT("AbsThresholdDescriptionTT", ""))
+		DEFINE_METASOUND_ENUM_ENTRY(EBufferTriggerType::RisingEdge, "RisingEdgeDescription", "Rising Edge", "RisingEdgeDescriptionTT", ""),
+		DEFINE_METASOUND_ENUM_ENTRY(EBufferTriggerType::FallingEdge, "FallingEdgeDescription", "Falling Edge", "FallingEdgeDescriptionTT", ""),
+		DEFINE_METASOUND_ENUM_ENTRY(EBufferTriggerType::AbsThreshold, "AbsThresholdDescription", "Abs Threshold", "AbsThresholdDescriptionTT", "")
 		DEFINE_METASOUND_ENUM_END()
 
 	class FTriggerOnThresholdOperator : public IOperator
@@ -183,12 +183,12 @@ namespace Metasound
 	{
 		static const FVertexInterface Interface(
 			FInputVertexInterface(
-				TInputDataVertexModel<FAudioBuffer>(InPinName, LOCTEXT("BufferInDescription", "Input")),
-				TInputDataVertexModel<float>(ThresholdPinName, LOCTEXT("ThresholdDescription", "Trigger Threshold"), DefaultThreshold),
-				TInputDataVertexModel<FEnumBufferTriggerType>(TriggerType, LOCTEXT("ThresholdDescription", "Trigger Threshold"))
+				TInputDataVertexModel<FAudioBuffer>(InPinName, METASOUND_LOCTEXT("BufferInDescription", "Input")),
+				TInputDataVertexModel<float>(ThresholdPinName, METASOUND_LOCTEXT("ThresholdDescription", "Trigger Threshold"), DefaultThreshold),
+				TInputDataVertexModel<FEnumBufferTriggerType>(TriggerType, METASOUND_LOCTEXT("ThresholdDescription", "Trigger Threshold"))
 			),
 			FOutputVertexInterface(
-				TOutputDataVertexModel<FTrigger>(OutPinName, LOCTEXT("TriggerOutDescription", "Output"))
+				TOutputDataVertexModel<FTrigger>(OutPinName, METASOUND_LOCTEXT("TriggerOutDescription", "Output"))
 			)
 		);
 		return Interface;
@@ -202,8 +202,8 @@ namespace Metasound
 			Info.ClassName = {StandardNodes::Namespace, TEXT("TriggerOnThreshold"), StandardNodes::AudioVariant};
 			Info.MajorVersion = 1;
 			Info.MinorVersion = 0;
-			Info.DisplayName = LOCTEXT("Metasound_TriggerOnThresholdNodeDisplayName", "Trigger On Threshold");
-			Info.Description = LOCTEXT("Metasound_TriggerOnThresholdNodeDescription", "Trigger based on a audio buffer input");
+			Info.DisplayName = METASOUND_LOCTEXT("Metasound_TriggerOnThresholdNodeDisplayName", "Trigger On Threshold");
+			Info.Description = METASOUND_LOCTEXT("Metasound_TriggerOnThresholdNodeDescription", "Trigger based on a audio buffer input");
 			Info.Author = PluginAuthor;
 			Info.PromptIfMissing = PluginNodeMissingPrompt;
 			Info.DefaultInterface = DeclareVertexInterface(); 
diff --git a/Engine/Plugins/Runtime/Metasound/Source/MetasoundStandardNodes/Private/MetasoundTriggerOnValueChangeNode.cpp b/Engine/Plugins/Runtime/Metasound/Source/MetasoundStandardNodes/Private/MetasoundTriggerOnValueChangeNode.cpp
--- a/Engine/Plugins/Runtime/Metasound/Source/MetasoundStandardNodes/Private/MetasoundTriggerOnValueChangeNode.cpp
+++ b/Engine/Plugins/Runtime/Metasound/Source/MetasoundStandardNodes/Private/MetasoundTriggerOnValueChangeNode.cpp
@@ -149,8 +149,8 @@ namespace Metasound
 		{
 			const FName DataTypeName = GetMetasoundDataTypeName<ValueType>();
 			const FName OperatorName = TEXT("Trigger On Value Change");
-			const FText NodeDisplayName = FText::Format(LOCTEXT("TriggerOnValueChangeName", "Trigger On Value Change ({0})"), GetMetasoundDataTypeDisplayText<ValueType>());
-			const FText NodeDescription = LOCTEXT("TriggerOnValueChangeNameDesc", "Triggers when a given value changes.");
+			const FText NodeDisplayName = METASOUND_LOCTEXT_FORMAT("TriggerOnValueChangeName", "Trigger On Value Change ({0})", GetMetasoundDataTypeDisplayText<ValueType>());
+			const FText NodeDescription = METASOUND_LOCTEXT("TriggerOnValueChangeNameDesc", "Triggers when a given value changes.");
 
 			FNodeClassMetadata Info;
 			Info.ClassName = { Metasound::StandardNodes::Namespace, OperatorName, DataTypeName };
diff --git a/Engine/Plugins/Runtime/Metasound/Source/MetasoundStandardNodes/Private/MetasoundTriggerOnceNode.cpp b/Engine/Plugins/Runtime/Metasound/Source/MetasoundStandardNodes/Private/MetasoundTriggerOnceNode.cpp
--- a/Engine/Plugins/Runtime/Metasound/Source/MetasoundStandardNodes/Private/MetasoundTriggerOnceNode.cpp
+++ b/Engine/Plugins/Runtime/Metasound/Source/MetasoundStandardNodes/Private/MetasoundTriggerOnceNode.cpp
@@ -166,8 +166,8 @@ namespace Metasound
 			Info.ClassName = { Metasound::StandardNodes::Namespace, "Trigger Once", FName() };
 			Info.MajorVersion = 1;
 			Info.MinorVersion = 0;
-			Info.DisplayName = LOCTEXT("Metasound_TriggerOnceNodeDisplayName", "Trigger Once");
-			Info.Description = LOCTEXT("Metasound_TriggerOnceNodeDescription", "Sends an output trigger the first time the node is triggered, and ignores all others (can be re-opened).");
+			Info.DisplayName = METASOUND_LOCTEXT("Metasound_TriggerOnceNodeDisplayName", "Trigger Once");
+			Info.Description = METASOUND_LOCTEXT("Metasound_TriggerOnceNodeDescription", "Sends an output trigger the first time the node is triggered, and ignores all others (can be re-opened).");
 			Info.Author = PluginAuthor;
 			Info.PromptIfMissing = PluginNodeMissingPrompt;
 			Info.DefaultInterface = GetVertexInterface();
diff --git a/Engine/Plugins/Runtime/Metasound/Source/MetasoundStandardNodes/Private/MetasoundTriggerPipeNode.cpp b/Engine/Plugins/Runtime/Metasound/Source/MetasoundStandardNodes/Private/MetasoundTriggerPipeNode.cpp
--- a/Engine/Plugins/Runtime/Metasound/Source/MetasoundStandardNodes/Private/MetasoundTriggerPipeNode.cpp
+++ b/Engine/Plugins/Runtime/Metasound/Source/MetasoundStandardNodes/Private/MetasoundTriggerPipeNode.cpp
@@ -175,8 +175,8 @@ namespace Metasound
 			Info.ClassName = FNodeClassName(StandardNodes::Namespace, "Pipe", "");
 			Info.MajorVersion = 1;
 			Info.MinorVersion = 0;
-			Info.DisplayName = LOCTEXT("PipeTriggerNode_NodeDisplayName", "Trigger Pipe");
-			Info.Description = LOCTEXT("Metasound_DelayNodeDescription", "Delays execution of the input trigger(s) by the given delay for all input trigger executions.");
+			Info.DisplayName = METASOUND_LOCTEXT("PipeTriggerNode_NodeDisplayName", "Trigger Pipe");
+			Info.Description = METASOUND_LOCTEXT("Metasound_DelayNodeDescription", "Delays execution of the input trigger(s) by the given delay for all input trigger executions.");
 			Info.Author = PluginAuthor;
 			Info.PromptIfMissing = PluginNodeMissingPrompt;
 			Info.DefaultInterface = GetVertexInterface();
diff --git a/Engine/Plugins/Runtime/Metasound/Source/MetasoundStandardNodes/Private/MetasoundTriggerRepeatNode.cpp b/Engine/Plugins/Runtime/Metasound/Source/MetasoundStandardNodes/Private/MetasoundTriggerRepeatNode.cpp
--- a/Engine/Plugins/Runtime/Metasound/Source/MetasoundStandardNodes/Private/MetasoundTriggerRepeatNode.cpp
+++ b/Engine/Plugins/Runtime/Metasound/Source/MetasoundStandardNodes/Private/MetasoundTriggerRepeatNode.cpp
@@ -172,8 +172,8 @@ namespace Metasound
 			Info.ClassName = { StandardNodes::Namespace, TEXT("TriggerRepeat"), TEXT("") };
 			Info.MajorVersion = 1;
 			Info.MinorVersion = 0;
-			Info.DisplayName = LOCTEXT("Metasound_PeriodicTriggerNodeDisplayName", "Trigger Repeat");
-			Info.Description = LOCTEXT("Metasound_PeriodicTriggerNodeDescription", "Emits a trigger periodically based on the period duration given.");
+			Info.DisplayName = METASOUND_LOCTEXT("Metasound_PeriodicTriggerNodeDisplayName", "Trigger Repeat");
+			Info.Description = METASOUND_LOCTEXT("Metasound_PeriodicTriggerNodeDescription", "Emits a trigger periodically based on the period duration given.");
 			Info.Author = PluginAuthor;
 			Info.PromptIfMissing = PluginNodeMissingPrompt;
 			Info.DefaultInterface = GetVertexInterface();
diff --git a/Engine/Plugins/Runtime/Metasound/Source/MetasoundStandardNodes/Private/MetasoundTriggerRouteNode.cpp b/Engine/Plugins/Runtime/Metasound/Source/MetasoundStandardNodes/Private/MetasoundTriggerRouteNode.cpp
--- a/Engine/Plugins/Runtime/Metasound/Source/MetasoundStandardNodes/Private/MetasoundTriggerRouteNode.cpp
+++ b/Engine/Plugins/Runtime/Metasound/Source/MetasoundStandardNodes/Private/MetasoundTriggerRouteNode.cpp
@@ -34,9 +34,9 @@ namespace Metasound
 		{
 			if (InIndex == 0)
 			{
-				return FText::Format(LOCTEXT("TriggerRouteInputTriggerDescInit", "The input trigger {0} to cause the corresponding input value to route to the output value. This trigger is the default output."), InIndex);
+				return METASOUND_LOCTEXT_FORMAT("TriggerRouteInputTriggerDescInit", "The input trigger {0} to cause the corresponding input value to route to the output value. This trigger is the default output.", InIndex);
 			}
-			return FText::Format(LOCTEXT("TriggerRouteInputTriggerDesc", "The input trigger {0} to cause the corresponding input value to route to the output value."), InIndex);
+			return METASOUND_LOCTEXT_FORMAT("TriggerRouteInputTriggerDesc", "The input trigger {0} to cause the corresponding input value to route to the output value.", InIndex);
 		}
 
 		const FVertexName GetInputValueName(uint32 InIndex)
@@ -46,7 +46,7 @@ namespace Metasound
 
 		const FText GetInputValueDescription(uint32 InIndex)
 		{
-			return FText::Format(LOCTEXT("TriggerRouteValueDesc", "The input value ({0}) to route to the output when triggered by Set {0}."), InIndex);
+			return METASOUND_LOCTEXT_FORMAT("TriggerRouteValueDesc", "The input value ({0}) to route to the output when triggered by Set {0}.", InIndex);
 		}
 
 		const FVertexName& GetOutputTriggerName()
@@ -57,7 +57,7 @@ namespace Metasound
 
 		const FText& GetOutputTriggerDescription()
 		{
-			static const FText Desc = LOCTEXT("TriggerRouteOnSetDesc", "Triggered when any of the input triggers are set.");
+			static const FText Desc = METASOUND_LOCTEXT("TriggerRouteOnSetDesc", "Triggered when any of the input triggers are set.");
 			return Desc;
 		}
 
@@ -69,7 +69,7 @@ namespace Metasound
 
 		const FText& GetOutputValueDescription()
 		{
-			static const FText Desc = LOCTEXT("TriggerRouteOutputValueDesc", "The output value set by the input triggers.");
+			static const FText Desc = METASOUND_LOCTEXT("TriggerRouteOutputValueDesc", "The output value set by the input triggers.");
 			return Desc;
 		}
 	}
@@ -131,8 +131,8 @@ namespace Metasound
 			{
 				FName DataTypeName = GetMetasoundDataTypeName<ValueType>();
 				FName OperatorName = *FString::Printf(TEXT("Trigger Route (%s, %d)"), *DataTypeName.ToString(), NumInputs);
-				FText NodeDisplayName = FText::Format(LOCTEXT("TriggerRouteDisplayNamePattern", "Trigger Route ({0}, {1})"), GetMetasoundDataTypeDisplayText<ValueType>(), NumInputs);
-				FText NodeDescription = LOCTEXT("TriggerRouteDescription", "Allows routing different values to the same output pin depending on trigger inputs.");
+				FText NodeDisplayName = METASOUND_LOCTEXT_FORMAT("TriggerRouteDisplayNamePattern", "Trigger Route ({0}, {1})", GetMetasoundDataTypeDisplayText<ValueType>(), NumInputs);
+				const FText NodeDescription = METASOUND_LOCTEXT("TriggerRouteDescription", "Allows routing different values to the same output pin depending on trigger inputs.");
 				FVertexInterface NodeInterface = GetDefaultInterface();
 
 				return MetasoundTriggerRouteNodePrivate::CreateNodeClassMetadata(DataTypeName, OperatorName, NodeDisplayName, NodeDescription, NodeInterface);
diff --git a/Engine/Plugins/Runtime/Metasound/Source/MetasoundStandardNodes/Private/MetasoundTriggerSequenceNode.cpp b/Engine/Plugins/Runtime/Metasound/Source/MetasoundStandardNodes/Private/MetasoundTriggerSequenceNode.cpp
--- a/Engine/Plugins/Runtime/Metasound/Source/MetasoundStandardNodes/Private/MetasoundTriggerSequenceNode.cpp
+++ b/Engine/Plugins/Runtime/Metasound/Source/MetasoundStandardNodes/Private/MetasoundTriggerSequenceNode.cpp
@@ -100,7 +100,7 @@ namespace Metasound
 		FDataReferenceCollection OutputDataReferences;
 		for (uint32 i = 0; i < NumOutputs; ++i)
 		{
-			OutputDataReferences.AddDataReadReference(METASOUND_GET_VARIABLE_PARAM_NAME(OutputTrigger, i), TriggerOutputs[i]);
+			OutputDataReferences.AddDataReadReference(METASOUND_GET_PARAM_NAME_WITH_INDEX(OutputTrigger, i), TriggerOutputs[i]);
 		}
 
 		return OutputDataReferences;
@@ -178,7 +178,7 @@ namespace Metasound
 			FOutputVertexInterface OutputInterface;
 			for (uint32 i = 0; i < NumOutputs; ++i)
 			{
-				OutputInterface.Add(TOutputDataVertexModel<FTrigger>(METASOUND_GET_VARIABLE_PARAM_NAME_AND_TT(OutputTrigger, i)));
+				OutputInterface.Add(TOutputDataVertexModel<FTrigger>(METASOUND_GET_PARAM_NAME_WITH_INDEX_AND_TT(OutputTrigger, i)));
 			}
 			return FVertexInterface(InputInterface, OutputInterface);
 		};
@@ -193,8 +193,8 @@ namespace Metasound
 		auto InitNodeInfo = []() -> FNodeClassMetadata
 		{
 			FName OperatorName = *FString::Printf(TEXT("Trigger Sequence (%d)"), NumOutputs);
-			FText NodeDisplayName = FText::Format(LOCTEXT("TriggerSequenceDisplayNamePattern", "Trigger Sequence ({0})"), NumOutputs);
-			FText NodeDescription = LOCTEXT("TriggerSequenceDescription", "Each time this node is triggered, it sends the next output trigger in the sequence.");
+			FText NodeDisplayName = METASOUND_LOCTEXT_FORMAT("TriggerSequenceDisplayNamePattern", "Trigger Sequence ({0})", NumOutputs);
+			const FText NodeDescription = METASOUND_LOCTEXT("TriggerSequenceDescription", "Each time this node is triggered, it sends the next output trigger in the sequence.");
 
 			FNodeClassMetadata Info;
 			Info.ClassName = { StandardNodes::Namespace, OperatorName, TEXT("") };
diff --git a/Engine/Plugins/Runtime/Metasound/Source/MetasoundStandardNodes/Private/MetasoundTriggerToggleNode.cpp b/Engine/Plugins/Runtime/Metasound/Source/MetasoundStandardNodes/Private/MetasoundTriggerToggleNode.cpp
--- a/Engine/Plugins/Runtime/Metasound/Source/MetasoundStandardNodes/Private/MetasoundTriggerToggleNode.cpp
+++ b/Engine/Plugins/Runtime/Metasound/Source/MetasoundStandardNodes/Private/MetasoundTriggerToggleNode.cpp
@@ -128,13 +128,13 @@ namespace Metasound
 
 		static const FVertexInterface Interface(
 			FInputVertexInterface(
-				TInputDataVertexModel<FTrigger>(GetInputOnTriggerName(), LOCTEXT("TriggerGateOnTT", "Trigger to toggle gate output to 1.")),
-				TInputDataVertexModel<FTrigger>(GetInputOffTriggerName(), LOCTEXT("TriggerGateOffTT", "Trigger to toggle gate output to 0.")),
-				TInputDataVertexModel<bool>(GetInputInitName(), LOCTEXT("TriggerGateInitTT", "Initial value of the output gate."))
+				TInputDataVertexModel<FTrigger>(GetInputOnTriggerName(), METASOUND_LOCTEXT("TriggerGateOnTT", "Trigger to toggle gate output to 1.")),
+				TInputDataVertexModel<FTrigger>(GetInputOffTriggerName(), METASOUND_LOCTEXT("TriggerGateOffTT", "Trigger to toggle gate output to 0.")),
+				TInputDataVertexModel<bool>(GetInputInitName(), METASOUND_LOCTEXT("TriggerGateInitTT", "Initial value of the output gate."))
 			),
 			FOutputVertexInterface(
-				TOutputDataVertexModel<FTrigger>(GetOutputTriggerName(), LOCTEXT("TriggerGateOutputTriggerTT", "Triggers output when gate is toggled.")),
-				TOutputDataVertexModel<bool>(GetOutputValueName(), LOCTEXT("BoolOutputTT", "Current output value of the toggle."))
+				TOutputDataVertexModel<FTrigger>(GetOutputTriggerName(), METASOUND_LOCTEXT("TriggerGateOutputTriggerTT", "Triggers output when gate is toggled.")),
+				TOutputDataVertexModel<bool>(GetOutputValueName(), METASOUND_LOCTEXT("BoolOutputTT", "Current output value of the toggle."))
 			)
 		);
 
@@ -149,8 +149,8 @@ namespace Metasound
 			Info.ClassName = {StandardNodes::Namespace, TEXT("Trigger Toggle"), TEXT("")};
 			Info.MajorVersion = 1;
 			Info.MinorVersion = 0;
-			Info.DisplayName = LOCTEXT("Metasound_TriggerToggleNodeDisplayName", "Trigger Toggle");
-			Info.Description = LOCTEXT("Metasound_TriggerToggleNodeDescription", "Toggles a boolean value on or off.");
+			Info.DisplayName = METASOUND_LOCTEXT("Metasound_TriggerToggleNodeDisplayName", "Trigger Toggle");
+			Info.Description = METASOUND_LOCTEXT("Metasound_TriggerToggleNodeDescription", "Toggles a boolean value on or off.");
 			Info.Author = PluginAuthor;
 			Info.PromptIfMissing = PluginNodeMissingPrompt;
 			Info.DefaultInterface = GetVertexInterface();
diff --git a/Engine/Plugins/Runtime/Metasound/Source/MetasoundStandardNodes/Private/MetasoundValueNode.cpp b/Engine/Plugins/Runtime/Metasound/Source/MetasoundStandardNodes/Private/MetasoundValueNode.cpp
--- a/Engine/Plugins/Runtime/Metasound/Source/MetasoundStandardNodes/Private/MetasoundValueNode.cpp
+++ b/Engine/Plugins/Runtime/Metasound/Source/MetasoundStandardNodes/Private/MetasoundValueNode.cpp
@@ -22,7 +22,7 @@ namespace Metasound
 				PluginAuthor,
 				PluginNodeMissingPrompt,
 				InDefaultInterface,
-				{ LOCTEXT("ValueCategory", "Value") },
+				{ METASOUND_LOCTEXT("ValueCategory", "Value") },
 				{ },
 				FNodeDisplayStyle{}
 			};
diff --git a/Engine/Plugins/Runtime/Metasound/Source/MetasoundStandardNodes/Private/MetasoundWaveShaperNode.cpp b/Engine/Plugins/Runtime/Metasound/Source/MetasoundStandardNodes/Private/MetasoundWaveShaperNode.cpp
--- a/Engine/Plugins/Runtime/Metasound/Source/MetasoundStandardNodes/Private/MetasoundWaveShaperNode.cpp
+++ b/Engine/Plugins/Runtime/Metasound/Source/MetasoundStandardNodes/Private/MetasoundWaveShaperNode.cpp
@@ -29,11 +29,11 @@ namespace Metasound
 	DECLARE_METASOUND_ENUM(Audio::EWaveShaperType, Audio::EWaveShaperType::ATan, METASOUNDSTANDARDNODES_API, FEnumEWaveShaperType, FEnumWaveShaperTypeInfo, FEnumWaveShaperReadRef, FEnumWaveShaperWriteRef);
 
 	DEFINE_METASOUND_ENUM_BEGIN(Audio::EWaveShaperType, FEnumEWaveShaperType, "WaveShaperType")
-		DEFINE_METASOUND_ENUM_ENTRY(Audio::EWaveShaperType::Sin, LOCTEXT("SinDescription", "Sine"), LOCTEXT("SinTT", "Sine WaveShaper")),
-		DEFINE_METASOUND_ENUM_ENTRY(Audio::EWaveShaperType::ATan, LOCTEXT("ATanDescription", "Inverse Tangent"), LOCTEXT("ATanTT", "Inverse Tangent WaveShaper")),
-		DEFINE_METASOUND_ENUM_ENTRY(Audio::EWaveShaperType::Tanh, LOCTEXT("TanHDescription", "Hyperbolic Tangent"), LOCTEXT("TanHTT", "Hyperbolic Tangent WaveShaper")),
-		DEFINE_METASOUND_ENUM_ENTRY(Audio::EWaveShaperType::Cubic, LOCTEXT("CubicDescription", "Cubic Polynomial"), LOCTEXT("CubicTT", "Cubic Polynomial WaveShaper")),
-		DEFINE_METASOUND_ENUM_ENTRY(Audio::EWaveShaperType::HardClip, LOCTEXT("HarClipDescription", "Hard Clip"), LOCTEXT("HardClipTT", "Hard Clipper"))
+		DEFINE_METASOUND_ENUM_ENTRY(Audio::EWaveShaperType::Sin, "SinDescription", "Sine", "SinTT", "Sine WaveShaper"),
+		DEFINE_METASOUND_ENUM_ENTRY(Audio::EWaveShaperType::ATan, "ATanDescription", "Inverse Tangent", "ATanTT", "Inverse Tangent WaveShaper"),
+		DEFINE_METASOUND_ENUM_ENTRY(Audio::EWaveShaperType::Tanh, "TanHDescription", "Hyperbolic Tangent", "TanHTT", "Hyperbolic Tangent WaveShaper"),
+		DEFINE_METASOUND_ENUM_ENTRY(Audio::EWaveShaperType::Cubic, "CubicDescription", "Cubic Polynomial", "CubicTT", "Cubic Polynomial WaveShaper"),
+		DEFINE_METASOUND_ENUM_ENTRY(Audio::EWaveShaperType::HardClip, "HarClipDescription", "Hard Clip", "HardClipTT", "Hard Clipper")
 		DEFINE_METASOUND_ENUM_END()
 
 	class FWaveShaperOperator : public TExecutableOperator<FWaveShaperOperator>
@@ -132,14 +132,14 @@ namespace Metasound
 	{
 		static const FVertexInterface Interface(
 			FInputVertexInterface(
-				TInputDataVertexModel<FAudioBuffer>(InParamNameAudioInput, LOCTEXT("AudioInputTT", "Audio input.")),
-				TInputDataVertexModel<float>(InParamNameWaveShapeAmount, LOCTEXT("WaveShapeAmountTT", "The amount of wave shaping to apply."), 1.0f),
-				TInputDataVertexModel<float>(InParamNameWaveShapeBias, LOCTEXT("WaveShaperBiasTT", "DC offset to apply before wave shaping."), 0.0f),
-				TInputDataVertexModel<float>(InParamNameOutputGain, LOCTEXT("WaveShaperGainTT", "The amount of gain to apply after processing."), 1.0f),
-				TInputDataVertexModel<FEnumEWaveShaperType>(InParamNameType, LOCTEXT("WaveShaperTypeTT", "Algorithm to use to process the audio."))
+				TInputDataVertexModel<FAudioBuffer>(InParamNameAudioInput, METASOUND_LOCTEXT("AudioInputTT", "Audio input.")),
+				TInputDataVertexModel<float>(InParamNameWaveShapeAmount, METASOUND_LOCTEXT("WaveShapeAmountTT", "The amount of wave shaping to apply."), 1.0f),
+				TInputDataVertexModel<float>(InParamNameWaveShapeBias, METASOUND_LOCTEXT("WaveShaperBiasTT", "DC offset to apply before wave shaping."), 0.0f),
+				TInputDataVertexModel<float>(InParamNameOutputGain, METASOUND_LOCTEXT("WaveShaperGainTT", "The amount of gain to apply after processing."), 1.0f),
+				TInputDataVertexModel<FEnumEWaveShaperType>(InParamNameType, METASOUND_LOCTEXT("WaveShaperTypeTT", "Algorithm to use to process the audio."))
 			),
 			FOutputVertexInterface(
-				TOutputDataVertexModel<FAudioBuffer>(OutParamNameAudio, LOCTEXT("WaveShaperOutputTT", "Audio output."))
+				TOutputDataVertexModel<FAudioBuffer>(OutParamNameAudio, METASOUND_LOCTEXT("WaveShaperOutputTT", "Audio output."))
 			)
 		);
 
@@ -154,8 +154,8 @@ namespace Metasound
 			Info.ClassName = { StandardNodes::Namespace, TEXT("WaveShaper"), StandardNodes::AudioVariant };
 			Info.MajorVersion = 1;
 			Info.MinorVersion = 0;
-			Info.DisplayName = LOCTEXT("Metasound_WaveShaperDisplayName", "WaveShaper");
-			Info.Description = LOCTEXT("Metasound_WaveShaperNodeDescription", "Applies non-linear shaping to the audio input.");
+			Info.DisplayName = METASOUND_LOCTEXT("Metasound_WaveShaperDisplayName", "WaveShaper");
+			Info.Description = METASOUND_LOCTEXT("Metasound_WaveShaperNodeDescription", "Applies non-linear shaping to the audio input.");
 			Info.Author = PluginAuthor;
 			Info.PromptIfMissing = PluginNodeMissingPrompt;
 			Info.DefaultInterface = GetVertexInterface();
diff --git a/Engine/Plugins/Runtime/Metasound/Source/MetasoundStandardNodes/Private/MetasoundWaveWriterNode.cpp b/Engine/Plugins/Runtime/Metasound/Source/MetasoundStandardNodes/Private/MetasoundWaveWriterNode.cpp
--- a/Engine/Plugins/Runtime/Metasound/Source/MetasoundStandardNodes/Private/MetasoundWaveWriterNode.cpp
+++ b/Engine/Plugins/Runtime/Metasound/Source/MetasoundStandardNodes/Private/MetasoundWaveWriterNode.cpp
@@ -123,7 +123,11 @@ namespace Metasound
 		virtual ~FFileWriteError() = default;
 
 		FFileWriteError(const FNode& InNode, const FString& InFilename)
-			: FBuildErrorBase(ErrorType, FText::Format(FTextFormat(LOCTEXT("MetasoundFileWriterErrorDescription", "File Writer Error while trying to write '{0}'")), FText::FromString(InFilename)))
+#if WITH_EDITOR
+			: FBuildErrorBase(ErrorType, METASOUND_LOCTEXT_FORMAT("MetasoundFileWriterErrorDescription", "File Writer Error while trying to write '{0}'", FText::FromString(InFilename)))
+#else 
+			: FBuildErrorBase(ErrorType, FText::GetEmpty())
+#endif // WITH_EDITOR
 		{
 			AddNode(InNode);
 		}
@@ -269,8 +273,8 @@ namespace Metasound
 
 				// inputs
 				FInputVertexInterface InputInterface(
-					TInputDataVertexModel<FString>(GetFilenamePrefixPinName(), LOCTEXT("WaveWriterFilenamePrefixDescription", "Filename Prefix of file you are writing."), GetDefaultFileName()),
-					TInputDataVertexModel<bool>(GetEnabledPinName(), LOCTEXT("WaveWriterEnabledDescription", "If this wave writer is enabled or not."), true)
+					TInputDataVertexModel<FString>(GetFilenamePrefixPinName(), METASOUND_LOCTEXT("WaveWriterFilenamePrefixDescription", "Filename Prefix of file you are writing."), GetDefaultFileName()),
+					TInputDataVertexModel<bool>(GetEnabledPinName(), METASOUND_LOCTEXT("WaveWriterEnabledDescription", "If this wave writer is enabled or not."), true)
 				);
 
 				// For backwards compatibility with previous (mono) node, in the case of 1 channels, just provide the old interface.
@@ -294,8 +298,8 @@ namespace Metasound
 			{
 				// For backwards compatibility with previous (mono) WaveWriters keep the node name the same.
 				FName OperatorName = TEXT("WaveWriter");
-				FText NodeDisplayName = LOCTEXT("Metasound_WaveWriterNodeMonoDisplayName", "Wave Writer (Mono)");
-				FText NodeDescription = LOCTEXT("Metasound_WaveWriterNodeMonoDescription", "Write a mono audio signal to disk");
+				FText NodeDisplayName = METASOUND_LOCTEXT("Metasound_WaveWriterNodeMonoDisplayName", "Wave Writer (Mono)");
+				const FText NodeDescription = METASOUND_LOCTEXT("Metasound_WaveWriterNodeMonoDescription", "Write a mono audio signal to disk");
 				FVertexInterface NodeInterface = DeclareVertexInterface();
 
 				return CreateNodeClassMetadata(OperatorName, NodeDisplayName, NodeDescription, NodeInterface);
@@ -305,8 +309,8 @@ namespace Metasound
 			auto CreateNodeClassMetadataStereo = []() -> FNodeClassMetadata
 			{
 				FName OperatorName = TEXT("Wave Writer (Stereo)");
-				FText NodeDisplayName = LOCTEXT("Metasound_WaveWriterNodeStereoDisplayName", "Wave Writer (Stereo)");
-				FText NodeDescription = LOCTEXT("Metasound_WaveWriterNodeStereoDescription", "Write a stereo audio signal to disk");
+				FText NodeDisplayName = METASOUND_LOCTEXT("Metasound_WaveWriterNodeStereoDisplayName", "Wave Writer (Stereo)");
+				const FText NodeDescription = METASOUND_LOCTEXT("Metasound_WaveWriterNodeStereoDescription", "Write a stereo audio signal to disk");
 				FVertexInterface NodeInterface = DeclareVertexInterface();
 
 				return  CreateNodeClassMetadata(OperatorName, NodeDisplayName, NodeDescription, NodeInterface);
@@ -316,8 +320,8 @@ namespace Metasound
 			auto CreateNodeClassMetadataMultiChan = []() -> FNodeClassMetadata
 			{
 				FName OperatorName = *FString::Printf(TEXT("Wave Writer (%d-Channel)"), NumInputChannels);
-				FText NodeDisplayName = FText::Format(LOCTEXT("Metasound_WaveWriterNodeMultiChannelDisplayName", "Wave Writer ({0}-channel)"), NumInputChannels);
-				FText NodeDescription = LOCTEXT("Metasound_WaveWriterNodeMultiDescription", "Write a multi-channel audio signal to disk");
+				FText NodeDisplayName = METASOUND_LOCTEXT_FORMAT("Metasound_WaveWriterNodeMultiChannelDisplayName", "Wave Writer ({0}-channel)", NumInputChannels);
+				const FText NodeDescription = METASOUND_LOCTEXT("Metasound_WaveWriterNodeMultiDescription", "Write a multi-channel audio signal to disk");
 				FVertexInterface NodeInterface = DeclareVertexInterface();
 
 				return  CreateNodeClassMetadata(OperatorName, NodeDisplayName, NodeDescription, NodeInterface);
@@ -378,7 +382,7 @@ namespace Metasound
 
 		static const FText GetAudioInputDescription(int32 InputIndex)
 		{
-			return FText::Format(LOCTEXT("WaveWriterAudioInputDescription", "Audio Input #: {0}"), InputIndex);
+			return METASOUND_LOCTEXT_FORMAT("WaveWriterAudioInputDescription", "Audio Input #: {0}", InputIndex);
 		}
 
 		static FNodeClassMetadata CreateNodeClassMetadata(const FName& InOperatorName, const FText& InDisplayName, const FText& InDescription, const FVertexInterface& InDefaultInterface)
@@ -394,7 +398,7 @@ namespace Metasound
 				PluginNodeMissingPrompt,
 				InDefaultInterface,
 				{ NodeCategories::Io },
-				{ LOCTEXT("Metasound_AudioMixerKeyword", "Writer") },
+				{ METASOUND_LOCTEXT("Metasound_AudioMixerKeyword", "Writer") },
 				FNodeDisplayStyle{}
 			};
 			return Metadata;
diff --git a/Engine/Plugins/Runtime/Metasound/Source/MetasoundStandardNodes/Public/MetasoundParamHelper.h b/Engine/Plugins/Runtime/Metasound/Source/MetasoundStandardNodes/Public/MetasoundParamHelper.h
--- a/Engine/Plugins/Runtime/Metasound/Source/MetasoundStandardNodes/Public/MetasoundParamHelper.h
+++ b/Engine/Plugins/Runtime/Metasound/Source/MetasoundStandardNodes/Public/MetasoundParamHelper.h
@@ -34,22 +34,31 @@ For variable parameters (e.g. templated on number), do the following:
 
 Then to retrieve the name/tooltip use:
 
-METASOUND_GET_VARIABLE_PARAM_NAME(InputParam, NUMBER);
+METASOUND_GET_PARAM_NAME_WITH_INDEX(InputParam, NUMBER);
 
 Where NUMBER is a number (or whatever) to slot into the format specifier you defined in the param def.
 
 */
 
-
+#if WITH_EDITOR
 #define METASOUND_PARAM(NAME, NAME_TEXT, TOOLTIP_TEXT) \
 	static const TCHAR* NAME##Name = TEXT(NAME_TEXT); \
-	static const FText NAME##Tooltip = LOCTEXT(#NAME "Tooltip", TOOLTIP_TEXT); \
-
+	static const FText NAME##Tooltip = LOCTEXT(#NAME "Tooltip", TOOLTIP_TEXT);
+#else 
+#define METASOUND_PARAM(NAME, NAME_TEXT, TOOLTIP_TEXT) \
+	static const TCHAR* NAME##Name = TEXT(NAME_TEXT); \
+	static const FText NAME##Tooltip = FText::GetEmpty();
+#endif // WITH_EDITOR
 
 #define METASOUND_GET_PARAM_NAME(NAME) NAME##Name
 #define METASOUND_GET_PARAM_TT(NAME) NAME##Tooltip
 #define METASOUND_GET_PARAM_NAME_AND_TT(NAME) NAME##Name, NAME##Tooltip
 
-#define METASOUND_GET_VARIABLE_PARAM_NAME(NAME, INDEX) *FString::Format(NAME##Name, {INDEX})
-#define METASOUND_GET_VARIABLE_TT(NAME, INDEX)  FText::Format(NAME##Tooltip, INDEX)
-#define METASOUND_GET_VARIABLE_PARAM_NAME_AND_TT(NAME, INDEX)  *FString::Format(NAME##Name, {INDEX}), FText::Format(NAME##Tooltip, INDEX)
+#define METASOUND_GET_PARAM_NAME_WITH_INDEX(NAME, INDEX) *FString::Format(NAME##Name, {INDEX})
+#if WITH_EDITOR
+#define METASOUND_GET_PARAM_NAME_WITH_INDEX_AND_TT(NAME, INDEX)  *FString::Format(NAME##Name, {INDEX}), FText::Format(NAME##Tooltip, INDEX)
+#define METASOUND_GET_PARAM_TT_WITH_INDEX(NAME, INDEX)  FText::Format(NAME##Tooltip, INDEX)
+#else 
+#define METASOUND_GET_PARAM_NAME_WITH_INDEX_AND_TT(NAME, INDEX)  *FString::Format(NAME##Name, {INDEX}), FText::GetEmpty()
+#define METASOUND_GET_PARAM_TT_WITH_INDEX(NAME, INDEX)  FText::GetEmpty();
+#endif // WITH_EDITOR
diff --git a/Engine/Plugins/Runtime/Metasound/Source/MetasoundStandardNodes/Public/MetasoundRandomNode.h b/Engine/Plugins/Runtime/Metasound/Source/MetasoundStandardNodes/Public/MetasoundRandomNode.h
--- a/Engine/Plugins/Runtime/Metasound/Source/MetasoundStandardNodes/Public/MetasoundRandomNode.h
+++ b/Engine/Plugins/Runtime/Metasound/Source/MetasoundStandardNodes/Public/MetasoundRandomNode.h
@@ -58,12 +58,12 @@ namespace Metasound
 
 		static FText GetDisplayName()
 		{
-			return LOCTEXT("RandomNode_Int32RandomValueDisplayName", "Random (Int)");
+			return METASOUND_LOCTEXT("RandomNode_Int32RandomValueDisplayName", "Random (Int)");
 		}
 
 		static FText GetDescription()
 		{
-			return LOCTEXT("RandomNode_Int32RandomDescription", "Generates a seedable random integer in the given value range.");
+			return METASOUND_LOCTEXT("RandomNode_Int32RandomDescription", "Generates a seedable random integer in the given value range.");
 		}
 
 		static bool HasRange()
@@ -107,12 +107,12 @@ namespace Metasound
 
 		static FText GetDisplayName()
 		{
-			return LOCTEXT("RandomNode_FloatRandomValueDisplayName", "Random (Float)");
+			return METASOUND_LOCTEXT("RandomNode_FloatRandomValueDisplayName", "Random (Float)");
 		}
 
 		static FText GetDescription()
 		{
-			return LOCTEXT("RandomNode_FloatRandomDescription", "Generates a seedable random float in the given value range.");
+			return METASOUND_LOCTEXT("RandomNode_FloatRandomDescription", "Generates a seedable random float in the given value range.");
 		}
 
 		static bool HasRange()
@@ -156,12 +156,12 @@ namespace Metasound
 
 		static FText GetDisplayName()
 		{
-			return LOCTEXT("RandomNode_BoolRandomValueDisplayName", "Random (Bool)");
+			return METASOUND_LOCTEXT("RandomNode_BoolRandomValueDisplayName", "Random (Bool)");
 		}
 
 		static FText GetDescription()
 		{
-			return LOCTEXT("RandomNode_BoolRandomDiscription", "Generates a random bool value.");
+			return METASOUND_LOCTEXT("RandomNode_BoolRandomDiscription", "Generates a random bool value.");
 		}
 
 		static bool HasRange()
@@ -205,12 +205,12 @@ namespace Metasound
 
 		static FText GetDisplayName()
 		{
-			return LOCTEXT("RandomNode_TimeRandomValueDisplayName", "Random (Time)");
+			return METASOUND_LOCTEXT("RandomNode_TimeRandomValueDisplayName", "Random (Time)");
 		}
 
 		static FText GetDescription()
 		{
-			return LOCTEXT("RandomNode_TimeRandomDiscription", "Generates a random time value.");
+			return METASOUND_LOCTEXT("RandomNode_TimeRandomDiscription", "Generates a random time value.");
 		}
 
 		static bool HasRange()
@@ -305,7 +305,7 @@ namespace Metasound
 				Info.MajorVersion = 1;
 				Info.MinorVersion = 1;
 				Info.DisplayName = TRandomNodeSpecialization<ValueType>::GetDisplayName();
-				Info.Description = LOCTEXT("RandomNode_Description", "Generates a random value.");
+				Info.Description = METASOUND_LOCTEXT("RandomNode_Description", "Generates a random value.");
 				Info.Author = PluginAuthor;
 				Info.PromptIfMissing = PluginNodeMissingPrompt;
 				Info.DefaultInterface = GetDefaultInterface();
diff --git a/Engine/Plugins/Runtime/Metasound/Source/MetasoundStandardNodes/Public/MetasoundScaleEnums.h b/Engine/Plugins/Runtime/Metasound/Source/MetasoundStandardNodes/Public/MetasoundScaleEnums.h
--- a/Engine/Plugins/Runtime/Metasound/Source/MetasoundStandardNodes/Public/MetasoundScaleEnums.h
+++ b/Engine/Plugins/Runtime/Metasound/Source/MetasoundStandardNodes/Public/MetasoundScaleEnums.h
@@ -21,43 +21,43 @@ namespace Metasound
 	DEFINE_METASOUND_ENUM_BEGIN(Audio::EMusicalScale::Scale, FEnumEMusicalScale, "MusicalScale")
 
 	// modes
-	DEFINE_METASOUND_ENUM_ENTRY(Audio::EMusicalScale::Scale::Major, LOCTEXT("MajorDescription", "Major Scale"), LOCTEXT("MajorDescriptionTT", "Major (Ionian)")),
-	DEFINE_METASOUND_ENUM_ENTRY(Audio::EMusicalScale::Scale::Minor_Dorian, LOCTEXT("Minor_DorianDescription", "Minor (Dorian)"), LOCTEXT("Minor_DorianDescriptionTT", "Dorian Minor")),
-	DEFINE_METASOUND_ENUM_ENTRY(Audio::EMusicalScale::Scale::Phrygian, LOCTEXT("PhrygianDescription", "Phrygian "), LOCTEXT("PhrygianDescriptionTT", "Phrygian")),
-	DEFINE_METASOUND_ENUM_ENTRY(Audio::EMusicalScale::Scale::Lydian, LOCTEXT("LydianDescription", "Lydian"), LOCTEXT("LydianDescriptionTT", "Lydian (sharp-4)")),
-	DEFINE_METASOUND_ENUM_ENTRY(Audio::EMusicalScale::Scale::Dominant7th_Mixolydian, LOCTEXT("Dominant7th_MixolydianDescription", "Dominant 7th (Mixolydian)"), LOCTEXT("Dominant7th_MixolydianDescriptionTT", "Mioxlydian (Dominant 7)")),
-	DEFINE_METASOUND_ENUM_ENTRY(Audio::EMusicalScale::Scale::NaturalMinor_Aeolian, LOCTEXT("NaturalMinor_AeolianDescription", "Natural Minor (Aeolian)"), LOCTEXT("NaturalMinor_AeolianDescriptionTT", "Natural Minor (Aeolian)")),
-	DEFINE_METASOUND_ENUM_ENTRY(Audio::EMusicalScale::Scale::HalfDiminished_Locrian, LOCTEXT("HalfDiminished_LocrianDescription", "Half Diminished (Locrian)"), LOCTEXT("HalfDiminished_LocrianDescriptionTT", "Half-Diminished (Locrian)")),
-	DEFINE_METASOUND_ENUM_ENTRY(Audio::EMusicalScale::Scale::Diminished, LOCTEXT("DiminishedDescription", "Diminished "), LOCTEXT("DiminishedDescriptionTT", "Diminished")),
+	DEFINE_METASOUND_ENUM_ENTRY(Audio::EMusicalScale::Scale::Major, "MajorDescription", "Major Scale", "MajorDescriptionTT", "Major (Ionian)"),
+	DEFINE_METASOUND_ENUM_ENTRY(Audio::EMusicalScale::Scale::Minor_Dorian, "Minor_DorianDescription", "Minor (Dorian)", "Minor_DorianDescriptionTT", "Dorian Minor"),
+	DEFINE_METASOUND_ENUM_ENTRY(Audio::EMusicalScale::Scale::Phrygian, "PhrygianDescription", "Phrygian ", "PhrygianDescriptionTT", "Phrygian"),
+	DEFINE_METASOUND_ENUM_ENTRY(Audio::EMusicalScale::Scale::Lydian, "LydianDescription", "Lydian", "LydianDescriptionTT", "Lydian (sharp-4)"),
+	DEFINE_METASOUND_ENUM_ENTRY(Audio::EMusicalScale::Scale::Dominant7th_Mixolydian, "Dominant7th_MixolydianDescription", "Dominant 7th (Mixolydian)", "Dominant7th_MixolydianDescriptionTT", "Mioxlydian (Dominant 7)"),
+	DEFINE_METASOUND_ENUM_ENTRY(Audio::EMusicalScale::Scale::NaturalMinor_Aeolian, "NaturalMinor_AeolianDescription", "Natural Minor (Aeolian)", "NaturalMinor_AeolianDescriptionTT", "Natural Minor (Aeolian)"),
+	DEFINE_METASOUND_ENUM_ENTRY(Audio::EMusicalScale::Scale::HalfDiminished_Locrian, "HalfDiminished_LocrianDescription", "Half Diminished (Locrian)", "HalfDiminished_LocrianDescriptionTT", "Half-Diminished (Locrian)"),
+	DEFINE_METASOUND_ENUM_ENTRY(Audio::EMusicalScale::Scale::Diminished, "DiminishedDescription", "Diminished ", "DiminishedDescriptionTT", "Diminished"),
 	// non-diatonic
-	DEFINE_METASOUND_ENUM_ENTRY(Audio::EMusicalScale::Scale::Chromatic, LOCTEXT("ChromaticDescription", "Chromatic"), LOCTEXT("ChromaticDescriptionTT", "Chromatic")),
-	DEFINE_METASOUND_ENUM_ENTRY(Audio::EMusicalScale::Scale::WholeTone, LOCTEXT("WholeToneDescription", "Whole-Tone"), LOCTEXT("WholeToneDescriptionTT", "Whole Tone")),
-	DEFINE_METASOUND_ENUM_ENTRY(Audio::EMusicalScale::Scale::DiminishedWholeTone, LOCTEXT("DiminishedWholeToneDescription", "Diminished Whole-Tone"), LOCTEXT("DiminishedWholeToneDescriptionTT", "Diminished Whole Tone")),
+	DEFINE_METASOUND_ENUM_ENTRY(Audio::EMusicalScale::Scale::Chromatic, "ChromaticDescription", "Chromatic", "ChromaticDescriptionTT", "Chromatic"),
+	DEFINE_METASOUND_ENUM_ENTRY(Audio::EMusicalScale::Scale::WholeTone, "WholeToneDescription", "Whole-Tone", "WholeToneDescriptionTT", "Whole Tone"),
+	DEFINE_METASOUND_ENUM_ENTRY(Audio::EMusicalScale::Scale::DiminishedWholeTone, "DiminishedWholeToneDescription", "Diminished Whole-Tone", "DiminishedWholeToneDescriptionTT", "Diminished Whole Tone"),
 	// petantonic
-	DEFINE_METASOUND_ENUM_ENTRY(Audio::EMusicalScale::Scale::MajorPentatonic, LOCTEXT("MajorPentatonicDescription", "Major Pentatonic "), LOCTEXT("MajorPentatonicDescriptionTT", "Major Pentatonic")),
-	DEFINE_METASOUND_ENUM_ENTRY(Audio::EMusicalScale::Scale::MinorPentatonic, LOCTEXT("MinorPentatonicDescription", "Minor Pentatonic "), LOCTEXT("MinorPentatonicDescriptionTT", "Minor Pentatonic")),
-	DEFINE_METASOUND_ENUM_ENTRY(Audio::EMusicalScale::Scale::Blues, LOCTEXT("BluesDescription", "Blues "), LOCTEXT("BluesDescriptionTT", "Blues")),
+	DEFINE_METASOUND_ENUM_ENTRY(Audio::EMusicalScale::Scale::MajorPentatonic, "MajorPentatonicDescription", "Major Pentatonic ", "MajorPentatonicDescriptionTT", "Major Pentatonic"),
+	DEFINE_METASOUND_ENUM_ENTRY(Audio::EMusicalScale::Scale::MinorPentatonic, "MinorPentatonicDescription", "Minor Pentatonic ", "MinorPentatonicDescriptionTT", "Minor Pentatonic"),
+	DEFINE_METASOUND_ENUM_ENTRY(Audio::EMusicalScale::Scale::Blues, "BluesDescription", "Blues ", "BluesDescriptionTT", "Blues"),
 	// bebop
-	DEFINE_METASOUND_ENUM_ENTRY(Audio::EMusicalScale::Scale::Bebop_Major, LOCTEXT("Bebop_MajorDescription", "Bebop (Major)"), LOCTEXT("Bebop_MajorDescriptionTT", "Bebop Major")),
-	DEFINE_METASOUND_ENUM_ENTRY(Audio::EMusicalScale::Scale::Bebop_Minor, LOCTEXT("Bebop_MinorDescription", "Bebop (Minor)"), LOCTEXT("Bebop_MinorDescriptionTT", "Bebop Minor")),
-	DEFINE_METASOUND_ENUM_ENTRY(Audio::EMusicalScale::Scale::Bebop_MinorNumber2, LOCTEXT("Bebop_MinorNumber2Description", "Bebop (Minor) #2"), LOCTEXT("Bebop_MinorNumber2DescriptionTT", "Bebop Minor #2")),
-	DEFINE_METASOUND_ENUM_ENTRY(Audio::EMusicalScale::Scale::Bebop_Dominant, LOCTEXT("Bebop_DominantDescription", "Bebop (Dominant)"), LOCTEXT("Bebop_DominantDescriptionTT", "Bebop Dominant")),
+	DEFINE_METASOUND_ENUM_ENTRY(Audio::EMusicalScale::Scale::Bebop_Major, "Bebop_MajorDescription", "Bebop (Major)", "Bebop_MajorDescriptionTT", "Bebop Major"),
+	DEFINE_METASOUND_ENUM_ENTRY(Audio::EMusicalScale::Scale::Bebop_Minor, "Bebop_MinorDescription", "Bebop (Minor)", "Bebop_MinorDescriptionTT", "Bebop Minor"),
+	DEFINE_METASOUND_ENUM_ENTRY(Audio::EMusicalScale::Scale::Bebop_MinorNumber2, "Bebop_MinorNumber2Description", "Bebop (Minor) #2", "Bebop_MinorNumber2DescriptionTT", "Bebop Minor #2"),
+	DEFINE_METASOUND_ENUM_ENTRY(Audio::EMusicalScale::Scale::Bebop_Dominant, "Bebop_DominantDescription", "Bebop (Dominant)", "Bebop_DominantDescriptionTT", "Bebop Dominant"),
 	// common major/minors
-	DEFINE_METASOUND_ENUM_ENTRY(Audio::EMusicalScale::Scale::HarmonicMajor, LOCTEXT("HarmonicMajorDescription", "Harmonic Major"), LOCTEXT("HarmonicMajorDescriptionTT", "Harmonic Major")),
-	DEFINE_METASOUND_ENUM_ENTRY(Audio::EMusicalScale::Scale::HarmonicMinor, LOCTEXT("HarmonicMinorDescription", "Harmonic Minor "), LOCTEXT("HarmonicMinorDescriptionTT", "Harmonic Minor")),
-	DEFINE_METASOUND_ENUM_ENTRY(Audio::EMusicalScale::Scale::MelodicMinor, LOCTEXT("MelodicMinorDescription", "Melodic Minor "), LOCTEXT("MelodicMinorDescriptionTT", "Melodic Minor")),
-	DEFINE_METASOUND_ENUM_ENTRY(Audio::EMusicalScale::Scale::SixthModeOfHarmonicMinor, LOCTEXT("SixthModeOfHarmonicMinorDescription", "Sixth Mode of Harmonic Minor"), LOCTEXT("SixthModeOfHarmonicMinorDescriptionTT", "Sixth Mode of Harmonic Minor")),
+	DEFINE_METASOUND_ENUM_ENTRY(Audio::EMusicalScale::Scale::HarmonicMajor, "HarmonicMajorDescription", "Harmonic Major", "HarmonicMajorDescriptionTT", "Harmonic Major"),
+	DEFINE_METASOUND_ENUM_ENTRY(Audio::EMusicalScale::Scale::HarmonicMinor, "HarmonicMinorDescription", "Harmonic Minor ", "HarmonicMinorDescriptionTT", "Harmonic Minor"),
+	DEFINE_METASOUND_ENUM_ENTRY(Audio::EMusicalScale::Scale::MelodicMinor, "MelodicMinorDescription", "Melodic Minor ", "MelodicMinorDescriptionTT", "Melodic Minor"),
+	DEFINE_METASOUND_ENUM_ENTRY(Audio::EMusicalScale::Scale::SixthModeOfHarmonicMinor, "SixthModeOfHarmonicMinorDescription", "Sixth Mode of Harmonic Minor", "SixthModeOfHarmonicMinorDescriptionTT", "Sixth Mode of Harmonic Minor"),
 	// lydian/augmented
-	DEFINE_METASOUND_ENUM_ENTRY(Audio::EMusicalScale::Scale::LydianAugmented, LOCTEXT("LydianAugmentedDescription", "Lydian Augmented"), LOCTEXT("LydianAugmentedDescriptionTT", "Lydian Augmented")),
-	DEFINE_METASOUND_ENUM_ENTRY(Audio::EMusicalScale::Scale::LydianDominant, LOCTEXT("LydianDominantDescription", "Lydian Dominant "), LOCTEXT("LydianDominantDescriptionTT", "Lydian Dominant")),
-	DEFINE_METASOUND_ENUM_ENTRY(Audio::EMusicalScale::Scale::Augmented, LOCTEXT("AugmentedDescription", "Augmented"), LOCTEXT("AugmentedDescriptionTT", "Augmented")),
+	DEFINE_METASOUND_ENUM_ENTRY(Audio::EMusicalScale::Scale::LydianAugmented, "LydianAugmentedDescription", "Lydian Augmented", "LydianAugmentedDescriptionTT", "Lydian Augmented"),
+	DEFINE_METASOUND_ENUM_ENTRY(Audio::EMusicalScale::Scale::LydianDominant, "LydianDominantDescription", "Lydian Dominant ", "LydianDominantDescriptionTT", "Lydian Dominant"),
+	DEFINE_METASOUND_ENUM_ENTRY(Audio::EMusicalScale::Scale::Augmented, "AugmentedDescription", "Augmented", "AugmentedDescriptionTT", "Augmented"),
 	// diminished
-	DEFINE_METASOUND_ENUM_ENTRY(Audio::EMusicalScale::Scale::Diminished_BeginWithHalfStep, LOCTEXT("Diminished_BeginWithHalfStepDescription", "Diminished (Begin With Half-Step)"), LOCTEXT("Diminished_BeginWithHalfStepDescriptionTT", "Diminished (begins with Half Step)")),
-	DEFINE_METASOUND_ENUM_ENTRY(Audio::EMusicalScale::Scale::Diminished_BeginWithWholeStep, LOCTEXT("Diminished_BeginWithWholeStepDescription", "Diminished (Begin With Whole-Step"), LOCTEXT("Diminished_BeginWithWholeStepDescriptionTT", "Diminished (begins with Whole Step)")),
-	DEFINE_METASOUND_ENUM_ENTRY(Audio::EMusicalScale::Scale::HalfDiminished_LocrianNumber2, LOCTEXT("HalfDiminished_LocrianNumber2Description", "Half-Diminished (Locrian #2)"), LOCTEXT("HalfDiminished_LocrianNumber2DescriptionTT", "Half Diminished Locrian (#2)")),
+	DEFINE_METASOUND_ENUM_ENTRY(Audio::EMusicalScale::Scale::Diminished_BeginWithHalfStep, "Diminished_BeginWithHalfStepDescription", "Diminished (Begin With Half-Step)", "Diminished_BeginWithHalfStepDescriptionTT", "Diminished (begins with Half Step)"),
+	DEFINE_METASOUND_ENUM_ENTRY(Audio::EMusicalScale::Scale::Diminished_BeginWithWholeStep, "Diminished_BeginWithWholeStepDescription", "Diminished (Begin With Whole-Step", "Diminished_BeginWithWholeStepDescriptionTT", "Diminished (begins with Whole Step)"),
+	DEFINE_METASOUND_ENUM_ENTRY(Audio::EMusicalScale::Scale::HalfDiminished_LocrianNumber2, "HalfDiminished_LocrianNumber2Description", "Half-Diminished (Locrian #2)", "HalfDiminished_LocrianNumber2DescriptionTT", "Half Diminished Locrian (#2)"),
 	// other
-	DEFINE_METASOUND_ENUM_ENTRY(Audio::EMusicalScale::Scale::Spanish_or_Jewish, LOCTEXT("Spanish_or_JewishDescription", "Spanish or Jewish Scale"), LOCTEXT("Spanish_or_JewishDescriptionTT", "Spanish/Jewish")),
-	DEFINE_METASOUND_ENUM_ENTRY(Audio::EMusicalScale::Scale::Hindu, LOCTEXT("HinduDescription", "Hindu "), LOCTEXT("HinduDescriptionTT", "Hindu"))
+	DEFINE_METASOUND_ENUM_ENTRY(Audio::EMusicalScale::Scale::Spanish_or_Jewish, "Spanish_or_JewishDescription", "Spanish or Jewish Scale", "Spanish_or_JewishDescriptionTT", "Spanish/Jewish"),
+	DEFINE_METASOUND_ENUM_ENTRY(Audio::EMusicalScale::Scale::Hindu, "HinduDescription", "Hindu ", "HinduDescriptionTT", "Hindu")
 
 	DEFINE_METASOUND_ENUM_END()
 
diff --git a/Engine/Plugins/Runtime/Metasound/Source/MetasoundStandardNodes/Public/MetasoundTriggerAccumulatorNode.h b/Engine/Plugins/Runtime/Metasound/Source/MetasoundStandardNodes/Public/MetasoundTriggerAccumulatorNode.h
--- a/Engine/Plugins/Runtime/Metasound/Source/MetasoundStandardNodes/Public/MetasoundTriggerAccumulatorNode.h
+++ b/Engine/Plugins/Runtime/Metasound/Source/MetasoundStandardNodes/Public/MetasoundTriggerAccumulatorNode.h
@@ -4,6 +4,7 @@
 #include "CoreMinimal.h"
 #include "MetasoundFacade.h"
 #include "MetasoundExecutableOperator.h"
+#include "MetasoundNodeRegistrationMacro.h"
 #include "MetasoundPrimitives.h"
 #include "MetasoundStandardNodesNames.h"
 #include "MetasoundTrigger.h"
@@ -40,7 +41,7 @@ namespace Metasound
 
 				for (uint32 i = 0; i < NumInputs; ++i)
 				{
-					InputInterface.Add(TInputDataVertexModel<FTrigger>(METASOUND_GET_VARIABLE_PARAM_NAME_AND_TT(InputTrigger, i)));
+					InputInterface.Add(TInputDataVertexModel<FTrigger>(METASOUND_GET_PARAM_NAME_WITH_INDEX_AND_TT(InputTrigger, i)));
 				}
 
 				InputInterface.Add(TInputDataVertexModel<bool>(METASOUND_GET_PARAM_NAME_AND_TT(InputAutoReset)));
@@ -60,8 +61,8 @@ namespace Metasound
 			auto CreateNodeClassMetadata = []() -> FNodeClassMetadata
 			{
 				FName OperatorName = *FString::Printf(TEXT("Trigger Accumulate (%d)") , NumInputs);
-				FText NodeDisplayName = FText::Format(LOCTEXT("TriggerAccumulateDisplayNamePattern", "Trigger Accumulate ({0})"), NumInputs);
-				FText NodeDescription = LOCTEXT("TriggerAccumulateDescription", "Will trigger output once all input triggers have been hit at some point in the past.");
+				FText NodeDisplayName = METASOUND_LOCTEXT_FORMAT("TriggerAccumulateDisplayNamePattern", "Trigger Accumulate ({0})", NumInputs);
+				const FText NodeDescription = METASOUND_LOCTEXT("TriggerAccumulateDescription", "Will trigger output once all input triggers have been hit at some point in the past.");
 				FVertexInterface NodeInterface = GetDefaultInterface();
 
 				return MetasoundTriggerAccumulatorNodePrivate::CreateNodeClassMetadata(OperatorName, NodeDisplayName, NodeDescription, NodeInterface);
@@ -83,7 +84,7 @@ namespace Metasound
 
 			for (uint32 i = 0; i < NumInputs; ++i)
 			{
-				InputTriggers.Add(InputCollection.GetDataReadReferenceOrConstruct<FTrigger>(METASOUND_GET_VARIABLE_PARAM_NAME(InputTrigger, i), InParams.OperatorSettings));
+				InputTriggers.Add(InputCollection.GetDataReadReferenceOrConstruct<FTrigger>(METASOUND_GET_PARAM_NAME_WITH_INDEX(InputTrigger, i), InParams.OperatorSettings));
 			}
 
 			return MakeUnique<TTriggerAccumulatorOperator<NumInputs>>(InParams.OperatorSettings, bInAutoReset, MoveTemp(InputTriggers));
@@ -109,7 +110,7 @@ namespace Metasound
 			FDataReferenceCollection Inputs;
 			for (uint32 i = 0; i < NumInputs; ++i)
 			{
-				Inputs.AddDataReadReference(METASOUND_GET_VARIABLE_PARAM_NAME(InputTrigger, i), InputTriggers[i]);
+				Inputs.AddDataReadReference(METASOUND_GET_PARAM_NAME_WITH_INDEX(InputTrigger, i), InputTriggers[i]);
 			}
 			Inputs.AddDataReadReference(METASOUND_GET_PARAM_NAME(InputAutoReset), bAutoReset);
 
diff --git a/Engine/Plugins/Runtime/Metasound/Source/MetasoundStandardNodes/Public/MetasoundTriggerCompareNode.h b/Engine/Plugins/Runtime/Metasound/Source/MetasoundStandardNodes/Public/MetasoundTriggerCompareNode.h
--- a/Engine/Plugins/Runtime/Metasound/Source/MetasoundStandardNodes/Public/MetasoundTriggerCompareNode.h
+++ b/Engine/Plugins/Runtime/Metasound/Source/MetasoundStandardNodes/Public/MetasoundTriggerCompareNode.h
@@ -51,12 +51,12 @@ namespace Metasound
 	FEnumTriggerComparisonType, FEnumTriggerComparisonTypeInfo, FEnumTriggerComparisonTypeReadRef, FEnumTriggerComparisonTypeWriteRef);
 
 	DEFINE_METASOUND_ENUM_BEGIN(ETriggerComparisonType, FEnumTriggerComparisonType, "TriggerComparisonType")
-		DEFINE_METASOUND_ENUM_ENTRY(ETriggerComparisonType::Equals, LOCTEXT("EqualsDescription", "Equals"), LOCTEXT("EqualsDescriptionTT", "True if A and B are equal.")),
-		DEFINE_METASOUND_ENUM_ENTRY(ETriggerComparisonType::NotEquals, LOCTEXT("NotEqualsDescriptioin", "Not Equals"), LOCTEXT("NotEqualsTT", "True if A and B are not equal.")),
-		DEFINE_METASOUND_ENUM_ENTRY(ETriggerComparisonType::LessThan, LOCTEXT("LessThanDescription", "Less Than"), LOCTEXT("LessThanTT", "True if A is less than B.")),
-		DEFINE_METASOUND_ENUM_ENTRY(ETriggerComparisonType::GreaterThan, LOCTEXT("GreaterThanDescription", "Greater Than"), LOCTEXT("GreaterThanTT", "True if A is greater than B.")),
-		DEFINE_METASOUND_ENUM_ENTRY(ETriggerComparisonType::LessThanOrEquals, LOCTEXT("LessThanOrEqualsDescription", "Less Than Or Equals"), LOCTEXT("LessThanOrEqualsTT", "True if A is less than or equal to B.")),
-		DEFINE_METASOUND_ENUM_ENTRY(ETriggerComparisonType::GreaterThanOrEquals, LOCTEXT("GreaterThanOrEqualsDescription", "Greater Than Or Equals"), LOCTEXT("GreaterThanOrEqualsTT", "True if A is greater than or equal to B.")),
+		DEFINE_METASOUND_ENUM_ENTRY(ETriggerComparisonType::Equals, "EqualsDescription", "Equals", "EqualsDescriptionTT", "True if A and B are equal."),
+		DEFINE_METASOUND_ENUM_ENTRY(ETriggerComparisonType::NotEquals, "NotEqualsDescriptioin", "Not Equals", "NotEqualsTT", "True if A and B are not equal."),
+		DEFINE_METASOUND_ENUM_ENTRY(ETriggerComparisonType::LessThan, "LessThanDescription", "Less Than", "LessThanTT", "True if A is less than B."),
+		DEFINE_METASOUND_ENUM_ENTRY(ETriggerComparisonType::GreaterThan, "GreaterThanDescription", "Greater Than", "GreaterThanTT", "True if A is greater than B."),
+		DEFINE_METASOUND_ENUM_ENTRY(ETriggerComparisonType::LessThanOrEquals, "LessThanOrEqualsDescription", "Less Than Or Equals", "LessThanOrEqualsTT", "True if A is less than or equal to B."),
+		DEFINE_METASOUND_ENUM_ENTRY(ETriggerComparisonType::GreaterThanOrEquals, "GreaterThanOrEqualsDescription", "Greater Than Or Equals", "GreaterThanOrEqualsTT", "True if A is greater than or equal to B."),
 	DEFINE_METASOUND_ENUM_END()
 
 
@@ -91,8 +91,8 @@ namespace Metasound
 			{
 				const FName DataTypeName = GetMetasoundDataTypeName<ValueType>();
 				const FName OperatorName = TEXT("Clamp");
-				const FText NodeDisplayName = FText::Format(LOCTEXT("TriggerCompareDisplayPattern", "Trigger Compare ({0})"), GetMetasoundDataTypeDisplayText<ValueType>());
-				const FText NodeDescription = LOCTEXT("TriggerCompareDisc", "Output triggers (True or False) based on comparing inputs, A and B.");
+				const FText NodeDisplayName = METASOUND_LOCTEXT_FORMAT("TriggerCompareDisplayPattern", "Trigger Compare ({0})", GetMetasoundDataTypeDisplayText<ValueType>());
+				const FText NodeDescription = METASOUND_LOCTEXT("TriggerCompareDisc", "Output triggers (True or False) based on comparing inputs, A and B.");
 				const FVertexInterface NodeInterface = GetDefaultInterface();
 
 				return MetasoundTriggerCompareNodePrivate::CreateNodeClassMetadata(DataTypeName, OperatorName, NodeDisplayName, NodeDescription, NodeInterface);
diff --git a/Engine/Plugins/Runtime/Metasound/Source/MetasoundStandardNodes/Public/MetasoundValueNode.h b/Engine/Plugins/Runtime/Metasound/Source/MetasoundStandardNodes/Public/MetasoundValueNode.h
--- a/Engine/Plugins/Runtime/Metasound/Source/MetasoundStandardNodes/Public/MetasoundValueNode.h
+++ b/Engine/Plugins/Runtime/Metasound/Source/MetasoundStandardNodes/Public/MetasoundValueNode.h
@@ -5,6 +5,7 @@
 #include "CoreMinimal.h"
 #include "MetasoundFacade.h"
 #include "MetasoundExecutableOperator.h"
+#include "MetasoundNodeRegistrationMacro.h"
 #include "MetasoundPrimitives.h"
 #include "MetasoundStandardNodesNames.h"
 #include "MetasoundTrigger.h"
@@ -62,8 +63,8 @@ namespace Metasound
 			{
 				const FName DataTypeName = GetMetasoundDataTypeName<ValueType>();
 				const FName OperatorName = "Value";
-				const FText NodeDisplayName = FText::Format(LOCTEXT("ValueDisplayNamePattern", "Value ({0})"), GetMetasoundDataTypeDisplayText<ValueType>());
-				const FText NodeDescription = LOCTEXT("ValueDescription", "Allows setting a value to output on trigger.");
+				const FText NodeDisplayName = METASOUND_LOCTEXT_FORMAT("ValueDisplayNamePattern", "Value ({0})", GetMetasoundDataTypeDisplayText<ValueType>());
+				const FText NodeDescription = METASOUND_LOCTEXT("ValueDescription", "Allows setting a value to output on trigger.");
 				const FVertexInterface NodeInterface = GetDefaultInterface();
 
 				return MetasoundValueNodePrivate::CreateNodeClassMetadata(DataTypeName, OperatorName, NodeDisplayName, NodeDescription, NodeInterface);
