Fix crash when creating preset with incompatible literal
#rnx

diff --git a/Engine/Plugins/Runtime/Metasound/Source/MetasoundEditor/Private/MetasoundEditorGraphMemberDefaults.cpp b/Engine/Plugins/Runtime/Metasound/Source/MetasoundEditor/Private/MetasoundEditorGraphMemberDefaults.cpp
--- a/Engine/Plugins/Runtime/Metasound/Source/MetasoundEditor/Private/MetasoundEditorGraphMemberDefaults.cpp
+++ b/Engine/Plugins/Runtime/Metasound/Source/MetasoundEditor/Private/MetasoundEditorGraphMemberDefaults.cpp
@@ -31,17 +31,16 @@ namespace Metasound
 			template <typename TPODType>
 			void ConvertLiteral(const FMetasoundFrontendLiteral& InLiteral, TPODType& OutValue)
 			{
-				const FName& TypeName = GetMetasoundDataTypeName<TPODType>();
-				OutValue = InLiteral.ToLiteral(TypeName).Value.Get<TPODType>();
+				InLiteral.TryGet(OutValue);
 			}
 
 			template <typename TPODType, typename TLiteralType = TPODType>
 			void ConvertLiteralToArray(const FMetasoundFrontendLiteral& InLiteral, TArray<TLiteralType>& OutArray)
 			{
 				OutArray.Reset();
-				const FName TypeName = CreateArrayTypeNameFromElementTypeName(GetMetasoundDataTypeName<TPODType>());
-				TArray<TPODType> NewValue = InLiteral.ToLiteral(TypeName).Value.Get<TArray<TPODType>>();
-				Algo::Transform(NewValue, OutArray, [](const TPODType& InValue) { return TLiteralType { InValue }; });
+				TArray<TPODType> Values;
+				InLiteral.TryGet(Values);
+				Algo::Transform(Values, OutArray, [](const TPODType& InValue) { return TLiteralType{ InValue }; });
 			}
 		}
 	}
