Fix for Undoing/Redoing change to a MetaSound not propagating to other referencing MetaSounds.
#rnx

diff --git a/Engine/Plugins/Runtime/Metasound/Source/MetasoundEditor/Private/MetasoundEditor.cpp b/Engine/Plugins/Runtime/Metasound/Source/MetasoundEditor/Private/MetasoundEditor.cpp
--- a/Engine/Plugins/Runtime/Metasound/Source/MetasoundEditor/Private/MetasoundEditor.cpp
+++ b/Engine/Plugins/Runtime/Metasound/Source/MetasoundEditor/Private/MetasoundEditor.cpp
@@ -574,11 +574,6 @@ namespace Metasound
 				MetasoundGraphEditor->NotifyGraphChanged();
 			}
 
-			// Synchronize will update the interface
-			FMetasoundAssetBase* MetasoundAsset = IMetasoundUObjectRegistry::Get().GetObjectAsAssetBase(Metasound);
-			check(MetasoundAsset);
-			MetasoundAsset->SetSynchronizationRequired();
-
 			FSlateApplication::Get().DismissAllMenus();
 		}
 
diff --git a/Engine/Plugins/Runtime/Metasound/Source/MetasoundEngine/Private/Metasound.cpp b/Engine/Plugins/Runtime/Metasound/Source/MetasoundEngine/Private/Metasound.cpp
--- a/Engine/Plugins/Runtime/Metasound/Source/MetasoundEngine/Private/Metasound.cpp
+++ b/Engine/Plugins/Runtime/Metasound/Source/MetasoundEngine/Private/Metasound.cpp
@@ -53,7 +53,11 @@ void UMetaSound::PostDuplicate(EDuplicateMode::Type InDuplicateMode)
 void UMetaSound::PostEditUndo()
 {
 	Super::PostEditUndo();
-	bSynchronizationRequired = true;
+
+	if (Graph)
+	{
+		Graph->RegisterGraphWithFrontend();
+	}
 }
 
 void UMetaSound::PostEditChangeProperty(FPropertyChangedEvent& InEvent)
diff --git a/Engine/Plugins/Runtime/Metasound/Source/MetasoundEngine/Private/MetasoundSource.cpp b/Engine/Plugins/Runtime/Metasound/Source/MetasoundEngine/Private/MetasoundSource.cpp
--- a/Engine/Plugins/Runtime/Metasound/Source/MetasoundEngine/Private/MetasoundSource.cpp
+++ b/Engine/Plugins/Runtime/Metasound/Source/MetasoundEngine/Private/MetasoundSource.cpp
@@ -96,7 +96,11 @@ UMetaSoundSource::UMetaSoundSource(const FObjectInitializer& ObjectInitializer)
 void UMetaSoundSource::PostEditUndo()
 {
 	Super::PostEditUndo();
-	SetSynchronizationRequired();
+
+	if (Graph)
+	{
+		Graph->RegisterGraphWithFrontend();
+	}
 }
 
 void UMetaSoundSource::PostDuplicate(EDuplicateMode::Type InDuplicateMode)
