Fix for Bhaskara and Pure Math variations of the Sine MetaSound Generator Node not correctly handling negative frequencies

diff --git a/Engine/Plugins/Runtime/Metasound/Source/MetasoundStandardNodes/Private/MetasoundOscillatorNodes.cpp b/Engine/Plugins/Runtime/Metasound/Source/MetasoundStandardNodes/Private/MetasoundOscillatorNodes.cpp
--- a/Engine/Plugins/Runtime/Metasound/Source/MetasoundStandardNodes/Private/MetasoundOscillatorNodes.cpp
+++ b/Engine/Plugins/Runtime/Metasound/Source/MetasoundStandardNodes/Private/MetasoundOscillatorNodes.cpp
@@ -50,6 +50,11 @@ namespace Metasound
 				{
 					InPhase -= 1.f;
 				}
+
+				while (InPhase <= -1.f)
+				{
+					InPhase += 1.f;
+				}
 			}
 		};
 
@@ -341,7 +346,9 @@ namespace Metasound
 			{
 				static const float TwoPi = 2.f * PI;
 				float PhaseRadians = InPhase * TwoPi;
-				return Audio::FastSin3(PhaseRadians - PI); // Expects [-PI, PI] 
+
+				float InBhaskaraDomain = (PhaseRadians < 0) ? PhaseRadians + PI : PhaseRadians - PI;
+				return Audio::FastSin3(InBhaskaraDomain); // Expects [-PI, PI] 
 			}
 		};
 		struct FSineWaveTableGenerator
