Fix for log spam on Sequencer audio track when using procedural SoundWaves.

diff --git a/Engine/Plugins/Runtime/Metasound/Source/MetasoundEngine/Private/MetasoundSource.cpp b/Engine/Plugins/Runtime/Metasound/Source/MetasoundEngine/Private/MetasoundSource.cpp
--- a/Engine/Plugins/Runtime/Metasound/Source/MetasoundEngine/Private/MetasoundSource.cpp
+++ b/Engine/Plugins/Runtime/Metasound/Source/MetasoundEngine/Private/MetasoundSource.cpp
@@ -490,63 +490,89 @@ bool UMetaSoundSource::GetAllDefaultParameters(TArray<FAudioParameter>& OutParam
 	using namespace Metasound::Frontend;
 	using namespace Metasound::Engine;
 
-	const FMetasoundFrontendDocument& Document = GetDocumentChecked();
-	for(const FMetasoundFrontendClassInput& Input : Document.RootGraph.Interface.Inputs)
-	{		
+	TArray<const FMetasoundFrontendClassInput*> TransmittableInputs = GetTransmittableClassInputs();
+	for(const FMetasoundFrontendClassInput* Input : TransmittableInputs)
+	{
+		check(Input);
+
 		FAudioParameter Params;
-		Params.ParamName = Input.Name;
-		switch (Input.DefaultLiteral.GetType())
+		Params.ParamName = Input->Name;
+
+		switch (Input->DefaultLiteral.GetType())
 		{
-		// bool
-		case EMetasoundFrontendLiteralType::Boolean:
-			Params.ParamType = EAudioParameterType::Boolean;
-			ensure(Input.DefaultLiteral.TryGet(Params.BoolParam));
+			case EMetasoundFrontendLiteralType::Boolean:
+			{
+				Params.ParamType = EAudioParameterType::Boolean;
+				ensure(Input->DefaultLiteral.TryGet(Params.BoolParam));
+			}
 			break;
-		case EMetasoundFrontendLiteralType::BooleanArray:
-			Params.ParamType = EAudioParameterType::BooleanArray;
-			ensure(Input.DefaultLiteral.TryGet(Params.ArrayBoolParam));
+
+			case EMetasoundFrontendLiteralType::BooleanArray:
+			{
+				Params.ParamType = EAudioParameterType::BooleanArray;
+				ensure(Input->DefaultLiteral.TryGet(Params.ArrayBoolParam));
+			}
 			break;
 
-		// int32
-		case EMetasoundFrontendLiteralType::Integer:
-			Params.ParamType = EAudioParameterType::Integer;
-			ensure(Input.DefaultLiteral.TryGet(Params.IntParam));
+			case EMetasoundFrontendLiteralType::Integer:
+			{
+				Params.ParamType = EAudioParameterType::Integer;
+				ensure(Input->DefaultLiteral.TryGet(Params.IntParam));
+			}
 			break;
-		case EMetasoundFrontendLiteralType::IntegerArray:
-			Params.ParamType = EAudioParameterType::IntegerArray;
-			ensure(Input.DefaultLiteral.TryGet(Params.ArrayIntParam));
+
+			case EMetasoundFrontendLiteralType::IntegerArray:
+			{
+				Params.ParamType = EAudioParameterType::IntegerArray;
+				ensure(Input->DefaultLiteral.TryGet(Params.ArrayIntParam));
+			}
 			break;
 
-		// float
-		case EMetasoundFrontendLiteralType::Float:
-			Params.ParamType = EAudioParameterType::Float;
-			ensure(Input.DefaultLiteral.TryGet(Params.FloatParam));
+			case EMetasoundFrontendLiteralType::Float:
+			{
+				Params.ParamType = EAudioParameterType::Float;
+				ensure(Input->DefaultLiteral.TryGet(Params.FloatParam));
+			}
 			break;
-		case EMetasoundFrontendLiteralType::FloatArray:
-			Params.ParamType = EAudioParameterType::FloatArray;
-			ensure(Input.DefaultLiteral.TryGet(Params.ArrayFloatParam));
+
+			case EMetasoundFrontendLiteralType::FloatArray:
+			{
+				Params.ParamType = EAudioParameterType::FloatArray;
+				ensure(Input->DefaultLiteral.TryGet(Params.ArrayFloatParam));
+			}
+			break;
+
+			case EMetasoundFrontendLiteralType::String:
+			{
+				Params.ParamType = EAudioParameterType::String;
+				ensure(Input->DefaultLiteral.TryGet(Params.StringParam));
+			}
 			break;
 
-		// FStrings.
-		case EMetasoundFrontendLiteralType::String:
-			Params.ParamType = EAudioParameterType::String;
-			ensure(Input.DefaultLiteral.TryGet(Params.StringParam));
+			case EMetasoundFrontendLiteralType::StringArray:
+			{
+				Params.ParamType = EAudioParameterType::StringArray;
+				ensure(Input->DefaultLiteral.TryGet(Params.ArrayStringParam));
+			}
 			break;
-		case EMetasoundFrontendLiteralType::StringArray:
-			Params.ParamType = EAudioParameterType::StringArray;
-			ensure(Input.DefaultLiteral.TryGet(Params.ArrayStringParam));
+
+			case EMetasoundFrontendLiteralType::UObject:
+			{
+				Params.ParamType = EAudioParameterType::Object;
+				UObject* Object = nullptr;
+				ensure(Input->DefaultLiteral.TryGet(Object));
+				Params.ObjectParam = Object;
+			}
 			break;
 
-		// UObjects.
-		//case EMetasoundFrontendLiteralType::UObject:
-		//	Params.ParamType = EAudioParameterType::Object;
-		//	ensure(Input.DefaultLiteral.TryGet(Params.ObjectParam));
-		//	break;
-		case EMetasoundFrontendLiteralType::UObjectArray:
-			Params.ParamType = EAudioParameterType::ObjectArray;
-			ensure(Input.DefaultLiteral.TryGet(Params.ArrayObjectParam));
+			case EMetasoundFrontendLiteralType::UObjectArray:
+			{
+				Params.ParamType = EAudioParameterType::ObjectArray;
+				ensure(Input->DefaultLiteral.TryGet(Params.ArrayObjectParam));
+			}
 			break;
-		default:
+
+			default:
 			break;
 		}
 
@@ -558,7 +584,6 @@ bool UMetaSoundSource::GetAllDefaultParameters(TArray<FAudioParameter>& OutParam
 	return true;
 }
 
-
 bool UMetaSoundSource::IsParameterValid(const FAudioParameter& InParameter) const
 {
 	TMap<FName, FName> InputNameTypeMap;
diff --git a/Engine/Plugins/Runtime/Metasound/Source/MetasoundFrontend/Public/MetasoundAssetBase.h b/Engine/Plugins/Runtime/Metasound/Source/MetasoundFrontend/Public/MetasoundAssetBase.h
--- a/Engine/Plugins/Runtime/Metasound/Source/MetasoundFrontend/Public/MetasoundAssetBase.h
+++ b/Engine/Plugins/Runtime/Metasound/Source/MetasoundFrontend/Public/MetasoundAssetBase.h
@@ -186,6 +186,9 @@ protected:
 	bool bSynchronizationUpdateDetails = false;
 #endif // WITH_EDITORONLY_DATA
 
+protected:
+	TArray<const FMetasoundFrontendClassInput*> GetTransmittableClassInputs() const;
+
 private:
 	Metasound::Frontend::FNodeRegistryKey RegistryKey;
 
@@ -210,7 +213,6 @@ private:
 	const FRuntimeData& GetRuntimeData() const;
 
 	TSharedPtr<Metasound::IGraph, ESPMode::ThreadSafe> BuildMetasoundDocument() const;
-	TArray<const FMetasoundFrontendClassInput*> GetTransmittableClassInputs() const;
 	Metasound::FSendAddress CreateSendAddress(uint64 InInstanceID, const Metasound::FVertexName& InVertexName, const FName& InDataTypeName) const;
 	Metasound::Frontend::FNodeHandle AddInputPinForSendAddress(const Metasound::FMetaSoundParameterTransmitter::FSendInfo& InSendInfo, Metasound::Frontend::FGraphHandle InGraph) const;
 };
diff --git a/Engine/Source/Editor/MovieSceneTools/Private/TrackEditors/AudioTrackEditor.cpp b/Engine/Source/Editor/MovieSceneTools/Private/TrackEditors/AudioTrackEditor.cpp
--- a/Engine/Source/Editor/MovieSceneTools/Private/TrackEditors/AudioTrackEditor.cpp
+++ b/Engine/Source/Editor/MovieSceneTools/Private/TrackEditors/AudioTrackEditor.cpp
@@ -311,6 +311,12 @@ void FAudioThumbnail::GenerateWaveformPreview(TArray<uint8>& OutData, TRange<flo
 	
 	check(SoundWave->NumChannels == 1 || SoundWave->NumChannels == 2);
 
+	// If this SoundWave is generated procedurally, it's not possible to render a thumbnail.
+	if (SoundWave->bProcedural)
+	{
+		return;
+	}
+
 	uint32 SampleRate;
 	uint16 NumChannels;
 	TArray<uint8> RawPCMData;
diff --git a/Engine/Source/Runtime/AudioMixer/Private/AudioMixerSubmix.cpp b/Engine/Source/Runtime/AudioMixer/Private/AudioMixerSubmix.cpp
--- a/Engine/Source/Runtime/AudioMixer/Private/AudioMixerSubmix.cpp
+++ b/Engine/Source/Runtime/AudioMixer/Private/AudioMixerSubmix.cpp
@@ -1161,10 +1161,7 @@ namespace Audio
 		if (bIsCurrentlyDisabled)
 		{
 			bIsCurrentlyDisabled = false;
-			if (LogSubmixEnablementCVar == 1)
-			{
-				UE_LOG(LogAudioMixer, Log, TEXT("Submix Enabled: %s"), *SubmixName);
-			}
+			UE_CLOG(LogSubmixEnablementCVar == 1, LogAudioMixer, Verbose, TEXT("Submix Enabled: %s"), *SubmixName);
 		}
 
 		// If this is a Soundfield Submix, process our soundfield and decode it to a OutAudioBuffer.
