(Re)Fix for renaming a MetaSound not propagating the rename to referencing graphs
#rnx

diff --git a/Engine/Plugins/Runtime/Metasound/Source/MetasoundEditor/Private/MetasoundEditorGraphBuilder.cpp b/Engine/Plugins/Runtime/Metasound/Source/MetasoundEditor/Private/MetasoundEditorGraphBuilder.cpp
--- a/Engine/Plugins/Runtime/Metasound/Source/MetasoundEditor/Private/MetasoundEditorGraphBuilder.cpp
+++ b/Engine/Plugins/Runtime/Metasound/Source/MetasoundEditor/Private/MetasoundEditorGraphBuilder.cpp
@@ -404,10 +404,15 @@ namespace Metasound
 				const bool bInterfaceChange = Result.Node->ContainsInterfaceChange();
 				const bool bMetadataChange = Result.Node->ContainsMetadataChange();
 				const bool bStyleChange = Result.Node->ContainsStyleChange();
-				if (Result.bIsDirty || bMetadataChange || bInterfaceChange || bStyleChange || bForceRefreshNodes)
+
+				const FText Title = Result.Node->GetCachedTitle();
+				Result.Node->CacheTitle();
+				const bool bTitleUpdated = !Title.IdenticalTo(Result.Node->GetCachedTitle());
+
+
+				if (Result.bIsDirty || bTitleUpdated || bMetadataChange || bInterfaceChange || bStyleChange || bForceRefreshNodes)
 				{
 					Result.Node->SyncChangeIDs();
-					Result.Node->CacheTitle();
 
 					if (GraphEditor.IsValid())
 					{
diff --git a/Engine/Plugins/Runtime/Metasound/Source/MetasoundEditor/Private/MetasoundEditorModule.cpp b/Engine/Plugins/Runtime/Metasound/Source/MetasoundEditor/Private/MetasoundEditorModule.cpp
--- a/Engine/Plugins/Runtime/Metasound/Source/MetasoundEditor/Private/MetasoundEditorModule.cpp
+++ b/Engine/Plugins/Runtime/Metasound/Source/MetasoundEditor/Private/MetasoundEditorModule.cpp
@@ -356,22 +356,10 @@ namespace Metasound
 					// MetaSoundAssetSubsystem so as to properly refresh respective open editors.
 					constexpr bool bReregisterWithFrontend = false;
 					AssetSubsystem->RenameAsset(InAssetData, bReregisterWithFrontend);
-					UObject* AssetObject = InAssetData.GetAsset();
-					FGraphBuilder::RegisterGraphWithFrontend(*AssetObject);
 
-					// Request sync of other open MetaSounds to ensure rename is reflected in other asset editors with
-					// nodes referencing the renamed asset.
-					TArray<UObject*> EditedAssets = GEditor->GetEditorSubsystem<UAssetEditorSubsystem>()->GetAllEditedAssets();
-					for (UObject* EditedAsset : EditedAssets)
-					{
-						if (AssetObject != EditedAsset)
-						{
-							if (FMetasoundAssetBase* EditedMetaSoundAsset = IMetasoundUObjectRegistry::Get().GetObjectAsAssetBase(EditedAsset))
-							{
-								EditedMetaSoundAsset->SetSynchronizationRequired();
-							}
-						}
-					}
+					constexpr bool bForceViewSynchronization = true;
+					UObject* AssetObject = InAssetData.GetAsset();
+					FGraphBuilder::RegisterGraphWithFrontend(*AssetObject, bForceViewSynchronization);
 				}
 			}
 
diff --git a/Engine/Plugins/Runtime/Metasound/Source/MetasoundEditor/Public/MetasoundEditorGraphNode.h b/Engine/Plugins/Runtime/Metasound/Source/MetasoundEditor/Public/MetasoundEditorGraphNode.h
--- a/Engine/Plugins/Runtime/Metasound/Source/MetasoundEditor/Public/MetasoundEditorGraphNode.h
+++ b/Engine/Plugins/Runtime/Metasound/Source/MetasoundEditor/Public/MetasoundEditorGraphNode.h
@@ -103,6 +103,7 @@ public:
 	// Mark node for refresh
 	void SyncChangeIDs();
 
+	FText GetCachedTitle() const { return CachedTitle; }
 
 	// Returns whether or not Metadata has been changed since the last node refresh
 	bool ContainsMetadataChange() const;
