MetaSound Member localization options menu support for DisplayName & Tooltip

diff --git a/Engine/Plugins/Runtime/Metasound/Source/MetasoundEditor/Private/MetasoundNodeDetailCustomization.h b/Engine/Plugins/Runtime/Metasound/Source/MetasoundEditor/Private/MetasoundNodeDetailCustomization.h
--- a/Engine/Plugins/Runtime/Metasound/Source/MetasoundEditor/Private/MetasoundNodeDetailCustomization.h
+++ b/Engine/Plugins/Runtime/Metasound/Source/MetasoundEditor/Private/MetasoundNodeDetailCustomization.h
@@ -23,6 +23,7 @@
 #include "PropertyRestriction.h"
 #include "ScopedTransaction.h"
 #include "SSearchableComboBox.h"
+#include "STextPropertyEditableTextBox.h"
 #include "Types/SlateEnums.h"
 #include "UObject/NameTypes.h"
 #include "Widgets/Input/SCheckBox.h"
@@ -37,6 +38,124 @@ namespace Metasound
 {
 	namespace Editor
 	{
+		class FGraphMemberEditableTextBase : public IEditableTextProperty
+		{
+		protected:
+			TWeakObjectPtr<UMetasoundEditorGraphMember> GraphMember;
+			FText ToolTip;
+
+		public:
+			FGraphMemberEditableTextBase(TWeakObjectPtr<UMetasoundEditorGraphMember> InGraphMember, const FText& InToolTip)
+				: GraphMember(InGraphMember)
+				, ToolTip(InToolTip)
+			{
+			}
+
+			virtual ~FGraphMemberEditableTextBase() = default;
+
+			virtual bool IsMultiLineText() const override { return true; }
+			virtual bool IsPassword() const override { return false; }
+			virtual bool IsReadOnly() const override { return false; }
+			virtual int32 GetNumTexts() const override { return 1; }
+			virtual bool IsValidText(const FText& InText, FText& OutErrorMsg) const override { return true; }
+			virtual void RequestRefresh() override { }
+
+			virtual FText GetToolTipText() const override
+			{
+				return ToolTip;
+			}
+
+			virtual bool IsDefaultValue() const override
+			{
+				return GetText(0).EqualTo(FText::GetEmpty());
+			}
+
+#if USE_STABLE_LOCALIZATION_KEYS
+			virtual void GetStableTextId(const int32 InIndex, const ETextPropertyEditAction InEditAction, const FString& InTextSource, const FString& InProposedNamespace, const FString& InProposedKey, FString& OutStableNamespace, FString& OutStableKey) const override
+			{
+				check(InIndex == 0);
+				StaticStableTextId(GraphMember->GetPackage(), InEditAction, InTextSource, InProposedNamespace, InProposedKey, OutStableNamespace, OutStableKey);
+			}
+#endif // USE_STABLE_LOCALIZATION_KEYS
+
+		};
+
+		class FGraphMemberEditableTextDescription : public FGraphMemberEditableTextBase
+		{
+		public:
+			FGraphMemberEditableTextDescription(TWeakObjectPtr<UMetasoundEditorGraphMember> InGraphMember, const FText& InToolTip)
+				: FGraphMemberEditableTextBase(InGraphMember, InToolTip)
+			{
+			}
+
+			virtual ~FGraphMemberEditableTextDescription() = default;
+
+			virtual FText GetText(const int32 InIndex) const override
+			{
+				check(InIndex == 0);
+
+				if (GraphMember.IsValid())
+				{
+					return GraphMember->GetDescription();
+				}
+
+				return FText::GetEmpty();
+			}
+
+			virtual void SetText(const int32 InIndex, const FText& InText) override
+			{
+				check(InIndex == 0);
+
+				if (GraphMember.IsValid())
+				{
+					GraphMember->SetDescription(InText, true);
+				}
+			}
+		};
+
+		class FGraphMemberEditableTextDisplayName : public FGraphMemberEditableTextBase
+		{
+		public:
+			FGraphMemberEditableTextDisplayName(TWeakObjectPtr<UMetasoundEditorGraphMember> InGraphMember, const FText& InToolTip)
+				: FGraphMemberEditableTextBase(InGraphMember, InToolTip)
+			{
+			}
+
+			virtual ~FGraphMemberEditableTextDisplayName() = default;
+
+			virtual FText GetText(const int32 InIndex) const override
+			{
+				check(InIndex == 0);
+
+				if (GraphMember.IsValid())
+				{
+					if (const UMetasoundEditorGraphVertex* Vertex = Cast<UMetasoundEditorGraphVertex>(GraphMember.Get()))
+					{
+						return Vertex->GetConstNodeHandle()->GetDisplayName();
+					}
+
+					if (const UMetasoundEditorGraphVariable* Variable = Cast<UMetasoundEditorGraphVariable>(GraphMember.Get()))
+					{
+						return Variable->GetConstVariableHandle()->GetDisplayName();
+					}
+
+					return GraphMember->GetDisplayName();
+				}
+
+				return FText::GetEmpty();
+			}
+
+			virtual void SetText(const int32 InIndex, const FText& InText) override
+			{
+				check(InIndex == 0);
+
+				if (GraphMember.IsValid())
+				{
+					GraphMember->SetDisplayName(InText, true);
+				}
+			}
+		};
+
 		// TODO: Move to actual style
 		namespace MemberCustomizationStyle
 		{
@@ -47,7 +166,6 @@ namespace Metasound
 
 			static const FText DataTypeNameText = LOCTEXT("Node_DataTypeName", "Type");
 			static const FText DefaultPropertyText = LOCTEXT("Node_DefaultPropertyName", "Default Value");
-			static const FText NodeTooltipText = LOCTEXT("Node_Tooltip", "Tooltip");
 		} // namespace MemberCustomizationStyle
 
 		class FMetasoundFloatLiteralCustomization : public FMetasoundDefaultLiteralCustomizationBase
@@ -173,7 +291,7 @@ namespace Metasound
 			FName ArrayTypeName;
 		};
 
-		template <typename TMemberType, typename TChildClass>
+		template <typename TMemberType>
 		class TMetasoundGraphMemberDetailCustomization : public IDetailCustomization
 		{
 		public:
@@ -191,39 +309,10 @@ namespace Metasound
 		protected:
 			TWeakObjectPtr<TMemberType> GraphMember;
 			TSharedPtr<SEditableTextBox> NameEditableTextBox;
-// 			TSharedPtr<SEditableTextBox> DisplayNameEditableTextBox;
 			TSharedPtr<FMetasoundDataTypeSelector> DataTypeSelector;
 
 			bool bIsNameInvalid = false;
 
-			static const FText& GetMemberDisplayNameText()
-			{
-				static const FText MemberDisplayNameTextFormat = LOCTEXT("GraphMember_NameFormat", "{0} Display Name");
-				static const FText DisplayName = FText::Format(MemberDisplayNameTextFormat, TChildClass::MemberNameText);
-				return DisplayName;
-			}
-
-			static const FText& GetMemberDisplayNameDescriptionText()
-			{
-				static const FText MemberDisplayNameDescriptionFormat = LOCTEXT("GraphMember_DisplayNameDescriptionFormat", "Optional, localized name used within the MetaSounds editor(s) to describe the given {0}.");
-				static const FText DisplayName = FText::Format(MemberDisplayNameDescriptionFormat, TChildClass::MemberNameText);
-				return DisplayName;
-			}
-
-			static const FText& GetMemberNameText()
-			{
-				static const FText MemberNameTextFormat = LOCTEXT("GraphMember_DisplayNameFormat", "{0} Name");
-				static const FText Name = FText::Format(MemberNameTextFormat, TChildClass::MemberNameText);
-				return Name;
-			}
-
-			static const FText& GetMemberNameDescriptionText()
-			{
-				static const FText MemberNameDescriptionFormat = LOCTEXT("GraphMember_DisplayDescriptionFormat", "Name used by external systems/referencing MetaSounds to identify {0}. Used as DisplayName within MetaSound Graph Editor if no DisplayName is provided.");
-				static const FText Name = FText::Format(MemberNameDescriptionFormat, TChildClass::MemberNameText);
-				return Name;
-			}
-
 			static IDetailCategoryBuilder& GetDefaultCategoryBuilder(IDetailLayoutBuilder& InDetailLayout)
 			{
 				return InDetailLayout.EditCategory("DefaultValue");
@@ -279,66 +368,66 @@ namespace Metasound
 				const bool bIsGraphEditable = IsGraphEditable();
 
 				NameEditableTextBox = SNew(SEditableTextBox)
-					.Text(this, &TChildClass::GetName)
-					.OnTextChanged(this, &TChildClass::OnNameChanged)
-					.OnTextCommitted(this, &TChildClass::OnNameCommitted)
+					.Text(this, &TMetasoundGraphMemberDetailCustomization::GetName)
+					.OnTextChanged(this, &TMetasoundGraphMemberDetailCustomization::OnNameChanged)
+					.OnTextCommitted(this, &TMetasoundGraphMemberDetailCustomization::OnNameCommitted)
 					.IsReadOnly(bIsInterfaceMember || !bIsGraphEditable)
 					.SelectAllTextWhenFocused(true)
 					.Font(IDetailLayoutBuilder::GetDetailFont());
 
-				// 				DisplayNameEditableTextBox = SNew(SEditableTextBox)
-				// 					.Text(this, &TChildClass::GetDisplayName)
-				// 					.OnTextCommitted(this, &TChildClass::OnDisplayNameCommitted)
-				// 					.IsReadOnly(bIsInterfaceMember || !bIsGraphEditable)
-				// 					.Font(IDetailLayoutBuilder::GetDetailFont());
-
-				CategoryBuilder.AddCustomRow(TChildClass::MemberNameText)
+				static const FText MemberNameToolTipFormat = LOCTEXT("GraphMember_NameDescriptionFormat", "Name used within the MetaSounds editor(s) and transacting systems (ex. Blueprints) if applicable to reference the given {0}.");
+				CategoryBuilder.AddCustomRow(LOCTEXT("GraphMember_NameProperty", "Name"))
 					.EditCondition(!bIsInterfaceMember && bIsGraphEditable, nullptr)
 					.NameContent()
 					[
 						SNew(STextBlock)
 						.Font(IDetailLayoutBuilder::GetDetailFontBold())
-						.Text(TChildClass::MemberNameText)
-						.ToolTipText(GetMemberNameDescriptionText())
+						.Text(GraphMember->GetGraphMemberLabel())
+						.ToolTipText(FText::Format(MemberNameToolTipFormat, GraphMember->GetGraphMemberLabel()))
 					]
 					.ValueContent()
 					[
 						NameEditableTextBox.ToSharedRef()
 					];
 
-				// TODO: Enable and use proper FText property editor
-// 				CategoryBuilder.AddCustomRow(GetMemberDisplayNameText())
-// 				.EditCondition(!bIsInterfaceMember && bIsGraphEditable, nullptr)
-// 				.NameContent()
-// 				[
-// 					SNew(STextBlock)
-// 					.Font(IDetailLayoutBuilder::GetDetailFontBold())
-// 					.Text(GetMemberDisplayNameText())
-// 					.ToolTipText(GetMemberDisplayNameDescriptionText())
-// 				]
-// 				.ValueContent()
-// 				[
-// 					DisplayNameEditableTextBox.ToSharedRef()
-// 				];
-
-				CategoryBuilder.AddCustomRow(MemberCustomizationStyle::NodeTooltipText)
+				static const FText MemberDisplayNameText = LOCTEXT("GraphMember_DisplayNameProperty", "Display Name");
+				static const FText MemberDisplayNameToolTipFormat = LOCTEXT("GraphMember_DisplayNameDescriptionFormat", "Optional, localized name used within the MetaSounds editor(s) to describe the given {0}.");
+				const FText MemberDisplayNameTooltipText = FText::Format(MemberDisplayNameToolTipFormat, GraphMember->GetGraphMemberLabel());
+				CategoryBuilder.AddCustomRow(MemberDisplayNameText)
+					.EditCondition(!bIsInterfaceMember && bIsGraphEditable, nullptr)
+					.NameContent()
+					[
+						SNew(STextBlock)
+						.Font(IDetailLayoutBuilder::GetDetailFontBold())
+						.Text(MemberDisplayNameText)
+						.ToolTipText(MemberDisplayNameTooltipText)
+					]
+					.ValueContent()
+					[
+						SNew(STextPropertyEditableTextBox, MakeShared<FGraphMemberEditableTextDisplayName>(GraphMember, MemberDisplayNameTooltipText))
+						.WrapTextAt(500)
+						.MinDesiredWidth(25.0f)
+						.MaxDesiredHeight(200)
+					];
+
+				static const FText MemberDescriptionText = LOCTEXT("Member_DescriptionPropertyName", "Description");
+				static const FText MemberDescriptionToolTipFormat = LOCTEXT("Member_DescriptionToolTipFormat", "Description for {0}. For example, used as a tooltip when displayed on another graph's referencing node.");
+				const FText MemberDescriptionToolTipText = FText::Format(MemberDescriptionToolTipFormat, GraphMember->GetGraphMemberLabel());
+				CategoryBuilder.AddCustomRow(MemberDescriptionText)
 					.EditCondition(!bIsInterfaceMember && bIsGraphEditable, nullptr)
 					.NameContent()
 					[
 						SNew(STextBlock)
 						.Font(IDetailLayoutBuilder::GetDetailFontBold())
-					.Text(MemberCustomizationStyle::NodeTooltipText)
+						.Text(MemberDescriptionText)
+						.ToolTipText(MemberDescriptionToolTipText)
 					]
 					.ValueContent()
 					[
-						SNew(SMultiLineEditableTextBox)
-						.Text(this, &TChildClass::GetTooltip)
-						.OnTextCommitted(this, &TChildClass::OnTooltipCommitted)
-						.IsReadOnly(bIsInterfaceMember || !bIsGraphEditable)
-						.ModiferKeyForNewLine(EModifierKey::Shift)
-						.RevertTextOnEscape(true)
-						.WrapTextAt(MemberCustomizationStyle::DetailsTitleMaxWidth - MemberCustomizationStyle::DetailsTitleWrapPadding)
-						.Font(IDetailLayoutBuilder::GetDetailFont())
+						SNew(STextPropertyEditableTextBox, MakeShared<FGraphMemberEditableTextDescription>(GraphMember, MemberDescriptionToolTipText))
+						.WrapTextAt(500)
+						.MinDesiredWidth(25.0f)
+						.MaxDesiredHeight(200)
 					];
 
 				DataTypeSelector->AddDataTypeSelector(InDetailLayout, MemberCustomizationStyle::DataTypeNameText, GraphMember, !bIsInterfaceMember && bIsGraphEditable);
@@ -403,7 +492,6 @@ namespace Metasound
 
 				bIsNameInvalid = false;
 				NameEditableTextBox->SetError(FText::GetEmpty());
-// 				DisplayNameEditableTextBox->SetError(FText::GetEmpty());
 
 				if (!ensure(GraphMember.IsValid()))
 				{
@@ -415,7 +503,6 @@ namespace Metasound
 				{
 					bIsNameInvalid = true;
 					NameEditableTextBox->SetError(Error);
-// 					DisplayNameEditableTextBox->SetError(Error);
 				}
 			}
 
@@ -454,20 +541,6 @@ namespace Metasound
 				return FText::GetEmpty();
 			}
 
-// 			void OnDisplayNameCommitted(const FText& InNewName, ETextCommit::Type InTextCommit)
-// 			{
-// 				using namespace Frontend;
-// 
-// 				if (!bIsNameInvalid && GraphMember.IsValid())
-// 				{
-// 					constexpr bool bPostTransaction = true;
-// 					GraphMember->SetDisplayName(FText::GetEmpty(), bPostTransaction);
-// 				}
-// 
-// 				DisplayNameEditableTextBox->SetError(FText::GetEmpty());
-// 				bIsNameInvalid = false;
-// 			}
-
 			void OnTooltipCommitted(const FText& InNewText, ETextCommit::Type InTextCommit)
 			{
 				using namespace Frontend;
@@ -495,7 +568,7 @@ namespace Metasound
 
 				if (!bIsNameInvalid && GraphMember.IsValid())
 				{
-					const FText TransactionLabel = FText::Format(LOCTEXT("RenameGraphMember_Format", "Set MetaSound {0}'s Name"), TChildClass::MemberNameText);
+					const FText TransactionLabel = FText::Format(LOCTEXT("RenameGraphMember_Format", "Set MetaSound {0}'s Name"), GraphMember->GetGraphMemberLabel());
 					const FScopedTransaction Transaction(TransactionLabel);
 
 					constexpr bool bPostTransaction = false;
@@ -523,11 +596,11 @@ namespace Metasound
 			FDelegateHandle RenameRequestedHandle;
 		};
 
-		class FMetasoundInputDetailCustomization : public TMetasoundGraphMemberDetailCustomization<UMetasoundEditorGraphInput, FMetasoundInputDetailCustomization>
+		class FMetasoundInputDetailCustomization : public TMetasoundGraphMemberDetailCustomization<UMetasoundEditorGraphInput>
 		{
 		public:
 			FMetasoundInputDetailCustomization()
-				: TMetasoundGraphMemberDetailCustomization<UMetasoundEditorGraphInput, FMetasoundInputDetailCustomization>()
+				: TMetasoundGraphMemberDetailCustomization<UMetasoundEditorGraphInput>()
 			{
 			}
 
@@ -545,13 +618,13 @@ namespace Metasound
 			void ClearInputInheritsDefault();
 		};
 
-		class FMetasoundOutputDetailCustomization : public TMetasoundGraphMemberDetailCustomization<UMetasoundEditorGraphOutput, FMetasoundOutputDetailCustomization>
+		class FMetasoundOutputDetailCustomization : public TMetasoundGraphMemberDetailCustomization<UMetasoundEditorGraphOutput>
 		{
 		public:
 			static const FText MemberNameText;
 
 			FMetasoundOutputDetailCustomization()
-				: TMetasoundGraphMemberDetailCustomization<UMetasoundEditorGraphOutput, FMetasoundOutputDetailCustomization>()
+				: TMetasoundGraphMemberDetailCustomization<UMetasoundEditorGraphOutput>()
 			{
 			}
 
diff --git a/Engine/Plugins/Runtime/Metasound/Source/MetasoundEditor/Private/MetasoundVariableDetailCustomization.h b/Engine/Plugins/Runtime/Metasound/Source/MetasoundEditor/Private/MetasoundVariableDetailCustomization.h
--- a/Engine/Plugins/Runtime/Metasound/Source/MetasoundEditor/Private/MetasoundVariableDetailCustomization.h
+++ b/Engine/Plugins/Runtime/Metasound/Source/MetasoundEditor/Private/MetasoundVariableDetailCustomization.h
@@ -16,13 +16,13 @@ namespace Metasound
 {
 	namespace Editor
 	{
-		class FMetasoundVariableDetailCustomization : public TMetasoundGraphMemberDetailCustomization<UMetasoundEditorGraphVariable, FMetasoundVariableDetailCustomization>
+		class FMetasoundVariableDetailCustomization : public TMetasoundGraphMemberDetailCustomization<UMetasoundEditorGraphVariable>
 		{
 		public:
 			static const FText MemberNameText;
 
 			FMetasoundVariableDetailCustomization()
-				: TMetasoundGraphMemberDetailCustomization<UMetasoundEditorGraphVariable, FMetasoundVariableDetailCustomization>()
+				: TMetasoundGraphMemberDetailCustomization<UMetasoundEditorGraphVariable>()
 			{
 			}
 
