Add missing 'BreakSinglePinLink' implementation for MetaSounds

diff --git a/Engine/Plugins/Runtime/Metasound/Source/MetasoundEditor/Private/MetasoundEditorGraphSchema.cpp b/Engine/Plugins/Runtime/Metasound/Source/MetasoundEditor/Private/MetasoundEditorGraphSchema.cpp
--- a/Engine/Plugins/Runtime/Metasound/Source/MetasoundEditor/Private/MetasoundEditorGraphSchema.cpp
+++ b/Engine/Plugins/Runtime/Metasound/Source/MetasoundEditor/Private/MetasoundEditorGraphSchema.cpp
@@ -1392,6 +1392,46 @@ void UMetasoundEditorGraphSchema::BreakPinLinks(UEdGraphPin& TargetPin, bool bSe
 	Super::BreakPinLinks(TargetPin, bSendsNodeNotifcation);
 }
 
+void UMetasoundEditorGraphSchema::BreakSinglePinLink(UEdGraphPin* SourcePin, UEdGraphPin* TargetPin) const
+{
+	using namespace Metasound::Editor;
+
+	UEdGraphPin* InputPin = nullptr;
+	UEdGraphPin* OutputPin = nullptr;
+	if (!SourcePin || !TargetPin || !SourcePin->LinkedTo.Contains(TargetPin) || !TargetPin->LinkedTo.Contains(SourcePin))
+	{
+		return;
+	}
+
+	if (SourcePin->Direction == EGPD_Input)
+	{
+		InputPin = SourcePin;
+	}
+	else if (TargetPin->Direction == EGPD_Input)
+	{
+		InputPin = TargetPin;
+	}
+	else
+	{
+		return;
+	}
+
+	UEdGraphNode* OwningNode = InputPin->GetOwningNode();
+	if (!OwningNode)
+	{
+		return;
+	}
+
+	const FScopedTransaction Transaction(LOCTEXT("BreakSinglePinLink", "Break Single Pin Link"));
+	UMetasoundEditorGraph* Graph = CastChecked<UMetasoundEditorGraph>(OwningNode->GetGraph());
+	Graph->GetMetasoundChecked().Modify();
+	SourcePin->Modify();
+	TargetPin->Modify();
+
+	FGraphBuilder::DisconnectPinVertex(*InputPin);
+	Super::BreakSinglePinLink(SourcePin, TargetPin);
+}
+
 void UMetasoundEditorGraphSchema::GetAssetsGraphHoverMessage(const TArray<FAssetData>& Assets, const UEdGraph* HoverGraph, FString& OutTooltipText, bool& OutOkIcon) const
 {
 	using namespace Metasound;
diff --git a/Engine/Plugins/Runtime/Metasound/Source/MetasoundEditor/Public/MetasoundEditorGraphSchema.h b/Engine/Plugins/Runtime/Metasound/Source/MetasoundEditor/Public/MetasoundEditorGraphSchema.h
--- a/Engine/Plugins/Runtime/Metasound/Source/MetasoundEditor/Public/MetasoundEditorGraphSchema.h
+++ b/Engine/Plugins/Runtime/Metasound/Source/MetasoundEditor/Public/MetasoundEditorGraphSchema.h
@@ -420,6 +420,7 @@ class UMetasoundEditorGraphSchema : public UEdGraphSchema
 	virtual FLinearColor GetPinTypeColor(const FEdGraphPinType& PinType) const override;
 	virtual void BreakNodeLinks(UEdGraphNode& TargetNode) const override;
 	virtual void BreakPinLinks(UEdGraphPin& TargetPin, bool bSendsNodeNotifcation) const override;
+	virtual void BreakSinglePinLink(UEdGraphPin* SourcePin, UEdGraphPin* TargetPin) const override;
 	virtual void GetAssetsGraphHoverMessage(const TArray<FAssetData>& Assets, const UEdGraph* HoverGraph, FString& OutTooltipText, bool& OutOkIcon) const override;
 	virtual void DroppedAssetsOnGraph(const TArray<struct FAssetData>& Assets, const FVector2D& GraphPosition, UEdGraph* Graph) const override;
 	virtual void DroppedAssetsOnNode(const TArray<FAssetData>& Assets, const FVector2D& GraphPosition, UEdGraphNode* Node) const override;
