Fix for deleting a pasted variable node deleting all instances of that node on the graph
- Also fixes pasting a variable member from another graph or pasting a variable accessor/delayed accessor from ensuring if no other accessor exists
#rnx

diff --git a/Engine/Plugins/Runtime/Metasound/Source/MetasoundEditor/Private/MetasoundEditor.cpp b/Engine/Plugins/Runtime/Metasound/Source/MetasoundEditor/Private/MetasoundEditor.cpp
--- a/Engine/Plugins/Runtime/Metasound/Source/MetasoundEditor/Private/MetasoundEditor.cpp
+++ b/Engine/Plugins/Runtime/Metasound/Source/MetasoundEditor/Private/MetasoundEditor.cpp
@@ -1808,11 +1808,31 @@ namespace Metasound
 					// Can only have one setter node
 					if (const UMetasoundEditorGraphVariable* Variable = VariableNode->Variable)
 					{
-						FConstVariableHandle VariableHandle = Variable->GetConstVariableHandle();
-						FConstNodeHandle VariableMutatorNodeHandle = VariableHandle->FindMutatorNode();
-						if (VariableNode->GetNodeID() == VariableMutatorNodeHandle->GetID())
+						if (Graph.ContainsVariable(*Variable))
+						{
+							FConstVariableHandle VariableHandle = Variable->GetConstVariableHandle();
+							if (VariableHandle->IsValid())
+							{
+								FConstNodeHandle VariableMutatorNodeHandle = VariableHandle->FindMutatorNode();
+								if (VariableNode->GetNodeID() == VariableMutatorNodeHandle->GetID())
+								{
+									bNotifyMultipleVariableSetters = true;
+									NodesToRemove.Add(GraphNode);
+								}
+								else
+								{
+									// Fix-up if variable getter node does not exist but variable does
+									FConstNodeHandle NodeHandle = VariableNode->GetConstNodeHandle();
+									if (!NodeHandle->IsValid())
+									{
+										const FNodeClassName NodeClassName = VariableNode->GetClassName().ToNodeClassName();
+										NodeHandle = FGraphBuilder::AddVariableNodeHandle(*Metasound, Variable->GetVariableID(), NodeClassName, VariableNode);
+									}
+								}
+							}
+						}
+						else
 						{
-							bNotifyMultipleVariableSetters = true;
 							NodesToRemove.Add(GraphNode);
 						}
 					}
diff --git a/Engine/Plugins/Runtime/Metasound/Source/MetasoundEditor/Private/MetasoundEditorGraphBuilder.cpp b/Engine/Plugins/Runtime/Metasound/Source/MetasoundEditor/Private/MetasoundEditorGraphBuilder.cpp
--- a/Engine/Plugins/Runtime/Metasound/Source/MetasoundEditor/Private/MetasoundEditorGraphBuilder.cpp
+++ b/Engine/Plugins/Runtime/Metasound/Source/MetasoundEditor/Private/MetasoundEditorGraphBuilder.cpp
@@ -919,7 +919,7 @@ namespace Metasound
 			return Variable;
 		}
 
-		Frontend::FNodeHandle FGraphBuilder::AddVariableNodeHandle(UObject& InMetaSound, const FGuid& InVariableID, const Metasound::FNodeClassName& InVariableNodeClassName)
+		Frontend::FNodeHandle FGraphBuilder::AddVariableNodeHandle(UObject& InMetaSound, const FGuid& InVariableID, const Metasound::FNodeClassName& InVariableNodeClassName, UMetasoundEditorGraphVariableNode* InVariableNode)
 		{
 			using namespace Frontend;
 
@@ -968,6 +968,13 @@ namespace Metasound
 				}
 			}
 
+			if (InVariableNode)
+			{
+				InVariableNode->ClassName = FrontendNode->GetClassMetadata().GetClassName();
+				InVariableNode->ClassType = FrontendNode->GetClassMetadata().GetType();
+				InVariableNode->SetNodeID(FrontendNode->GetID());
+			}
+
 			return FrontendNode;
 		}
 
@@ -1212,8 +1219,6 @@ namespace Metasound
 				return false;
 			}
 
-			bool bWasErroredNode = InNode.ErrorType == EMessageSeverity::Error;
-
 			// If node isn't a MetasoundEditorGraphNode, just remove and return (ex. comment nodes)
 			UMetasoundEditorGraphNode* Node = Cast<UMetasoundEditorGraphNode>(&InNode);
 			UMetasoundEditorGraph* Graph = CastChecked<UMetasoundEditorGraph>(InNode.GetGraph());
@@ -1237,37 +1242,62 @@ namespace Metasound
 
 			FNodeHandle NodeHandle = Node->GetNodeHandle();
 			Frontend::FGraphHandle GraphHandle = NodeHandle->GetOwningGraph();
+
+			auto RemoveNodeLocation = [](FNodeHandle InNodeHandle, const FGuid& InNodeGuid)
+			{
+				FMetasoundFrontendNodeStyle Style = InNodeHandle->GetNodeStyle();
+				Style.Display.Locations.Remove(InNodeGuid);
+				InNodeHandle->SetNodeStyle(Style);
+			};	
+
+			auto RemoveNodeHandle = [] (FGraphHandle InGraphHandle, FNodeHandle InNodeHandle)
+			{
+				if (ensure(InGraphHandle->RemoveNode(*InNodeHandle)))
+				{
+					InGraphHandle->GetOwningDocument()->RemoveUnreferencedDependencies();
+				}
+			};
+
 			if (GraphHandle->IsValid())
 			{
-				switch (NodeHandle->GetClassMetadata().GetType())
+				const EMetasoundFrontendClassType ClassType = NodeHandle->GetClassMetadata().GetType();
+				switch (ClassType)
 				{
+					// NodeHandle does not get removed in these cases as EdGraph Inputs/Outputs
+					// Frontend node is represented by the editor graph as a respective member
+					// (not a node) on the MetasoundGraph. Therefore, just the editor position
+					// data is removed.
 					case EMetasoundFrontendClassType::Output:
 					case EMetasoundFrontendClassType::Input:
 					{
-						// NodeHandle does not get removed in these cases as EdGraph Inputs/Outputs
-						// merely reference their respective types set on the MetasoundGraph. It must
-						// be removed from the location display data however for graph sync reasons.
-						FMetasoundFrontendNodeStyle Style = NodeHandle->GetNodeStyle();
-						Style.Display.Locations.Remove(InNode.NodeGuid);
-						NodeHandle->SetNodeStyle(Style);
+						RemoveNodeLocation(NodeHandle, InNode.NodeGuid);
+					}
+					break;
+					
+					// NodeHandle is only removed for variable accessors if the editor graph
+					// no longer contains nodes representing the given accessor on the MetasoundGraph.
+					// Therefore, just the editor position data is removed unless no location remains
+					// on the Frontend node.
+					case EMetasoundFrontendClassType::VariableAccessor:
+					case EMetasoundFrontendClassType::VariableDeferredAccessor:
+					{
+						RemoveNodeLocation(NodeHandle, InNode.NodeGuid);
+						if (NodeHandle->GetNodeStyle().Display.Locations.IsEmpty())
+						{
+							RemoveNodeHandle(GraphHandle, NodeHandle);
+						}
 					}
 					break;
 
 					case EMetasoundFrontendClassType::Graph:
 					case EMetasoundFrontendClassType::Literal:
-					case EMetasoundFrontendClassType::VariableAccessor:
-					case EMetasoundFrontendClassType::VariableDeferredAccessor:
 					case EMetasoundFrontendClassType::VariableMutator:
 					case EMetasoundFrontendClassType::Variable:
 					case EMetasoundFrontendClassType::External:
 					default:
 					{
 						static_assert(static_cast<int32>(EMetasoundFrontendClassType::Invalid) == 9, "Possible missing MetasoundFrontendClassType switch case coverage.");
-
-						if (ensure(GraphHandle->RemoveNode(*NodeHandle)))
-						{
-							GraphHandle->GetOwningDocument()->RemoveUnreferencedDependencies();
-						}
+						RemoveNodeHandle(GraphHandle, NodeHandle);
 					}
 					break;
 				}
diff --git a/Engine/Plugins/Runtime/Metasound/Source/MetasoundEditor/Private/MetasoundEditorGraphBuilder.h b/Engine/Plugins/Runtime/Metasound/Source/MetasoundEditor/Private/MetasoundEditorGraphBuilder.h
--- a/Engine/Plugins/Runtime/Metasound/Source/MetasoundEditor/Private/MetasoundEditorGraphBuilder.h
+++ b/Engine/Plugins/Runtime/Metasound/Source/MetasoundEditor/Private/MetasoundEditorGraphBuilder.h
@@ -155,9 +155,10 @@ namespace Metasound
 			// @param InMetaSound - FMetasoundAssetBase derived UObject.
 			// @param InVariableID - ID of variable existing on the root graph.
 			// @param InVariableNodeClassName - FNodeClassName of the variable node to add.
+			// @param InEditorNode - (Optional) Editor node to set as referencing new variable node handle.
 			//
 			// @return The added frontend node handle. On error, the returned handle is invalid.
-			static Frontend::FNodeHandle AddVariableNodeHandle(UObject& InMetaSound, const FGuid& InVariableID, const Metasound::FNodeClassName& InVariableNodeClassName);
+			static Frontend::FNodeHandle AddVariableNodeHandle(UObject& InMetaSound, const FGuid& InVariableID, const Metasound::FNodeClassName& InVariableNodeClassName, UMetasoundEditorGraphVariableNode* InEditorNode = nullptr);
 
 			// Attempts to connect Frontend node counterparts together for provided pins.  Returns true if succeeded,
 			// and breaks pin link and returns false if failed.  If bConnectEdPins is set, will attempt to connect
