Fixing random get "shared state" to use GUIDs vs metasound IDs to preserve state correctly.

diff --git a/Engine/Plugins/Runtime/Metasound/Source/MetasoundFrontend/Private/MetasoundArrayRandomNode.cpp b/Engine/Plugins/Runtime/Metasound/Source/MetasoundFrontend/Private/MetasoundArrayRandomNode.cpp
--- a/Engine/Plugins/Runtime/Metasound/Source/MetasoundFrontend/Private/MetasoundArrayRandomNode.cpp
+++ b/Engine/Plugins/Runtime/Metasound/Source/MetasoundFrontend/Private/MetasoundArrayRandomNode.cpp
@@ -39,9 +39,7 @@ namespace Metasound
 	{
 		if (InSeed == INDEX_NONE)
 		{
-			static FThreadSafeCounter SeedCounter(FMath::Rand());
-			SeedCounter.Increment();
-			RandomStream.Initialize(SeedCounter.GetValue());
+			RandomStream.Initialize(FPlatformTime::Cycles());
 		}
 		else
 		{
