Naming gain inputs to clarify linear vs dB. (misinterpretation would appear broken)

diff --git a/Engine/Plugins/Runtime/Metasound/Source/MetasoundFrontend/Private/MetasoundFrontendDocument.cpp b/Engine/Plugins/Runtime/Metasound/Source/MetasoundFrontend/Private/MetasoundFrontendDocument.cpp
--- a/Engine/Plugins/Runtime/Metasound/Source/MetasoundFrontend/Private/MetasoundFrontendDocument.cpp
+++ b/Engine/Plugins/Runtime/Metasound/Source/MetasoundFrontend/Private/MetasoundFrontendDocument.cpp
@@ -210,6 +210,7 @@ FMetasoundFrontendClassInterface FMetasoundFrontendClassInterface::GenerateClass
 
 #if WITH_EDITOR
 			ClassInput.Metadata.SetSerializeText(false);
+			ClassInput.Metadata.SetDisplayName(VertexMetadata.DisplayName);
 			ClassInput.Metadata.SetDescription(VertexMetadata.Description);
 			ClassInput.Metadata.bIsAdvancedDisplay = VertexMetadata.bIsAdvancedDisplay;
 
@@ -261,6 +262,7 @@ FMetasoundFrontendClassInterface FMetasoundFrontendClassInterface::GenerateClass
 
 #if WITH_EDITOR
 			ClassOutput.Metadata.SetSerializeText(false);
+			ClassOutput.Metadata.SetDisplayName(VertexMetadata.DisplayName);
 			ClassOutput.Metadata.SetDescription(VertexMetadata.Description);
 			ClassOutput.Metadata.bIsAdvancedDisplay = VertexMetadata.bIsAdvancedDisplay;
 
diff --git a/Engine/Plugins/Runtime/Metasound/Source/MetasoundGraphCore/Public/MetasoundVertex.h b/Engine/Plugins/Runtime/Metasound/Source/MetasoundGraphCore/Public/MetasoundVertex.h
--- a/Engine/Plugins/Runtime/Metasound/Source/MetasoundGraphCore/Public/MetasoundVertex.h
+++ b/Engine/Plugins/Runtime/Metasound/Source/MetasoundGraphCore/Public/MetasoundVertex.h
@@ -26,6 +26,7 @@ namespace Metasound
 	struct FDataVertexMetadata
 	{
 		FText Description;
+		FText DisplayName;
 		bool bIsAdvancedDisplay = false;
 	};
 
diff --git a/Engine/Plugins/Runtime/Metasound/Source/MetasoundStandardNodes/Private/MetasoundBasicFilters.cpp b/Engine/Plugins/Runtime/Metasound/Source/MetasoundStandardNodes/Private/MetasoundBasicFilters.cpp
--- a/Engine/Plugins/Runtime/Metasound/Source/MetasoundStandardNodes/Private/MetasoundBasicFilters.cpp
+++ b/Engine/Plugins/Runtime/Metasound/Source/MetasoundStandardNodes/Private/MetasoundBasicFilters.cpp
@@ -827,12 +827,18 @@ namespace Metasound
 
 	FVertexInterface FBiquadFilterOperator::DeclareVertexInterface()
 	{
+		static const FDataVertexMetadata GainPinMetaData
+		{
+			  METASOUND_GET_PARAM_TT(ParamGainDb) // description
+			, METASOUND_LOCTEXT("Biquad_Filter_DisplayName", "Gain (dB)") // display name
+		};
+
 		static const FVertexInterface Interface(
 			FInputVertexInterface(
 				TInputDataVertexModel<FAudioBuffer>(METASOUND_GET_PARAM_NAME(ParamAudioInput), METASOUND_GET_PARAM_TT(ParamAudioInput)),
 				TInputDataVertexModel<float>(METASOUND_GET_PARAM_NAME(ParamCutoffFrequency), METASOUND_GET_PARAM_TT(ParamCutoffFrequency), 20000.f),
 				TInputDataVertexModel<float>(METASOUND_GET_PARAM_NAME(ParamBandwidth), METASOUND_GET_PARAM_TT(ParamBandwidth), 1.f),
-				TInputDataVertexModel<float>(METASOUND_GET_PARAM_NAME(ParamGainDb), METASOUND_GET_PARAM_TT(ParamGainDb), 0.f),
+				TInputDataVertexModel<float>(METASOUND_GET_PARAM_NAME(ParamGainDb), GainPinMetaData, 0.f),
 				TInputDataVertexModel<FEnumEBiquadFilterType>(METASOUND_GET_PARAM_NAME(ParamFilterType), METASOUND_GET_PARAM_TT(ParamFilterType))
 			),
 			FOutputVertexInterface(
diff --git a/Engine/Plugins/Runtime/Metasound/Source/MetasoundStandardNodes/Private/MetasoundDynamicFilterNode.cpp b/Engine/Plugins/Runtime/Metasound/Source/MetasoundStandardNodes/Private/MetasoundDynamicFilterNode.cpp
--- a/Engine/Plugins/Runtime/Metasound/Source/MetasoundStandardNodes/Private/MetasoundDynamicFilterNode.cpp
+++ b/Engine/Plugins/Runtime/Metasound/Source/MetasoundStandardNodes/Private/MetasoundDynamicFilterNode.cpp
@@ -170,7 +170,11 @@ namespace Metasound
 					TInputDataVertexModel<float>(METASOUND_GET_PARAM_NAME_AND_TT(InputRatio), 4.0f),
 					TInputDataVertexModel<float>(METASOUND_GET_PARAM_NAME_AND_TT(InputKnee), 12.0f),
 					TInputDataVertexModel<float>(METASOUND_GET_PARAM_NAME_AND_TT(InputRange), -60.0f),
-					TInputDataVertexModel<float>(METASOUND_GET_PARAM_NAME_AND_TT(InputGain), 0.0f),
+					TInputDataVertexModel<float>(METASOUND_GET_PARAM_NAME(InputGain),
+					{ 
+						METASOUND_GET_PARAM_TT(InputGain),
+						METASOUND_LOCTEXT("InputGainName", "Gain (dB)")
+					}, 0.0f),
 					TInputDataVertexModel<FTime>(METASOUND_GET_PARAM_NAME_AND_TT(InputAttackTime), 0.01f),
 					TInputDataVertexModel<FTime>(METASOUND_GET_PARAM_NAME_AND_TT(InputReleaseTime), 0.1f),
 					TInputDataVertexModel<FEnumEnvelopePeakMode>(METASOUND_GET_PARAM_NAME_AND_TT(InputEnvelopeMode), 1.0f),
diff --git a/Engine/Plugins/Runtime/Metasound/Source/MetasoundStandardNodes/Private/MetasoundMixerNode.cpp b/Engine/Plugins/Runtime/Metasound/Source/MetasoundStandardNodes/Private/MetasoundMixerNode.cpp
--- a/Engine/Plugins/Runtime/Metasound/Source/MetasoundStandardNodes/Private/MetasoundMixerNode.cpp
+++ b/Engine/Plugins/Runtime/Metasound/Source/MetasoundStandardNodes/Private/MetasoundMixerNode.cpp
@@ -62,7 +62,14 @@ namespace Metasound
 					}
 
 					// gain scalar
-					InputInterface.Add(TInputDataVertexModel<float>(GetGainInputName(InputIndex), GetGainInputDescription(InputIndex), 1.0f));
+					FDataVertexMetadata GainPinMetaData
+					{
+						GetGainInputDescription(InputIndex),
+						GetGainInputDisplayName(InputIndex)
+					};
+					TInputDataVertexModel<float>GainVertexModel(GetGainInputName(InputIndex), GainPinMetaData, 1.0f);
+
+					InputInterface.Add(GainVertexModel);
 				}
 
 				// outputs
@@ -250,6 +257,11 @@ namespace Metasound
 			return *FString::Printf(TEXT("Gain %i"), InputIndex);
 		}
 
+		static const FText GetGainInputDisplayName(uint32 InputIndex)
+		{
+			return METASOUND_LOCTEXT_FORMAT("AudioMixerGainInputDisplayName", "Gain {0} (Lin)", InputIndex);
+		}
+
 		static const FText GetGainInputDescription(uint32 InputIndex)
 		{
 			return METASOUND_LOCTEXT_FORMAT("AudioMixerGainInputDescription", "Gain Input #: {0}", InputIndex);
