Fix MetaSound sine generator enum mismatch

diff --git a/Engine/Plugins/Runtime/Metasound/Source/MetasoundStandardNodes/Private/MetasoundOscillatorNodes.cpp b/Engine/Plugins/Runtime/Metasound/Source/MetasoundStandardNodes/Private/MetasoundOscillatorNodes.cpp
--- a/Engine/Plugins/Runtime/Metasound/Source/MetasoundStandardNodes/Private/MetasoundOscillatorNodes.cpp
+++ b/Engine/Plugins/Runtime/Metasound/Source/MetasoundStandardNodes/Private/MetasoundOscillatorNodes.cpp
@@ -738,7 +738,6 @@ namespace Metasound
 	{
 		Rotation,
 		Sinf,
-		FilterResonanceFeedback,
 		Bhaskara,
 		Wavetable,
 	};
@@ -751,7 +750,6 @@ namespace Metasound
 		DEFINE_METASOUND_ENUM_ENTRY(ESineGenerationType::Sinf, "SinfDescription", "Pure Math", "SinfDescriptionTT", "Uses the standard math library (Sinf) to generate the sine (most expensive)"),
 		DEFINE_METASOUND_ENUM_ENTRY(ESineGenerationType::Bhaskara, "BhaskaraDescription", "Bhaskara", "BhaskaraDescriptionTT", "Sine approximation using Bhaskara technique discovered in 7th century"),
 		DEFINE_METASOUND_ENUM_ENTRY(ESineGenerationType::Wavetable, "WavetableDescription", "Wavetable", "WavetableDescriptionTT", "Uses a wavetable to generate the sine"),
-		//DEFINE_METASOUND_ENUM_ENTRY(ESineGenerationType::FilterResonanceFeedback, "FRFDescription", "Filter resonance", "FRFDescriptionTT", "Filter resonance/feedback"),
 	DEFINE_METASOUND_ENUM_END()
 
 	class FSineOscilatorNode::FFactory : public FOscilatorFactoryBase
