Fixing tooltip warning

diff --git a/Engine/Plugins/Runtime/Metasound/Source/MetasoundStandardNodes/Private/MetasoundMathNodes.cpp b/Engine/Plugins/Runtime/Metasound/Source/MetasoundStandardNodes/Private/MetasoundMathNodes.cpp
--- a/Engine/Plugins/Runtime/Metasound/Source/MetasoundStandardNodes/Private/MetasoundMathNodes.cpp
+++ b/Engine/Plugins/Runtime/Metasound/Source/MetasoundStandardNodes/Private/MetasoundMathNodes.cpp
@@ -671,8 +671,8 @@ namespace Metasound
 		{
 			static const FVertexInterface Interface(
 				FInputVertexInterface(
-					TInputDataVertexModel<FAudioBuffer>(MathOpNames::PrimaryOperandName, METASOUND_LOCTEXT("MathOpAddendTooltip1", "First addend.")),
-					TInputDataVertexModel<FAudioBuffer>(MathOpNames::AdditionalOperandsName, METASOUND_LOCTEXT("MathOpAddendAdditionalTooltip1", "Additional addends."))
+					TInputDataVertexModel<FAudioBuffer>(MathOpNames::PrimaryOperandName, METASOUND_LOCTEXT("MathOpAddAudioBufferNameTooltip", "First addend.")),
+					TInputDataVertexModel<FAudioBuffer>(MathOpNames::AdditionalOperandsName, METASOUND_LOCTEXT("MathOpAddAudioBufferAdditionalTooltip", "Additional addends."))
 				),
 				FOutputVertexInterface(
 					TOutputDataVertexModel<FAudioBuffer>(TEXT("Out"), METASOUND_LOCTEXT("MathOpOutTooltip", "Math operation result"))
@@ -748,8 +748,8 @@ namespace Metasound
 		{
 			static const FVertexInterface Interface(
 				FInputVertexInterface(
-					TInputDataVertexModel<FTime>(MathOpNames::PrimaryOperandName, METASOUND_LOCTEXT("MathOpAddendTooltip2", "First addend."), 0.0f),
-					TInputDataVertexModel<FTime>(MathOpNames::AdditionalOperandsName, METASOUND_LOCTEXT("MathOpAddendAdditionalTooltip2", "Additional addends."), 0.0f)
+					TInputDataVertexModel<FTime>(MathOpNames::PrimaryOperandName, METASOUND_LOCTEXT("MathOpAddTimeTooltip", "First addend."), 0.0f),
+					TInputDataVertexModel<FTime>(MathOpNames::AdditionalOperandsName, METASOUND_LOCTEXT("MathOpAddTimeAdditionalalTooltip", "Additional addends."), 0.0f)
 				),
 				FOutputVertexInterface(
 					TOutputDataVertexModel<FTime>(TEXT("Out"), METASOUND_LOCTEXT("MathOpOutTooltip", "Math operation result"))
@@ -792,8 +792,8 @@ namespace Metasound
 		{
 			static const FVertexInterface Interface(
 				FInputVertexInterface(
-					TInputDataVertexModel<FAudioBuffer>(MathOpNames::PrimaryOperandName, METASOUND_LOCTEXT("MathOpAddendAudioTooltip", "Audio Buffer to add offset(s) to.")),
-					TInputDataVertexModel<float>(MathOpNames::AdditionalOperandsName, METASOUND_LOCTEXT("MathOpAddendAdditionalTooltip2", "Float addends of which to offset buffer samples."), 0.0f)
+					TInputDataVertexModel<FAudioBuffer>(MathOpNames::PrimaryOperandName, METASOUND_LOCTEXT("MathOpAddAudioTooltip", "Audio Buffer to add offset(s) to.")),
+					TInputDataVertexModel<float>(MathOpNames::AdditionalOperandsName, METASOUND_LOCTEXT("MathOpAddAdditionalTooltip", "Float addends of which to offset buffer samples."), 0.0f)
 				),
 				FOutputVertexInterface(
 					TOutputDataVertexModel<FAudioBuffer>(TEXT("Out"), METASOUND_LOCTEXT("MathOpAudioFloatAddOutTooltip", "Resulting buffer"))
