Fix for regression with undo/redo when adding/removing MetaSound Graph Members causing the interface list or details panel to go stale
#rnx

diff --git a/Engine/Plugins/Runtime/Metasound/Source/MetasoundEditor/Private/MetasoundEditor.cpp b/Engine/Plugins/Runtime/Metasound/Source/MetasoundEditor/Private/MetasoundEditor.cpp
--- a/Engine/Plugins/Runtime/Metasound/Source/MetasoundEditor/Private/MetasoundEditor.cpp
+++ b/Engine/Plugins/Runtime/Metasound/Source/MetasoundEditor/Private/MetasoundEditor.cpp
@@ -1450,7 +1450,7 @@ namespace Metasound
 				UMetasoundEditorGraphMember* NextToSelect = Graph.FindAdjacentMember(*GraphMember);
 
 				{
-					const FScopedTransaction Transaction(LOCTEXT("MetaSoundEditorDeleteSelectedNode1", "Delete MetaSound Graph Member"));
+					const FScopedTransaction Transaction(LOCTEXT("MetaSoundEditorDeleteSelectedMember", "Delete MetaSound Graph Member"));
 					Metasound->Modify();
 					Graph.Modify();
 					Graph.RemoveMember(*GraphMember);
@@ -1520,6 +1520,16 @@ namespace Metasound
 								}
 								else
 								{
+									if (MetasoundDetails.IsValid())
+									{
+										if (MetasoundDetails->GetSelectedObjects().Contains(GraphMember))
+										{
+											FMetasoundAssetBase* MetasoundAsset = IMetasoundUObjectRegistry::Get().GetObjectAsAssetBase(Metasound);
+											check(MetasoundAsset);
+											MetasoundAsset->SetUpdateDetailsOnSynchronization();
+										}
+									}
+
 									ActionToDelete = MetasoundAction;
 									break;
 								}
@@ -1545,6 +1555,10 @@ namespace Metasound
 			const FGraphPanelSelectionSet SelectedNodes = MetasoundGraphEditor->GetSelectedNodes();
 			MetasoundGraphEditor->ClearSelectionSet();
 
+			FMetasoundAssetBase* MetasoundAsset = IMetasoundUObjectRegistry::Get().GetObjectAsAssetBase(Metasound);
+			check(MetasoundAsset);
+			MetasoundAsset->SetUpdateDetailsOnSynchronization();
+
 			const FScopedTransaction Transaction(LOCTEXT("MetaSoundEditorDeleteSelectedNode2", "Delete Selected MetaSound Node(s)"));
 			check(Metasound);
 			Metasound->Modify();
diff --git a/Engine/Plugins/Runtime/Metasound/Source/MetasoundEngine/Private/Metasound.cpp b/Engine/Plugins/Runtime/Metasound/Source/MetasoundEngine/Private/Metasound.cpp
--- a/Engine/Plugins/Runtime/Metasound/Source/MetasoundEngine/Private/Metasound.cpp
+++ b/Engine/Plugins/Runtime/Metasound/Source/MetasoundEngine/Private/Metasound.cpp
@@ -53,17 +53,7 @@ void UMetaSound::PostDuplicate(EDuplicateMode::Type InDuplicateMode)
 void UMetaSound::PostEditUndo()
 {
 	Super::PostEditUndo();
-
-	if (Graph)
-	{
-		Graph->RegisterGraphWithFrontend();
-	}
-}
-
-void UMetaSound::PostEditChangeProperty(FPropertyChangedEvent& InEvent)
-{
-	Super::PostEditChangeProperty(InEvent);
-	Metasound::PostEditChangeProperty(*this, InEvent);
+	Metasound::PostEditUndo(*this);
 }
 #endif // WITHEDITOR
 
diff --git a/Engine/Plugins/Runtime/Metasound/Source/MetasoundEngine/Private/MetasoundSource.cpp b/Engine/Plugins/Runtime/Metasound/Source/MetasoundEngine/Private/MetasoundSource.cpp
--- a/Engine/Plugins/Runtime/Metasound/Source/MetasoundEngine/Private/MetasoundSource.cpp
+++ b/Engine/Plugins/Runtime/Metasound/Source/MetasoundEngine/Private/MetasoundSource.cpp
@@ -96,11 +96,7 @@ UMetaSoundSource::UMetaSoundSource(const FObjectInitializer& ObjectInitializer)
 void UMetaSoundSource::PostEditUndo()
 {
 	Super::PostEditUndo();
-
-	if (Graph)
-	{
-		Graph->RegisterGraphWithFrontend();
-	}
+	Metasound::PostEditUndo(*this);
 }
 
 void UMetaSoundSource::PostDuplicate(EDuplicateMode::Type InDuplicateMode)
@@ -122,7 +118,6 @@ void UMetaSoundSource::PostEditChangeProperty(FPropertyChangedEvent& InEvent)
 	using namespace Metasound::Frontend;
 
 	Super::PostEditChangeProperty(InEvent);
-	Metasound::PostEditChangeProperty(*this, InEvent);
 
 	if (InEvent.GetPropertyName() == GET_MEMBER_NAME_CHECKED(UMetaSoundSource, OutputFormat))
 	{
diff --git a/Engine/Plugins/Runtime/Metasound/Source/MetasoundEngine/Public/Metasound.h b/Engine/Plugins/Runtime/Metasound/Source/MetasoundEngine/Public/Metasound.h
--- a/Engine/Plugins/Runtime/Metasound/Source/MetasoundEngine/Public/Metasound.h
+++ b/Engine/Plugins/Runtime/Metasound/Source/MetasoundEngine/Public/Metasound.h
@@ -50,8 +50,13 @@ namespace Metasound
 
 #if WITH_EDITOR
 	template <typename TMetaSoundObject>
-	void PostEditChangeProperty(TMetaSoundObject& InMetaSound, FPropertyChangedEvent& InEvent)
+	void PostEditUndo(TMetaSoundObject& InMetaSound)
 	{
+		InMetaSound.SetSynchronizationRequired();
+		if (UMetasoundEditorGraphBase* Graph = Cast<UMetasoundEditorGraphBase>(InMetaSound.GetGraph()))
+		{
+			Graph->RegisterGraphWithFrontend();
+		}
 	}
 #endif // WITH_EDITOR
 
@@ -230,7 +235,6 @@ public:
 #if WITH_EDITOR
 	virtual void PostDuplicate(EDuplicateMode::Type DuplicateMode) override;
 	virtual void PostEditUndo() override;
-	virtual void PostEditChangeProperty(FPropertyChangedEvent& InEvent) override;
 #endif // WITH_EDITOR
 
 	virtual void BeginDestroy() override;
