- Drag/Drop/Double-Click Add/JumpTo Member Support
- Fix another spot for when member is removed from the editor views on next tick, don't show garbage member in details panel
- Variable node body fix-ups
#rnx

diff --git a/Engine/Plugins/Runtime/Metasound/Source/MetasoundEditor/Private/MetasoundEditor.cpp b/Engine/Plugins/Runtime/Metasound/Source/MetasoundEditor/Private/MetasoundEditor.cpp
--- a/Engine/Plugins/Runtime/Metasound/Source/MetasoundEditor/Private/MetasoundEditor.cpp
+++ b/Engine/Plugins/Runtime/Metasound/Source/MetasoundEditor/Private/MetasoundEditor.cpp
@@ -11,17 +11,21 @@
 #include "EdGraphUtilities.h"
 #include "Editor.h"
 #include "EditorStyleSet.h"
+#include "Framework/Application/SlateApplication.h"
 #include "Framework/Commands/GenericCommands.h"
 #include "Framework/Docking/TabManager.h"
 #include "Framework/MultiBox/MultiBoxBuilder.h"
 #include "Framework/Notifications/NotificationManager.h"
+#include "GenericPlatform/GenericApplication.h"
 #include "GraphEditor.h"
 #include "GraphEditorActions.h"
+#include "GraphEditorDragDropAction.h"
 #include "HAL/IConsoleManager.h"
 #include "HAL/PlatformApplicationMisc.h"
 #include "IAudioExtensionPlugin.h"
 #include "IDetailsView.h"
 #include "Kismet2/BlueprintEditorUtils.h"
+#include "Kismet2/KismetEditorUtilities.h"
 #include "Logging/TokenizedMessage.h"
 #include "Metasound.h"
 #include "MetasoundEditorCommands.h"
@@ -46,6 +50,7 @@
 #include "PropertyCustomizationHelpers.h"
 #include "PropertyEditorModule.h"
 #include "ScopedTransaction.h"
+#include "SGraphPanel.h"
 #include "SMetasoundActionMenu.h"
 #include "SMetasoundPalette.h"
 #include "SNodePanel.h"
@@ -56,6 +61,7 @@
 #include "Widgets/Notifications/SNotificationList.h"
 #include "Widgets/SBoxPanel.h"
 #include "Widgets/SOverlay.h"
+#include "Widgets/SWindow.h"
 #include "Stats/Stats.h"
 
 struct FGraphActionNode;
@@ -147,6 +153,261 @@ namespace Metasound
 			}
 		};
 
+		class FMetaSoundDragDropMemberAction : public FGraphSchemaActionDragDropAction
+		{
+			TSharedPtr<FEditor> Editor;
+			TWeakObjectPtr<UMetasoundEditorGraphMember> GraphMember;
+
+		public:
+			FMetaSoundDragDropMemberAction(TSharedPtr<FEditor> InEditor, UMetasoundEditorGraphMember* InGraphMember)
+				: Editor(InEditor)
+				, GraphMember(InGraphMember)
+			{
+				Construct();
+			}
+
+			DRAG_DROP_OPERATOR_TYPE(FMetaSoundDragDropMemberAction, FGraphSchemaActionDragDropAction)
+
+			virtual FReply DroppedOnPanel(const TSharedRef<SWidget>& InPanel, FVector2D InScreenPosition, FVector2D InGraphPosition, UEdGraph& InGraph) override
+			{
+				using namespace Frontend;
+
+				if (!GraphMember.IsValid() || &InGraph != GraphMember->GetOwningGraph())
+				{
+					return FReply::Unhandled();
+				}
+
+				UMetasoundEditorGraph* MetasoundGraph = CastChecked<UMetasoundEditorGraph>(&InGraph);
+				UObject& ParentMetasound = MetasoundGraph->GetMetasoundChecked();
+
+				if (UMetasoundEditorGraphInput* Input = Cast<UMetasoundEditorGraphInput>(GraphMember.Get()))
+				{
+					const FScopedTransaction Transaction(LOCTEXT("DropAddNewInputNode", "Drop New MetaSound Input Node"));
+					ParentMetasound.Modify();
+					MetasoundGraph->Modify();
+					Input->Modify();
+
+					if (UMetasoundEditorGraphInputNode* NewGraphNode = FGraphBuilder::AddInputNode(ParentMetasound, Input->GetNodeHandle(), InGraphPosition))
+					{
+						NewGraphNode->Modify();
+						FGraphBuilder::RegisterGraphWithFrontend(ParentMetasound);
+						return FReply::Handled();
+					}
+				}
+
+				if (UMetasoundEditorGraphOutput* Output = Cast<UMetasoundEditorGraphOutput>(GraphMember.Get()))
+				{
+					TArray<UMetasoundEditorGraphNode*> Nodes = Output->GetNodes();
+					if (Nodes.IsEmpty())
+					{
+						const FScopedTransaction Transaction(LOCTEXT("DropAddNewOutputNode", "Drop New MetaSound Output Node"));
+						ParentMetasound.Modify();
+						MetasoundGraph->Modify();
+						Output->Modify();
+
+						FNodeHandle OutputHandle = Output->GetNodeHandle();
+						if (UMetasoundEditorGraphOutputNode* NewGraphNode = FGraphBuilder::AddOutputNode(ParentMetasound, OutputHandle, InGraphPosition))
+						{
+							NewGraphNode->Modify();
+							FGraphBuilder::RegisterGraphWithFrontend(ParentMetasound);
+							return FReply::Handled();
+						}
+					}
+					else
+					{
+						if (Editor.IsValid())
+						{
+							Editor->JumpToNodes(Nodes);
+							return FReply::Handled();
+						}
+					}
+				}
+
+				if (UMetasoundEditorGraphVariable* Variable = Cast<UMetasoundEditorGraphVariable>(GraphMember.Get()))
+				{
+					const FScopedTransaction Transaction(LOCTEXT("DropAddNewVariableNode", "Drop New MetaSound Variable Node"));
+					ParentMetasound.Modify();
+					MetasoundGraph->Modify();
+					Variable->Modify();
+
+					FVariableHandle VariableHandle = Variable->GetVariableHandle();
+					FMetasoundFrontendClass VariableClass;
+
+					const bool bMakeOrJumpToMutator = FSlateApplication::Get().GetModifierKeys().AreModifersDown(EModifierKey::Shift);
+					if (bMakeOrJumpToMutator)
+					{
+						FConstNodeHandle MutatorNodeHandle = Variable->GetConstVariableHandle()->FindMutatorNode();
+						if (MutatorNodeHandle->IsValid())
+						{
+							if (Editor.IsValid())
+							{
+								auto IsMutatorNode = [&MutatorNodeHandle](const UMetasoundEditorGraphNode* Node)
+								{
+									return Node->GetNodeID() == MutatorNodeHandle->GetID();
+								};
+								TArray<UMetasoundEditorGraphNode*> Nodes = Variable->GetNodes();
+								if (UMetasoundEditorGraphNode** MutatorNode = Nodes.FindByPredicate(IsMutatorNode))
+								{
+									check(*MutatorNode);
+									Editor->JumpToNodes({ *MutatorNode });
+									return FReply::Handled();
+								}
+							}
+						}
+						else
+						{
+							ensure(IDataTypeRegistry::Get().GetFrontendVariableMutatorClass(Variable->GetDataType(), VariableClass));
+						}
+					}
+					else
+					{
+						const bool bJumpToGetters = FSlateApplication::Get().GetModifierKeys().AreModifersDown(EModifierKey::Control);
+						if (bJumpToGetters)
+						{
+							TArray<UMetasoundEditorGraphNode*> Nodes = Variable->GetNodes();
+							for (int32 i = Nodes.Num() - 1; i >= 0; --i)
+							{
+								const UMetasoundEditorGraphVariableNode* VariableNode = CastChecked<UMetasoundEditorGraphVariableNode>(Nodes[i]);
+								const EMetasoundFrontendClassType ClassType = VariableNode->GetClassType();
+								if (ClassType != EMetasoundFrontendClassType::VariableAccessor
+									&& ClassType != EMetasoundFrontendClassType::VariableDeferredAccessor)
+								{
+									constexpr bool bAllowShrinking = false;
+									Nodes.RemoveAtSwap(i, 1, bAllowShrinking);
+								}
+							}
+							Editor->JumpToNodes(Nodes);
+							return FReply::Handled();
+						}
+						else
+						{
+							const bool bMakeGetDeferred = FSlateApplication::Get().GetModifierKeys().AreModifersDown(EModifierKey::Alt);
+							if (bMakeGetDeferred)
+							{
+								ensure(IDataTypeRegistry::Get().GetFrontendVariableDeferredAccessorClass(Variable->GetDataType(), VariableClass));
+							}
+							else
+							{
+								ensure(IDataTypeRegistry::Get().GetFrontendVariableAccessorClass(Variable->GetDataType(), VariableClass));
+							}
+						}
+					}
+
+					const FNodeClassName ClassName = VariableClass.Metadata.GetClassName().ToNodeClassName();
+					FNodeHandle NodeHandle = FGraphBuilder::AddVariableNodeHandle(ParentMetasound, Variable->GetVariableID(), ClassName);
+					if (UMetasoundEditorGraphVariableNode* NewGraphNode = FGraphBuilder::AddVariableNode(ParentMetasound, NodeHandle, InGraphPosition))
+					{
+						NewGraphNode->Modify();
+						FGraphBuilder::RegisterGraphWithFrontend(ParentMetasound);
+						return FReply::Handled();
+					}
+				}
+
+				return FReply::Unhandled();
+			}
+			virtual FReply DroppedOnNode(FVector2D ScreenPosition, FVector2D GraphPosition) override { return FReply::Unhandled(); }
+			virtual FReply DroppedOnPin(FVector2D ScreenPosition, FVector2D GraphPosition) override { return FReply::Unhandled(); }
+			virtual FReply DroppedOnAction(TSharedRef<FEdGraphSchemaAction> Action) { return FReply::Unhandled(); }
+			virtual FReply DroppedOnCategory(FText Category) override { return FReply::Unhandled(); }
+			virtual void HoverTargetChanged() override
+			{
+				using namespace Frontend;
+
+				bDropTargetValid = true;
+
+				const FSlateBrush* PrimarySymbol = nullptr;
+				const FSlateBrush* SecondarySymbol = nullptr;
+				FSlateColor PrimaryColor;
+				FSlateColor SecondaryColor;
+				GetDefaultStatusSymbol(PrimarySymbol, PrimaryColor, SecondarySymbol, SecondaryColor);
+
+				FText Message;
+				if (GraphMember.IsValid())
+				{
+					Message = GraphMember->GetDisplayName();
+					if (GetHoveredGraph() == GraphMember->GetOwningGraph())
+					{
+						if (UMetasoundEditorGraphInput* Input = Cast<UMetasoundEditorGraphInput>(GraphMember.Get()))
+						{
+							if (const ISlateStyle* MetasoundStyle = FSlateStyleRegistry::FindSlateStyle("MetaSoundStyle"))
+							{
+								PrimarySymbol = MetasoundStyle->GetBrush("MetasoundEditor.Graph.Node.Class.Input");
+								SecondarySymbol = nullptr;
+							}
+
+							if (const UMetasoundEditorSettings* EditorSettings = GetDefault<UMetasoundEditorSettings>())
+							{
+								PrimaryColor = EditorSettings->InputNodeTitleColor;
+								SecondaryColor = EditorSettings->InputNodeTitleColor;
+							}
+						}
+
+						if (UMetasoundEditorGraphOutput* Output = Cast<UMetasoundEditorGraphOutput>(GraphMember.Get()))
+						{
+							if (!Output->GetNodes().IsEmpty())
+							{
+								PrimarySymbol = FEditorStyle::GetBrush(TEXT("Graph.ConnectorFeedback.ShowNode"));
+								SecondarySymbol = nullptr;
+								Message = FText::Format(LOCTEXT("DropTargetShowOutput", "Show '{0}' (One per graph)"), GraphMember->GetDisplayName());
+							}
+							else
+							{
+								if (const ISlateStyle* MetasoundStyle = FSlateStyleRegistry::FindSlateStyle("MetaSoundStyle"))
+								{
+									PrimarySymbol = MetasoundStyle->GetBrush("MetasoundEditor.Graph.Node.Class.Output");
+									SecondarySymbol = nullptr;
+								}
+
+								if (const UMetasoundEditorSettings* EditorSettings = GetDefault<UMetasoundEditorSettings>())
+								{
+									PrimaryColor = EditorSettings->OutputNodeTitleColor;
+									SecondaryColor = EditorSettings->OutputNodeTitleColor;
+								}
+							}
+						}
+
+						if (UMetasoundEditorGraphVariable* Variable = Cast<UMetasoundEditorGraphVariable>(GraphMember.Get()))
+						{
+							PrimarySymbol = FEditorStyle::GetBrush(TEXT("Graph.ConnectorFeedback.ShowNode"));
+
+							if (const ISlateStyle* MetasoundStyle = FSlateStyleRegistry::FindSlateStyle("MetaSoundStyle"))
+							{
+								PrimarySymbol = MetasoundStyle->GetBrush("MetasoundEditor.Graph.Node.Class.Variable");
+								SecondarySymbol = nullptr;
+							}
+
+							if (const UMetasoundEditorSettings* EditorSettings = GetDefault<UMetasoundEditorSettings>())
+							{
+								PrimaryColor = EditorSettings->VariableNodeTitleColor;
+								SecondaryColor = EditorSettings->VariableNodeTitleColor;
+							}
+
+							const FText DisplayName = GraphMember->GetDisplayName();
+							const FText GetterToolTip = FText::Format(LOCTEXT("DropTargetGetterVariableToolTipFormat", "{0}\nAdd:\n* Get (Drop)\n* Get Delayed (Alt+Drop)\n"), DisplayName);
+							static const FText GetShowToolTip = LOCTEXT("ShowGettersToolTip", "Get (Ctrl+Drop)");
+							static const FText AddOrShowSetToolTip = LOCTEXT("AddOrShowSetToolTip", "");
+							FConstNodeHandle MutatorNodeHandle = Variable->GetConstVariableHandle()->FindMutatorNode();
+							if (MutatorNodeHandle->IsValid())
+							{
+								Message = FText::Format(LOCTEXT("DropTargetVariableShowFormat", "{0}\nShow:\n* {1}\n* Set (Shift+Drop, One per graph)"), GetterToolTip, GetShowToolTip);
+							}
+							else
+							{
+								Message = FText::Format(LOCTEXT("DropTargetVariableAddSetShowGetFormat", "{0}* Set (Shift+Drop)\n\nShow:\n* {1}"), GetterToolTip, GetShowToolTip);
+							}
+						}
+					}
+				}
+				else
+				{
+					bDropTargetValid = false;
+				}
+
+				SetSimpleFeedbackMessage(PrimarySymbol, PrimaryColor, Message, SecondarySymbol, SecondaryColor);
+			}
+		};
+
+
 		class SMetaSoundGraphPaletteItem : public SGraphPaletteItem
 		{
 		private:
@@ -680,7 +941,7 @@ namespace Metasound
 
 			SAssignNew(MetasoundInterfaceMenu, SGraphActionMenu, false)
 				.AlphaSortItems(true)
-// 				.OnActionDoubleClicked(this, &FEditor::OnActionDoubleClicked)
+				.OnActionDoubleClicked(this, &FEditor::OnMemberActionDoubleClicked)
 				.OnActionDragged(this, &FEditor::OnActionDragged)
 				.OnActionMatchesName(this, &FEditor::HandleActionMatchesName)
 				.OnActionSelected(this, &FEditor::OnActionSelected)
@@ -1469,6 +1730,10 @@ namespace Metasound
 						SetSelection(GraphMembersToSelect);
 					}
 				}
+				else
+				{
+					MetasoundAsset->SetUpdateDetailsOnSynchronization();
+				}
 			}
 
 			FGraphBuilder::RegisterGraphWithFrontend(*Metasound);
@@ -2401,12 +2666,78 @@ namespace Metasound
 			return false;
 		}
 
-		// TODO: Implement
 		FReply FEditor::OnActionDragged(const TArray<TSharedPtr<FEdGraphSchemaAction>>& InActions, const FPointerEvent& MouseEvent)
 		{
+			if (!MetasoundGraphEditor.IsValid() || InActions.IsEmpty())
+			{
+				return FReply::Unhandled();
+			}
+
+			TSharedPtr<FEdGraphSchemaAction> DragAction = InActions.Last();
+			if (FMetasoundGraphMemberSchemaAction* MemberAction = static_cast<FMetasoundGraphMemberSchemaAction*>(DragAction.Get()))
+			{
+				if (UEdGraph* ActionGraph = MemberAction->Graph)
+				{
+					if (&GetMetaSoundGraphChecked() == ActionGraph)
+					{
+						TSharedPtr<FEditor> ThisEditor = StaticCastSharedRef<FEditor>(AsShared());
+						return FReply::Handled().BeginDragDrop(MakeShared<FMetaSoundDragDropMemberAction>(ThisEditor, MemberAction->GetGraphMember()));
+					}
+				}
+			}
+
 			return FReply::Unhandled();
 		}
 
+		void FEditor::OnMemberActionDoubleClicked(const TArray<TSharedPtr<FEdGraphSchemaAction>>& InActions)
+		{
+			if (!MetasoundGraphEditor.IsValid() || InActions.IsEmpty())
+			{
+				return;
+			}
+
+			TSharedPtr<FMetasoundGraphMemberSchemaAction> MemberAction = StaticCastSharedPtr<FMetasoundGraphMemberSchemaAction>(InActions.Last());
+			if (UMetasoundEditorGraphMember* Member = MemberAction->GetGraphMember())
+			{
+				JumpToNodes(Member->GetNodes());
+			}
+		}
+
+		void FEditor::JumpToNodes(const TArray<UMetasoundEditorGraphNode*>& InNodes)
+		{
+			if (!MetasoundGraphEditor.IsValid())
+			{
+				return;
+			}
+
+			MetasoundGraphEditor->ClearSelectionSet();
+			const UMetasoundEditorGraph& Graph = GetMetaSoundGraphChecked();
+			if (!InNodes.IsEmpty())
+			{
+				if (SGraphPanel* GraphPanel = MetasoundGraphEditor->GetGraphPanel())
+				{
+					FVector2D BottomLeft = { TNumericLimits<float>::Max(), TNumericLimits<float>::Max() };
+					FVector2D TopRight = { TNumericLimits<float>::Min(), TNumericLimits<float>::Min() };
+					for (UMetasoundEditorGraphNode* Node : InNodes)
+					{
+						if (!Node || Node->GetGraph() != &Graph)
+						{
+							continue;
+						}
+
+						constexpr bool bSelected = true;
+						MetasoundGraphEditor->SetNodeSelection(Node, bSelected);
+						BottomLeft.X = FMath::Min(BottomLeft.X, Node->NodePosX);
+						BottomLeft.Y = FMath::Min(BottomLeft.Y, Node->NodePosY);
+						TopRight.X = FMath::Max(TopRight.X, Node->NodePosX + Node->EstimateNodeWidth());
+						TopRight.Y = FMath::Max(TopRight.Y, Node->NodePosY);
+					}
+
+					GraphPanel->JumpToRect(BottomLeft, TopRight);
+				}
+			}
+		}
+
 		FActionMenuContent FEditor::OnCreateGraphActionMenu(UEdGraph* InGraph, const FVector2D& InNodePosition, const TArray<UEdGraphPin*>& InDraggedPins, bool bAutoExpand, SGraphEditor::FActionMenuClosed InOnMenuClosed)
 		{
 			TSharedRef<SMetasoundActionMenu> ActionMenu = SNew(SMetasoundActionMenu)
diff --git a/Engine/Plugins/Runtime/Metasound/Source/MetasoundEditor/Private/MetasoundEditor.h b/Engine/Plugins/Runtime/Metasound/Source/MetasoundEditor/Private/MetasoundEditor.h
--- a/Engine/Plugins/Runtime/Metasound/Source/MetasoundEditor/Private/MetasoundEditor.h
+++ b/Engine/Plugins/Runtime/Metasound/Source/MetasoundEditor/Private/MetasoundEditor.h
@@ -39,6 +39,7 @@ class SVerticalBox;
 class UEdGraphNode;
 class UMetaSound;
 class UMetasoundEditorGraph;
+class UMetasoundEditorGraphNode;
 
 struct FGraphActionNode;
 struct FMeterResults;
@@ -157,6 +158,9 @@ namespace Metasound
 			/** Destroys analyzers */
 			void DestroyAnalyzers();
 
+			/** Jumps to the given nodes on the graph */
+			void JumpToNodes(const TArray<UMetasoundEditorGraphNode*>& InNodes);
+
 		protected:
 			// Callbacks for action tree
 			bool CanRenameOnActionNode(TWeakPtr<FGraphActionNode> InSelectedNode) const;
@@ -167,6 +171,7 @@ namespace Metasound
 			FText GetFilterText() const;
 			bool HandleActionMatchesName(FEdGraphSchemaAction* InAction, const FName& InName) const;
 			FReply OnActionDragged(const TArray<TSharedPtr<FEdGraphSchemaAction>>& InActions, const FPointerEvent& MouseEvent);
+			void OnMemberActionDoubleClicked(const TArray< TSharedPtr<FEdGraphSchemaAction> >& InActions);
 			FActionMenuContent OnCreateGraphActionMenu(UEdGraph* InGraph, const FVector2D& InNodePosition, const TArray<UEdGraphPin*>& InDraggedPins, bool bAutoExpand, SGraphEditor::FActionMenuClosed InOnMenuClosed);
 			void OnActionSelected(const TArray<TSharedPtr<FEdGraphSchemaAction>>& InActions, ESelectInfo::Type InSelectionType);
 			FReply OnAddButtonClickedOnSection(int32 InSectionID);
@@ -202,9 +207,6 @@ namespace Metasound
 			/** Cut the currently selected nodes */
 			void CutSelectedNodes();
 
-			/** Copy the currently selected nodes to the given string */
-			void CopySelectedNodes(FString& OutString) const;
-
 			/** Copy the currently selected nodes */
 			void CopySelectedNodes() const;
 
diff --git a/Engine/Plugins/Runtime/Metasound/Source/MetasoundEditor/Private/MetasoundEditorGraph.cpp b/Engine/Plugins/Runtime/Metasound/Source/MetasoundEditor/Private/MetasoundEditorGraph.cpp
--- a/Engine/Plugins/Runtime/Metasound/Source/MetasoundEditor/Private/MetasoundEditorGraph.cpp
+++ b/Engine/Plugins/Runtime/Metasound/Source/MetasoundEditor/Private/MetasoundEditorGraph.cpp
@@ -1623,6 +1623,12 @@ bool UMetasoundEditorGraph::RemoveMember(UMetasoundEditorGraphMember& InGraphMem
 		bSuccess = false;
 	}
 
+	// Required to ensure member is removed from the editor views on next tick when refreshing details.
+	if (NumRemoved > 0)
+	{
+		InGraphMember.MarkAsGarbage();
+	}
+
 	bSuccess = bSuccess && (NumRemoved > 0);
 
 	return bSuccess;
diff --git a/Engine/Plugins/Runtime/Metasound/Source/MetasoundEditor/Private/MetasoundEditorGraphNode.cpp b/Engine/Plugins/Runtime/Metasound/Source/MetasoundEditor/Private/MetasoundEditorGraphNode.cpp
--- a/Engine/Plugins/Runtime/Metasound/Source/MetasoundEditor/Private/MetasoundEditorGraphNode.cpp
+++ b/Engine/Plugins/Runtime/Metasound/Source/MetasoundEditor/Private/MetasoundEditorGraphNode.cpp
@@ -1013,11 +1013,26 @@ FMetasoundFrontendClassName UMetasoundEditorGraphVariableNode::GetClassName() co
 	return ClassName;
 }
 
+EMetasoundFrontendClassType UMetasoundEditorGraphVariableNode::GetClassType() const
+{
+	return ClassType;
+}
+
 FGuid UMetasoundEditorGraphVariableNode::GetNodeID() const
 {
 	return NodeID;
 }
 
+FName UMetasoundEditorGraphVariableNode::GetCornerIcon() const
+{
+	if (ClassType == EMetasoundFrontendClassType::VariableDeferredAccessor)
+	{
+		return TEXT("Graph.Latent.LatentIcon");
+	}
+
+	return Super::GetCornerIcon();
+}
+
 void UMetasoundEditorGraphVariableNode::PinDefaultValueChanged(UEdGraphPin* Pin)
 {
 	using namespace Metasound::Editor;
@@ -1068,7 +1083,7 @@ FLinearColor UMetasoundEditorGraphVariableNode::GetNodeTitleColor() const
 
 FSlateIcon UMetasoundEditorGraphVariableNode::GetNodeTitleIcon() const
 {
-	return FSlateIcon("MetaSoundStyle", "MetasoundEditor.Graph.Node.Class.Variable");
+	return FSlateIcon();
 }
 
 void UMetasoundEditorGraphVariableNode::SetNodeID(FGuid InNodeID)
diff --git a/Engine/Plugins/Runtime/Metasound/Source/MetasoundEditor/Private/SMetasoundGraphNode.cpp b/Engine/Plugins/Runtime/Metasound/Source/MetasoundEditor/Private/SMetasoundGraphNode.cpp
--- a/Engine/Plugins/Runtime/Metasound/Source/MetasoundEditor/Private/SMetasoundGraphNode.cpp
+++ b/Engine/Plugins/Runtime/Metasound/Source/MetasoundEditor/Private/SMetasoundGraphNode.cpp
@@ -71,6 +71,29 @@ namespace Metasound
 			}
 		}
 
+		bool SMetaSoundGraphNode::IsVariableAccessor() const
+		{
+			const EMetasoundFrontendClassType ClassType = GetMetaSoundNode().GetNodeHandle()->GetClassMetadata().GetType();
+			return ClassType == EMetasoundFrontendClassType::VariableAccessor
+				|| ClassType == EMetasoundFrontendClassType::VariableDeferredAccessor;
+		}
+
+		bool SMetaSoundGraphNode::IsVariableMutator() const
+		{
+			const EMetasoundFrontendClassType ClassType = GetMetaSoundNode().GetNodeHandle()->GetClassMetadata().GetType();
+			return ClassType == EMetasoundFrontendClassType::VariableMutator;
+		}
+
+		const FSlateBrush* SMetaSoundGraphNode::GetShadowBrush(bool bSelected) const
+		{
+			if (IsVariableAccessor() || IsVariableMutator())
+			{
+				return bSelected ? FEditorStyle::GetBrush(TEXT("Graph.VarNode.ShadowSelected")) : FEditorStyle::GetBrush(TEXT("Graph.VarNode.Shadow"));
+			}
+
+			return SGraphNode::GetShadowBrush(bSelected);
+		}
+
 		void SMetaSoundGraphNode::Construct(const FArguments& InArgs, class UEdGraphNode* InNode)
 		{
 			GraphNode = InNode;
@@ -381,6 +404,24 @@ namespace Metasound
 			return TitleBoxWidget.ToSharedRef();
 		}
 
+		void SMetaSoundGraphNode::GetOverlayBrushes(bool bSelected, const FVector2D WidgetSize, TArray<FOverlayBrushInfo>& Brushes) const
+		{
+			FName CornerIcon = GetMetaSoundNode().GetCornerIcon();
+			if (CornerIcon != NAME_None)
+			{
+
+				if (const FSlateBrush* Brush = FEditorStyle::GetBrush(CornerIcon))
+				{
+					FOverlayBrushInfo OverlayInfo = { Brush };
+
+					// Logic copied from SGraphNodeK2Base
+					OverlayInfo.OverlayOffset.X = (WidgetSize.X - (OverlayInfo.Brush->ImageSize.X / 2.f)) - 3.f;
+					OverlayInfo.OverlayOffset.Y = (OverlayInfo.Brush->ImageSize.Y / -2.f) + 2.f;
+					Brushes.Add(MoveTemp(OverlayInfo));
+				}
+			}
+		}
+
 		FLinearColor SMetaSoundGraphNode::GetNodeTitleColorOverride() const
 		{
 			FLinearColor ReturnTitleColor = GraphNode->IsDeprecated() ? FLinearColor::Red : GetNodeObj()->GetNodeTitleColor();
@@ -475,21 +516,29 @@ namespace Metasound
 
 		const FSlateBrush* SMetaSoundGraphNode::GetNodeBodyBrush() const
 		{
-		// 	if (const ISlateStyle* MetasoundStyle = FSlateStyleRegistry::FindSlateStyle("MetaSoundStyle"))
-		// 	{
-		// 		if (GetMetaSoundNode().GetNodeHandle()->GetClassType() == EMetasoundFrontendClassType::Input)
-		// 		{
-		// 			if (const FSlateBrush* InputBrush = MetasoundStyle->GetBrush("MetasoundEditor.Graph.Node.Body.Input"))
-		// 			{
-		// 				return InputBrush;
-		// 			}
-		// 		}
-		// 
-		// 		if (const FSlateBrush* DefaultBrush = MetasoundStyle->GetBrush("MetasoundEditor.Graph.Node.Body.Default"))
-		// 		{
-		// 			return DefaultBrush;
-		// 		}
-		// 	}
+			// TODO: Add tweak & add custom bodies
+			if (GraphNode)
+			{
+				const EMetasoundFrontendClassType ClassType = GetMetaSoundNode().GetNodeHandle()->GetClassMetadata().GetType();
+				switch (ClassType)
+				{
+					case EMetasoundFrontendClassType::Variable:
+					case EMetasoundFrontendClassType::VariableAccessor:
+					case EMetasoundFrontendClassType::VariableDeferredAccessor:
+					case EMetasoundFrontendClassType::VariableMutator:
+					{
+						return FEditorStyle::GetBrush("Graph.VarNode.Body");
+					}
+					break;
+
+					case EMetasoundFrontendClassType::Input:
+					case EMetasoundFrontendClassType::Output:
+					default:
+					{
+					}
+					break;
+				}
+			}
 
 			return FEditorStyle::GetBrush("Graph.Node.Body");
 		}
@@ -786,46 +835,46 @@ namespace Metasound
 				}
 			}
 	
-			if (StyleDisplay.ImageName.IsNone())
-			{
-				ContentBox->AddSlot()
-					.HAlign(HAlign_Left)
-					.VAlign(VAlign_Top)
-					.AutoWidth()
-					[
-						SAssignNew(LeftNodeBox, SVerticalBox)
-					];
-			}
-			else
-			{
-				ContentBox->AddSlot()
-					.HAlign(HAlign_Left)
-					.VAlign(VAlign_Top)
-					.AutoWidth()
-					[
-						SAssignNew(LeftNodeBox, SVerticalBox)
-					];
+			static const float GrabPadding = 28.0f;
 
+			// Gives more space for user to grab a bit easier as variables do not have any title area nor icon
+			const float LeftNodeGrabPadding = IsVariableMutator() ? GrabPadding : 0.0f;
+			ContentBox->AddSlot()
+			.HAlign(HAlign_Left)
+			.VAlign(VAlign_Top)
+			.AutoWidth()
+			.Padding(0.0f, 0.0f, LeftNodeGrabPadding, 0.0f)
+			[
+				SAssignNew(LeftNodeBox, SVerticalBox)
+			];
+
+			if (!StyleDisplay.ImageName.IsNone())
+			{
 				if (const ISlateStyle* MetasoundStyle = FSlateStyleRegistry::FindSlateStyle("MetaSoundStyle"))
 				{
-					const FSlateBrush* ImageBrush = MetasoundStyle->GetBrush(StyleDisplay.ImageName);
-					ContentBox->AddSlot()
-					.AutoWidth()
-					.HAlign(HAlign_Center)
-					.VAlign(VAlign_Center)
-					[
-						SNew(SImage)
-						.Image(ImageBrush)
-						.ColorAndOpacity(FSlateColor::UseForeground())
-						.DesiredSizeOverride(FVector2D(20, 20))
-					];
+					if (const FSlateBrush* ImageBrush = MetasoundStyle->GetBrush(StyleDisplay.ImageName))
+					{
+						ContentBox->AddSlot()
+						.AutoWidth()
+						.HAlign(HAlign_Center)
+						.VAlign(VAlign_Center)
+						[
+							SNew(SImage)
+							.Image(ImageBrush)
+							.ColorAndOpacity(FSlateColor::UseForeground())
+							.DesiredSizeOverride(FVector2D(20, 20))
+						];
+					}
 				}
 			}
 
+			// Gives more space for user to grab a bit easier as variables do not have any title area nor icon
+			const float RightNodeGrabPadding = IsVariableAccessor() ? GrabPadding : 0.0f;
 			ContentBox->AddSlot()
 				.AutoWidth()
 				.HAlign(HAlign_Right)
 				.VAlign(VAlign_Center)
+				.Padding(RightNodeGrabPadding, 0.0f, 0.0f, 0.0f)
 				[
 					SAssignNew(RightNodeBox, SVerticalBox)
 				];
diff --git a/Engine/Plugins/Runtime/Metasound/Source/MetasoundEditor/Private/SMetasoundGraphNode.h b/Engine/Plugins/Runtime/Metasound/Source/MetasoundEditor/Private/SMetasoundGraphNode.h
--- a/Engine/Plugins/Runtime/Metasound/Source/MetasoundEditor/Private/SMetasoundGraphNode.h
+++ b/Engine/Plugins/Runtime/Metasound/Source/MetasoundEditor/Private/SMetasoundGraphNode.h
@@ -37,8 +37,11 @@ namespace Metasound
 			void Construct(const FArguments& InArgs, class UEdGraphNode* InNode);
 
 		protected:
+			bool IsVariableAccessor() const;
+			bool IsVariableMutator() const;
 
 			// SGraphNode Interface
+			virtual const FSlateBrush* GetShadowBrush(bool bSelected) const override;
 			virtual void CreateInputSideAddButton(TSharedPtr<SVerticalBox> InputBox) override;
 			virtual void CreateOutputSideAddButton(TSharedPtr<SVerticalBox> OutputBox) override;
 			virtual TSharedPtr<SGraphPin> CreatePinWidget(UEdGraphPin* InPin) const override;
@@ -51,6 +54,7 @@ namespace Metasound
 			virtual FReply OnAddPin() override;
 			virtual void MoveTo(const FVector2D& NewPosition, FNodeSet& NodeFilter, bool bMarkDirty = true) override;
 			virtual void SetDefaultTitleAreaWidget(TSharedRef<SOverlay> DefaultTitleAreaWidget) override;
+			virtual void GetOverlayBrushes(bool bSelected, const FVector2D WidgetSize, TArray<FOverlayBrushInfo>& Brushes) const override;
 
 			FLinearColor GetNodeTitleColorOverride() const;
 
diff --git a/Engine/Plugins/Runtime/Metasound/Source/MetasoundEditor/Public/MetasoundEditorGraphNode.h b/Engine/Plugins/Runtime/Metasound/Source/MetasoundEditor/Public/MetasoundEditorGraphNode.h
--- a/Engine/Plugins/Runtime/Metasound/Source/MetasoundEditor/Public/MetasoundEditorGraphNode.h
+++ b/Engine/Plugins/Runtime/Metasound/Source/MetasoundEditor/Public/MetasoundEditorGraphNode.h
@@ -76,7 +76,7 @@ public:
 	// End of UObject interface
 
 	virtual FSlateIcon GetNodeTitleIcon() const { return FSlateIcon(); }
-
+	virtual FName GetCornerIcon() const { return FName(); }
 	virtual bool CanAddInputPin() const
 	{
 		return false;
@@ -249,13 +249,18 @@ public:
 	UPROPERTY()
 	UMetasoundEditorGraphVariable* Variable;
 
+	// Variables do not have titles to distinguish more visually from vertex types
+	virtual void CacheTitle() override { }
+
 	virtual UMetasoundEditorGraphMember* GetMember() const override;
 	virtual bool EnableInteractWidgets() const override;
 	virtual FMetasoundFrontendClassName GetClassName() const override;
 	virtual FGuid GetNodeID() const override;
-	
+	virtual FName GetCornerIcon() const override;
 	virtual void PinDefaultValueChanged(UEdGraphPin* Pin) override;
 
+	virtual EMetasoundFrontendClassType GetClassType() const;
+
 protected:
 	virtual FLinearColor GetNodeTitleColor() const override;
 	virtual FSlateIcon GetNodeTitleIcon() const override;
diff --git a/Engine/Plugins/Runtime/Metasound/Source/MetasoundGraphCore/Public/MetasoundVariableNodes.h b/Engine/Plugins/Runtime/Metasound/Source/MetasoundGraphCore/Public/MetasoundVariableNodes.h
--- a/Engine/Plugins/Runtime/Metasound/Source/MetasoundGraphCore/Public/MetasoundVariableNodes.h
+++ b/Engine/Plugins/Runtime/Metasound/Source/MetasoundGraphCore/Public/MetasoundVariableNodes.h
@@ -181,7 +181,9 @@ namespace Metasound
 			Info.ClassName = VariableNames::GetVariableNodeClassName<DataType>();
 			Info.MajorVersion = 1;
 			Info.MinorVersion = 0;
+
 #if WITH_EDITOR
+			Info.DisplayStyle.bShowName = false;
 			Info.Description = LOCTEXT("Metasound_InitVariableNodeDescription", "Initialize a variable of a MetaSound graph.");
 			Info.CategoryHierarchy = { LOCTEXT("Metasound_VariableCategory", "Variable") };
 #endif // WITH_EDITOR
@@ -331,6 +333,7 @@ namespace Metasound
 			Info.MinorVersion = 0;
 
 #if WITH_EDITOR
+			Info.DisplayStyle.bShowName = false;
 			Info.DisplayName = LOCTEXT("Metasound_VariableMutatorNodeDisplayName", "Set");
 			Info.Description = LOCTEXT("Metasound_VariableMutatorNodeDescription", "Set variable on MetaSound graph.");
 			Info.CategoryHierarchy = { LOCTEXT("Metasound_VariableCategory", "Variable") };
@@ -441,6 +444,7 @@ namespace Metasound
 			Info.MajorVersion = 1;
 			Info.MinorVersion = 0;
 #if WITH_EDITOR
+			Info.DisplayStyle.bShowName = false;
 			Info.DisplayName = LOCTEXT("Metasound_VariableDeferredAccessorNodeDisplayName", "Get Delayed");
 			Info.Description = LOCTEXT("Metasound_VariableDeferredAccessorNodeDescription", "Get a delayed variable on MetaSound graph.");
 			Info.CategoryHierarchy = { LOCTEXT("Metasound_VariableCategory", "Variable") };
@@ -550,6 +554,7 @@ namespace Metasound
 			Info.MajorVersion = 1;
 			Info.MinorVersion = 0;
 #if WITH_EDITOR
+			Info.DisplayStyle.bShowName = false;
 			Info.DisplayName = LOCTEXT("Metasound_VariableAccessorNodeDisplayName", "Get");
 			Info.Description = LOCTEXT("Metasound_VariableAccessorNodeDescription", "Get variable on MetaSound graph.");
 #endif // WITH_EDITOR
