Fix for generating dynamic LiteralNodes that provide default values for vertices that are already connected with explicitly defined edges to the same destination vertex.
#rnx

diff --git a/Engine/Plugins/Runtime/Metasound/Source/MetasoundFrontend/Private/MetasoundFrontendGraph.cpp b/Engine/Plugins/Runtime/Metasound/Source/MetasoundFrontend/Private/MetasoundFrontendGraph.cpp
--- a/Engine/Plugins/Runtime/Metasound/Source/MetasoundFrontend/Private/MetasoundFrontendGraph.cpp
+++ b/Engine/Plugins/Runtime/Metasound/Source/MetasoundFrontend/Private/MetasoundFrontendGraph.cpp
@@ -237,7 +237,7 @@ namespace Metasound
 		return TUniquePtr<INode>(nullptr);
 	}
 
-	TUniquePtr<INode> FFrontendGraphBuilder::CreateOutputNode(const FMetasoundFrontendNode& InNode, const FMetasoundFrontendClass& InClass, FBuildGraphContext& InGraphContext)
+	TUniquePtr<INode> FFrontendGraphBuilder::CreateOutputNode(const FMetasoundFrontendNode& InNode, const FMetasoundFrontendClass& InClass, FBuildGraphContext& InGraphContext, const TSet<FGuid>& InEdgeDestinations)
 	{
 		using namespace Metasound::Frontend;
 
@@ -257,7 +257,7 @@ namespace Metasound
 
 			{
 				const FNodeInitData InitData = FrontendGraphPrivate::CreateNodeInitData(InNode);
-				TArray<FDefaultLiteralData> DefaultLiteralData = GetInputDefaultLiteralData(InNode, InitData);
+				TArray<FDefaultLiteralData> DefaultLiteralData = GetInputDefaultLiteralData(InNode, InitData, InEdgeDestinations);
 				for (FDefaultLiteralData& Data : DefaultLiteralData)
 				{
 					InGraphContext.DefaultInputs.Emplace(FNodeIDVertexID { InNode.GetID(), Data.DestinationVertexID }, MoveTemp(Data));
@@ -269,16 +269,16 @@ namespace Metasound
 		return TUniquePtr<INode>(nullptr);
 	}
 
-	TUniquePtr<INode> FFrontendGraphBuilder::CreateExternalNode(const FMetasoundFrontendNode& InNode, const FMetasoundFrontendClass& InClass, FBuildGraphContext& InGraphContext)
+	TUniquePtr<INode> FFrontendGraphBuilder::CreateExternalNode(const FMetasoundFrontendNode& InNode, const FMetasoundFrontendClass& InClass, FBuildGraphContext& InGraphContext, const TSet<FGuid>& InEdgeDestinations)
 	{
 		check(InNode.ClassID == InClass.ID);
 
 		const FNodeInitData InitData = FrontendGraphPrivate::CreateNodeInitData(InNode);
 		{
-			TArray<FDefaultLiteralData> DefaultLiteralData = GetInputDefaultLiteralData(InNode, InitData);
+			TArray<FDefaultLiteralData> DefaultLiteralData = GetInputDefaultLiteralData(InNode, InitData, InEdgeDestinations);
 			for (FDefaultLiteralData& Data : DefaultLiteralData)
 			{
-				InGraphContext.DefaultInputs.Emplace(FNodeIDVertexID { InNode.GetID(), Data.DestinationVertexID }, MoveTemp(Data));
+				InGraphContext.DefaultInputs.Emplace(FNodeIDVertexID{ InNode.GetID(), Data.DestinationVertexID }, MoveTemp(Data));
 			}
 		}
 
@@ -397,6 +397,13 @@ namespace Metasound
 	// TODO: add errors here. Most will be a "PromptIfMissing"...
 	void FFrontendGraphBuilder::AddNodesToGraph(FBuildGraphContext& InGraphContext)
 	{
+		TSet<FGuid> GraphEdgeDestinations;
+		const TArray<FMetasoundFrontendEdge>& GraphEdges = InGraphContext.GraphClass.Graph.Edges;
+		Algo::Transform(GraphEdges, GraphEdgeDestinations, [](const FMetasoundFrontendEdge& Edge) 
+		{
+			return Edge.ToVertexID;
+		});
+
 		for (const FMetasoundFrontendNode& Node : InGraphContext.GraphClass.Graph.Nodes)
 		{
 			const FMetasoundFrontendClass* NodeClass = InGraphContext.BuildContext.FrontendClasses.FindRef(Node.ClassID);
@@ -429,7 +436,7 @@ namespace Metasound
 						const FMetasoundFrontendClassOutput* ClassOutput = FindClassOutputForOutputNode(InGraphContext.GraphClass, Node, OutputIndex);
 						if ((nullptr != ClassOutput) && (INDEX_NONE != OutputIndex))
 						{
-							TSharedPtr<const INode> OutputNode(CreateOutputNode(Node, *NodeClass, InGraphContext).Release());
+							TSharedPtr<const INode> OutputNode(CreateOutputNode(Node, *NodeClass, InGraphContext, GraphEdgeDestinations).Release());
 							InGraphContext.Graph->AddOutputNode(Node.GetID(), OutputIndex, ClassOutput->Name, OutputNode);
 						}
 						else
@@ -475,7 +482,7 @@ namespace Metasound
 					case EMetasoundFrontendClassType::External:
 					default:
 					{
-						TSharedPtr<const INode> ExternalNode(CreateExternalNode(Node, *NodeClass, InGraphContext).Release());
+						TSharedPtr<const INode> ExternalNode(CreateExternalNode(Node, *NodeClass, InGraphContext, GraphEdgeDestinations).Release());
 						InGraphContext.Graph->AddNode(Node.GetID(), ExternalNode);
 					}
 					break;
@@ -617,7 +624,7 @@ namespace Metasound
 		InGraphContext.DefaultInputs.Reset();
 	}
 
-	TArray<FFrontendGraphBuilder::FDefaultLiteralData> FFrontendGraphBuilder::GetInputDefaultLiteralData(const FMetasoundFrontendNode& InNode, const FNodeInitData& InInitData)
+	TArray<FFrontendGraphBuilder::FDefaultLiteralData> FFrontendGraphBuilder::GetInputDefaultLiteralData(const FMetasoundFrontendNode& InNode, const FNodeInitData& InInitData, const TSet<FGuid>& InEdgeDestinations)
 	{
 		TArray<FDefaultLiteralData> DefaultLiteralData;
 
@@ -629,16 +636,23 @@ namespace Metasound
 
 			FGuid VertexID = Literal.VertexID;
 			FVertexName DestinationVertexName;
+			bool bRequiresDefault = false;
 			auto RemoveAndBuildParams = [&](const FMetasoundFrontendVertex& Vertex)
 			{
 				if (Vertex.VertexID == VertexID)
 				{
-					InitParams.Literal = Literal.Value.ToLiteral(Vertex.TypeName);
-					InitParams.InstanceID = FGuid::NewGuid();
-					InitParams.NodeName = "Literal";
-					TypeName = Vertex.TypeName;
+					// Only build params and forward along to connect to dynamically generated
+					// literal node if edge is not explicitly connected to vertex as destination.
+					bRequiresDefault = !InEdgeDestinations.Contains(Vertex.VertexID);
+					if (bRequiresDefault)
+					{
+						InitParams.Literal = Literal.Value.ToLiteral(Vertex.TypeName);
+						InitParams.InstanceID = FGuid::NewGuid();
+						InitParams.NodeName = "Literal";
+						TypeName = Vertex.TypeName;
 
-					DestinationVertexName = Vertex.Name;
+						DestinationVertexName = Vertex.Name;
+					}
 
 					return true;
 				}
@@ -649,14 +663,17 @@ namespace Metasound
 			const bool bRemoved = InputVertices.RemoveAllSwap(RemoveAndBuildParams, false /* bAllowShrinking */) > 0;
 			if (ensure(bRemoved))
 			{
-				DefaultLiteralData.Emplace(FDefaultLiteralData
+				if (bRequiresDefault)
 				{
-					InNode.GetID(),
-					VertexID,
-					DestinationVertexName,
-					TypeName,
-					MoveTemp(InitParams)
-				});
+					DefaultLiteralData.Emplace(FDefaultLiteralData
+						{
+							InNode.GetID(),
+							VertexID,
+							DestinationVertexName,
+							TypeName,
+							MoveTemp(InitParams)
+						});
+				}
 			}
 		}
 
diff --git a/Engine/Plugins/Runtime/Metasound/Source/MetasoundFrontend/Public/MetasoundFrontendGraph.h b/Engine/Plugins/Runtime/Metasound/Source/MetasoundFrontend/Public/MetasoundFrontendGraph.h
--- a/Engine/Plugins/Runtime/Metasound/Source/MetasoundFrontend/Public/MetasoundFrontendGraph.h
+++ b/Engine/Plugins/Runtime/Metasound/Source/MetasoundFrontend/Public/MetasoundFrontendGraph.h
@@ -144,11 +144,11 @@ namespace Metasound
 			TUniquePtr<FFrontendGraph> Graph;
 			const FMetasoundFrontendGraphClass& GraphClass;
 			FBuildContext& BuildContext;
-
+			TSet<FGuid> DestinationsWithEdges;
 			FDefaultInputByIDMap DefaultInputs;
 		};
 
-		static TArray<FDefaultLiteralData> GetInputDefaultLiteralData(const FMetasoundFrontendNode& InNode, const FNodeInitData& InInitData);
+		static TArray<FDefaultLiteralData> GetInputDefaultLiteralData(const FMetasoundFrontendNode& InNode, const FNodeInitData& InInitData, const TSet<FGuid>& InEdgeDestinations);
 
 		static bool SortSubgraphDependencies(TArray<const FMetasoundFrontendGraphClass*>& Subgraphs);
 
@@ -160,9 +160,9 @@ namespace Metasound
 		static const FMetasoundFrontendVariable* FindVariableForVariableNode(const FMetasoundFrontendNode& InVariableNode, const FMetasoundFrontendGraph& InGraph);
 
 		static TUniquePtr<INode> CreateInputNode(const FMetasoundFrontendNode& InNode, const FMetasoundFrontendClass& InClass, const FMetasoundFrontendClassInput& InOwningGraphClassInput);
-		static TUniquePtr<INode> CreateOutputNode(const FMetasoundFrontendNode& InNode, const FMetasoundFrontendClass& InClass, FBuildGraphContext& InGraphContext);
+		static TUniquePtr<INode> CreateOutputNode(const FMetasoundFrontendNode& InNode, const FMetasoundFrontendClass& InClass, FBuildGraphContext& InGraphContext, const TSet<FGuid>& InEdgeDestinations);
 		static TUniquePtr<INode> CreateVariableNode(const FMetasoundFrontendNode& InNode, const FMetasoundFrontendGraph& InGraph);
-		static TUniquePtr<INode> CreateExternalNode(const FMetasoundFrontendNode& InNode, const FMetasoundFrontendClass& InClass, FBuildGraphContext& InGraphContext);
+		static TUniquePtr<INode> CreateExternalNode(const FMetasoundFrontendNode& InNode, const FMetasoundFrontendClass& InClass, FBuildGraphContext& InGraphContext, const TSet<FGuid>& InEdgeDestinations);
 
 		// TODO: add errors here. Most will be a "PromptIfMissing"...
 		static void AddNodesToGraph(FBuildGraphContext& InGraphContext);
