Fix for incorrect reset behavior on triggercounter node

diff --git a/Engine/Plugins/Runtime/Metasound/Source/MetasoundStandardNodes/Private/MetasoundTriggerCounterNode.cpp b/Engine/Plugins/Runtime/Metasound/Source/MetasoundStandardNodes/Private/MetasoundTriggerCounterNode.cpp
--- a/Engine/Plugins/Runtime/Metasound/Source/MetasoundStandardNodes/Private/MetasoundTriggerCounterNode.cpp
+++ b/Engine/Plugins/Runtime/Metasound/Source/MetasoundStandardNodes/Private/MetasoundTriggerCounterNode.cpp
@@ -67,6 +67,7 @@ namespace Metasound
 			int32 CurrentAutoResetCount = 0;
 			int32 CurrentTriggerCount = 0;
 			float CurrentValue = 0.0f;
+			bool bIsFirstTrigger = true;
 	};
 
 	FTriggerCounterOperator::FTriggerCounterOperator(const FOperatorSettings& InSettings, 
@@ -129,7 +130,7 @@ namespace Metasound
 			},
 			[this](int32 StartFrame, int32 EndFrame)
 			{
-				CurrentTriggerCount = 0;
+				*OutCount = CurrentTriggerCount = 0;
 				*OutValue = *StartValue;
 
 				TriggerOnReset->TriggerFrame(StartFrame);
@@ -144,7 +145,13 @@ namespace Metasound
 			{
 				++CurrentTriggerCount;
 
-				if (CurrentAutoResetCount > 0 && CurrentTriggerCount > CurrentAutoResetCount)
+				if (bIsFirstTrigger)
+				{
+					CurrentTriggerCount = 1;
+					*OutValue = *StartValue;
+					bIsFirstTrigger = false;
+				}
+				else if (CurrentAutoResetCount > 0 && CurrentTriggerCount > CurrentAutoResetCount)
 				{
 					CurrentTriggerCount = 1;
 					*OutValue = *StartValue;
