Fixing one-shot interface  issues

diff --git a/Engine/Plugins/Runtime/Metasound/Source/MetasoundEditor/Private/MetasoundEditorGraph.cpp b/Engine/Plugins/Runtime/Metasound/Source/MetasoundEditor/Private/MetasoundEditorGraph.cpp
--- a/Engine/Plugins/Runtime/Metasound/Source/MetasoundEditor/Private/MetasoundEditorGraph.cpp
+++ b/Engine/Plugins/Runtime/Metasound/Source/MetasoundEditor/Private/MetasoundEditorGraph.cpp
@@ -1588,6 +1588,9 @@ bool UMetasoundEditorGraph::ValidateInternal(Metasound::Editor::FGraphValidation
 			bMarkDirty |= NodeResult.bIsDirty;
 		}
 
+		bIsValid &= Node->Validate(NodeResult);
+		bMarkDirty |= NodeResult.bIsDirty;
+
 		OutResults.NodeResults.Add(NodeResult);
 	}
 
diff --git a/Engine/Plugins/Runtime/Metasound/Source/MetasoundEditor/Private/MetasoundEditorGraphInputNode.cpp b/Engine/Plugins/Runtime/Metasound/Source/MetasoundEditor/Private/MetasoundEditorGraphInputNode.cpp
--- a/Engine/Plugins/Runtime/Metasound/Source/MetasoundEditor/Private/MetasoundEditorGraphInputNode.cpp
+++ b/Engine/Plugins/Runtime/Metasound/Source/MetasoundEditor/Private/MetasoundEditorGraphInputNode.cpp
@@ -15,6 +15,7 @@
 #include "MetasoundFrontendDocument.h"
 #include "MetasoundFrontendLiteral.h"
 #include "MetasoundFrontendRegistries.h"
+#include "MetasoundFrontendSearchEngine.h"
 #include "MetasoundPrimitives.h"
 #include "Misc/Guid.h"
 #include "Sound/SoundWave.h"
@@ -73,6 +74,43 @@ FSlateIcon UMetasoundEditorGraphInputNode::GetNodeTitleIcon() const
 	return FSlateIcon("MetaSoundStyle", NativeIconName);
 }
 
+bool UMetasoundEditorGraphInputNode::Validate(Metasound::Editor::FGraphNodeValidationResult& OutResult)
+{
+#if WITH_EDITOR
+	using namespace Metasound::Editor;
+	using namespace Metasound::Frontend;
+
+	OutResult = CreateNewValidationResult();
+
+	FNodeHandle NodeHandle = GetNodeHandle();
+	const FMetasoundFrontendClassMetadata& Metadata = NodeHandle->GetClassMetadata();
+
+	const FMetasoundFrontendVersion& MetasoundFrontendVersion = NodeHandle->GetInterfaceVersion();
+
+	FName InterfaceNameToValidate = MetasoundFrontendVersion.Name;
+	FMetasoundFrontendInterface InterfaceToValidate;
+	if (ISearchEngine::Get().FindInterfaceWithHighestVersion(InterfaceNameToValidate, InterfaceToValidate))
+	{
+		const FName& NodeName = NodeHandle->GetNodeName();
+		FText RequiredText;
+		if (InterfaceToValidate.IsMemberInputRequired(NodeName, RequiredText))
+		{
+			TArray<FConstInputHandle> InputHandles = NodeHandle->GetConstInputs();
+			if (ensure(!InputHandles.IsEmpty()))
+			{
+				const FConstInputHandle& InputHandle = InputHandles.Last();
+				if (!InputHandle->IsConnected())
+				{
+					GraphNode::SetMessage(*this, EMessageSeverity::Warning, *RequiredText.ToString());
+					return false;
+				}
+			}
+		}
+	}
+#endif // #if WITH_EDITOR
+	return true;
+}
+
 void UMetasoundEditorGraphInputNode::SetNodeID(FGuid InNodeID)
 {
 	Input->NodeID = InNodeID;
diff --git a/Engine/Plugins/Runtime/Metasound/Source/MetasoundEditor/Private/MetasoundEditorGraphNode.cpp b/Engine/Plugins/Runtime/Metasound/Source/MetasoundEditor/Private/MetasoundEditorGraphNode.cpp
--- a/Engine/Plugins/Runtime/Metasound/Source/MetasoundEditor/Private/MetasoundEditorGraphNode.cpp
+++ b/Engine/Plugins/Runtime/Metasound/Source/MetasoundEditor/Private/MetasoundEditorGraphNode.cpp
@@ -29,29 +29,45 @@
 
 #define LOCTEXT_NAMESPACE "MetaSoundEditor"
 
-
 namespace Metasound
 {
 	namespace Editor
 	{
-		namespace GraphNodePrivate
+		void GraphNode::SetMessage(UEdGraphNode& InNode, EMessageSeverity::Type InSeverity, const FString& InMessage)
 		{
-			void SetGraphNodeMessage(UEdGraphNode& InNode, EMessageSeverity::Type InSeverity, const FString& InMessage)
-			{
-				InNode.bHasCompilerMessage = true;
-				InNode.ErrorMsg = InMessage;
-				InNode.ErrorType = InSeverity;
+			InNode.bHasCompilerMessage = true;
+			InNode.ErrorMsg = InMessage;
+			InNode.ErrorType = InSeverity;
 
-				if (InSeverity == EMessageSeverity::Error)
-				{
-					UE_LOG(LogMetasoundEditor, Error, TEXT("%s"), *InMessage);
-				}
+			if (InSeverity == EMessageSeverity::Error)
+			{
+				UE_LOG(LogMetasoundEditor, Error, TEXT("%s"), *InMessage);
 			}
-		} // namespace GraphNodePrivate
-	} // namespace Editor
-} // namespace Metasound
+		}
+	}
+}
+
+Metasound::Editor::FGraphNodeValidationResult UMetasoundEditorGraphMemberNode::CreateNewValidationResult()
+{
+	using namespace Metasound::Editor;
+
+	FGraphNodeValidationResult OutResult = FGraphNodeValidationResult(*this);
 
+	// Reset the node validation state
+	if (ErrorType != EMessageSeverity::Info)
+	{
+		ErrorType = EMessageSeverity::Info;
+		OutResult.bIsDirty = true;
+	}
 
+	if (bHasCompilerMessage)
+	{
+		bHasCompilerMessage = false;
+		ErrorMsg.Reset();
+		OutResult.bIsDirty = true;
+	}
+	return OutResult;
+}
 UMetasoundEditorGraphNode::UMetasoundEditorGraphNode(const FObjectInitializer& ObjectInitializer)
 	: Super(ObjectInitializer)
 {
@@ -192,6 +208,13 @@ void UMetasoundEditorGraphNode::CacheTitle()
 	CachedTitle = NodeHandle->GetDisplayTitle();
 }
 
+bool UMetasoundEditorGraphNode::Validate(Metasound::Editor::FGraphNodeValidationResult& OutResult)
+{
+	// Do no validation by default
+	OutResult = Metasound::Editor::FGraphNodeValidationResult(*this);
+	return true;
+}
+
 bool UMetasoundEditorGraphNode::ContainsMetadataChange() const
 {
 	using namespace Metasound::Frontend;
@@ -580,6 +603,43 @@ bool UMetasoundEditorGraphOutputNode::EnableInteractWidgets() const
 	return bEnabled;
 }
 
+bool UMetasoundEditorGraphOutputNode::Validate(Metasound::Editor::FGraphNodeValidationResult& OutResult) 
+{
+#if WITH_EDITOR
+	using namespace Metasound::Editor;
+	using namespace Metasound::Frontend;
+
+	OutResult = CreateNewValidationResult();
+
+	// 2. Check if node is invalid, version is missing and cache if interface changes exist between the document's records and the registry
+	FNodeHandle NodeHandle = GetNodeHandle();
+	const FMetasoundFrontendVersion& MetasoundFrontendVersion = NodeHandle->GetInterfaceVersion();
+
+	FName InterfaceNameToValidate = MetasoundFrontendVersion.Name;
+	FMetasoundFrontendInterface InterfaceToValidate;
+	if (ISearchEngine::Get().FindInterfaceWithHighestVersion(InterfaceNameToValidate, InterfaceToValidate))
+	{
+		const FName& NodeName = NodeHandle->GetNodeName();
+		FText RequiredText;
+		if (InterfaceToValidate.IsMemberOutputRequired(NodeName, RequiredText))
+		{
+			TArray<FConstInputHandle> InputHandles = NodeHandle->GetConstInputs();
+			if (ensure(!InputHandles.IsEmpty()))
+			{
+				bool bIsConnected = InputHandles.Last()->IsConnected();
+				if (!bIsConnected)
+				{
+					GraphNode::SetMessage(*this, EMessageSeverity::Warning, *RequiredText.ToString());
+					return false;
+				}
+			}
+		}
+	}
+#endif // #if WITH_EDITOR
+
+	return true;
+}
+
 FMetasoundFrontendClassName UMetasoundEditorGraphOutputNode::GetClassName() const
 {
 	if (ensure(Output))
@@ -770,19 +830,19 @@ bool UMetasoundEditorGraphExternalNode::Validate(Metasound::Editor::FGraphNodeVa
 				*PromptIfMissing->ToString()
 			});
 
-			GraphNodePrivate::SetGraphNodeMessage(*this, EMessageSeverity::Error, NewErrorMsg);
+			GraphNode::SetMessage(*this, EMessageSeverity::Error, NewErrorMsg);
 		}
 		else
 		{
 			if (bIsClassNative)
 			{
-				GraphNodePrivate::SetGraphNodeMessage(*this, EMessageSeverity::Error, FString::Format(
+				GraphNode::SetMessage(*this, EMessageSeverity::Error, FString::Format(
 					TEXT("Class '{0}' definition missing for last known natively defined node."),
 					{ *ClassName.ToString() }));
 			}
 			else
 			{
-				GraphNodePrivate::SetGraphNodeMessage(*this, EMessageSeverity::Error,
+				GraphNode::SetMessage(*this, EMessageSeverity::Error,
 					FString::Format(TEXT("Class definition missing for asset with guid '{0}': Asset is either missing or invalid"),
 					{ *ClassName.Name.ToString() }));
 			}
@@ -853,7 +913,7 @@ bool UMetasoundEditorGraphExternalNode::Validate(Metasound::Editor::FGraphNodeVa
 		if (RegisteredClass.Metadata.GetIsDeprecated())
 		{
 			constexpr bool bIncludeNamespace = true;
-			GraphNodePrivate::SetGraphNodeMessage(*this, EMessageSeverity::Warning,
+			GraphNode::SetMessage(*this, EMessageSeverity::Warning,
 				FString::Format(TEXT("Class '{0} {1}' is deprecated."),
 				{
 					*FGraphBuilder::GetDisplayName(RegisteredClass.Metadata, { }, bIncludeNamespace).ToString(),
@@ -867,7 +927,7 @@ bool UMetasoundEditorGraphExternalNode::Validate(Metasound::Editor::FGraphNodeVa
 	const TArray<FMetasoundFrontendClass> SortedClasses = ISearchEngine::Get().FindClassesWithName(NodeClassName, true /* bSortByVersion */);
 	if (SortedClasses.IsEmpty())
 	{
-		GraphNodePrivate::SetGraphNodeMessage(*this, EMessageSeverity::Error,
+		GraphNode::SetMessage(*this, EMessageSeverity::Error,
 			FString::Format(TEXT("Class '{0} {1}' not registered."),
 			{
 				*Metadata.GetClassName().ToString(),
@@ -912,14 +972,14 @@ bool UMetasoundEditorGraphExternalNode::Validate(Metasound::Editor::FGraphNodeVa
 				OutResult.bIsInvalid = true;
 			}
 
-			GraphNodePrivate::SetGraphNodeMessage(*this, Severity, NodeMsg);
+			GraphNode::SetMessage(*this, Severity, NodeMsg);
 			OutResult.bIsDirty = true;
 		}
 		else if (HighestRegistryClass.Metadata.GetVersion() == CurrentVersion)
 		{
 			if (InterfaceUpdates.ContainsChanges())
 			{
-				GraphNodePrivate::SetGraphNodeMessage(*this, EMessageSeverity::Error,
+				GraphNode::SetMessage(*this, EMessageSeverity::Error,
 				FString::Format(TEXT("Node & registered class interface mismatch: '{0} {1}'. Class either versioned improperly, class key collision exists, or AutoUpdate disabled in 'MetaSound' Developer Settings."),
 				{
 					*Metadata.GetClassName().ToString(),
@@ -931,7 +991,7 @@ bool UMetasoundEditorGraphExternalNode::Validate(Metasound::Editor::FGraphNodeVa
 		}
 		else
 		{
-			GraphNodePrivate::SetGraphNodeMessage(*this, EMessageSeverity::Error,
+			GraphNode::SetMessage(*this, EMessageSeverity::Error,
 				FString::Format(TEXT("Node with class '{0} {1}' interface version higher than that of highest minor revision ({2}) in class registry."),
 				{
 					*Metadata.GetClassName().ToString(),
diff --git a/Engine/Plugins/Runtime/Metasound/Source/MetasoundEditor/Public/MetasoundEditorGraph.h b/Engine/Plugins/Runtime/Metasound/Source/MetasoundEditor/Public/MetasoundEditorGraph.h
--- a/Engine/Plugins/Runtime/Metasound/Source/MetasoundEditor/Public/MetasoundEditorGraph.h
+++ b/Engine/Plugins/Runtime/Metasound/Source/MetasoundEditor/Public/MetasoundEditorGraph.h
@@ -3,6 +3,7 @@
 
 #include "AudioParameterControllerInterface.h"
 #include "MetasoundEditor.h"
+#include "MetasoundEditorGraphBuilder.h"
 #include "MetasoundEditorGraphNode.h"
 #include "MetasoundFrontendController.h"
 #include "MetasoundFrontendDocument.h"
diff --git a/Engine/Plugins/Runtime/Metasound/Source/MetasoundEditor/Public/MetasoundEditorGraphInputNode.h b/Engine/Plugins/Runtime/Metasound/Source/MetasoundEditor/Public/MetasoundEditorGraphInputNode.h
--- a/Engine/Plugins/Runtime/Metasound/Source/MetasoundEditor/Public/MetasoundEditorGraphInputNode.h
+++ b/Engine/Plugins/Runtime/Metasound/Source/MetasoundEditor/Public/MetasoundEditorGraphInputNode.h
@@ -34,6 +34,7 @@ public:
 	virtual FGuid GetNodeID() const override;
 	virtual FLinearColor GetNodeTitleColor() const override;
 	virtual FSlateIcon GetNodeTitleIcon() const override;
+	virtual bool Validate(Metasound::Editor::FGraphNodeValidationResult& OutResult) override;
 
 protected:
 	virtual void SetNodeID(FGuid InNodeID) override;
diff --git a/Engine/Plugins/Runtime/Metasound/Source/MetasoundEditor/Public/MetasoundEditorGraphNode.h b/Engine/Plugins/Runtime/Metasound/Source/MetasoundEditor/Public/MetasoundEditorGraphNode.h
--- a/Engine/Plugins/Runtime/Metasound/Source/MetasoundEditor/Public/MetasoundEditorGraphNode.h
+++ b/Engine/Plugins/Runtime/Metasound/Source/MetasoundEditor/Public/MetasoundEditorGraphNode.h
@@ -4,6 +4,7 @@
 
 #include "CoreMinimal.h"
 #include "EdGraph/EdGraphNode.h"
+#include "MetasoundEditorGraphValidation.h"
 #include "MetasoundFrontend.h"
 #include "MetasoundFrontendController.h"
 #include "MetasoundFrontendDocument.h"
@@ -32,6 +33,11 @@ namespace Metasound
 
 		// Map of class names to sorted array of registered version numbers
 		using FSortedClassVersionMap = TMap<FName, TArray<FMetasoundFrontendVersionNumber>>;
+
+		namespace GraphNode
+		{
+			void SetMessage(UEdGraphNode& InNode, EMessageSeverity::Type InSeverity, const FString& InMessage);
+		}
 	} // namespace Editor
 } // namespace Metasound
 
@@ -99,10 +105,12 @@ public:
 	virtual FGuid GetNodeID() const { return FGuid(); }
 	virtual FText GetDisplayName() const;
 	virtual void CacheTitle();
+	virtual bool Validate(Metasound::Editor::FGraphNodeValidationResult& OutResult);
 
 	// Mark node for refresh
 	void SyncChangeIDs();
 
+
 	FText GetCachedTitle() const { return CachedTitle; }
 
 	// Returns whether or not Metadata has been changed since the last node refresh
@@ -143,6 +151,10 @@ public:
 	{
 		return true;
 	}
+
+protected:
+	// Utility to construct a new validaton result for child classes
+	Metasound::Editor::FGraphNodeValidationResult CreateNewValidationResult();
 };
 
 /** Node that represents a graph output */
@@ -170,6 +182,8 @@ public:
 	// Disables interact widgets (ex. sliders, knobs) when input is connected
 	virtual bool EnableInteractWidgets() const override;
 
+	virtual bool Validate(Metasound::Editor::FGraphNodeValidationResult& OutResult) override;
+
 protected:
 	virtual FLinearColor GetNodeTitleColor() const override;
 	virtual FSlateIcon GetNodeTitleIcon() const override;
@@ -210,7 +224,7 @@ public:
 	bool CanAutoUpdate() const;
 
 	// Validates node and returns whether or not the node is valid.
-	bool Validate(Metasound::Editor::FGraphNodeValidationResult& OutResult);
+	virtual bool Validate(Metasound::Editor::FGraphNodeValidationResult& OutResult) override;
 
 
 protected:
diff --git a/Engine/Plugins/Runtime/Metasound/Source/MetasoundEditor/Public/MetasoundEditorGraphValidation.h b/Engine/Plugins/Runtime/Metasound/Source/MetasoundEditor/Public/MetasoundEditorGraphValidation.h
--- a/Engine/Plugins/Runtime/Metasound/Source/MetasoundEditor/Public/MetasoundEditorGraphValidation.h
+++ b/Engine/Plugins/Runtime/Metasound/Source/MetasoundEditor/Public/MetasoundEditorGraphValidation.h
@@ -5,6 +5,7 @@
 #include "MetasoundEditorGraph.h"
 #include "MetasoundEditorGraphNode.h"
 
+class UMetasoundEditorGraphNode;
 
 namespace Metasound
 {
diff --git a/Engine/Plugins/Runtime/Metasound/Source/MetasoundEngine/Private/MetasoundInterface.cpp b/Engine/Plugins/Runtime/Metasound/Source/MetasoundEngine/Private/MetasoundInterface.cpp
--- a/Engine/Plugins/Runtime/Metasound/Source/MetasoundEngine/Private/MetasoundInterface.cpp
+++ b/Engine/Plugins/Runtime/Metasound/Source/MetasoundEngine/Private/MetasoundInterface.cpp
@@ -134,6 +134,15 @@ namespace Metasound
 
 			Algo::Transform(Interface->GetInputs(), FrontendInterface.Inputs, [&](const Audio::FParameterInterface::FInput& Input)
 			{
+				// Setup required inputs by telling the style that the input is required
+				// This will later be validated against.
+#if WITH_EDITOR
+				if (!Input.RequiredText.IsEmpty())
+				{
+					FrontendInterface.AddRequiredInputToStyle(Input.InitValue.ParamName, Input.RequiredText);
+				}
+#endif // #if WITH_EDITOR
+
 				FMetasoundFrontendClassInput ClassInput;
 				ClassInput.Name = Input.InitValue.ParamName;
 				ClassInput.DefaultLiteral = FMetasoundFrontendLiteral(Input.InitValue);
@@ -154,6 +163,15 @@ namespace Metasound
 
 			Algo::Transform(Interface->GetOutputs(), FrontendInterface.Outputs, [&](const Audio::FParameterInterface::FOutput& Output)
 			{
+				// Setup required outputs by telling the style that the output is required
+				// This will later be validated against.
+#if WITH_EDITOR
+				if (!Output.RequiredText.IsEmpty())
+				{
+					FrontendInterface.AddRequiredOutputToStyle(Output.ParamName, Output.RequiredText);
+				}
+#endif // #if WITH_EDITOR
+
 				FMetasoundFrontendClassOutput ClassOutput;
 				ClassOutput.Name = Output.ParamName;
 				ClassOutput.TypeName = ResolveMemberDataType(Output.DataType, Output.ParamType);
@@ -205,6 +223,8 @@ namespace Metasound
 
 				RegisterInterface(SourceInterface::CreateInterface(*UMetaSoundSource::StaticClass()), MakeUnique<SourceInterface::FUpdateInterface>(), bIsDefault, bEditorCanAddOrRemove, IDataReference::RouterName);
 				RegisterInterface(OutputFormatMonoInterface::CreateInterface(*UMetaSoundSource::StaticClass()), nullptr, bIsDefault, bEditorCanAddOrRemove, IDataReference::RouterName);
+
+				RegisterInterface(SourceOneShotInterface::CreateInterface(*UMetaSoundSource::StaticClass()), nullptr, bIsDefault, true, IDataReference::RouterName);
 			}
 
 			// Register Non-Default Internal Interfaces (Not managed directly by end-user & not added by default when creating new MetaSound assets).
@@ -226,9 +246,7 @@ namespace Metasound
 				RegisterInterface<UMetaSoundSource>(MetasoundOutputFormatMonoV1_1::GetInterface(), MakeUnique<MetasoundOutputFormatMonoV1_1::FUpdateInterface>(), bIsDefault, bEditorCanAddOrRemove, IDataReference::RouterName);
 				RegisterInterface<UMetaSoundSource>(MetasoundOutputFormatMonoV1_2::GetInterface(), MakeUnique<MetasoundOutputFormatMonoV1_2::FUpdateInterface>(), bIsDefault, bEditorCanAddOrRemove, IDataReference::RouterName);
 
-				RegisterInterface(OutputFormatStereoInterface::CreateInterface(*UMetaSoundSource::StaticClass()), nullptr, bIsDefault, bEditorCanAddOrRemove, IDataReference::RouterName);
-				
-				RegisterInterface(SourceOneShotInterface::CreateInterface(*UMetaSoundSource::StaticClass()), nullptr, bIsDefault, true, IDataReference::RouterName);
+				RegisterInterface(OutputFormatStereoInterface::CreateInterface(*UMetaSoundSource::StaticClass()), nullptr, bIsDefault, bEditorCanAddOrRemove, IDataReference::RouterName);			
 			}
 
 			// Register External Interfaces (Interfaces defined externally & can be managed directly by end-user).
diff --git a/Engine/Plugins/Runtime/Metasound/Source/MetasoundFrontend/Private/MetasoundFrontendTransform.cpp b/Engine/Plugins/Runtime/Metasound/Source/MetasoundFrontend/Private/MetasoundFrontendTransform.cpp
--- a/Engine/Plugins/Runtime/Metasound/Source/MetasoundFrontend/Private/MetasoundFrontendTransform.cpp
+++ b/Engine/Plugins/Runtime/Metasound/Source/MetasoundFrontend/Private/MetasoundFrontendTransform.cpp
@@ -77,16 +77,38 @@ namespace Metasound
 				InputsToRemove.Append(FromInterface.Inputs);
 				OutputsToRemove.Append(FromInterface.Outputs);
 			}
+
+			// This function combines all the inputs of all interfaces into one input list and ptrs to their originating interfaces.
+			// The interface ptr will be used to query the interface for required validations on inputs. Interfaces define required inputs (and possibly other validation requirements).
 			for (const FMetasoundFrontendInterface& ToInterface : InterfacesToAdd)
 			{
-				InputsToAdd.Append(ToInterface.Inputs);
-				OutputsToAdd.Append(ToInterface.Outputs);
+				TArray<FInputData> NewInputDataArray;
+				for (const FMetasoundFrontendClassInput& Input : ToInterface.Inputs)
+				{
+					FInputData NewData;
+					NewData.Input = Input;
+					NewData.InputInterface = &ToInterface;
+					NewInputDataArray.Add(NewData);
+				}
+
+				InputsToAdd.Append(NewInputDataArray);
+
+				TArray<FOutputData> NewOutputDataArray;
+				for (const FMetasoundFrontendClassOutput& Output : ToInterface.Outputs)
+				{
+					FOutputData NewData;
+					NewData.Output = Output;
+					NewData.OutputInterface = &ToInterface;
+					NewOutputDataArray.Add(NewData);
+				}
+
+				OutputsToAdd.Append(NewOutputDataArray);
 			}
 
 			// Iterate in reverse to allow removal from `InputsToAdd`
 			for (int32 AddIndex = InputsToAdd.Num() - 1; AddIndex >= 0; AddIndex--)
 			{
-				const FMetasoundFrontendClassVertex& VertexToAdd = InputsToAdd[AddIndex];
+				const FMetasoundFrontendClassVertex& VertexToAdd = InputsToAdd[AddIndex].Input;
 
 				const int32 RemoveIndex = InputsToRemove.IndexOfByPredicate([&](const FMetasoundFrontendClassVertex& VertexToRemove)
 				{
@@ -111,7 +133,7 @@ namespace Metasound
 
 				if (INDEX_NONE != RemoveIndex)
 				{
-					PairedInputs.Add(FVertexPair{InputsToRemove[RemoveIndex], InputsToAdd[AddIndex]});
+					PairedInputs.Add(FVertexPair{InputsToRemove[RemoveIndex], InputsToAdd[AddIndex].Input});
 					InputsToRemove.RemoveAtSwap(RemoveIndex);
 					InputsToAdd.RemoveAtSwap(AddIndex);
 				}
@@ -120,7 +142,7 @@ namespace Metasound
 			// Iterate in reverse to allow removal from `OutputsToAdd`
 			for (int32 AddIndex = OutputsToAdd.Num() - 1; AddIndex >= 0; AddIndex--)
 			{
-				const FMetasoundFrontendClassVertex& VertexToAdd = OutputsToAdd[AddIndex];
+				const FMetasoundFrontendClassVertex& VertexToAdd = OutputsToAdd[AddIndex].Output;
 
 				const int32 RemoveIndex = OutputsToRemove.IndexOfByPredicate([&](const FMetasoundFrontendClassVertex& VertexToRemove)
 				{
@@ -145,11 +167,11 @@ namespace Metasound
 
 				if (INDEX_NONE != RemoveIndex)
 				{
-					PairedOutputs.Add(FVertexPair{OutputsToRemove[RemoveIndex], OutputsToAdd[AddIndex]});
+					PairedOutputs.Add(FVertexPair{OutputsToRemove[RemoveIndex], OutputsToAdd[AddIndex].Output});
 					OutputsToRemove.RemoveAtSwap(RemoveIndex);
 					OutputsToAdd.RemoveAtSwap(AddIndex);
 				}
-			}
+			}			
 		}
 
 		bool FModifyRootGraphInterfaces::Transform(FDocumentHandle InDocument) const
@@ -229,33 +251,43 @@ namespace Metasound
 #if WITH_EDITOR
 			auto FindLowestNodeLocationOfClassType = [](EMetasoundFrontendClassType ClassType, FGraphHandle Graph, FName DataType, TFunctionRef<bool(FConstNodeHandle, FName)> NodeDataTypeFilter)
 			{
-				FVector2D LowestLocation;
+				FVector2D LowestLocation = { TNumericLimits<float>::Min(), TNumericLimits<float>::Min() };
+				bool bFoundLocation = false;
+
 				Graph->IterateConstNodes([&](FConstNodeHandle NodeHandle)
 				{
-					for (const TPair<FGuid, FVector2D>& Pair : NodeHandle->GetNodeStyle().Display.Locations)
+					const TMap<FGuid, FVector2D>& Locations = NodeHandle->GetNodeStyle().Display.Locations;
+					for (const TPair<FGuid, FVector2D>& Pair : Locations)
 					{
 						if (Pair.Value.Y > LowestLocation.Y)
 						{
-							if (NodeDataTypeFilter(NodeHandle, DataType))
-							{
-								LowestLocation = Pair.Value;
-							}
+							LowestLocation = Pair.Value;
+							bFoundLocation = true;
 						}
 					}
 				}, ClassType);
 
+				if (!bFoundLocation)
+				{
+					LowestLocation = { DisplayStyle::NodeLayout::DefaultOffsetX.X * 2.0f, 0.0f };
+				}
+
 				return LowestLocation;
 			};
 #endif // WITH_EDITOR
 
 			// Add missing inputs
-			for (const FMetasoundFrontendClassInput& InputToAdd : InputsToAdd)
+			for (const FInputData& InputData : InputsToAdd)
 			{
 				bDidEdit = true;
+				const FMetasoundFrontendClassInput& InputToAdd = InputData.Input;
 				FNodeHandle NewInputNode = GraphHandle->AddInputVertex(InputToAdd);
 
 #if WITH_EDITOR
-				if (bSetDefaultNodeLocations)
+				FName InputName = InputToAdd.Name;
+				bool bRequiredInput = InputData.InputInterface->GetInputStyle().RequiredMembers.Contains(InputName);
+
+				if (bSetDefaultNodeLocations || bRequiredInput)
 				{
 					FMetasoundFrontendNodeStyle Style = NewInputNode->GetNodeStyle();
 					const FVector2D LastOutputLocation = FindLowestNodeLocationOfClassType(EMetasoundFrontendClassType::Input, GraphHandle, InputToAdd.TypeName, InputDataTypeCompareFilter);
@@ -266,13 +298,17 @@ namespace Metasound
 			}
 
 			// Add missing outputs
-			for (const FMetasoundFrontendClassOutput& OutputToAdd : OutputsToAdd)
+			for (const FOutputData& OutputData : OutputsToAdd)
 			{
 				bDidEdit = true;
+				const FMetasoundFrontendClassOutput& OutputToAdd = OutputData.Output;
 				FNodeHandle NewOutputNode = GraphHandle->AddOutputVertex(OutputToAdd);
 
 #if WITH_EDITOR
-				if (bSetDefaultNodeLocations)
+				FName OutputName = OutputToAdd.Name;
+				bool bRequiredOutput = OutputData.OutputInterface->GetOutputStyle().RequiredMembers.Contains(OutputName);
+
+				if (bSetDefaultNodeLocations || bRequiredOutput)
 				{
 					FMetasoundFrontendNodeStyle Style = NewOutputNode->GetNodeStyle();
 					const FVector2D LastOutputLocation = FindLowestNodeLocationOfClassType(EMetasoundFrontendClassType::Output, GraphHandle, OutputToAdd.TypeName, OutputDataTypeCompareFilter);
diff --git a/Engine/Plugins/Runtime/Metasound/Source/MetasoundFrontend/Private/MetasoundSourceInterface.cpp b/Engine/Plugins/Runtime/Metasound/Source/MetasoundFrontend/Private/MetasoundSourceInterface.cpp
--- a/Engine/Plugins/Runtime/Metasound/Source/MetasoundFrontend/Private/MetasoundSourceInterface.cpp
+++ b/Engine/Plugins/Runtime/Metasound/Source/MetasoundFrontend/Private/MetasoundSourceInterface.cpp
@@ -43,7 +43,8 @@ namespace Metasound
 								LOCTEXT("OnFinished", "On Finished"),
 								LOCTEXT("OnFinishedDescription", "Trigger executed to initiate stopping the source."),
 								GetMetasoundDataTypeName<FTrigger>(),
-								Outputs::OnFinished
+								Outputs::OnFinished,								
+								LOCTEXT("OnFinishedWarning", "\"On Finished\" should be connected for OneShot MetaSound sources. For sources with undefined duration (e.g. looping), remove the OneShot interface and use an audio component to avoid leaking the source."),
 							}
 						};
 					}
diff --git a/Engine/Plugins/Runtime/Metasound/Source/MetasoundFrontend/Public/MetasoundFrontendDocument.h b/Engine/Plugins/Runtime/Metasound/Source/MetasoundFrontend/Public/MetasoundFrontendDocument.h
--- a/Engine/Plugins/Runtime/Metasound/Source/MetasoundFrontend/Public/MetasoundFrontendDocument.h
+++ b/Engine/Plugins/Runtime/Metasound/Source/MetasoundFrontend/Public/MetasoundFrontendDocument.h
@@ -792,6 +792,12 @@ struct FMetasoundFrontendInterfaceStyle
 	UPROPERTY()
 	TArray<int32> DefaultSortOrder;
 
+#if WITH_EDITORONLY_DATA
+	// Map of member names with FText to be used as warnings if not hooked up
+	UPROPERTY()
+	TMap<FName, FText> RequiredMembers;
+#endif // #if WITH_EDITORONLY_DATA
+
 	template <typename HandleType>
 	TArray<HandleType> SortDefaults(const TArray<HandleType>& InHandles) const
 	{
@@ -832,6 +838,7 @@ struct METASOUNDFRONTEND_API FMetasoundFrontendClassInterface
 	GENERATED_BODY()
 
 private:
+
 	// Style info for inputs.
 	UPROPERTY()
 	FMetasoundFrontendInterfaceStyle InputStyle;
@@ -841,6 +848,7 @@ private:
 	FMetasoundFrontendInterfaceStyle OutputStyle;
 
 public:
+
 	// Generates class interface intended to be used as a registry descriptor from FNodeClassMetadata.
 	// Does not initialize a change ID as it is not considered to be transactional.
 	static FMetasoundFrontendClassInterface GenerateClassInterface(const Metasound::FVertexInterface& InVertexInterface);
@@ -884,6 +892,39 @@ public:
 		ChangeID = FGuid::NewGuid();
 	}
 
+#if WITH_EDITORONLY_DATA
+	void AddRequiredInputToStyle(const FName& InInputName, const FText& InRequiredText)
+	{
+		InputStyle.RequiredMembers.Add(InInputName, InRequiredText);
+	}
+
+	void AddRequiredOutputToStyle(const FName& InOutputName, const FText& InRequiredText)
+	{
+		OutputStyle.RequiredMembers.Add(InOutputName, InRequiredText);
+	}
+
+	bool IsMemberInputRequired(const FName& InInputName, FText& OutRequiredText)
+	{
+		if (FText* RequiredText = InputStyle.RequiredMembers.Find(InInputName))
+		{
+			OutRequiredText = *RequiredText;
+			return true;
+		}
+		return false;
+	}
+
+	bool IsMemberOutputRequired(const FName& InOutputName, FText& OutRequiredText)
+	{
+		if (FText* RequiredText = OutputStyle.RequiredMembers.Find(InOutputName))
+		{
+			OutRequiredText = *RequiredText;
+			return true;
+		}
+		return false;
+	}
+#endif // #if WITH_EDITORONLY_DATA
+
+
 	const FGuid& GetChangeID() const
 	{
 		return ChangeID;
diff --git a/Engine/Plugins/Runtime/Metasound/Source/MetasoundFrontend/Public/MetasoundFrontendTransform.h b/Engine/Plugins/Runtime/Metasound/Source/MetasoundFrontend/Public/MetasoundFrontendTransform.h
--- a/Engine/Plugins/Runtime/Metasound/Source/MetasoundFrontend/Public/MetasoundFrontendTransform.h
+++ b/Engine/Plugins/Runtime/Metasound/Source/MetasoundFrontend/Public/MetasoundFrontendTransform.h
@@ -65,9 +65,21 @@ namespace Metasound
 			TArray<FVertexPair> PairedInputs;
 			TArray<FVertexPair> PairedOutputs;
 
-			TArray<FMetasoundFrontendClassInput> InputsToAdd;
+			struct FInputData
+			{
+				FMetasoundFrontendClassInput Input;
+				const FMetasoundFrontendInterface* InputInterface = nullptr;
+			};
+
+			struct FOutputData
+			{
+				FMetasoundFrontendClassOutput Output;
+				const FMetasoundFrontendInterface* OutputInterface = nullptr;
+			};
+
+			TArray<FInputData> InputsToAdd;
 			TArray<FMetasoundFrontendClassInput> InputsToRemove;
-			TArray<FMetasoundFrontendClassOutput> OutputsToAdd;
+			TArray<FOutputData> OutputsToAdd;
 			TArray<FMetasoundFrontendClassOutput> OutputsToRemove;
 
 		};
diff --git a/Engine/Source/Runtime/AudioExtensions/Public/IAudioParameterInterfaceRegistry.h b/Engine/Source/Runtime/AudioExtensions/Public/IAudioParameterInterfaceRegistry.h
--- a/Engine/Source/Runtime/AudioExtensions/Public/IAudioParameterInterfaceRegistry.h
+++ b/Engine/Source/Runtime/AudioExtensions/Public/IAudioParameterInterfaceRegistry.h
@@ -29,6 +29,7 @@ namespace Audio
 			const FName DataType;
 
 			const FAudioParameter InitValue;
+			const FText RequiredText;
 		};
 
 		struct FOutput
@@ -38,6 +39,7 @@ namespace Audio
 			const FName DataType;
 
 			const FName ParamName;
+			const FText RequiredText = FText();
 			const EAudioParameterType ParamType = EAudioParameterType::None;
 		};
 
