Fix for one-shot metasound sources still presenting as looping. (moving init of Duration/bIsLooping to PostLoad() override).

diff --git a/Engine/Plugins/Runtime/Metasound/Source/MetasoundEngine/Private/MetasoundSource.cpp b/Engine/Plugins/Runtime/Metasound/Source/MetasoundEngine/Private/MetasoundSource.cpp
--- a/Engine/Plugins/Runtime/Metasound/Source/MetasoundEngine/Private/MetasoundSource.cpp
+++ b/Engine/Plugins/Runtime/Metasound/Source/MetasoundEngine/Private/MetasoundSource.cpp
@@ -84,9 +84,7 @@ UMetaSoundSource::UMetaSoundSource(const FObjectInitializer& ObjectInitializer)
 {
 	bRequiresStopFade = true;
 	NumChannels = 1;
-	Duration = INDEFINITELY_LOOPING_DURATION;
-	bLooping = true;
-	
+
 	// todo: ensure that we have a method so that the audio engine can be authoritative over the sample rate the UMetaSoundSource runs at.
 	SampleRate = 48000.f;
 
@@ -257,6 +255,14 @@ FText UMetaSoundSource::GetDisplayName() const
 	return FMetasoundAssetBase::GetDisplayName(MoveTemp(TypeName));
 }
 
+void UMetaSoundSource::PostLoad()
+{
+	Super::PostLoad();
+
+	Duration = GetDuration();
+	bLooping = IsLooping();
+}
+
 void UMetaSoundSource::SetRegistryAssetClassInfo(const Metasound::Frontend::FNodeClassInfo& InNodeInfo)
 {
 	Metasound::SetMetaSoundRegistryAssetClassInfo(*this, InNodeInfo);
@@ -484,9 +490,6 @@ ISoundGeneratorPtr UMetaSoundSource::CreateSoundGenerator(const FSoundGeneratorI
 
 	METASOUND_TRACE_CPUPROFILER_EVENT_SCOPE(UMetaSoundSource::CreateSoundGenerator);
 
-	Duration = INDEFINITELY_LOOPING_DURATION;
-	bLooping = true;
-
 	SampleRate = InParams.SampleRate;
 	FOperatorSettings InSettings = GetOperatorSettings(static_cast<FSampleRate>(SampleRate));
 	FMetasoundEnvironment Environment = CreateEnvironment(InParams);
diff --git a/Engine/Plugins/Runtime/Metasound/Source/MetasoundEngine/Public/MetasoundSource.h b/Engine/Plugins/Runtime/Metasound/Source/MetasoundEngine/Public/MetasoundSource.h
--- a/Engine/Plugins/Runtime/Metasound/Source/MetasoundEngine/Public/MetasoundSource.h
+++ b/Engine/Plugins/Runtime/Metasound/Source/MetasoundEngine/Public/MetasoundSource.h
@@ -78,6 +78,10 @@ public:
 	UPROPERTY(AssetRegistrySearchable)
 	int32 RegistryVersionMinor = 0;
 
+	//~ Begin UObject Interface.
+	virtual void PostLoad() override;
+	//~ End UObject Interface.
+
 	// Sets Asset Registry Metadata associated with this MetaSoundSource
 	virtual void SetRegistryAssetClassInfo(const Metasound::Frontend::FNodeClassInfo& InNodeInfo) override;
 
