Follow up fix for Metasound Trigger Counter Node logic (reset behavior)

diff --git a/Engine/Plugins/Runtime/Metasound/Source/MetasoundStandardNodes/Private/MetasoundTriggerCounterNode.cpp b/Engine/Plugins/Runtime/Metasound/Source/MetasoundStandardNodes/Private/MetasoundTriggerCounterNode.cpp
--- a/Engine/Plugins/Runtime/Metasound/Source/MetasoundStandardNodes/Private/MetasoundTriggerCounterNode.cpp
+++ b/Engine/Plugins/Runtime/Metasound/Source/MetasoundStandardNodes/Private/MetasoundTriggerCounterNode.cpp
@@ -133,6 +133,8 @@ namespace Metasound
 				*OutCount = CurrentTriggerCount = 0;
 				*OutValue = *StartValue;
 
+				bIsFirstTrigger = true;
+
 				TriggerOnReset->TriggerFrame(StartFrame);
 			}
 		);
