Fix for MetaSound node literal inputs working inconsistently.
#rnx

diff --git a/Engine/Plugins/Runtime/Metasound/Source/MetasoundFrontend/Private/MetasoundFrontendGraph.cpp b/Engine/Plugins/Runtime/Metasound/Source/MetasoundFrontend/Private/MetasoundFrontendGraph.cpp
--- a/Engine/Plugins/Runtime/Metasound/Source/MetasoundFrontend/Private/MetasoundFrontendGraph.cpp
+++ b/Engine/Plugins/Runtime/Metasound/Source/MetasoundFrontend/Private/MetasoundFrontendGraph.cpp
@@ -237,7 +237,7 @@ namespace Metasound
 		return TUniquePtr<INode>(nullptr);
 	}
 
-	TUniquePtr<INode> FFrontendGraphBuilder::CreateOutputNode(const FMetasoundFrontendNode& InNode, const FMetasoundFrontendClass& InClass, FBuildGraphContext& InGraphContext, const TSet<FGuid>& InEdgeDestinations)
+	TUniquePtr<INode> FFrontendGraphBuilder::CreateOutputNode(const FMetasoundFrontendNode& InNode, const FMetasoundFrontendClass& InClass, FBuildGraphContext& InGraphContext, const TSet<FNodeIDVertexID>& InEdgeDestinations)
 	{
 		using namespace Metasound::Frontend;
 
@@ -269,7 +269,7 @@ namespace Metasound
 		return TUniquePtr<INode>(nullptr);
 	}
 
-	TUniquePtr<INode> FFrontendGraphBuilder::CreateExternalNode(const FMetasoundFrontendNode& InNode, const FMetasoundFrontendClass& InClass, FBuildGraphContext& InGraphContext, const TSet<FGuid>& InEdgeDestinations)
+	TUniquePtr<INode> FFrontendGraphBuilder::CreateExternalNode(const FMetasoundFrontendNode& InNode, const FMetasoundFrontendClass& InClass, FBuildGraphContext& InGraphContext, const TSet<FNodeIDVertexID>& InEdgeDestinations)
 	{
 		check(InNode.ClassID == InClass.ID);
 
@@ -397,11 +397,11 @@ namespace Metasound
 	// TODO: add errors here. Most will be a "PromptIfMissing"...
 	void FFrontendGraphBuilder::AddNodesToGraph(FBuildGraphContext& InGraphContext)
 	{
-		TSet<FGuid> GraphEdgeDestinations;
+		TSet<FNodeIDVertexID> GraphEdgeDestinations;
 		const TArray<FMetasoundFrontendEdge>& GraphEdges = InGraphContext.GraphClass.Graph.Edges;
 		Algo::Transform(GraphEdges, GraphEdgeDestinations, [](const FMetasoundFrontendEdge& Edge) 
 		{
-			return Edge.ToVertexID;
+			return FNodeIDVertexID{Edge.ToNodeID, Edge.ToVertexID};
 		});
 
 		for (const FMetasoundFrontendNode& Node : InGraphContext.GraphClass.Graph.Nodes)
@@ -624,7 +624,7 @@ namespace Metasound
 		InGraphContext.DefaultInputs.Reset();
 	}
 
-	TArray<FFrontendGraphBuilder::FDefaultLiteralData> FFrontendGraphBuilder::GetInputDefaultLiteralData(const FMetasoundFrontendNode& InNode, const FNodeInitData& InInitData, const TSet<FGuid>& InEdgeDestinations)
+	TArray<FFrontendGraphBuilder::FDefaultLiteralData> FFrontendGraphBuilder::GetInputDefaultLiteralData(const FMetasoundFrontendNode& InNode, const FNodeInitData& InInitData, const TSet<FNodeIDVertexID>& InEdgeDestinations)
 	{
 		TArray<FDefaultLiteralData> DefaultLiteralData;
 
@@ -643,7 +643,7 @@ namespace Metasound
 				{
 					// Only build params and forward along to connect to dynamically generated
 					// literal node if edge is not explicitly connected to vertex as destination.
-					bRequiresDefault = !InEdgeDestinations.Contains(Vertex.VertexID);
+					bRequiresDefault = !InEdgeDestinations.Contains({InNode.GetID(), Vertex.VertexID});
 					if (bRequiresDefault)
 					{
 						InitParams.Literal = Literal.Value.ToLiteral(Vertex.TypeName);
diff --git a/Engine/Plugins/Runtime/Metasound/Source/MetasoundFrontend/Public/MetasoundFrontendGraph.h b/Engine/Plugins/Runtime/Metasound/Source/MetasoundFrontend/Public/MetasoundFrontendGraph.h
--- a/Engine/Plugins/Runtime/Metasound/Source/MetasoundFrontend/Public/MetasoundFrontendGraph.h
+++ b/Engine/Plugins/Runtime/Metasound/Source/MetasoundFrontend/Public/MetasoundFrontendGraph.h
@@ -144,11 +144,10 @@ namespace Metasound
 			TUniquePtr<FFrontendGraph> Graph;
 			const FMetasoundFrontendGraphClass& GraphClass;
 			FBuildContext& BuildContext;
-			TSet<FGuid> DestinationsWithEdges;
 			FDefaultInputByIDMap DefaultInputs;
 		};
 
-		static TArray<FDefaultLiteralData> GetInputDefaultLiteralData(const FMetasoundFrontendNode& InNode, const FNodeInitData& InInitData, const TSet<FGuid>& InEdgeDestinations);
+		static TArray<FDefaultLiteralData> GetInputDefaultLiteralData(const FMetasoundFrontendNode& InNode, const FNodeInitData& InInitData, const TSet<FNodeIDVertexID>& InEdgeDestinations);
 
 		static bool SortSubgraphDependencies(TArray<const FMetasoundFrontendGraphClass*>& Subgraphs);
 
@@ -160,9 +159,9 @@ namespace Metasound
 		static const FMetasoundFrontendVariable* FindVariableForVariableNode(const FMetasoundFrontendNode& InVariableNode, const FMetasoundFrontendGraph& InGraph);
 
 		static TUniquePtr<INode> CreateInputNode(const FMetasoundFrontendNode& InNode, const FMetasoundFrontendClass& InClass, const FMetasoundFrontendClassInput& InOwningGraphClassInput);
-		static TUniquePtr<INode> CreateOutputNode(const FMetasoundFrontendNode& InNode, const FMetasoundFrontendClass& InClass, FBuildGraphContext& InGraphContext, const TSet<FGuid>& InEdgeDestinations);
+		static TUniquePtr<INode> CreateOutputNode(const FMetasoundFrontendNode& InNode, const FMetasoundFrontendClass& InClass, FBuildGraphContext& InGraphContext, const TSet<FNodeIDVertexID>& InEdgeDestinations);
 		static TUniquePtr<INode> CreateVariableNode(const FMetasoundFrontendNode& InNode, const FMetasoundFrontendGraph& InGraph);
-		static TUniquePtr<INode> CreateExternalNode(const FMetasoundFrontendNode& InNode, const FMetasoundFrontendClass& InClass, FBuildGraphContext& InGraphContext, const TSet<FGuid>& InEdgeDestinations);
+		static TUniquePtr<INode> CreateExternalNode(const FMetasoundFrontendNode& InNode, const FMetasoundFrontendClass& InClass, FBuildGraphContext& InGraphContext, const TSet<FNodeIDVertexID>& InEdgeDestinations);
 
 		// TODO: add errors here. Most will be a "PromptIfMissing"...
 		static void AddNodesToGraph(FBuildGraphContext& InGraphContext);
