MetaSound interfaces menu changes
- Add right click menu and actions for graph members
- Move add/remove interfaces menu to a new details tab

diff --git a/Engine/Plugins/Runtime/Metasound/Source/MetasoundEditor/Private/MetasoundDetailCustomization.cpp b/Engine/Plugins/Runtime/Metasound/Source/MetasoundEditor/Private/MetasoundDetailCustomization.cpp
--- a/Engine/Plugins/Runtime/Metasound/Source/MetasoundEditor/Private/MetasoundDetailCustomization.cpp
+++ b/Engine/Plugins/Runtime/Metasound/Source/MetasoundEditor/Private/MetasoundDetailCustomization.cpp
@@ -76,181 +76,6 @@
 			});
 		}
 
-		void FMetasoundDetailCustomization::CustomizeInterfaces(IDetailLayoutBuilder& DetailLayout)
-		{
-			UpdateInterfaceNames();
-
-			SAssignNew(InterfaceComboBox, SSearchableComboBox)
-				.OptionsSource(&AddableInterfaceNames)
-				.OnGenerateWidget_Lambda([](TSharedPtr<FString> InItem)
-				{
-					return SNew(STextBlock)
-						.Text(FText::FromString(*InItem));
-				})
-				.OnSelectionChanged_Lambda([this](TSharedPtr<FString> NameToAdd, ESelectInfo::Type InSelectInfo)
-				{
-					using namespace Metasound;
-					using namespace Metasound::Frontend;
-
-					FMetasoundAssetBase* MetaSoundAsset = IMetasoundUObjectRegistry::Get().GetObjectAsAssetBase(MetaSound.Get());
-					if (!ensure(MetaSoundAsset))
-					{
-						return;
-					}
-
-					if (InSelectInfo != ESelectInfo::OnNavigation)
-					{
-						FMetasoundFrontendInterface InterfaceToAdd;
-						if (ensure(ISearchEngine::Get().FindInterfaceWithHighestVersion(FName(*NameToAdd.Get()), InterfaceToAdd)))
-						{
-							FScopedTransaction(LOCTEXT("AddInterfaceTransaction", "Add MetaSound Interface"));
-							MetaSound.Get()->Modify();
-							MetaSoundAsset->GetGraphChecked().Modify();
-
-							FDocumentHandle DocumentHandle = MetaSoundAsset->GetDocumentHandle();
-							FModifyRootGraphInterfaces ModifyTransform({ }, { InterfaceToAdd });
-							ModifyTransform.SetDefaultNodeLocations(false); // Don't automatically add nodes to ed graph
-							ModifyTransform.Transform(DocumentHandle);
-							MetaSoundAsset->SetUpdateDetailsOnSynchronization();
-						}
-
-						UpdateInterfaceNames();
-						InterfaceComboBox->RefreshOptions();
-						MetaSoundAsset->SetSynchronizationRequired();
-					}
-				})
-				.Content()
-				[
-					SNew(STextBlock)
-					.Text(LOCTEXT("UpdateInterfaceAction", "Add Interface..."))
-					.IsEnabled(IsGraphEditableAttribute)
-				];
-
-			TSharedRef<SWidget> InterfaceUtilities = SNew(SHorizontalBox)
-			+ SHorizontalBox::Slot()
-			.Padding(2.0f)
-			.HAlign(HAlign_Left)
-			.VAlign(VAlign_Center)
-			.AutoWidth()
-			[
-				InterfaceComboBox->AsShared()
-			]
-			+ SHorizontalBox::Slot()
-			.Padding(2.0f)
-			.HAlign(HAlign_Left)
-			.VAlign(VAlign_Center)
-			.AutoWidth()
-			[
-				PropertyCustomizationHelpers::MakeDeleteButton(FSimpleDelegate::CreateLambda([this]()
-				{
-					using namespace Frontend;
-
-					FMetasoundAssetBase* MetaSoundAsset = IMetasoundUObjectRegistry::Get().GetObjectAsAssetBase(MetaSound.Get());
-					if (!ensure(MetaSoundAsset))
-					{
-						return;
-					}
-
-					TArray<FMetasoundFrontendInterface> ImplementedInterfaces;
-					Algo::Transform(ImplementedInterfaceNames, ImplementedInterfaces, [](const FName& Name)
-					{
-						FMetasoundFrontendInterface Interface;
-						ISearchEngine::Get().FindInterfaceWithHighestVersion(Name, Interface);
-						return Interface;
-					});
-
-					{
-						FScopedTransaction(LOCTEXT("RemoveAllInterfacesTransaction", "Remove All MetaSound Interfaces"));
-						MetaSound.Get()->Modify();
-						MetaSoundAsset->GetGraphChecked().Modify();
-
-						FDocumentHandle DocumentHandle = MetaSoundAsset->GetDocumentHandle();
-						FModifyRootGraphInterfaces({ ImplementedInterfaces }, { }).Transform(DocumentHandle);
-					}
-
-					UpdateInterfaceNames();
-					InterfaceComboBox->RefreshOptions();
-					MetaSoundAsset->SetUpdateDetailsOnSynchronization();
-					MetaSoundAsset->SetSynchronizationRequired();
-				}), LOCTEXT("RemoveInterfaceTooltip1", "Removes all interfaces from the given MetaSound."))
-			];
-
-			const FText HeaderName = LOCTEXT("InterfacesGroupDisplayName", "Interfaces");
-			IDetailCategoryBuilder& InterfaceCategory = DetailLayout.EditCategory("Interfaces", HeaderName);
-
-			InterfaceCategory.AddCustomRow(HeaderName)
-			.ValueContent()
-			[
-				InterfaceUtilities
-			];
-
-			auto CreateInterfaceEntryWidget = [&](FName InInterfaceName) -> TSharedPtr<SWidget>
-			{
-				using namespace Frontend;
-
-				FMetasoundFrontendInterface InterfaceEntry;
-				if (!ensure(ISearchEngine::Get().FindInterfaceWithHighestVersion(InInterfaceName, InterfaceEntry)))
-				{
-					return SNullWidget::NullWidget;
-				}
-
-				TSharedRef<SWidget> RemoveButtonWidget = PropertyCustomizationHelpers::MakeDeleteButton(FSimpleDelegate::CreateLambda([this, InterfaceEntry]()
-				{
-					using namespace Frontend;
-
-					FMetasoundAssetBase* MetaSoundAsset = IMetasoundUObjectRegistry::Get().GetObjectAsAssetBase(MetaSound.Get());
-					if (!ensure(MetaSoundAsset))
-					{
-						return;
-					}
-
-					{
-						FScopedTransaction(LOCTEXT("RemoveInterfaceTransaction", "Remove MetaSound Interface"));
-						MetaSound.Get()->Modify();
-						MetaSoundAsset->GetGraphChecked().Modify();
-
-						FDocumentHandle DocumentHandle = MetaSoundAsset->GetDocumentHandle();
-						FModifyRootGraphInterfaces({ InterfaceEntry }, { }).Transform(DocumentHandle);
-					}
-
-					UpdateInterfaceNames();
-					InterfaceComboBox->RefreshOptions();
-					MetaSoundAsset->SetUpdateDetailsOnSynchronization();
-					MetaSoundAsset->SetSynchronizationRequired();
-				}), LOCTEXT("RemoveInterfaceTooltip2", "Removes the associated interface from the MetaSound."));
-
-				return SNew(SHorizontalBox)
-					+ SHorizontalBox::Slot()
-					.Padding(2.0f)
-					.HAlign(HAlign_Center)
-					.VAlign(VAlign_Center)
-					.AutoWidth()
-					[
-						SNew(STextBlock)
-						.Text(FText::FromName(InterfaceEntry.Version.Name))
-					]
-					+ SHorizontalBox::Slot()
-					.Padding(2.0f)
-					.HAlign(HAlign_Center)
-					.VAlign(VAlign_Center)
-					.AutoWidth()
-					[
-						RemoveButtonWidget
-					];
-			};
-
-			TArray<FName> InterfaceNames = ImplementedInterfaceNames.Array();
-			InterfaceNames.Sort([](const FName& A, const FName& B) { return A.LexicalLess(B); });
-			for (const FName& InterfaceName : InterfaceNames)
-			{
-				InterfaceCategory.AddCustomRow(FText::FromName(InterfaceName))
-				.ValueContent()
-				[
-					CreateInterfaceEntryWidget(InterfaceName)->AsShared()
-				];
-			}
-		}
-
 		FName FMetasoundDetailCustomization::GetInterfaceVersionsPath() const
 		{
 			return Metasound::Editor::BuildChildPath(DocumentPropertyName, GET_MEMBER_NAME_CHECKED(FMetasoundFrontendDocument, Interfaces));
@@ -267,47 +92,6 @@
 			return Metasound::Editor::BuildChildPath(RootClass, GET_MEMBER_NAME_CHECKED(FMetasoundFrontendClass, Metadata));
 		}
 
-		void FMetasoundDetailCustomization::UpdateInterfaceNames()
-		{
-			using namespace Frontend;
-
-			AddableInterfaceNames.Reset();
-			ImplementedInterfaceNames.Reset();
-
-			if (const FMetasoundAssetBase* MetaSoundAsset = IMetasoundUObjectRegistry::Get().GetObjectAsAssetBase(MetaSound.Get()))
-			{
-				auto GetVersionName = [](const FMetasoundFrontendVersion& Version) { return Version.Name; };
-				auto CanAddOrRemoveInterface = [](const FMetasoundFrontendVersion& Version)
-				{
-					using namespace Metasound::Frontend;
-
-					const FInterfaceRegistryKey Key = GetInterfaceRegistryKey(Version);
-					if (const IInterfaceRegistryEntry* Entry = IInterfaceRegistry::Get().FindInterfaceRegistryEntry(Key))
-					{
-						return Entry->EditorCanAddOrRemove();
-					}
-
-					return false;
-				};
-
-				const TSet<FMetasoundFrontendVersion>& ImplementedInterfaces = MetaSoundAsset->GetDocumentChecked().Interfaces;
-				Algo::TransformIf(ImplementedInterfaces, ImplementedInterfaceNames, CanAddOrRemoveInterface, GetVersionName);
-
-				TArray<FMetasoundFrontendInterface> Interfaces = ISearchEngine::Get().FindAllInterfaces();
-				for (const FMetasoundFrontendInterface& Interface : Interfaces)
-				{
-					if (!ImplementedInterfaceNames.Contains(Interface.Version.Name))
-					{
-						if (CanAddOrRemoveInterface(Interface.Version))
-						{
-							FString Name = Interface.Version.Name.ToString();
-							AddableInterfaceNames.Add(MakeShared<FString>(MoveTemp(Name)));
-						}
-					}
-				}
-			}
-		}
-
 		void FMetasoundDetailCustomization::CustomizeDetails(IDetailLayoutBuilder& DetailLayout)
 		{
 			using namespace Frontend;
@@ -463,9 +247,6 @@
 					GeneralCategoryBuilder.AddProperty(CategoryHierarchyHandle);
 					GeneralCategoryBuilder.AddProperty(KeywordsHandle);
 
-					CustomizeInterfaces(DetailLayout);
-
-
 					DetailLayout.HideCategory("Attenuation");
 					DetailLayout.HideCategory("Effects");
 					DetailLayout.HideCategory("Loading");
@@ -521,6 +302,250 @@
 			DetailLayout.HideCategory("Analysis");
 			DetailLayout.HideCategory("Advanced");
 		}
+
+		FMetasoundInterfacesDetailCustomization::FMetasoundInterfacesDetailCustomization()
+		{
+			IsGraphEditableAttribute = TAttribute<bool>::Create([this]()
+			{
+				using namespace Metasound;
+				using namespace Metasound::Frontend;
+				if (const FMetasoundAssetBase* MetaSoundAsset = IMetasoundUObjectRegistry::Get().GetObjectAsAssetBase(MetaSound.Get()))
+				{
+					FConstGraphHandle GraphHandle = MetaSoundAsset->GetRootGraphHandle();
+					return GraphHandle->GetGraphStyle().bIsGraphEditable;
+				}
+
+				return false;
+			});
+		}
+
+		void FMetasoundInterfacesDetailCustomization::CustomizeDetails(IDetailLayoutBuilder& DetailLayout)
+		{
+			TArray<TWeakObjectPtr<UObject>> Objects;
+			DetailLayout.GetObjectsBeingCustomized(Objects);
+
+			// Only support modifying a single MetaSound at a time (Multiple
+			// MetaSound editing will be covered most likely by separate tool).
+			if (Objects.Num() > 1)
+			{
+				return;
+			}
+			if (UMetasoundInterfacesView* InterfacesView = CastChecked<UMetasoundInterfacesView>(Objects.Last()))
+			{
+				MetaSound = InterfacesView->GetMetasound();
+			}
+
+			UpdateInterfaceNames();
+
+			SAssignNew(InterfaceComboBox, SSearchableComboBox)
+				.OptionsSource(&AddableInterfaceNames)
+				.OnGenerateWidget_Lambda([](TSharedPtr<FString> InItem)
+				{
+					return SNew(STextBlock)
+						.Text(FText::FromString(*InItem));
+				})
+				.OnSelectionChanged_Lambda([this](TSharedPtr<FString> NameToAdd, ESelectInfo::Type InSelectInfo)
+				{
+					using namespace Metasound;
+					using namespace Metasound::Frontend;
+
+					FMetasoundAssetBase* MetaSoundAsset = IMetasoundUObjectRegistry::Get().GetObjectAsAssetBase(MetaSound.Get());
+					if (!ensure(MetaSoundAsset))
+					{
+						return;
+					}
+
+					if (InSelectInfo != ESelectInfo::OnNavigation)
+					{
+						FMetasoundFrontendInterface InterfaceToAdd;
+						if (ensure(ISearchEngine::Get().FindInterfaceWithHighestVersion(FName(*NameToAdd.Get()), InterfaceToAdd)))
+						{
+							FScopedTransaction(LOCTEXT("AddInterfaceTransaction", "Add MetaSound Interface"));
+							MetaSound.Get()->Modify();
+							MetaSoundAsset->GetGraphChecked().Modify();
+
+							FDocumentHandle DocumentHandle = MetaSoundAsset->GetDocumentHandle();
+							FModifyRootGraphInterfaces ModifyTransform({ }, { InterfaceToAdd });
+							ModifyTransform.SetDefaultNodeLocations(false); // Don't automatically add nodes to ed graph
+							ModifyTransform.Transform(DocumentHandle);
+							MetaSoundAsset->SetUpdateDetailsOnSynchronization();
+						}
+
+						UpdateInterfaceNames();
+						InterfaceComboBox->RefreshOptions();
+						MetaSoundAsset->SetSynchronizationRequired();
+					}
+				})
+				.Content()
+				[
+					SNew(STextBlock)
+					.Text(LOCTEXT("UpdateInterfaceAction", "Add Interface..."))
+					.IsEnabled(IsGraphEditableAttribute)
+				];
+
+			TSharedRef<SWidget> InterfaceUtilities = SNew(SHorizontalBox)
+				+ SHorizontalBox::Slot()
+				.Padding(2.0f)
+				.HAlign(HAlign_Left)
+				.VAlign(VAlign_Center)
+				.AutoWidth()
+				[
+					InterfaceComboBox->AsShared()
+				]
+			+ SHorizontalBox::Slot()
+				.Padding(2.0f)
+				.HAlign(HAlign_Left)
+				.VAlign(VAlign_Center)
+				.AutoWidth()
+				[
+					PropertyCustomizationHelpers::MakeDeleteButton(FSimpleDelegate::CreateLambda([this]()
+					{
+						using namespace Frontend;
+
+						FMetasoundAssetBase* MetaSoundAsset = IMetasoundUObjectRegistry::Get().GetObjectAsAssetBase(MetaSound.Get());
+						if (!ensure(MetaSoundAsset))
+						{
+							return;
+						}
+
+						TArray<FMetasoundFrontendInterface> ImplementedInterfaces;
+						Algo::Transform(ImplementedInterfaceNames, ImplementedInterfaces, [](const FName& Name)
+						{
+							FMetasoundFrontendInterface Interface;
+							ISearchEngine::Get().FindInterfaceWithHighestVersion(Name, Interface);
+							return Interface;
+						});
+
+						{
+							FScopedTransaction(LOCTEXT("RemoveAllInterfacesTransaction", "Remove All MetaSound Interfaces"));
+							MetaSound.Get()->Modify();
+							MetaSoundAsset->GetGraphChecked().Modify();
+
+							FDocumentHandle DocumentHandle = MetaSoundAsset->GetDocumentHandle();
+							FModifyRootGraphInterfaces({ ImplementedInterfaces }, { }).Transform(DocumentHandle);
+						}
+
+						UpdateInterfaceNames();
+						InterfaceComboBox->RefreshOptions();
+						MetaSoundAsset->SetUpdateDetailsOnSynchronization();
+						MetaSoundAsset->SetSynchronizationRequired();
+					}), LOCTEXT("RemoveInterfaceTooltip1", "Removes all interfaces from the given MetaSound."))
+				];
+
+			const FText HeaderName = LOCTEXT("InterfacesGroupDisplayName", "Interfaces");
+			IDetailCategoryBuilder& InterfaceCategory = DetailLayout.EditCategory("Interfaces", HeaderName);
+
+			InterfaceCategory.AddCustomRow(HeaderName)
+			[
+				InterfaceUtilities
+			];
+
+			auto CreateInterfaceEntryWidget = [&](FName InInterfaceName) -> TSharedPtr<SWidget>
+			{
+				using namespace Frontend;
+
+				FMetasoundFrontendInterface InterfaceEntry;
+				if (!ensure(ISearchEngine::Get().FindInterfaceWithHighestVersion(InInterfaceName, InterfaceEntry)))
+				{
+					return SNullWidget::NullWidget;
+				}
+
+				TSharedRef<SWidget> RemoveButtonWidget = PropertyCustomizationHelpers::MakeDeleteButton(FSimpleDelegate::CreateLambda([this, InterfaceEntry]()
+				{
+					using namespace Frontend;
+
+					FMetasoundAssetBase* MetaSoundAsset = IMetasoundUObjectRegistry::Get().GetObjectAsAssetBase(MetaSound.Get());
+					if (!ensure(MetaSoundAsset))
+					{
+						return;
+					}
+
+					{
+						FScopedTransaction(LOCTEXT("RemoveInterfaceTransaction", "Remove MetaSound Interface"));
+						MetaSound.Get()->Modify();
+						MetaSoundAsset->GetGraphChecked().Modify();
+
+						FDocumentHandle DocumentHandle = MetaSoundAsset->GetDocumentHandle();
+						FModifyRootGraphInterfaces({ InterfaceEntry }, { }).Transform(DocumentHandle);
+					}
+
+					UpdateInterfaceNames();
+					InterfaceComboBox->RefreshOptions();
+					MetaSoundAsset->SetUpdateDetailsOnSynchronization();
+					MetaSoundAsset->SetSynchronizationRequired();
+				}), LOCTEXT("RemoveInterfaceTooltip2", "Removes the associated interface from the MetaSound."));
+
+				return SNew(SHorizontalBox)
+					+ SHorizontalBox::Slot()
+					.Padding(2.0f)
+					.HAlign(HAlign_Center)
+					.VAlign(VAlign_Center)
+					.AutoWidth()
+					[
+						SNew(STextBlock)
+						.Text(FText::FromName(InterfaceEntry.Version.Name))
+					]
+				+ SHorizontalBox::Slot()
+					.Padding(2.0f)
+					.HAlign(HAlign_Center)
+					.VAlign(VAlign_Center)
+					.AutoWidth()
+					[
+						RemoveButtonWidget
+					];
+			};
+
+			TArray<FName> InterfaceNames = ImplementedInterfaceNames.Array();
+			InterfaceNames.Sort([](const FName& A, const FName& B) { return A.LexicalLess(B); });
+			for (const FName& InterfaceName : InterfaceNames)
+			{
+				InterfaceCategory.AddCustomRow(FText::FromName(InterfaceName))
+				[
+					CreateInterfaceEntryWidget(InterfaceName)->AsShared()
+				];
+			}
+		}
+
+		void FMetasoundInterfacesDetailCustomization::UpdateInterfaceNames()
+		{
+			using namespace Frontend;
+
+			AddableInterfaceNames.Reset();
+			ImplementedInterfaceNames.Reset();
+
+			if (const FMetasoundAssetBase* MetaSoundAsset = IMetasoundUObjectRegistry::Get().GetObjectAsAssetBase(MetaSound.Get()))
+			{
+				auto GetVersionName = [](const FMetasoundFrontendVersion& Version) { return Version.Name; };
+				auto CanAddOrRemoveInterface = [](const FMetasoundFrontendVersion& Version)
+				{
+					using namespace Metasound::Frontend;
+
+					const FInterfaceRegistryKey Key = GetInterfaceRegistryKey(Version);
+					if (const IInterfaceRegistryEntry* Entry = IInterfaceRegistry::Get().FindInterfaceRegistryEntry(Key))
+					{
+						return Entry->EditorCanAddOrRemove();
+					}
+
+					return false;
+				};
+
+				const TSet<FMetasoundFrontendVersion>& ImplementedInterfaces = MetaSoundAsset->GetDocumentChecked().Interfaces;
+				Algo::TransformIf(ImplementedInterfaces, ImplementedInterfaceNames, CanAddOrRemoveInterface, GetVersionName);
+
+				TArray<FMetasoundFrontendInterface> Interfaces = ISearchEngine::Get().FindAllInterfaces();
+				for (const FMetasoundFrontendInterface& Interface : Interfaces)
+				{
+					if (!ImplementedInterfaceNames.Contains(Interface.Version.Name))
+					{
+						if (CanAddOrRemoveInterface(Interface.Version))
+						{
+							FString Name = Interface.Version.Name.ToString();
+							AddableInterfaceNames.Add(MakeShared<FString>(MoveTemp(Name)));
+						}
+					}
+				}
+			}
+		}
 	} // namespace Editor
 } // namespace Metasound
 #undef LOCTEXT_NAMESPACE
diff --git a/Engine/Plugins/Runtime/Metasound/Source/MetasoundEditor/Private/MetasoundDetailCustomization.h b/Engine/Plugins/Runtime/Metasound/Source/MetasoundEditor/Private/MetasoundDetailCustomization.h
--- a/Engine/Plugins/Runtime/Metasound/Source/MetasoundEditor/Private/MetasoundDetailCustomization.h
+++ b/Engine/Plugins/Runtime/Metasound/Source/MetasoundEditor/Private/MetasoundDetailCustomization.h
@@ -28,12 +28,27 @@ namespace Metasound
 			// End of IDetailCustomization interface
 
 		private:
-			void CustomizeInterfaces(IDetailLayoutBuilder& DetailLayout);
-
 			FName GetInterfaceVersionsPath() const;
 			FName GetMetadataRootClassPath() const;
 			FName GetMetadataPropertyPath() const;
 
+			TAttribute<bool> IsGraphEditableAttribute;
+
+			TWeakObjectPtr<UObject> MetaSound;
+
+			FName DocumentPropertyName;
+		};
+
+		class FMetasoundInterfacesDetailCustomization : public IDetailCustomization
+		{
+		public:
+			FMetasoundInterfacesDetailCustomization();
+
+			// IDetailCustomization interface
+			virtual void CustomizeDetails(IDetailLayoutBuilder& DetailLayout) override;
+			// End of IDetailCustomization interface
+
+		private:
 			void UpdateInterfaceNames();
 
 			TArray<TSharedPtr<FString>> AddableInterfaceNames;
@@ -43,8 +58,6 @@ namespace Metasound
 			TAttribute<bool> IsGraphEditableAttribute;
 
 			TWeakObjectPtr<UObject> MetaSound;
-
-			FName DocumentPropertyName;
 		};
 	} // namespace Editor
 } // namespace Metasound
\ No newline at end of file
diff --git a/Engine/Plugins/Runtime/Metasound/Source/MetasoundEditor/Private/MetasoundEditor.cpp b/Engine/Plugins/Runtime/Metasound/Source/MetasoundEditor/Private/MetasoundEditor.cpp
--- a/Engine/Plugins/Runtime/Metasound/Source/MetasoundEditor/Private/MetasoundEditor.cpp
+++ b/Engine/Plugins/Runtime/Metasound/Source/MetasoundEditor/Private/MetasoundEditor.cpp
@@ -384,16 +384,25 @@ namespace Metasound
 
 							const FText DisplayName = GraphMember->GetDisplayName();
 							const FText GetterToolTip = FText::Format(LOCTEXT("DropTargetGetterVariableToolTipFormat", "{0}\nAdd:\n* Get (Drop)\n* Get Delayed (Alt+Drop)\n"), DisplayName);
-							static const FText GetShowToolTip = LOCTEXT("ShowGettersToolTip", "Get (Ctrl+Drop)");
-							static const FText AddOrShowSetToolTip = LOCTEXT("AddOrShowSetToolTip", "");
+							static const FText GetJumpToToolTip = LOCTEXT("JumpToGettersToolTip", "Get (Ctrl+Drop)");
+							static const FText AddOrJumpToSetToolTip = LOCTEXT("AddOrJumpToSetToolTip", "");
 							FConstNodeHandle MutatorNodeHandle = Variable->GetConstVariableHandle()->FindMutatorNode();
 							if (MutatorNodeHandle->IsValid())
 							{
-								Message = FText::Format(LOCTEXT("DropTargetVariableShowFormat", "{0}\nShow:\n* {1}\n* Set (Shift+Drop, One per graph)"), GetterToolTip, GetShowToolTip);
+								Message = FText::Format(LOCTEXT("DropTargetVariableJumpToFormat", "{0}\nJump To:\n* {1}\n* Set (Shift+Drop, One per graph)"), GetterToolTip, GetJumpToToolTip);
 							}
 							else
 							{
-								Message = FText::Format(LOCTEXT("DropTargetVariableAddSetShowGetFormat", "{0}* Set (Shift+Drop)\n\nShow:\n* {1}"), GetterToolTip, GetShowToolTip);
+								TArray<FConstNodeHandle> AccessorNodeHandles = Variable->GetConstVariableHandle()->FindAccessorNodes();
+
+								if (AccessorNodeHandles.IsEmpty())
+								{
+									Message = FText::Format(LOCTEXT("DropTargetVariableAddSetGetFormat", "{0}* Set (Shift+Drop)"), GetterToolTip);
+								}
+								else
+								{
+									Message = FText::Format(LOCTEXT("DropTargetVariableAddSetJumpToGetFormat", "{0}* Set (Shift+Drop)\n\nJump To:\n* {1}"), GetterToolTip, GetJumpToToolTip);
+								}
 							}
 						}
 					}
@@ -562,21 +571,21 @@ namespace Metasound
 			.SetGroup(WorkspaceMenuCategoryRef)
 			.SetIcon(FSlateIcon(FEditorStyle::GetStyleSetName(), "GraphEditor.EventGraph_16x"));
 
-			InTabManager->RegisterTabSpawner(TabFactory::Names::Inspector, FOnSpawnTab::CreateLambda([InMetasoundDetails = MetasoundDetails](const FSpawnTabArgs& Args)
+			InTabManager->RegisterTabSpawner(TabFactory::Names::Details, FOnSpawnTab::CreateLambda([InMetasoundDetails = MetasoundDetails](const FSpawnTabArgs& Args)
 			{
-				return TabFactory::CreateInspectorTab(InMetasoundDetails, Args);
+				return TabFactory::CreateDetailsTab(InMetasoundDetails, Args);
 			}))
-			.SetDisplayName(LOCTEXT("InspectorTab", "Inspector"))
+			.SetDisplayName(LOCTEXT("DetailsTab", "Details"))
 			.SetGroup(WorkspaceMenuCategoryRef)
 			.SetIcon(FSlateIcon(FEditorStyle::GetStyleSetName(), "LevelEditor.Tabs.Details"));
 
-			InTabManager->RegisterTabSpawner(TabFactory::Names::Metasound, FOnSpawnTab::CreateLambda([InMetasoundMenu = MetasoundInterfaceMenu](const FSpawnTabArgs& Args)
+			InTabManager->RegisterTabSpawner(TabFactory::Names::Members, FOnSpawnTab::CreateLambda([InGraphMembersMenu = GraphMembersMenu](const FSpawnTabArgs& Args)
 			{
-				return TabFactory::CreateMetasoundTab(InMetasoundMenu, Args);
+				return TabFactory::CreateMembersTab(InGraphMembersMenu, Args);
 			}))
-			.SetDisplayName(LOCTEXT("MetasoundTab", "MetaSound"))
+			.SetDisplayName(LOCTEXT("MembersTab", "Members"))
 			.SetGroup(WorkspaceMenuCategoryRef)
-			.SetIcon(FSlateIcon(FEditorStyle::GetStyleSetName(), "LevelEditor.Tabs.Settings"));
+			.SetIcon(FSlateIcon("MetaSoundStyle", "MetasoundEditor.Metasound.Icon"));
 
 			InTabManager->RegisterTabSpawner(TabFactory::Names::Analyzers, FOnSpawnTab::CreateLambda([InAnalyzerWidget = BuildAnalyzerWidget()](const FSpawnTabArgs& Args)
 			{
@@ -585,6 +594,14 @@ namespace Metasound
 			.SetDisplayName(LOCTEXT("AnalyzersTab", "Analyzers"))
 			.SetGroup(WorkspaceMenuCategoryRef)
 			.SetIcon(FSlateIcon(FEditorStyle::GetStyleSetName(), "Kismet.Tabs.Palette"));
+
+			InTabManager->RegisterTabSpawner(TabFactory::Names::Interfaces, FOnSpawnTab::CreateLambda([InInterfacesDetails = InterfacesDetails](const FSpawnTabArgs& Args)
+			{
+				return TabFactory::CreateInterfacesTab(InInterfacesDetails, Args);
+			}))
+			.SetDisplayName(LOCTEXT("InterfacesTab", "Interfaces"))
+			.SetGroup(WorkspaceMenuCategoryRef)
+			.SetIcon(FSlateIcon(FEditorStyle::GetStyleSetName(), "ClassIcon.Interface"));
 		}
 
 		void FEditor::UnregisterTabSpawners(const TSharedRef<FTabManager>& InTabManager)
@@ -595,8 +612,9 @@ namespace Metasound
 
 			InTabManager->UnregisterTabSpawner(TabFactory::Names::Analyzers);
 			InTabManager->UnregisterTabSpawner(TabFactory::Names::GraphCanvas);
-			InTabManager->UnregisterTabSpawner(TabFactory::Names::Inspector);
-			InTabManager->UnregisterTabSpawner(TabFactory::Names::Metasound);
+			InTabManager->UnregisterTabSpawner(TabFactory::Names::Details);
+			InTabManager->UnregisterTabSpawner(TabFactory::Names::Members);
+			InTabManager->UnregisterTabSpawner(TabFactory::Names::Interfaces);
 		}
 
 		TSharedPtr<SWidget> FEditor::BuildAnalyzerWidget() const
@@ -681,8 +699,8 @@ namespace Metasound
 				NameChangeDelegateHandles.Reset();
 			}
 
+			InterfacesView.Reset();
 			DestroyAnalyzers();
-
 			check(GEditor);
 			GEditor->UnregisterForUndo(this);
 		}
@@ -724,7 +742,7 @@ namespace Metasound
 				NotifyAssetPrimeInProgress();
 			}
 
-			const TSharedRef<FTabManager::FLayout> StandaloneDefaultLayout = FTabManager::NewLayout("Standalone_MetasoundEditor_Layout_v9")
+			const TSharedRef<FTabManager::FLayout> StandaloneDefaultLayout = FTabManager::NewLayout("Standalone_MetasoundEditor_Layout_v10")
 				->AddArea
 				(
 					FTabManager::NewPrimaryArea()
@@ -741,14 +759,21 @@ namespace Metasound
 								FTabManager::NewStack()
 								->SetSizeCoefficient(0.25f)
 								->SetHideTabWell(false)
-								->AddTab(TabFactory::Names::Metasound, ETabState::OpenedTab)
+								->AddTab(TabFactory::Names::Members, ETabState::OpenedTab)
+							)
+							->Split
+							(
+								FTabManager::NewStack()
+								->SetSizeCoefficient(0.1f)
+								->SetHideTabWell(true)
+								->AddTab(TabFactory::Names::Interfaces, ETabState::OpenedTab)
 							)
 							->Split
 							(
 								FTabManager::NewStack()
 								->SetSizeCoefficient(0.50f)
 								->SetHideTabWell(false)
-								->AddTab(TabFactory::Names::Inspector, ETabState::OpenedTab)
+								->AddTab(TabFactory::Names::Details, ETabState::OpenedTab)
 							)
 						)
 						->Split
@@ -788,7 +813,7 @@ namespace Metasound
 			if (MetasoundDetails.IsValid())
 			{
 				MetasoundDetails->SetObjects(SelectedObjects);
-				MetasoundDetails->HideFilterArea(true);
+				MetasoundDetails->HideFilterArea(false);
 			}
 		}
 
@@ -939,7 +964,7 @@ namespace Metasound
 			Args.bHideSelectionTip = true;
 			Args.NotifyHook = this;
 
-			SAssignNew(MetasoundInterfaceMenu, SGraphActionMenu, false)
+			SAssignNew(GraphMembersMenu, SGraphActionMenu, false)
 				.AlphaSortItems(true)
 				.OnActionDoubleClicked(this, &FEditor::OnMemberActionDoubleClicked)
 				.OnActionDragged(this, &FEditor::OnActionDragged)
@@ -948,7 +973,7 @@ namespace Metasound
 // 				.OnCategoryTextCommitted(this, &FEditor::OnCategoryNameCommitted)
 				.OnCollectAllActions(this, &FEditor::CollectAllActions)
 				.OnCollectStaticSections(this, &FEditor::CollectStaticSections)
-// 				.OnContextMenuOpening(this, &FEditor::OnContextMenuOpening)
+ 				.OnContextMenuOpening(this, &FEditor::OnContextMenuOpening)
 				.OnCreateWidgetForAction(this, &FEditor::OnCreateWidgetForAction)
   				.OnCanRenameSelectedAction(this, &FEditor::CanRenameOnActionNode)
 				.OnGetFilterText(this, &FEditor::GetFilterText)
@@ -962,10 +987,21 @@ namespace Metasound
 
 			FPropertyEditorModule& PropertyModule = FModuleManager::LoadModuleChecked<FPropertyEditorModule>("PropertyEditor");
 			MetasoundDetails = PropertyModule.CreateDetailView(Args);
+			InterfacesDetails = PropertyModule.CreateDetailView(Args);
+			if (InterfacesDetails.IsValid())
+			{
+				InterfacesView = TStrongObjectPtr(NewObject<UMetasoundInterfacesView>());
+				InterfacesView->SetMetasound(Metasound);
+				const TArray<UObject*> InterfacesViewObj{ InterfacesView.Get() };
+
+				InterfacesDetails->SetObjects(InterfacesViewObj);
+				InterfacesDetails->HideFilterArea(true);
+			}
+
 			Palette = SNew(SMetasoundPalette);
 		}
 
-		// TODO: Tie in rename on GraphActionMenu.  For now, just renameable via field in inspector
+		// TODO: Tie in rename on GraphActionMenu.  For now, just renameable via field in details
 		bool FEditor::CanRenameOnActionNode(TWeakPtr<FGraphActionNode> InSelectedNode) const
 		{
 			return false;
@@ -1150,13 +1186,13 @@ namespace Metasound
 				Commands.ConvertFromPreset,
 				FExecuteAction::CreateSP(this, &FEditor::ConvertFromPreset));
 
-			ToolkitCommands->MapAction(
-				Commands.Delete,
-				FExecuteAction::CreateSP(this, &FEditor::DeleteSelected));
+			ToolkitCommands->MapAction(FGenericCommands::Get().Delete,
+				FExecuteAction::CreateSP(this, &FEditor::DeleteSelectedInterfaceItems),
+				FCanExecuteAction::CreateSP(this, &FEditor::CanDeleteInterfaceItems));
 
-			ToolkitCommands->MapAction(
-				FGenericCommands::Get().Rename,
-				FExecuteAction::CreateSP(this, &FEditor::RenameSelected));
+			ToolkitCommands->MapAction(FGenericCommands::Get().Rename,
+				FExecuteAction::CreateSP(this, &FEditor::RenameSelectedInterfaceItem),
+				FCanExecuteAction::CreateSP(this, &FEditor::CanRenameSelectedInterfaceItems));
 
 			ToolkitCommands->MapAction(
 				FEditorCommands::Get().UpdateNodeClass,
@@ -1424,9 +1460,9 @@ namespace Metasound
 
 		void FEditor::EditObjectSettings()
 		{
-			if (MetasoundInterfaceMenu.IsValid())
+			if (GraphMembersMenu.IsValid())
 			{
-				MetasoundInterfaceMenu->SelectItemByName(FName());
+				GraphMembersMenu->SelectItemByName(FName());
 			}
 
 			if (MetasoundGraphEditor.IsValid())
@@ -1581,7 +1617,7 @@ namespace Metasound
 					FCanExecuteAction::CreateLambda([this]() { return CanDuplicateNodes(); }));
 
 				GraphEditorCommands->MapAction(FGenericCommands::Get().Rename,
-					FExecuteAction::CreateLambda([this] { RenameSelected(); }),
+					FExecuteAction::CreateLambda([this] { RenameSelectedNode(); }),
 					FCanExecuteAction::CreateLambda([this]() { return CanRenameSelectedNodes(); }));
 
 				// Alignment Commands
@@ -1682,9 +1718,9 @@ namespace Metasound
 				}
 			}
 
-			if (MetasoundInterfaceMenu.IsValid() && !bManuallyClearingGraphSelection)
+			if (GraphMembersMenu.IsValid() && !bManuallyClearingGraphSelection)
 			{
-				MetasoundInterfaceMenu->SelectItemByName(FName());
+				GraphMembersMenu->SelectItemByName(FName());
 			}
 			SetSelection(Selection);
 		}
@@ -1724,7 +1760,7 @@ namespace Metasound
 
 				if (NextToSelect)
 				{
-					if (MetasoundInterfaceMenu->SelectItemByName(NextToSelect->GetMemberName(), ESelectInfo::Direct, static_cast<int32>(NextToSelect->GetSectionID())))
+					if (GraphMembersMenu->SelectItemByName(NextToSelect->GetMemberName(), ESelectInfo::Direct, static_cast<int32>(NextToSelect->GetSectionID())))
 					{
 						const TArray<UObject*> GraphMembersToSelect { NextToSelect };
 						SetSelection(GraphMembersToSelect);
@@ -1748,69 +1784,11 @@ namespace Metasound
 				return;
 			}
 
-			if (MetasoundInterfaceMenu.IsValid())
+			if (CanDeleteNodes())
 			{
-				TArray<TSharedPtr<FEdGraphSchemaAction>> Actions;
-				MetasoundInterfaceMenu->GetSelectedActions(Actions);
-
-				if (!Actions.IsEmpty())
-				{
-					TSharedPtr<FMetasoundGraphMemberSchemaAction> ActionToDelete;
-					for (const TSharedPtr<FEdGraphSchemaAction>& Action : Actions)
-					{
-						TSharedPtr<FMetasoundGraphMemberSchemaAction> MetasoundAction = StaticCastSharedPtr<FMetasoundGraphMemberSchemaAction>(Action);
-						if (MetasoundAction.IsValid())
-						{
-							const UMetasoundEditorGraphMember* GraphMember = MetasoundAction->GetGraphMember();
-							if (ensure(nullptr != GraphMember))
-							{
-								const FMetasoundFrontendVersion* InterfaceVersion = nullptr;
-								if (const UMetasoundEditorGraphVertex* Vertex = Cast<UMetasoundEditorGraphVertex>(GraphMember))
-								{
-									InterfaceVersion = &Vertex->GetInterfaceVersion();
-								}
-
-								if (InterfaceVersion && InterfaceVersion->IsValid())
-								{
-									if (MetasoundGraphEditor.IsValid())
-									{
-										const FText Notification = FText::Format(LOCTEXT("CannotDeleteInterfaceMemberNotificationFormat", "Cannot delete individual member of interface '{0}'."), FText::FromName(InterfaceVersion->Name));
-										FNotificationInfo Info(Notification);
-										Info.bFireAndForget = true;
-										Info.bUseSuccessFailIcons = false;
-										Info.ExpireDuration = 5.0f;
-
-										MetasoundGraphEditor->AddNotification(Info, false /* bSuccess */);
-									}
-								}
-								else
-								{
-									if (MetasoundDetails.IsValid())
-									{
-										if (MetasoundDetails->GetSelectedObjects().Contains(GraphMember))
-										{
-											FMetasoundAssetBase* MetasoundAsset = IMetasoundUObjectRegistry::Get().GetObjectAsAssetBase(Metasound);
-											check(MetasoundAsset);
-											MetasoundAsset->SetUpdateDetailsOnSynchronization();
-										}
-									}
-
-									ActionToDelete = MetasoundAction;
-									break;
-								}
-							}
-						}
-					}
-
-					if (ActionToDelete.IsValid())
-					{
-						DeleteInterfaceItem(ActionToDelete);
-					}
-					return;
-				}
+				DeleteSelectedNodes();
 			}
-
-			DeleteSelectedNodes();
+			DeleteSelectedInterfaceItems();
 		}
 
 		void FEditor::DeleteSelectedNodes()
@@ -1843,6 +1821,61 @@ namespace Metasound
 			}
 		}
 
+		void FEditor::DeleteSelectedInterfaceItems()
+		{
+			if (!IsGraphEditable() || !GraphMembersMenu.IsValid())
+			{
+				return;
+			}
+
+			TArray<TSharedPtr<FEdGraphSchemaAction>> Actions;
+			GraphMembersMenu->GetSelectedActions(Actions);
+			if (Actions.IsEmpty())
+			{
+				return;
+			}
+
+			TSharedPtr<FMetasoundGraphMemberSchemaAction> ActionToDelete;
+			for (const TSharedPtr<FEdGraphSchemaAction>& Action : Actions)
+			{
+				TSharedPtr<FMetasoundGraphMemberSchemaAction> MetasoundAction = StaticCastSharedPtr<FMetasoundGraphMemberSchemaAction>(Action);
+				if (MetasoundAction.IsValid())
+				{
+					const UMetasoundEditorGraphMember* GraphMember = MetasoundAction->GetGraphMember();
+					if (ensure(nullptr != GraphMember))
+					{
+						const FMetasoundFrontendVersion* InterfaceVersion = nullptr;
+						if (const UMetasoundEditorGraphVertex* Vertex = Cast<UMetasoundEditorGraphVertex>(GraphMember))
+						{
+							InterfaceVersion = &Vertex->GetInterfaceVersion();
+						}
+
+						if (InterfaceVersion && InterfaceVersion->IsValid())
+						{
+							if (MetasoundGraphEditor.IsValid())
+							{
+								const FText Notification = FText::Format(LOCTEXT("CannotDeleteInterfaceMemberNotificationFormat", "Cannot delete individual member of interface '{0}'."), FText::FromName(InterfaceVersion->Name));
+								FNotificationInfo Info(Notification);
+								Info.bFireAndForget = true;
+								Info.bUseSuccessFailIcons = false;
+								Info.ExpireDuration = 5.0f;
+
+								MetasoundGraphEditor->AddNotification(Info, false /* bSuccess */);
+							}
+						}
+						else
+						{
+							ActionToDelete = MetasoundAction;
+							if (ActionToDelete.IsValid())
+							{
+								DeleteInterfaceItem(ActionToDelete);
+							}
+						}
+					}
+				}
+			}
+		}
+
 		void FEditor::CutSelectedNodes()
 		{
 			CopySelectedNodes();
@@ -1961,6 +1994,53 @@ namespace Metasound
 			return false;
 		}
 
+		bool FEditor::CanDeleteInterfaceItems() const
+		{
+			if (!IsGraphEditable())
+			{
+				return false;
+			}
+
+			if (!GraphMembersMenu.IsValid())
+			{
+				return false;
+			}
+
+			TArray<TSharedPtr<FEdGraphSchemaAction>> Actions;
+			GraphMembersMenu->GetSelectedActions(Actions);
+
+			if (Actions.IsEmpty())
+			{
+				return false;
+			}
+
+			TSharedPtr<FMetasoundGraphMemberSchemaAction> ActionToDelete;
+			for (const TSharedPtr<FEdGraphSchemaAction>& Action : Actions)
+			{
+				TSharedPtr<FMetasoundGraphMemberSchemaAction> MetasoundAction = StaticCastSharedPtr<FMetasoundGraphMemberSchemaAction>(Action);
+				if (MetasoundAction.IsValid())
+				{
+					const UMetasoundEditorGraphMember* GraphMember = MetasoundAction->GetGraphMember();
+					if (ensure(nullptr != GraphMember))
+					{
+						const FMetasoundFrontendVersion* InterfaceVersion = nullptr;
+						if (const UMetasoundEditorGraphVertex* Vertex = Cast<UMetasoundEditorGraphVertex>(GraphMember))
+						{
+							InterfaceVersion = &Vertex->GetInterfaceVersion();
+						}
+
+						// Interface members cannot be deleted
+						const bool bIsInterfaceMember = InterfaceVersion && InterfaceVersion->IsValid();
+						if (!bIsInterfaceMember)
+						{
+							return true;
+						}
+					}
+				}
+			}
+			return false;
+		}
+
 		void FEditor::DuplicateNodes()
 		{
 			const FGraphPanelSelectionSet SelectedNodes = MetasoundGraphEditor->GetSelectedNodes();
@@ -2264,21 +2344,13 @@ namespace Metasound
 			}
 			return false;
 		}
-		
-		void FEditor::RenameSelected()
-		{
-			// Renaming selected nodes
-			if (CanRenameSelectedNodes())
-			{
-				RenameSelectedNode();
-				return;
-			}
 
-			// Renaming selected interface member
-			if (MetasoundInterfaceMenu.IsValid())
+		bool FEditor::CanRenameSelectedInterfaceItems() const
+		{
+			if (GraphMembersMenu.IsValid())
 			{
 				TArray<TSharedPtr<FEdGraphSchemaAction>> Actions;
-				MetasoundInterfaceMenu->GetSelectedActions(Actions);
+				GraphMembersMenu->GetSelectedActions(Actions);
 
 				if (!Actions.IsEmpty())
 				{
@@ -2291,29 +2363,14 @@ namespace Metasound
 							{
 								if (GraphMember->CanRename())
 								{
-									RenameInterfaceItem(MetasoundAction);
-									return;
+									return true;
 								}
 							}
 						}
 					}
 				}
 			}
-		}
-
-		void FEditor::RenameInterfaceItem(TSharedPtr<FMetasoundGraphMemberSchemaAction> ActionToRename)
-		{
-			using namespace Metasound::Frontend;
-			check(Metasound);
-
-			UMetasoundEditorGraphMember* GraphMember = ActionToRename->GetGraphMember();
-			if (ensure(GraphMember))
-			{
-				if (GraphMember->OnRenameRequested.IsBound())
-				{
-					GraphMember->OnRenameRequested.Broadcast();
-				} 
-			}
+			return false;
 		}
 
 		void FEditor::RenameSelectedNode() 
@@ -2337,7 +2394,7 @@ namespace Metasound
 				{
 					if (const UMetasoundEditorGraphMember* Member = MemberNode->GetMember()) 
 					{
-						MetasoundInterfaceMenu->SelectItemByName(Member->GetMemberName(), ESelectInfo::Direct, static_cast<int32>(Member->GetSectionID()));
+						GraphMembersMenu->SelectItemByName(Member->GetMemberName(), ESelectInfo::Direct, static_cast<int32>(Member->GetSectionID()));
 
 						if (Member->OnRenameRequested.IsBound())
 						{
@@ -2348,19 +2405,55 @@ namespace Metasound
 			}
 		}
 
+		void FEditor::RenameSelectedInterfaceItem()
+		{
+			if (GraphMembersMenu.IsValid())
+			{
+				TArray<TSharedPtr<FEdGraphSchemaAction>> Actions;
+				GraphMembersMenu->GetSelectedActions(Actions);
+
+				if (!Actions.IsEmpty())
+				{
+					for (const TSharedPtr<FEdGraphSchemaAction>& Action : Actions)
+					{
+						TSharedPtr<FMetasoundGraphMemberSchemaAction> MetasoundAction = StaticCastSharedPtr<FMetasoundGraphMemberSchemaAction>(Action);
+						if (MetasoundAction.IsValid())
+						{
+							if (const UMetasoundEditorGraphMember* GraphMember = MetasoundAction->GetGraphMember())
+							{
+								if (GraphMember->CanRename())
+								{
+									if (GraphMember->OnRenameRequested.IsBound())
+									{
+										GraphMember->OnRenameRequested.Broadcast();
+										return;
+									}
+								}
+							}
+						}
+					}
+				}
+			}
+		}
+
 		void FEditor::RefreshDetails()
 		{
 			if (MetasoundDetails.IsValid())
 			{
 				MetasoundDetails->ForceRefresh();
 			}
+
+			if (InterfacesDetails.IsValid())
+			{
+				InterfacesDetails->ForceRefresh();
+			}
 		}
 
 		void FEditor::RefreshInterface()
 		{
-			if (MetasoundInterfaceMenu.IsValid())
+			if (GraphMembersMenu.IsValid())
 			{
-				MetasoundInterfaceMenu->RefreshAllActions(true /* bPreserveExpansion */);
+				GraphMembersMenu->RefreshAllActions(true /* bPreserveExpansion */);
 			}
 		}
 
@@ -2442,14 +2535,14 @@ namespace Metasound
 
 		void FEditor::OnInputNameChanged(FGuid InNodeID)
 		{
-			if (!MetasoundInterfaceMenu.IsValid() || !Metasound)
+			if (!GraphMembersMenu.IsValid() || !Metasound)
 			{
 				return;
 			}
 
 			TArray<TSharedPtr<FEdGraphSchemaAction>> SelectedActions;
-			MetasoundInterfaceMenu->GetSelectedActions(SelectedActions);
-			MetasoundInterfaceMenu->RefreshAllActions(/* bPreserveExpansion */ true);
+			GraphMembersMenu->GetSelectedActions(SelectedActions);
+			GraphMembersMenu->RefreshAllActions(/* bPreserveExpansion */ true);
 
 			for(const TSharedPtr<FEdGraphSchemaAction>& Action : SelectedActions)
 			{
@@ -2461,7 +2554,7 @@ namespace Metasound
 						if (InNodeID == Member->GetMemberID())
 						{
 							const FName ActionName = Member->GetMemberName();
-							MetasoundInterfaceMenu->SelectItemByName(ActionName, ESelectInfo::Direct, Action->GetSectionID());
+							GraphMembersMenu->SelectItemByName(ActionName, ESelectInfo::Direct, Action->GetSectionID());
 							break;
 						}
 					}
@@ -2473,14 +2566,14 @@ namespace Metasound
 
 		void FEditor::OnOutputNameChanged(FGuid InNodeID)
 		{
-			if (!MetasoundInterfaceMenu.IsValid())
+			if (!GraphMembersMenu.IsValid())
 			{
 				return;
 			}
 
 			TArray<TSharedPtr<FEdGraphSchemaAction>> SelectedActions;
-			MetasoundInterfaceMenu->GetSelectedActions(SelectedActions);
-			MetasoundInterfaceMenu->RefreshAllActions(/* bPreserveExpansion */ true);
+			GraphMembersMenu->GetSelectedActions(SelectedActions);
+			GraphMembersMenu->RefreshAllActions(/* bPreserveExpansion */ true);
 
 			for (const TSharedPtr<FEdGraphSchemaAction>& Action : SelectedActions)
 			{
@@ -2492,7 +2585,7 @@ namespace Metasound
 						if (InNodeID == Member->GetMemberID())
 						{
 							const FName ActionName = Member->GetMemberName();
-							MetasoundInterfaceMenu->SelectItemByName(ActionName, ESelectInfo::Direct, Action->GetSectionID());
+							GraphMembersMenu->SelectItemByName(ActionName, ESelectInfo::Direct, Action->GetSectionID());
 							break;
 						}
 					}
@@ -2504,14 +2597,14 @@ namespace Metasound
 
 		void FEditor::OnVariableNameChanged(FGuid InVariableID)
 		{
-			if (!MetasoundInterfaceMenu.IsValid())
+			if (!GraphMembersMenu.IsValid())
 			{
 				return;
 			}
 
 			TArray<TSharedPtr<FEdGraphSchemaAction>> SelectedActions;
-			MetasoundInterfaceMenu->GetSelectedActions(SelectedActions);
-			MetasoundInterfaceMenu->RefreshAllActions(/* bPreserveExpansion */ true);
+			GraphMembersMenu->GetSelectedActions(SelectedActions);
+			GraphMembersMenu->RefreshAllActions(/* bPreserveExpansion */ true);
 
 			for (const TSharedPtr<FEdGraphSchemaAction>& Action : SelectedActions)
 			{
@@ -2522,7 +2615,7 @@ namespace Metasound
 					{
 						if (InVariableID == Variable->GetVariableID())
 						{
-							MetasoundInterfaceMenu->SelectItemByName(Variable->GetMemberName(), ESelectInfo::Direct, Action->GetSectionID());
+							GraphMembersMenu->SelectItemByName(Variable->GetMemberName(), ESelectInfo::Direct, Action->GetSectionID());
 							break;
 						}
 					}
@@ -2738,6 +2831,61 @@ namespace Metasound
 			}
 		}
 
+		bool FEditor::CanJumpToNodesForSelectedInterfaceItem() const 
+		{
+			if (!GraphMembersMenu.IsValid())
+			{
+				return false;
+			}
+			TArray<TSharedPtr<FEdGraphSchemaAction>> Actions;
+			GraphMembersMenu->GetSelectedActions(Actions);
+
+			if (!Actions.IsEmpty())
+			{
+				for (const TSharedPtr<FEdGraphSchemaAction>& Action : Actions)
+				{
+					TSharedPtr<FMetasoundGraphMemberSchemaAction> MetasoundAction = StaticCastSharedPtr<FMetasoundGraphMemberSchemaAction>(Action);
+					if (MetasoundAction.IsValid())
+					{
+						if (const UMetasoundEditorGraphMember* GraphMember = MetasoundAction->GetGraphMember())
+						{
+							TArray<UMetasoundEditorGraphNode*> Nodes = GraphMember->GetNodes();
+							if (!Nodes.IsEmpty())
+							{
+								return true;
+							}
+						}
+					}
+				}
+			}
+			return false;
+		}
+
+		void FEditor::JumpToNodesForSelectedInterfaceItem()
+		{
+			if (GraphMembersMenu.IsValid())
+			{
+				TArray<TSharedPtr<FEdGraphSchemaAction>> Actions;
+				GraphMembersMenu->GetSelectedActions(Actions);
+
+				if (!Actions.IsEmpty())
+				{
+					for (const TSharedPtr<FEdGraphSchemaAction>& Action : Actions)
+					{
+						TSharedPtr<FMetasoundGraphMemberSchemaAction> MetasoundAction = StaticCastSharedPtr<FMetasoundGraphMemberSchemaAction>(Action);
+						if (MetasoundAction.IsValid())
+						{
+							if (const UMetasoundEditorGraphMember* GraphMember = MetasoundAction->GetGraphMember())
+							{
+								JumpToNodes(GraphMember->GetNodes());
+								return;
+							}
+						}
+					}
+				}
+			}
+		}
+
 		FActionMenuContent FEditor::OnCreateGraphActionMenu(UEdGraph* InGraph, const FVector2D& InNodePosition, const TArray<UEdGraphPin*>& InDraggedPins, bool bAutoExpand, SGraphEditor::FActionMenuClosed InOnMenuClosed)
 		{
 			TSharedRef<SMetasoundActionMenu> ActionMenu = SNew(SMetasoundActionMenu)
@@ -2790,6 +2938,37 @@ namespace Metasound
 			return SNew(SMetaSoundGraphPaletteItem, InCreateData);
 		}
 
+		TSharedPtr<SWidget> FEditor::OnContextMenuOpening()
+		{
+			if (!GraphMembersMenu.IsValid())
+			{
+				return nullptr;
+			}
+
+			// Context menu should only open when graph members are selected
+			TArray<TSharedPtr<FEdGraphSchemaAction>> Actions;
+			GraphMembersMenu->GetSelectedActions(Actions);
+			if (Actions.IsEmpty())
+			{
+				return nullptr;
+			}
+
+			FMenuBuilder MenuBuilder(true, ToolkitCommands);
+
+			MenuBuilder.AddMenuEntry(FGenericCommands::Get().Delete);
+			MenuBuilder.AddMenuEntry(FGenericCommands::Get().Rename);
+
+			MenuBuilder.AddMenuEntry(
+				LOCTEXT("JumpToNodesMenuEntry", "Jump to Node(s) in Graph"),
+				LOCTEXT("JumpToNodesMenuEntryTooltip", "Jump to the corresponding node(s) in the MetaSound graph"),
+				FSlateIcon(),
+				FUIAction(
+					FExecuteAction::CreateSP(this, &FEditor::JumpToNodesForSelectedInterfaceItem), 
+					FCanExecuteAction::CreateSP(this, &FEditor::CanJumpToNodesForSelectedInterfaceItem)));
+
+			return MenuBuilder.MakeWidget();
+		}
+			
 		void FEditor::Tick(float DeltaTime)
 		{
 			if (!Metasound)
@@ -3027,12 +3206,12 @@ namespace Metasound
 
 			FGraphBuilder::RegisterGraphWithFrontend(*Metasound);
 
-			if (MetasoundInterfaceMenu.IsValid())
+			if (GraphMembersMenu.IsValid())
 			{
-				MetasoundInterfaceMenu->RefreshAllActions(/* bPreserveExpansion */ true);
+				GraphMembersMenu->RefreshAllActions(/* bPreserveExpansion */ true);
 				if (!NameToSelect.IsNone())
 				{
-					MetasoundInterfaceMenu->SelectItemByName(NameToSelect);
+					GraphMembersMenu->SelectItemByName(NameToSelect);
 					SetSelection(SelectedObjects);
 				}
 			}
diff --git a/Engine/Plugins/Runtime/Metasound/Source/MetasoundEditor/Private/MetasoundEditor.h b/Engine/Plugins/Runtime/Metasound/Source/MetasoundEditor/Private/MetasoundEditor.h
--- a/Engine/Plugins/Runtime/Metasound/Source/MetasoundEditor/Private/MetasoundEditor.h
+++ b/Engine/Plugins/Runtime/Metasound/Source/MetasoundEditor/Private/MetasoundEditor.h
@@ -26,6 +26,8 @@
 #include "WorkflowOrientedApp/WorkflowTabFactory.h"
 #include "WorkflowOrientedApp/WorkflowTabManager.h"
 
+#include "MetasoundEditor.generated.h"
+
 
 // Forward Declarations
 class FTabManager;
@@ -46,6 +48,20 @@ struct FMeterResults;
 struct FPropertyChangedEvent;
 
 
+// Simple class for the interfaces details tab to keep track of its corresponding Metasound 
+UCLASS(MinimalAPI)
+class UMetasoundInterfacesView : public UObject
+{
+	GENERATED_BODY()
+public:
+	UMetasoundInterfacesView() = default;
+
+	void SetMetasound(UObject* InMetasound) { Metasound = InMetasound; }
+	const TWeakObjectPtr<UObject> GetMetasound() { return Metasound; }
+private:
+	TWeakObjectPtr<UObject> Metasound;
+};
+
 namespace Metasound
 {
 	namespace Editor
@@ -179,6 +195,7 @@ namespace Metasound
 			FText GetSectionTitle(ENodeSection InSection) const;
 			FText OnGetSectionTitle(int32 InSectionID);
 			TSharedRef<SWidget> OnCreateWidgetForAction(struct FCreateWidgetForActionData* const InCreateData);
+			TSharedPtr<SWidget> OnContextMenuOpening();
 
 			/** Called when the selection changes in the GraphEditor */
 			void OnSelectedNodesChanged(const TSet<UObject*>& NewSelection);
@@ -204,6 +221,9 @@ namespace Metasound
 			/** Delete the currently selected nodes */
 			void DeleteSelectedNodes();
 
+			/** Delete the currently selected interface items */
+			void DeleteSelectedInterfaceItems();
+
 			/** Cut the currently selected nodes */
 			void CutSelectedNodes();
 
@@ -219,14 +239,26 @@ namespace Metasound
 			/** Whether the currently selected node(s) can be deleted */
 			bool CanDeleteNodes() const;
 
+			/** Whether the currently selected interface item(s) can be deleted */
+			bool CanDeleteInterfaceItems() const;
+
 			/** Whether at least one of the currently selected node(s) can be renamed. */
 			bool CanRenameSelectedNodes() const;
 
 			/** Rename selected node (currently applies to comments and member nodes). */
 			void RenameSelectedNode();
 
+			/** Whether at least one of the currently selected interface item(s) can be renamed. */
+			bool CanRenameSelectedInterfaceItems() const;
+
 			/** Rename selected interface item. */
-			void RenameInterfaceItem(TSharedPtr<FMetasoundGraphMemberSchemaAction> ActionToDelete);
+			void RenameSelectedInterfaceItem();
+
+			/** Whether there are nodes to jump to for the currently selected interface item. */
+			bool CanJumpToNodesForSelectedInterfaceItem() const;
+
+			/** Jumps to the nodes corresponding to the first valid currently selected interface item. */
+			void JumpToNodesForSelectedInterfaceItem();
 
 			/** Called to undo the last action */
 			void UndoGraphAction();
@@ -277,10 +309,10 @@ namespace Metasound
 			/** Converts the MetaSound from a preset to a fully modifiable MetaSound. */
 			void ConvertFromPreset();
 
-			/** Show the Metasound object's Source settings in the Inspector */
+			/** Show the Metasound object's Source settings in the Details panel */
 			void EditSourceSettings();
 
-			/** Show the Metasound object's settings in the Inspector */
+			/** Show the Metasound object's settings in the Details panel */
 			void EditMetasoundSettings();
 
 			/** Add an input to the currently selected node */
@@ -323,8 +355,12 @@ namespace Metasound
 			/** Details tab */
 			TSharedPtr<IDetailsView> MetasoundDetails;
 
-			/** Metasound Interface menu */
-			TSharedPtr<SGraphActionMenu> MetasoundInterfaceMenu;
+			/** Interfaces tab */
+			TSharedPtr<IDetailsView> InterfacesDetails;
+			TStrongObjectPtr<UMetasoundInterfacesView> InterfacesView;
+
+			/** Metasound graph members menu */
+			TSharedPtr<SGraphActionMenu> GraphMembersMenu;
 
 			/** Meter used in the analyzer tab for auditioning preview output. */
 			TSharedPtr<FEditorMeter> OutputMeter;
diff --git a/Engine/Plugins/Runtime/Metasound/Source/MetasoundEditor/Private/MetasoundEditorModule.cpp b/Engine/Plugins/Runtime/Metasound/Source/MetasoundEditor/Private/MetasoundEditorModule.cpp
--- a/Engine/Plugins/Runtime/Metasound/Source/MetasoundEditor/Private/MetasoundEditorModule.cpp
+++ b/Engine/Plugins/Runtime/Metasound/Source/MetasoundEditor/Private/MetasoundEditorModule.cpp
@@ -59,6 +59,8 @@ FAutoConsoleVariableRef CVarMetaSoundEditorAsyncRegistrationEnabled(
 	TEXT("0: Disabled, !0: Enabled (default)"),
 	ECVF_Default);
 
+// Forward declarations 
+class UMetasoundInterfacesView;
 
 namespace Metasound
 {
@@ -625,6 +627,10 @@ namespace Metasound
 					UMetaSoundSource::StaticClass()->GetFName(),
 					FOnGetDetailCustomizationInstance::CreateLambda([]() { return MakeShared<FMetasoundDetailCustomization>(UMetaSoundSource::GetDocumentPropertyName()); }));
 
+				PropertyModule.RegisterCustomClassLayout(
+					UMetasoundInterfacesView::StaticClass()->GetFName(),
+					FOnGetDetailCustomizationInstance::CreateLambda([]() { return MakeShared<FMetasoundInterfacesDetailCustomization>(); }));
+
 				PropertyModule.RegisterCustomClassLayout(
 					UMetasoundEditorGraphInput::StaticClass()->GetFName(),
 					FOnGetDetailCustomizationInstance::CreateLambda([]() { return MakeShared<FMetasoundInputDetailCustomization>(); }));
diff --git a/Engine/Plugins/Runtime/Metasound/Source/MetasoundEditor/Private/MetasoundEditorTabFactory.h b/Engine/Plugins/Runtime/Metasound/Source/MetasoundEditor/Private/MetasoundEditorTabFactory.h
--- a/Engine/Plugins/Runtime/Metasound/Source/MetasoundEditor/Private/MetasoundEditorTabFactory.h
+++ b/Engine/Plugins/Runtime/Metasound/Source/MetasoundEditor/Private/MetasoundEditorTabFactory.h
@@ -24,10 +24,11 @@ namespace Metasound
 			namespace Names
 			{
 				const FName Analyzers = "MetasoundEditor_Analyzers";
-				const FName Inspector = "MetasoundEditor_Inspector";
+				const FName Details = "MetasoundEditor_Details";
 				const FName GraphCanvas = "MetasoundEditor_GraphCanvas";
-				const FName Metasound = "MetasoundEditor_Metasound";
+				const FName Members = "MetasoundEditor_Members";
 				const FName Palette = "MetasoundEditor_Palette";
+				const FName Interfaces = "MetasoundEditor_Interfaces";
 			}
 
 			TSharedRef<SDockTab> CreateAnalyzersTab(TSharedPtr<SWidget> InAnalyzerWidget, const FSpawnTabArgs& Args)
@@ -56,23 +57,23 @@ namespace Metasound
 				return SpawnedTab;
 			}
 
-			TSharedRef<SDockTab> CreateInspectorTab(TSharedPtr<IDetailsView> DetailsView, const FSpawnTabArgs& Args)
+			TSharedRef<SDockTab> CreateDetailsTab(TSharedPtr<IDetailsView> DetailsView, const FSpawnTabArgs& Args)
 			{
-				check(Args.GetTabId() == Names::Inspector);
+				check(Args.GetTabId() == Names::Details);
 
 				return SNew(SDockTab)
-					.Label(LOCTEXT("MetaSoundInspectorTitle", "Inspector"))
+					.Label(LOCTEXT("MetaSoundDetailsTitle", "Details"))
 					[
 						DetailsView.ToSharedRef()
 					];
 			}
 
-			TSharedRef<SDockTab> CreateMetasoundTab(TSharedPtr<SGraphActionMenu> GraphActionMenu, const FSpawnTabArgs& Args)
+			TSharedRef<SDockTab> CreateMembersTab(TSharedPtr<SGraphActionMenu> GraphActionMenu, const FSpawnTabArgs& Args)
 			{
-				check(Args.GetTabId() == Names::Metasound);
+				check(Args.GetTabId() == Names::Members);
 
 				TSharedRef<SDockTab> NewTab = SNew(SDockTab)
-				.Label(LOCTEXT("MetaSoundGeneralTitle", "MetaSound"))
+				.Label(LOCTEXT("GraphMembersMenulTitle", "Members"))
 				[
 					GraphActionMenu.ToSharedRef()
 				];
@@ -95,6 +96,17 @@ namespace Metasound
 						Palette.ToSharedRef()
 					];
 			}
+
+			TSharedRef<SDockTab> CreateInterfacesTab(TSharedPtr<IDetailsView> InInterfacesDetails, const FSpawnTabArgs& Args)
+			{
+				check(Args.GetTabId() == Names::Interfaces);
+
+				return SNew(SDockTab)
+					.Label(LOCTEXT("MetasoundInterfacesDetailsTitle", "Interfaces"))
+					[
+						InInterfacesDetails.ToSharedRef()
+					];
+			}
 		}
 	} // namespace Editor
 } // namespace Metasound
diff --git a/Engine/Plugins/Runtime/Metasound/Source/MetasoundEditor/Public/MetasoundEditorCommands.h b/Engine/Plugins/Runtime/Metasound/Source/MetasoundEditor/Public/MetasoundEditorCommands.h
--- a/Engine/Plugins/Runtime/Metasound/Source/MetasoundEditor/Public/MetasoundEditorCommands.h
+++ b/Engine/Plugins/Runtime/Metasound/Source/MetasoundEditor/Public/MetasoundEditorCommands.h
@@ -46,10 +46,10 @@ namespace Metasound
 			/** Removes an input from the node */
 			TSharedPtr<FUICommandInfo> DeleteInput;
 
-			/** Shows the MetasoundSource's specific settings in the Inspector (if of respective type) */
+			/** Shows the MetasoundSource's specific settings in the Details panel (if of respective type) */
 			TSharedPtr<FUICommandInfo> EditSourceSettings;
 
-			/** Shows the Metasound's top-level general object settings in the Inspector */
+			/** Shows the Metasound's top-level general object settings in the Details panel */
 			TSharedPtr<FUICommandInfo> EditMetasoundSettings;
 
 			// Run updates on selected node(s) class(es) if required by any subset of selected nodes
