Fix MetaSound Variable display naming namespace syntax to follow that of other members (inputs & outputs)
#rnx

diff --git a/Engine/Plugins/Runtime/Metasound/Source/MetasoundEditor/Private/MetasoundEditor.cpp b/Engine/Plugins/Runtime/Metasound/Source/MetasoundEditor/Private/MetasoundEditor.cpp
--- a/Engine/Plugins/Runtime/Metasound/Source/MetasoundEditor/Private/MetasoundEditor.cpp
+++ b/Engine/Plugins/Runtime/Metasound/Source/MetasoundEditor/Private/MetasoundEditor.cpp
@@ -2631,11 +2631,11 @@ namespace Metasound
 			UMetasoundEditorGraph& EdGraph = GetMetaSoundGraphChecked();
 			Frontend::FConstGraphHandle FrontendGraph = MetasoundAsset->GetRootGraphHandle();
 
-			auto GetVertexCategory = [](const Frontend::FConstNodeHandle& NodeHandle)
+			auto GetMemberCategory = [](FName InFullCategoryName)
 			{
 				FName InterfaceName;
 				FName MemberName;
-				Audio::FParameterPath::SplitName(NodeHandle->GetNodeName(), InterfaceName, MemberName);
+				Audio::FParameterPath::SplitName(InFullCategoryName, InterfaceName, MemberName);
 
 				if (InterfaceName.IsNone())
 				{
@@ -2647,13 +2647,13 @@ namespace Metasound
 				return FText::FromString(CategoryString);
 			};
 
-			FrontendGraph->IterateConstNodes([this, &GetVertexCategory, &EdGraph, ActionList = &OutAllActions](const Frontend::FConstNodeHandle& Input)
+			FrontendGraph->IterateConstNodes([this, &GetMemberCategory, &EdGraph, ActionList = &OutAllActions](const Frontend::FConstNodeHandle& Input)
 			{
 				constexpr bool bIncludeNamespace = false;
 				const FText Tooltip = Input->GetDescription();
 				const FText MenuDesc = FGraphBuilder::GetDisplayName(*Input, bIncludeNamespace);
 				const FGuid NodeID = Input->GetID();
-				const FText Category = GetVertexCategory(Input);
+				const FText Category = GetMemberCategory(Input->GetNodeName());
 
 				TSharedPtr<FMetasoundGraphMemberSchemaAction> NewFuncAction = MakeShared<FMetasoundGraphMemberSchemaAction>(Category, MenuDesc, Tooltip, 1, ENodeSection::Inputs);
 				NewFuncAction->Graph = &EdGraph;
@@ -2672,14 +2672,14 @@ namespace Metasound
 				}
 			}, EMetasoundFrontendClassType::Input);
 
-			FrontendGraph->IterateConstNodes([this, &GetVertexCategory, &EdGraph, ActionList = &OutAllActions](const Frontend::FConstNodeHandle& Output)
+			FrontendGraph->IterateConstNodes([this, &GetMemberCategory, &EdGraph, ActionList = &OutAllActions](const Frontend::FConstNodeHandle& Output)
 			{
 				constexpr bool bIncludeNamespace = false;
 
 				const FText Tooltip = Output->GetDescription();
 				const FText MenuDesc = FGraphBuilder::GetDisplayName(*Output, bIncludeNamespace);
 				const FGuid NodeID = Output->GetID();
-				const FText Category = GetVertexCategory(Output);
+				const FText Category = GetMemberCategory(Output->GetNodeName());
 
 				TSharedPtr<FMetasoundGraphMemberSchemaAction> NewFuncAction = MakeShared<FMetasoundGraphMemberSchemaAction>(Category, MenuDesc, Tooltip, 1, ENodeSection::Outputs);
 				NewFuncAction->Graph = &EdGraph;
@@ -2702,8 +2702,9 @@ namespace Metasound
 			{
 				const FText MenuDesc = FGraphBuilder::GetDisplayName(*Variable);
 				const FGuid VariableID = Variable->GetID();
+				const FText Category = GetMemberCategory(Variable->GetName());
 
-				TSharedPtr<FMetasoundGraphMemberSchemaAction> NewFuncAction = MakeShared<FMetasoundGraphMemberSchemaAction>(FText::GetEmpty(), MenuDesc, FText::GetEmpty(), 1, ENodeSection::Variables);
+				TSharedPtr<FMetasoundGraphMemberSchemaAction> NewFuncAction = MakeShared<FMetasoundGraphMemberSchemaAction>(Category, MenuDesc, FText::GetEmpty(), 1, ENodeSection::Variables);
 				NewFuncAction->Graph = &EdGraph;
 				NewFuncAction->MemberID = VariableID; 
 				OutAllActions.AddAction(NewFuncAction);
diff --git a/Engine/Plugins/Runtime/Metasound/Source/MetasoundEditor/Private/MetasoundEditorGraphBuilder.cpp b/Engine/Plugins/Runtime/Metasound/Source/MetasoundEditor/Private/MetasoundEditorGraphBuilder.cpp
--- a/Engine/Plugins/Runtime/Metasound/Source/MetasoundEditor/Private/MetasoundEditorGraphBuilder.cpp
+++ b/Engine/Plugins/Runtime/Metasound/Source/MetasoundEditor/Private/MetasoundEditorGraphBuilder.cpp
@@ -180,13 +180,22 @@ namespace Metasound
 			return DisplayName;
 		}
 
-		FText FGraphBuilder::GetDisplayName(const Frontend::IVariableController& InFrontendVariable)
+		FText FGraphBuilder::GetDisplayName(const Frontend::IVariableController& InFrontendVariable, bool bInIncludeNamespace)
 		{
 			FText DisplayName = InFrontendVariable.GetDisplayName();
 			if (DisplayName.IsEmptyOrWhitespace())
 			{
-				DisplayName = FText::FromName(InFrontendVariable.GetName());
+				FName Namespace;
+				FName ParameterName;
+				Audio::FParameterPath::SplitName(InFrontendVariable.GetName(), Namespace, ParameterName);
+
+				DisplayName = FText::FromName(ParameterName);
+				if (bInIncludeNamespace && !Namespace.IsNone())
+				{
+					return FText::Format(LOCTEXT("ClassMetadataDisplayNameWithNamespaceFormat", "{0} ({1})"), DisplayName, FText::FromName(Namespace));
+				}
 			}
+
 			return DisplayName;
 		}
 
diff --git a/Engine/Plugins/Runtime/Metasound/Source/MetasoundEditor/Private/MetasoundEditorGraphBuilder.h b/Engine/Plugins/Runtime/Metasound/Source/MetasoundEditor/Private/MetasoundEditorGraphBuilder.h
--- a/Engine/Plugins/Runtime/Metasound/Source/MetasoundEditor/Private/MetasoundEditorGraphBuilder.h
+++ b/Engine/Plugins/Runtime/Metasound/Source/MetasoundEditor/Private/MetasoundEditorGraphBuilder.h
@@ -98,7 +98,7 @@ namespace Metasound
 
 			// Returns a display name for a variable. If the variable has an empty or whitespace
 			// only DisplayName, then the VariableName is used. 
-			static FText GetDisplayName(const Frontend::IVariableController& InFrontendVariable);
+			static FText GetDisplayName(const Frontend::IVariableController& InFrontendVariable, bool bInIncludeNamespace = false);
 
 			// Returns the PinName for an IOutputController.
 			static FName GetPinName(const Frontend::IOutputController& InFrontendOutput);
