Fix some loc duplicate keys batch 3 (hopefully the last batch)

diff --git a/Engine/Plugins/Experimental/ChaosEditor/Source/FractureEditor/Private/GeometryCollectionOutlinerDragDrop.cpp b/Engine/Plugins/Experimental/ChaosEditor/Source/FractureEditor/Private/GeometryCollectionOutlinerDragDrop.cpp
--- a/Engine/Plugins/Experimental/ChaosEditor/Source/FractureEditor/Private/GeometryCollectionOutlinerDragDrop.cpp
+++ b/Engine/Plugins/Experimental/ChaosEditor/Source/FractureEditor/Private/GeometryCollectionOutlinerDragDrop.cpp
@@ -48,7 +48,7 @@ bool FGeometryCollectionBoneDragDrop::ValidateDrop(const FGeometryCollection* Ot
 
 		if (ContainsEmbedded())
 		{
-			MessageText = NSLOCTEXT("GeometryCollectionOutliner", "EmbeddedToCluster_GeometryCollectionOutliner",
+			MessageText = NSLOCTEXT("GeometryCollectionOutliner", "EmbeddedToCluster_GeometryCollectionOutliner2",
 				"Embedded geometry will convert to rigids when parented to cluster.");
 			return true;
 		}
diff --git a/Engine/Plugins/FX/Niagara/Source/Niagara/Private/DataInterface/NiagaraDataInterfaceAsyncGpuTrace.cpp b/Engine/Plugins/FX/Niagara/Source/Niagara/Private/DataInterface/NiagaraDataInterfaceAsyncGpuTrace.cpp
--- a/Engine/Plugins/FX/Niagara/Source/Niagara/Private/DataInterface/NiagaraDataInterfaceAsyncGpuTrace.cpp
+++ b/Engine/Plugins/FX/Niagara/Source/Niagara/Private/DataInterface/NiagaraDataInterfaceAsyncGpuTrace.cpp
@@ -283,7 +283,7 @@ void UNiagaraDataInterfaceAsyncGpuTrace::GetFunctions(TArray<FNiagaraFunctionSig
 	}
 
 	{
-		const FText TraceCountDescription = LOCTEXT("QueryIDDescription", "Number of async raytrace requests to be reserved");
+		const FText TraceCountDescription = LOCTEXT("QueryIDDescription2", "Number of async raytrace requests to be reserved");
 		const FText FirstQueryIDValidDescription = LOCTEXT("FirstQueryIDValidDescription", "The first index in the block reserved through this call");
 		const FText ReserveValidDescription = LOCTEXT("ReserveValidDescription", "Returns true if the requested indices were reserved");
 
diff --git a/Engine/Plugins/Runtime/Metasound/Source/MetasoundEditor/Private/MetasoundEditorGraphSchema.cpp b/Engine/Plugins/Runtime/Metasound/Source/MetasoundEditor/Private/MetasoundEditorGraphSchema.cpp
--- a/Engine/Plugins/Runtime/Metasound/Source/MetasoundEditor/Private/MetasoundEditorGraphSchema.cpp
+++ b/Engine/Plugins/Runtime/Metasound/Source/MetasoundEditor/Private/MetasoundEditorGraphSchema.cpp
@@ -532,7 +532,7 @@ FMetasoundGraphSchemaAction_PromoteToVariable_AccessorNode::FMetasoundGraphSchem
 	: FMetasoundGraphSchemaAction_NodeWithMultipleOutputs(
 		FText(),
 		LOCTEXT("PromoteToVariableGetterName", "Promote To Graph Variable"),
-		LOCTEXT("PromoteToInputTooltip2", "Promotes node input to graph variable and creates a getter node"),
+		LOCTEXT("PromoteToInputTooltip3", "Promotes node input to graph variable and creates a getter node"),
 		Metasound::Editor::EPrimaryContextGroup::Common)
 {
 }
diff --git a/Engine/Plugins/Runtime/Metasound/Source/MetasoundFrontend/Private/MetasoundSourceInterface.cpp b/Engine/Plugins/Runtime/Metasound/Source/MetasoundFrontend/Private/MetasoundSourceInterface.cpp
--- a/Engine/Plugins/Runtime/Metasound/Source/MetasoundFrontend/Private/MetasoundSourceInterface.cpp
+++ b/Engine/Plugins/Runtime/Metasound/Source/MetasoundFrontend/Private/MetasoundSourceInterface.cpp
@@ -197,19 +197,19 @@ namespace Metasound
 						{
 							{
 								LOCTEXT("AudioDeviceIDDisplayName1", "Audio Device ID"),
-								LOCTEXT("AudioDeviceIDDescription", "ID of AudioDevice source is played from."),
+								LOCTEXT("AudioDeviceIDDescription2", "ID of AudioDevice source is played from."),
 								FName(), // TODO: Align environment data types with environment (ex. this is actually set/get as a uint32)
 								Environment::DeviceID
 							},
 							{
 								LOCTEXT("GraphNameDisplayName", "Graph Name"),
-								LOCTEXT("AudioDeviceIDDescription", "Name of source graph (for debugging/logging)."),
+								LOCTEXT("AudioDeviceIDDescription3", "Name of source graph (for debugging/logging)."),
 								GetMetasoundDataTypeName<FString>(),
 								Environment::GraphName
 							},
 							{
 								LOCTEXT("IsPreviewSoundDisplayName", "Is Preview Sound"),
-								LOCTEXT("IsPreviewSoundDescription", "Whether source is being played as a previewed sound."),
+								LOCTEXT("IsPreviewSoundDescription4", "Whether source is being played as a previewed sound."),
 								GetMetasoundDataTypeName<bool>(),
 								Environment::IsPreview
 							},
diff --git a/Engine/Plugins/VirtualProduction/DMX/DMXEngine/Source/DMXEditor/Private/Widgets/FixtureType/SDMXFixtureFunctionEditor.cpp b/Engine/Plugins/VirtualProduction/DMX/DMXEngine/Source/DMXEditor/Private/Widgets/FixtureType/SDMXFixtureFunctionEditor.cpp
--- a/Engine/Plugins/VirtualProduction/DMX/DMXEngine/Source/DMXEditor/Private/Widgets/FixtureType/SDMXFixtureFunctionEditor.cpp
+++ b/Engine/Plugins/VirtualProduction/DMX/DMXEngine/Source/DMXEditor/Private/Widgets/FixtureType/SDMXFixtureFunctionEditor.cpp
@@ -59,7 +59,7 @@ void SDMXFixtureFunctionEditor::NotifyPreChange(FProperty* PropertyAboutToChange
 		}
 		else if (PropertyName == GET_MEMBER_NAME_CHECKED(FDMXFixtureFunction, DefaultValue))
 		{
-			Transaction = MakeUnique<FScopedTransaction>(LOCTEXT("ChangeDefaultValueTransaction", "Default Value of Function"));
+			Transaction = MakeUnique<FScopedTransaction>(LOCTEXT("ChangeDefaultValueTransaction2", "Default Value of Function"));
 		}
 
 		FixtureType->Modify();
diff --git a/Engine/Plugins/VirtualProduction/DMX/DMXEngine/Source/DMXEditor/Private/Widgets/FixtureType/SDMXFixtureTypeFunctionsEditorFunctionRow.cpp b/Engine/Plugins/VirtualProduction/DMX/DMXEngine/Source/DMXEditor/Private/Widgets/FixtureType/SDMXFixtureTypeFunctionsEditorFunctionRow.cpp
--- a/Engine/Plugins/VirtualProduction/DMX/DMXEngine/Source/DMXEditor/Private/Widgets/FixtureType/SDMXFixtureTypeFunctionsEditorFunctionRow.cpp
+++ b/Engine/Plugins/VirtualProduction/DMX/DMXEngine/Source/DMXEditor/Private/Widgets/FixtureType/SDMXFixtureTypeFunctionsEditorFunctionRow.cpp
@@ -143,7 +143,7 @@ TSharedRef<SWidget> SDMXFixtureTypeFunctionsEditorFunctionRow::GenerateWidgetFor
 						}
 						else
 						{
-							return LOCTEXT("InvalidAttributeFunctionName", "<Unnamed>");
+							return LOCTEXT("InvalidAttributeFunctionName2", "<Unnamed>");
 						}
 					})
 				.IsEnabled_Lambda([this]()
diff --git a/Engine/Plugins/VirtualProduction/LevelSnapshots/Source/LevelSnapshotsEditor/Private/Widgets/SLevelSnapshotsEditorBrowser.cpp b/Engine/Plugins/VirtualProduction/LevelSnapshots/Source/LevelSnapshotsEditor/Private/Widgets/SLevelSnapshotsEditorBrowser.cpp
--- a/Engine/Plugins/VirtualProduction/LevelSnapshots/Source/LevelSnapshotsEditor/Private/Widgets/SLevelSnapshotsEditorBrowser.cpp
+++ b/Engine/Plugins/VirtualProduction/LevelSnapshots/Source/LevelSnapshotsEditor/Private/Widgets/SLevelSnapshotsEditorBrowser.cpp
@@ -222,7 +222,7 @@ TSharedPtr<SWidget> SLevelSnapshotsEditorBrowser::OnGetAssetContextMenu(const TA
 
 		MenuBuilder.AddMenuEntry(
 			LOCTEXT("AssetTypeActions_LevelSnapshot_CaptureThumbnail", "Capture Thumbnail"),
-			LOCTEXT("AssetTypeActions_LevelSnapshot_CaptureThumbnailsToolTip", "Capture and update the thumbnail only for this asset."),
+			LOCTEXT("AssetTypeActions_LevelSnapshot_CaptureThumbnailsToolTip2", "Capture and update the thumbnail only for this asset."),
 			FSlateIcon(FLevelSnapshotsEditorStyle::GetStyleSetName(), "LevelSnapshots.ToolbarButton.Small"),
 			FUIAction(
 			FExecuteAction::CreateLambda([SelectedAsset] {
diff --git a/Engine/Plugins/VirtualProduction/RemoteControlProtocolDMX/Source/RemoteControlProtocolDMX/Private/RemoteControlProtocolDMX.cpp b/Engine/Plugins/VirtualProduction/RemoteControlProtocolDMX/Source/RemoteControlProtocolDMX/Private/RemoteControlProtocolDMX.cpp
--- a/Engine/Plugins/VirtualProduction/RemoteControlProtocolDMX/Source/RemoteControlProtocolDMX/Private/RemoteControlProtocolDMX.cpp
+++ b/Engine/Plugins/VirtualProduction/RemoteControlProtocolDMX/Source/RemoteControlProtocolDMX/Private/RemoteControlProtocolDMX.cpp
@@ -285,7 +285,7 @@ void FRemoteControlProtocolDMX::ProcessAutoBinding(const FRemoteControlProtocolE
 				FRemoteControlLogger::Get().Log(ProtocolName, [&LastSignalPtr, FinalChannelValue, FoundChannelDifferenceValue]
 				{
 					return FText::Format(
-						LOCTEXT("DMXEventLog",
+						LOCTEXT("DMXEventLog2",
 								"AutoBinding new value. ExternUniverseID {0}, Channel {1}, New Value {2}"),
 						LastSignalPtr->ExternUniverseID, FinalChannelValue, FoundChannelDifferenceValue);
 				});
