Fix for slate "Construct" virtual being called in drag-drop constructor
#rnx

diff --git a/Engine/Plugins/Runtime/Metasound/Source/MetasoundEditor/Private/MetasoundEditor.cpp b/Engine/Plugins/Runtime/Metasound/Source/MetasoundEditor/Private/MetasoundEditor.cpp
--- a/Engine/Plugins/Runtime/Metasound/Source/MetasoundEditor/Private/MetasoundEditor.cpp
+++ b/Engine/Plugins/Runtime/Metasound/Source/MetasoundEditor/Private/MetasoundEditor.cpp
@@ -163,7 +163,11 @@ namespace Metasound
 				: Editor(InEditor)
 				, GraphMember(InGraphMember)
 			{
-				Construct();
+				CursorDecoratorWindow = SWindow::MakeCursorDecorator();
+				constexpr bool bShowImmediately = false;
+				FSlateApplication::Get().AddWindow(CursorDecoratorWindow.ToSharedRef(), bShowImmediately);
+
+				HoverTargetChanged();
 			}
 
 			DRAG_DROP_OPERATOR_TYPE(FMetaSoundDragDropMemberAction, FGraphSchemaActionDragDropAction)
