Fix missing ability to order graph members
- Remove base template for node customizations now that it isn't required to be a template to enable intra-engine versioning (allows moving implementation of category customization to .cpp)
- Move redundant "Literal" property to parent UMetasoundEditorGraphMember
#rnx

diff --git a/Engine/Plugins/Runtime/Metasound/Source/MetasoundEditor/Private/MetasoundEditorGraph.cpp b/Engine/Plugins/Runtime/Metasound/Source/MetasoundEditor/Private/MetasoundEditorGraph.cpp
--- a/Engine/Plugins/Runtime/Metasound/Source/MetasoundEditor/Private/MetasoundEditorGraph.cpp
+++ b/Engine/Plugins/Runtime/Metasound/Source/MetasoundEditor/Private/MetasoundEditorGraph.cpp
@@ -75,11 +75,10 @@ void UMetasoundEditorGraphMember::ConformLiteralDataType()
 		LiteralClass = UMetasoundEditorGraphMemberDefaultLiteral::StaticClass();
 	}
 
-	const UMetasoundEditorGraphMemberDefaultLiteral* Literal = GetLiteral();
 	if (!Literal || Literal->GetClass() != LiteralClass)
 	{
 		UMetasoundEditorGraphMemberDefaultLiteral* NewLiteral = NewObject<UMetasoundEditorGraphMemberDefaultLiteral>(this, LiteralClass, FName(), RF_Transactional);
-		SetLiteral(NewLiteral);
+		Literal = NewLiteral;
 	}
 }
 
@@ -128,7 +127,6 @@ void UMetasoundEditorGraphVertex::InitMember(FName InDataType, const FMetasoundF
 
 	ConformLiteralDataType();
 
-	UMetasoundEditorGraphMemberDefaultLiteral* Literal = GetLiteral();
 	if (ensure(Literal))
 	{
 		Literal->SetFromLiteral(InDefaultLiteral);
@@ -480,6 +478,34 @@ const FText& UMetasoundEditorGraphInput::GetGraphMemberLabel() const
 	return Label;
 }
 
+int32 UMetasoundEditorGraphInput::GetSortOrderIndex() const
+{
+	using namespace Metasound::Editor;
+	using namespace Metasound::Frontend;
+
+	const UMetasoundEditorGraph* MetaSoundGraph = GetOwningGraph();
+	FConstGraphHandle GraphHandle = MetaSoundGraph->GetGraphHandle();
+	FConstNodeHandle NodeHandle = GetConstNodeHandle();
+	const Metasound::FVertexName& NodeName = NodeHandle->GetNodeName();
+	return GraphHandle->GetSortOrderIndexForInput(NodeName);
+}
+
+void UMetasoundEditorGraphInput::SetSortOrderIndex(int32 InSortOrderIndex)
+{
+	using namespace Metasound::Editor;
+	using namespace Metasound::Frontend;
+
+	UMetasoundEditorGraph* MetaSoundGraph = GetOwningGraph();
+	check(MetaSoundGraph);
+
+	FGraphHandle GraphHandle = MetaSoundGraph->GetGraphHandle();
+	FConstNodeHandle NodeHandle = GetConstNodeHandle();
+	const Metasound::FVertexName& NodeName = NodeHandle->GetNodeName();
+
+	GraphHandle->SetSortOrderIndexForInput(NodeName, InSortOrderIndex);
+	MetaSoundGraph->SetSynchronizationRequired();
+}
+
 void UMetasoundEditorGraphInput::ResetToClassDefault()
 {
 	using namespace Metasound::Editor;
@@ -487,7 +513,7 @@ void UMetasoundEditorGraphInput::ResetToClassDefault()
 
 	UMetasoundEditorGraph* MetaSoundGraph = GetOwningGraph();
 	FGraphHandle GraphHandle = MetaSoundGraph->GetGraphHandle();
-	FNodeHandle NodeHandle = GraphHandle->GetNodeWithID(NodeID);
+	FConstNodeHandle NodeHandle = GraphHandle->GetNodeWithID(NodeID);
 
 	FMetasoundFrontendLiteral DefaultLiteral;
 	DefaultLiteral.SetFromLiteral(IDataTypeRegistry::Get().CreateDefaultLiteral(GetDataType()));
@@ -506,9 +532,7 @@ void UMetasoundEditorGraphInput::ResetToClassDefault()
 
 		if (TScriptInterface<IAudioParameterControllerInterface> ParamInterface = PreviewComponent)
 		{
-			Metasound::Frontend::FConstNodeHandle ConstNodeHandle = GetConstNodeHandle();
-			Metasound::FVertexName VertexKey = NodeHandle->GetNodeName();
-			Literal->UpdatePreviewInstance(VertexKey, ParamInterface);
+			Literal->UpdatePreviewInstance(NodeName, ParamInterface);
 		}
 	}
 
@@ -593,6 +617,34 @@ Metasound::Frontend::FNodeHandle UMetasoundEditorGraphOutput::AddNodeHandle(cons
 	return NewNodeHandle;
 }
 
+int32 UMetasoundEditorGraphOutput::GetSortOrderIndex() const
+{
+	using namespace Metasound::Editor;
+	using namespace Metasound::Frontend;
+
+	const UMetasoundEditorGraph* MetaSoundGraph = GetOwningGraph();
+	FConstGraphHandle GraphHandle = MetaSoundGraph->GetGraphHandle();
+	FConstNodeHandle NodeHandle = GetConstNodeHandle();
+	const Metasound::FVertexName& NodeName = NodeHandle->GetNodeName();
+	return GraphHandle->GetSortOrderIndexForOutput(NodeName);
+}
+
+void UMetasoundEditorGraphOutput::SetSortOrderIndex(int32 InSortOrderIndex)
+{
+	using namespace Metasound::Editor;
+	using namespace Metasound::Frontend;
+
+	UMetasoundEditorGraph* MetaSoundGraph = GetOwningGraph();
+	check(MetaSoundGraph);
+
+	FGraphHandle GraphHandle = MetaSoundGraph->GetGraphHandle();
+	FConstNodeHandle NodeHandle = GetConstNodeHandle();
+	const Metasound::FVertexName& NodeName = NodeHandle->GetNodeName();
+
+	GraphHandle->SetSortOrderIndexForOutput(NodeName, InSortOrderIndex);
+	MetaSoundGraph->SetSynchronizationRequired();
+}
+
 const FText& UMetasoundEditorGraphOutput::GetGraphMemberLabel() const
 {
 	static const FText Label = LOCTEXT("GraphMemberLabel_Output", "Output");
@@ -676,10 +728,9 @@ void UMetasoundEditorGraphVariable::InitMember(FName InDataType, const FMetasoun
 
 	ConformLiteralDataType();
 
-	UMetasoundEditorGraphMemberDefaultLiteral* MemberDefault = GetLiteral();
-	if (ensure(MemberDefault))
+	if (ensure(Literal))
 	{
-		MemberDefault->SetFromLiteral(InDefaultLiteral);
+		Literal->SetFromLiteral(InDefaultLiteral);
 	}
 }
 
diff --git a/Engine/Plugins/Runtime/Metasound/Source/MetasoundEditor/Private/MetasoundEditorGraphBuilder.cpp b/Engine/Plugins/Runtime/Metasound/Source/MetasoundEditor/Private/MetasoundEditorGraphBuilder.cpp
--- a/Engine/Plugins/Runtime/Metasound/Source/MetasoundEditor/Private/MetasoundEditorGraphBuilder.cpp
+++ b/Engine/Plugins/Runtime/Metasound/Source/MetasoundEditor/Private/MetasoundEditorGraphBuilder.cpp
@@ -1333,7 +1333,7 @@ namespace Metasound
 			}
 
 			TArray<FInputHandle> InputHandles = NodeHandle->GetInputs();
-			InputHandles = NodeHandle->GetInputStyle().SortDefaults(InputHandles);
+			NodeHandle->GetInputStyle().SortDefaults(InputHandles);
 			for (const FInputHandle& InputHandle : InputHandles)
 			{
 				// Only add pins of the node if the connection is user modifiable. 
@@ -1345,7 +1345,7 @@ namespace Metasound
 			}
 
 			TArray<FOutputHandle> OutputHandles = NodeHandle->GetOutputs();
-			OutputHandles = NodeHandle->GetOutputStyle().SortDefaults(OutputHandles);
+			NodeHandle->GetOutputStyle().SortDefaults(OutputHandles);
 			for (const FOutputHandle& OutputHandle : OutputHandles)
 			{
 				// Only add pins of the node if the connection is user modifiable. 
@@ -1923,16 +1923,21 @@ namespace Metasound
 
 			IMetasoundEditorModule& EditorModule = FModuleManager::GetModuleChecked<IMetasoundEditorModule>("MetaSoundEditor");
 
-			TArray<Frontend::FConstInputHandle> InputHandles = InNode->GetConstInputs();
-			TArray<Frontend::FConstOutputHandle> OutputHandles = InNode->GetConstOutputs();
-			TArray<UEdGraphPin*> EditorPins = InEditorNode.Pins;
-
+			TArray<Frontend::FConstInputHandle> InputHandles;
+			TArray<Frontend::FConstOutputHandle> OutputHandles;
+			auto GetUserModifiableHandles = [InNode](TArray<Frontend::FConstInputHandle>& InHandles, TArray<Frontend::FConstOutputHandle>& OutHandles)
+			{
+				InHandles = InNode->GetConstInputs();
+				OutHandles = InNode->GetConstOutputs();
 
-			// Remove input and output handles which are not user modifiable
-			InputHandles.RemoveAll([](const Frontend::FConstInputHandle& FrontendInput) { return !FrontendInput->IsConnectionUserModifiable(); });
-			OutputHandles.RemoveAll([](const Frontend::FConstOutputHandle& FrontendOutput) { return !FrontendOutput->IsConnectionUserModifiable(); });
+				// Remove input and output handles which are not user modifiable
+				InHandles.RemoveAll([](const Frontend::FConstInputHandle& FrontendInput) { return !FrontendInput->IsConnectionUserModifiable(); });
+				OutHandles.RemoveAll([](const Frontend::FConstOutputHandle& FrontendOutput) { return !FrontendOutput->IsConnectionUserModifiable(); });
+			};
+			GetUserModifiableHandles(InputHandles, OutputHandles);
 
 			// Filter out pins which are not paired.
+			TArray<UEdGraphPin*> EditorPins = InEditorNode.Pins;
 			for (int32 i = EditorPins.Num() - 1; i >= 0; i--)
 			{
 				UEdGraphPin* Pin = EditorPins[i];
@@ -1995,7 +2000,6 @@ namespace Metasound
 			if (!InputHandles.IsEmpty())
 			{
 				bIsNodeDirty = true;
-				InputHandles = InNode->GetInputStyle().SortDefaults(InputHandles);
 				for (Frontend::FConstInputHandle& InputHandle : InputHandles)
 				{
 					if (bLogChanges)
@@ -2012,7 +2016,6 @@ namespace Metasound
 			if (!OutputHandles.IsEmpty())
 			{
 				bIsNodeDirty = true;
-				OutputHandles = InNode->GetOutputStyle().SortDefaults(OutputHandles);
 				for (Frontend::FConstOutputHandle& OutputHandle : OutputHandles)
 				{
 					if (bLogChanges)
@@ -2026,6 +2029,58 @@ namespace Metasound
 				}
 			}
 
+			// Order pins
+			GetUserModifiableHandles(InputHandles, OutputHandles);
+			InNode->GetInputStyle().SortDefaults(InputHandles);
+			InNode->GetOutputStyle().SortDefaults(OutputHandles);
+
+			auto SwapAndDirty = [&](int32 IndexA, int32 IndexB)
+			{
+				const bool bRequiresSwap = IndexA != IndexB;
+				if (bRequiresSwap)
+				{
+					InEditorNode.Pins.Swap(IndexA, IndexB);
+					bIsNodeDirty |= bRequiresSwap;
+				}
+			};
+
+			for (int32 i = InEditorNode.Pins.Num() - 1; i >= 0; --i)
+			{
+				UEdGraphPin* Pin = InEditorNode.Pins[i];
+				if (Pin->Direction == EGPD_Input)
+				{
+					if (!InputHandles.IsEmpty())
+					{
+						constexpr bool bAllowShrinking = false;
+						Frontend::FConstInputHandle InputHandle = InputHandles.Pop(bAllowShrinking);
+						for (int32 j = i; j >= 0; --j)
+						{
+							if (IsMatchingInputHandleAndPin(InputHandle, *InEditorNode.Pins[j]))
+							{
+								SwapAndDirty(i, j);
+								break;
+							}
+						}
+					}
+				}
+				else /* Pin->Direction == EGPD_Output */
+				{
+					if (!OutputHandles.IsEmpty())
+					{
+						constexpr bool bAllowShrinking = false;
+						Frontend::FConstOutputHandle OutputHandle = OutputHandles.Pop(bAllowShrinking);
+						for (int32 j = i; j >= 0; --j)
+						{
+							if (IsMatchingOutputHandleAndPin(OutputHandle, *InEditorNode.Pins[j]))
+							{
+								SwapAndDirty(i, j);
+								break;
+							}
+						}
+					}
+				}
+			}
+
 			return bIsNodeDirty;
 		}
 
diff --git a/Engine/Plugins/Runtime/Metasound/Source/MetasoundEditor/Private/MetasoundEditorGraphNode.cpp b/Engine/Plugins/Runtime/Metasound/Source/MetasoundEditor/Private/MetasoundEditorGraphNode.cpp
--- a/Engine/Plugins/Runtime/Metasound/Source/MetasoundEditor/Private/MetasoundEditorGraphNode.cpp
+++ b/Engine/Plugins/Runtime/Metasound/Source/MetasoundEditor/Private/MetasoundEditorGraphNode.cpp
@@ -691,59 +691,6 @@ FSlateIcon UMetasoundEditorGraphOutputNode::GetNodeTitleIcon() const
 	return FSlateIcon("MetaSoundStyle", "MetasoundEditor.Graph.Node.Class.Output");
 }
 
-bool UMetasoundEditorGraphExternalNode::RefreshPinMetadata()
-{
-	using namespace Metasound::Frontend;
-
-	bool bModified = false;
-
-	FConstNodeHandle NodeHandle = GetConstNodeHandle();
-	const FMetasoundFrontendClassInterface& ClassInterface = NodeHandle->GetClassInterface();
-	TArray<FMetasoundFrontendClassInput> Inputs = ClassInterface.Inputs;
-	TArray<FMetasoundFrontendClassOutput> Outputs = ClassInterface.Outputs;
-
-	for (UEdGraphPin* Pin : Pins)
-	{
-		check(Pin);
-
-		if (Pin->Direction == EGPD_Input)
-		{
-			// Remove valid class inputs on iteration to avoid additional
-			// incurred look-up cost on subsequent tooltip update.
-			auto RemoveValidEntry = [&](const FMetasoundFrontendClassInput& Entry)
-			{
-				if (Entry.Name == Pin->GetFName())
-				{
-					Metasound::Editor::FGraphBuilder::RefreshPinMetadata(*Pin, Entry.Metadata);
-					return true;
-				}
-
-				return false;
-			};
-			bModified |= Inputs.RemoveAllSwap(RemoveValidEntry) > 0;
-		}
-
-		if (Pin->Direction == EGPD_Output)
-		{
-			// Remove valid class inputs on iteration to avoid additional
-			// incurred look-up cost on subsequent tooltip update.
-			auto RemoveValidEntry = [&](const FMetasoundFrontendClassOutput& Entry)
-			{
-				if (Entry.Name == Pin->GetFName())
-				{
-					Metasound::Editor::FGraphBuilder::RefreshPinMetadata(*Pin, Entry.Metadata);
-					return true;
-				}
-
-				return false;
-			};
-			bModified |= Outputs.RemoveAllSwap(RemoveValidEntry) > 0;
-		}
-	}
-
-	return bModified;
-}
-
 void UMetasoundEditorGraphExternalNode::ReconstructNode()
 {
 	using namespace Metasound::Editor;
diff --git a/Engine/Plugins/Runtime/Metasound/Source/MetasoundEditor/Private/MetasoundEditorModule.cpp b/Engine/Plugins/Runtime/Metasound/Source/MetasoundEditor/Private/MetasoundEditorModule.cpp
--- a/Engine/Plugins/Runtime/Metasound/Source/MetasoundEditor/Private/MetasoundEditorModule.cpp
+++ b/Engine/Plugins/Runtime/Metasound/Source/MetasoundEditor/Private/MetasoundEditorModule.cpp
@@ -33,7 +33,6 @@
 #include "MetasoundTime.h"
 #include "MetasoundTrigger.h"
 #include "MetasoundUObjectRegistry.h"
-#include "MetasoundVariableDetailCustomization.h"
 #include "Modules/ModuleInterface.h"
 #include "Modules/ModuleManager.h"
 #include "PropertyEditorDelegates.h"
diff --git a/Engine/Plugins/Runtime/Metasound/Source/MetasoundEditor/Private/MetasoundNodeDetailCustomization.cpp b/Engine/Plugins/Runtime/Metasound/Source/MetasoundEditor/Private/MetasoundNodeDetailCustomization.cpp
--- a/Engine/Plugins/Runtime/Metasound/Source/MetasoundEditor/Private/MetasoundNodeDetailCustomization.cpp
+++ b/Engine/Plugins/Runtime/Metasound/Source/MetasoundEditor/Private/MetasoundNodeDetailCustomization.cpp
@@ -828,96 +828,389 @@ namespace Metasound
 			}
 		}
 
-		const FText FMetasoundInputDetailCustomization::MemberNameText = LOCTEXT("InputGraphMemberLabel", "Input");
-
-		Frontend::FDocumentHandle FMetasoundInputDetailCustomization::GetDocumentHandle() const
+		void FMetasoundMemberDetailCustomization::UpdateRenameDelegate(UMetasoundEditorGraphMemberDefaultLiteral& InMemberDefaultLiteral)
 		{
-			if (MemberDefaultLiteral.IsValid())
+			if (UMetasoundEditorGraphMember* Member = InMemberDefaultLiteral.GetParentMember())
 			{
-				if (UMetasoundEditorGraphMember* Parent = MemberDefaultLiteral->GetParentMember())
+				if (Member->CanRename())
 				{
-					if (UMetasoundEditorGraph* Graph = Parent->GetOwningGraph())
+					if (!RenameRequestedHandle.IsValid())
 					{
-						return Graph->GetDocumentHandle();
+						Member->OnRenameRequested.Clear();
+						RenameRequestedHandle = Member->OnRenameRequested.AddLambda([this]()
+							{
+								FSlateApplication::Get().SetKeyboardFocus(NameEditableTextBox.ToSharedRef(), EFocusCause::SetDirectly);
+							});
 					}
 				}
 			}
+		}
 
-			return Frontend::IDocumentController::GetInvalidHandle();
+		void FMetasoundMemberDetailCustomization::CacheMemberData(IDetailLayoutBuilder& InDetailLayout)
+		{
+			TArray<TWeakObjectPtr<UObject>> Objects;
+			InDetailLayout.GetObjectsBeingCustomized(Objects);
+			if (!Objects.IsEmpty())
+			{
+				GraphMember = Cast<UMetasoundEditorGraphMember>(Objects.Last().Get());
+
+				TSharedPtr<IPropertyHandle> LiteralHandle = InDetailLayout.GetProperty(UMetasoundEditorGraphMember::GetLiteralPropertyName());
+				if (ensure(GraphMember.IsValid()) && ensure(LiteralHandle.IsValid()))
+				{
+					// Always hide, even if no customization (LiteralObject isn't found) as this is the case
+					// where the default object is not required (i.e. Default Member is default constructed)
+					LiteralHandle->MarkHiddenByCustomization();
+				}
+			}
 		}
 
-		bool FMetasoundInputDetailCustomization::GetInputInheritsDefault() const
+		TArray<IDetailPropertyRow*> FMetasoundMemberDetailCustomization::CustomizeDefaultCategory(IDetailLayoutBuilder& InDetailLayout)
 		{
-			if (!MemberDefaultLiteral.IsValid())
+			TArray<IDetailPropertyRow*> DefaultPropertyRows;
+
+			if (!GraphMember.IsValid())
 			{
-				return false;
+				return DefaultPropertyRows;
 			}
 
-			if (const UMetasoundEditorGraphVertex* Vertex = Cast<UMetasoundEditorGraphVertex>(MemberDefaultLiteral->GetParentMember()))
+			if (UMetasoundEditorGraphMemberDefaultLiteral* MemberDefaultLiteral = GraphMember->GetLiteral())
 			{
-				const TSet<FName>& InputsInheritingDefault = GetDocumentHandle()->GetRootGraph()->GetInputsInheritingDefault();
-				FName NodeName = Vertex->GetConstNodeHandle()->GetNodeName();
-				return InputsInheritingDefault.Contains(NodeName);
+				UpdateRenameDelegate(*MemberDefaultLiteral);
+
+				UClass* MemberClass = MemberDefaultLiteral->GetClass();
+				check(MemberClass);
+
+				IDetailCategoryBuilder& DefaultCategoryBuilder = GetDefaultCategoryBuilder(InDetailLayout);
+				IMetasoundEditorModule& EditorModule = FModuleManager::GetModuleChecked<IMetasoundEditorModule>("MetaSoundEditor");
+				TUniquePtr<FMetasoundDefaultLiteralCustomizationBase> LiteralCustomization = EditorModule.CreateMemberDefaultLiteralCustomization(*MemberClass, DefaultCategoryBuilder);
+				if (LiteralCustomization.IsValid())
+				{
+					DefaultPropertyRows = LiteralCustomization->CustomizeLiteral(*MemberDefaultLiteral, InDetailLayout);
+				}
+				else
+				{
+					IDetailPropertyRow* DefaultPropertyRow = DefaultCategoryBuilder.AddExternalObjectProperty(TArray<UObject*>({ MemberDefaultLiteral }), "Default");
+					ensureMsgf(DefaultPropertyRow, TEXT("Class '%s' missing expected 'Default' member."
+						"Either add/rename default member or register customization to display default value/opt out appropriately."),
+						*MemberClass->GetName());
+					DefaultPropertyRows.Add(DefaultPropertyRow);
+				}
 			}
 
-			return false;
+			return DefaultPropertyRows;
 		}
 
-		void FMetasoundInputDetailCustomization::SetInputInheritsDefault()
+		void FMetasoundMemberDetailCustomization::CustomizeGeneralCategory(IDetailLayoutBuilder& InDetailLayout)
+		{
+			const bool bIsInterfaceMember = IsInterfaceMember();
+			const bool bIsGraphEditable = IsGraphEditable();
+			
+			IDetailCategoryBuilder& CategoryBuilder = GetGeneralCategoryBuilder(InDetailLayout);
+
+			NameEditableTextBox = SNew(SEditableTextBox)
+				.Text(this, &FMetasoundMemberDetailCustomization::GetName)
+				.OnTextChanged(this, &FMetasoundMemberDetailCustomization::OnNameChanged)
+				.OnTextCommitted(this, &FMetasoundMemberDetailCustomization::OnNameCommitted)
+				.IsReadOnly(bIsInterfaceMember || !bIsGraphEditable)
+				.SelectAllTextWhenFocused(true)
+				.Font(IDetailLayoutBuilder::GetDetailFont());
+
+			static const FText MemberNameToolTipFormat = LOCTEXT("GraphMember_NameDescriptionFormat", "Name used within the MetaSounds editor(s) and transacting systems (ex. Blueprints) if applicable to reference the given {0}.");
+			CategoryBuilder.AddCustomRow(LOCTEXT("GraphMember_NameProperty", "Name"))
+				.EditCondition(!bIsInterfaceMember && bIsGraphEditable, nullptr)
+				.NameContent()
+				[
+					SNew(STextBlock)
+					.Font(IDetailLayoutBuilder::GetDetailFontBold())
+					.Text(GraphMember->GetGraphMemberLabel())
+					.ToolTipText(FText::Format(MemberNameToolTipFormat, GraphMember->GetGraphMemberLabel()))
+				]
+				.ValueContent()
+				[
+					NameEditableTextBox.ToSharedRef()
+				];
+
+			static const FText MemberDisplayNameText = LOCTEXT("GraphMember_DisplayNameProperty", "Display Name");
+			static const FText MemberDisplayNameToolTipFormat = LOCTEXT("GraphMember_DisplayNameDescriptionFormat", "Optional, localized name used within the MetaSounds editor(s) to describe the given {0}.");
+			const FText MemberDisplayNameTooltipText = FText::Format(MemberDisplayNameToolTipFormat, GraphMember->GetGraphMemberLabel());
+			CategoryBuilder.AddCustomRow(MemberDisplayNameText)
+				.EditCondition(!bIsInterfaceMember && bIsGraphEditable, nullptr)
+				.NameContent()
+				[
+					SNew(STextBlock)
+					.Font(IDetailLayoutBuilder::GetDetailFontBold())
+					.Text(MemberDisplayNameText)
+					.ToolTipText(MemberDisplayNameTooltipText)
+				]
+				.ValueContent()
+				[
+					SNew(STextPropertyEditableTextBox, MakeShared<FGraphMemberEditableTextDisplayName>(GraphMember, MemberDisplayNameTooltipText))
+					.WrapTextAt(500)
+					.MinDesiredWidth(25.0f)
+					.MaxDesiredHeight(200)
+				];
+
+			static const FText MemberDescriptionText = LOCTEXT("Member_DescriptionPropertyName", "Description");
+			static const FText MemberDescriptionToolTipFormat = LOCTEXT("Member_DescriptionToolTipFormat", "Description for {0}. For example, used as a tooltip when displayed on another graph's referencing node.");
+			const FText MemberDescriptionToolTipText = FText::Format(MemberDescriptionToolTipFormat, GraphMember->GetGraphMemberLabel());
+			CategoryBuilder.AddCustomRow(MemberDescriptionText)
+				.EditCondition(!bIsInterfaceMember && bIsGraphEditable, nullptr)
+				.NameContent()
+				[
+					SNew(STextBlock)
+					.Font(IDetailLayoutBuilder::GetDetailFontBold())
+					.Text(MemberDescriptionText)
+					.ToolTipText(MemberDescriptionToolTipText)
+				]
+				.ValueContent()
+				[
+					SNew(STextPropertyEditableTextBox, MakeShared<FGraphMemberEditableTextDescription>(GraphMember, MemberDescriptionToolTipText))
+					.WrapTextAt(500)
+					.MinDesiredWidth(25.0f)
+					.MaxDesiredHeight(200)
+				];
+
+			DataTypeSelector.AddDataTypeSelector(InDetailLayout, MemberCustomizationStyle::DataTypeNameText, GraphMember, !bIsInterfaceMember && bIsGraphEditable);
+		}
+
+		void FMetasoundMemberDetailCustomization::CustomizeDetails(IDetailLayoutBuilder& InDetailLayout)
 		{
-			if (!MemberDefaultLiteral.IsValid())
+			CacheMemberData(InDetailLayout);
+			if (GraphMember.IsValid())
+			{
+				CustomizeGeneralCategory(InDetailLayout);
+				CustomizeDefaultCategory(InDetailLayout);
+			}
+		}
+
+		void FMetasoundMemberDetailCustomization::OnNameChanged(const FText& InNewName)
+		{
+			using namespace Frontend;
+
+			bIsNameInvalid = false;
+			NameEditableTextBox->SetError(FText::GetEmpty());
+
+			if (!ensure(GraphMember.IsValid()))
 			{
 				return;
 			}
 
-			if (UMetasoundEditorGraphVertex* Vertex = Cast<UMetasoundEditorGraphVertex>(MemberDefaultLiteral->GetParentMember()))
+			FText Error;
+			if (!GraphMember->CanRename(InNewName, Error))
 			{
-				FScopedTransaction(LOCTEXT("SetPresetInputOverrideTransaction", "Set MetaSound Preset Input Overridden"));
+				bIsNameInvalid = true;
+				NameEditableTextBox->SetError(Error);
+			}
+		}
 
-				Vertex->GetOutermost()->Modify();
-				Vertex->Modify();
-				MemberDefaultLiteral->Modify();
+		FText FMetasoundMemberDetailCustomization::GetName() const
+		{
+			if (GraphMember.IsValid())
+			{
+				return FText::FromName(GraphMember->GetMemberName());
+			}
 
-				constexpr bool bDefaultIsInherited = true;
-				const FName NodeName = Vertex->GetConstNodeHandle()->GetNodeName();
-				GetDocumentHandle()->GetRootGraph()->SetInputInheritsDefault(NodeName, bDefaultIsInherited);
+			return FText::GetEmpty();
+		}
 
-				if (UObject* Metasound = Vertex->GetOutermostObject())
+		bool FMetasoundMemberDetailCustomization::IsGraphEditable() const
+		{
+			if (GraphMember.IsValid())
+			{
+				if (const UMetasoundEditorGraph* OwningGraph = GraphMember->GetOwningGraph())
 				{
-					FGraphBuilder::RegisterGraphWithFrontend(*Metasound);
+					return OwningGraph->IsEditable();
 				}
 			}
+
+			return false;
 		}
 
-		void FMetasoundInputDetailCustomization::ClearInputInheritsDefault()
+		FText FMetasoundMemberDetailCustomization::GetDisplayName() const
 		{
-			if (!MemberDefaultLiteral.IsValid())
+			using namespace Frontend;
+
+			if (GraphMember.IsValid())
 			{
-				return;
+				return GraphMember->GetDisplayName();
+			}
+
+			return FText::GetEmpty();
+		}
+
+		void FMetasoundMemberDetailCustomization::OnTooltipCommitted(const FText& InNewText, ETextCommit::Type InTextCommit)
+		{
+			using namespace Frontend;
+
+			if (GraphMember.IsValid())
+			{
+				constexpr bool bPostTransaction = true;
+				GraphMember->SetDescription(InNewText, bPostTransaction);
+			}
+		}
+
+		FText FMetasoundMemberDetailCustomization::GetTooltip() const
+		{
+			if (GraphMember.IsValid())
+			{
+				return GraphMember->GetDescription();
+			}
+
+			return FText::GetEmpty();
+		}
+
+		void FMetasoundMemberDetailCustomization::OnNameCommitted(const FText& InNewName, ETextCommit::Type InTextCommit)
+		{
+			using namespace Frontend;
+
+			if (!bIsNameInvalid && GraphMember.IsValid())
+			{
+				const FText TransactionLabel = FText::Format(LOCTEXT("RenameGraphMember_Format", "Set MetaSound {0}'s Name"), GraphMember->GetGraphMemberLabel());
+				const FScopedTransaction Transaction(TransactionLabel);
+
+				constexpr bool bPostTransaction = false;
+				GraphMember->SetDisplayName(FText::GetEmpty(), bPostTransaction);
+				GraphMember->SetMemberName(*InNewName.ToString(), bPostTransaction);
+			}
+
+			NameEditableTextBox->SetError(FText::GetEmpty());
+			bIsNameInvalid = false;
+		}
+
+		void FMetasoundVertexDetailCustomization::CustomizeGeneralCategory(IDetailLayoutBuilder& InDetailLayout)
+		{
+			FMetasoundMemberDetailCustomization::CustomizeGeneralCategory(InDetailLayout);
+
+			IDetailCategoryBuilder& CategoryBuilder = FMetasoundMemberDetailCustomization::GetGeneralCategoryBuilder(InDetailLayout);
+			UMetasoundEditorGraphVertex* Vertex = Cast<UMetasoundEditorGraphVertex>(GraphMember.Get());
+			if (ensure(Vertex))
+			{
+				TWeakObjectPtr<UMetasoundEditorGraphVertex> VertexPtr = Vertex;
+				static const FText SortOrderText = LOCTEXT("Vertex_SortOrderPropertyName", "Sort Order");
+				static const FText SortOrderToolTipFormat = LOCTEXT("Vertex_SortOrderToolTipFormat", "Sort Order for {0}. Used to organize pins in node view. The higher the number, the lower in the list.");
+				const FText SortOrderToolTipText = FText::Format(SortOrderToolTipFormat, GraphMember->GetGraphMemberLabel());
+				CategoryBuilder.AddCustomRow(SortOrderText)
+					.EditCondition(IsGraphEditable(), nullptr)
+					.NameContent()
+					[
+						SNew(STextBlock)
+						.Font(IDetailLayoutBuilder::GetDetailFontBold())
+						.Text(SortOrderText)
+						.ToolTipText(SortOrderToolTipText)
+					]
+					.ValueContent()
+					[
+						SNew(SNumericEntryBox<int32>)
+						.Value_Lambda([VertexPtr]() { return VertexPtr->GetSortOrderIndex(); })
+						.AllowSpin(false)
+						.UndeterminedString(LOCTEXT("Vertex_SortOrder_MultipleValues", "Multiple"))
+						.OnValueCommitted_Lambda([VertexPtr](int32 NewValue, ETextCommit::Type CommitInfo)
+						{
+							if (!VertexPtr.IsValid())
+							{
+								return;
+							}
+
+							const FText TransactionTitle = FText::Format(LOCTEXT("SetVertexSortOrderFormat", "Set MetaSound Graph {0} '{1}' SortOrder to {2}"),
+								VertexPtr->GetGraphMemberLabel(),
+								VertexPtr->GetDisplayName(),
+								FText::AsNumber(NewValue));
+							FScopedTransaction Transaction(TransactionTitle);
+
+							UObject* MetaSoundObject = VertexPtr->GetOutermostObject();
+							FMetasoundAssetBase* MetaSoundAsset = Metasound::IMetasoundUObjectRegistry::Get().GetObjectAsAssetBase(MetaSoundObject);
+							check(MetaSoundAsset);
+
+							MetaSoundObject->Modify();
+							MetaSoundAsset->GetGraphChecked().Modify();
+							VertexPtr->Modify();
+
+							VertexPtr->SetSortOrderIndex(NewValue);
+
+							constexpr bool bInForceViewSynchronization = true;
+							FGraphBuilder::RegisterGraphWithFrontend(*MetaSoundObject, bInForceViewSynchronization);
+						})
+						.Font(IDetailLayoutBuilder::GetDetailFont())
+					];
 			}
+		}
 
-			if (UMetasoundEditorGraphVertex* Vertex = Cast<UMetasoundEditorGraphVertex>(MemberDefaultLiteral->GetParentMember()))
+		Frontend::FDocumentHandle FMetasoundInputDetailCustomization::GetDocumentHandle() const
+		{
+			if (GraphMember.IsValid())
 			{
-				FScopedTransaction(LOCTEXT("ClearPresetInputOverrideTransaction", "Clear MetaSound Preset Input Overridden"));
+				if (UMetasoundEditorGraph* Graph = GraphMember->GetOwningGraph())
+				{
+					return Graph->GetDocumentHandle();
+				}
+			}
 
-				Vertex->GetOutermost()->Modify();
-				Vertex->Modify();
-				MemberDefaultLiteral->Modify();
+			return Frontend::IDocumentController::GetInvalidHandle();
+		}
 
-				constexpr bool bDefaultIsInherited = false;
-				const FName NodeName = Vertex->GetConstNodeHandle()->GetNodeName();
-				GetDocumentHandle()->GetRootGraph()->SetInputInheritsDefault(NodeName, bDefaultIsInherited);
+		bool FMetasoundInputDetailCustomization::GetInputInheritsDefault() const
+		{
+			if (UMetasoundEditorGraphInput* Input = Cast<UMetasoundEditorGraphInput>(GraphMember.Get()))
+			{
+				const TSet<FName>& InputsInheritingDefault = GetDocumentHandle()->GetRootGraph()->GetInputsInheritingDefault();
+				FName NodeName = Input->GetConstNodeHandle()->GetNodeName();
+				return InputsInheritingDefault.Contains(NodeName);
+			}
 
-				Vertex->UpdateFrontendDefaultLiteral(false /* bPostTransaction */);
+			return false;
+		}
 
-				if (UMetasoundEditorGraphMemberDefaultLiteral* Literal = Vertex->GetLiteral())
+		void FMetasoundInputDetailCustomization::SetInputInheritsDefault()
+		{
+			if (UMetasoundEditorGraphInput* Input = Cast<UMetasoundEditorGraphInput>(GraphMember.Get()))
+			{
+				if (UMetasoundEditorGraphMemberDefaultLiteral* MemberDefaultLiteral = Input->GetLiteral())
 				{
-					Literal->ForceRefresh();
+					FScopedTransaction Transaction(LOCTEXT("SetPresetInputOverrideTransaction", "Set MetaSound Preset Input Overridden"));
+
+					Input->GetOutermost()->Modify();
+					Input->Modify();
+					MemberDefaultLiteral->Modify();
+
+					constexpr bool bDefaultIsInherited = true;
+					const FName NodeName = Input->GetConstNodeHandle()->GetNodeName();
+					GetDocumentHandle()->GetRootGraph()->SetInputInheritsDefault(NodeName, bDefaultIsInherited);
+
+					if (UObject* Metasound = Input->GetOutermostObject())
+					{
+						FGraphBuilder::RegisterGraphWithFrontend(*Metasound);
+					}
 				}
+			}
+		}
 
-				if (UObject* Metasound = Vertex->GetOutermostObject())
+		void FMetasoundInputDetailCustomization::ClearInputInheritsDefault()
+		{
+			if (UMetasoundEditorGraphInput* Input = Cast<UMetasoundEditorGraphInput>(GraphMember.Get()))
+			{
+				if (UMetasoundEditorGraphMemberDefaultLiteral* MemberDefaultLiteral = Input->GetLiteral())
 				{
-					FGraphBuilder::RegisterGraphWithFrontend(*Metasound);
+					FScopedTransaction Transaction(LOCTEXT("ClearPresetInputOverrideTransaction", "Clear MetaSound Preset Input Overridden"));
+
+					Input->GetOutermost()->Modify();
+					Input->Modify();
+					MemberDefaultLiteral->Modify();
+
+					constexpr bool bDefaultIsInherited = false;
+					const FName NodeName = Input->GetConstNodeHandle()->GetNodeName();
+					GetDocumentHandle()->GetRootGraph()->SetInputInheritsDefault(NodeName, bDefaultIsInherited);
+
+					Input->UpdateFrontendDefaultLiteral(false /* bPostTransaction */);
+
+					if (UMetasoundEditorGraphMemberDefaultLiteral* Literal = Input->GetLiteral())
+					{
+						Literal->ForceRefresh();
+					}
+
+					if (UObject* Metasound = Input->GetOutermostObject())
+					{
+						FGraphBuilder::RegisterGraphWithFrontend(*Metasound);
+					}
 				}
 			}
 		}
@@ -932,7 +1225,8 @@ namespace Metasound
 
 			CustomizeGeneralCategory(InDetailLayout);
 
-			if (!MemberDefaultLiteral.IsValid())
+			UMetasoundEditorGraphMemberDefaultLiteral* MemberDefaultLiteral = GraphMember->GetLiteral();
+			if (!MemberDefaultLiteral)
 			{
 				return;
 			}
@@ -1013,8 +1307,6 @@ namespace Metasound
 			return false;
 		}
 
-		const FText FMetasoundOutputDetailCustomization::MemberNameText = LOCTEXT("OutputGraphMemberLabel", "Output");
-
 		bool FMetasoundOutputDetailCustomization::IsInterfaceMember() const
 		{
 			if (GraphMember.IsValid())
diff --git a/Engine/Plugins/Runtime/Metasound/Source/MetasoundEditor/Private/MetasoundNodeDetailCustomization.h b/Engine/Plugins/Runtime/Metasound/Source/MetasoundEditor/Private/MetasoundNodeDetailCustomization.h
--- a/Engine/Plugins/Runtime/Metasound/Source/MetasoundEditor/Private/MetasoundNodeDetailCustomization.h
+++ b/Engine/Plugins/Runtime/Metasound/Source/MetasoundEditor/Private/MetasoundNodeDetailCustomization.h
@@ -29,6 +29,7 @@
 #include "Widgets/Input/SCheckBox.h"
 #include "Widgets/Input/SEditableTextBox.h"
 #include "Widgets/Input/SMultiLineEditableTextBox.h"
+#include "Widgets/Input/SNumericEntryBox.h"
 #include "Widgets/Input/STextComboBox.h"
 #include "WorkflowOrientedApp/SModeWidget.h"
 
@@ -268,7 +269,7 @@ namespace Metasound
 			virtual TSharedRef<SWidget> CreateStructureWidget(TSharedPtr<IPropertyHandle>& StructPropertyHandle) const override;
 		};
 
-		class FMetasoundDataTypeSelector : public TSharedFromThis<FMetasoundDataTypeSelector>
+		class FMetasoundDataTypeSelector
 		{
 		public:
 			void AddDataTypeSelector(IDetailLayoutBuilder& InDetailLayoutBuilder, const FText& InRowName, TWeakObjectPtr<UMetasoundEditorGraphMember> InGraphMember, bool bInIsInterfaceMember);
@@ -291,178 +292,29 @@ namespace Metasound
 			FName ArrayTypeName;
 		};
 
-		template <typename TMemberType>
-		class TMetasoundGraphMemberDetailCustomization : public IDetailCustomization
+		class FMetasoundMemberDetailCustomization : public IDetailCustomization
 		{
 		public:
-			TMetasoundGraphMemberDetailCustomization()
-				: IDetailCustomization()
-			{
-				DataTypeSelector = MakeShared<FMetasoundDataTypeSelector>();
-			}
-
-			virtual ~TMetasoundGraphMemberDetailCustomization()
+			virtual ~FMetasoundMemberDetailCustomization()
 			{
 				RenameRequestedHandle.Reset();
 			}
 
 		protected:
-			TWeakObjectPtr<TMemberType> GraphMember;
-			TSharedPtr<SEditableTextBox> NameEditableTextBox;
-			TSharedPtr<FMetasoundDataTypeSelector> DataTypeSelector;
-
-			bool bIsNameInvalid = false;
-
 			static IDetailCategoryBuilder& GetDefaultCategoryBuilder(IDetailLayoutBuilder& InDetailLayout)
 			{
 				return InDetailLayout.EditCategory("DefaultValue");
 			}
 
-			void UpdateRenameDelegate(UMetasoundEditorGraphMemberDefaultLiteral& InMemberDefaultLiteral)
-			{
-				if (UMetasoundEditorGraphMember* Member = InMemberDefaultLiteral.GetParentMember())
-				{
-					if (Member->CanRename())
-					{
-						if (!RenameRequestedHandle.IsValid())
-						{
-							Member->OnRenameRequested.Clear();
-							RenameRequestedHandle = Member->OnRenameRequested.AddLambda([this]()
-							{
-								FSlateApplication::Get().SetKeyboardFocus(NameEditableTextBox.ToSharedRef(), EFocusCause::SetDirectly);
-							});
-						}
-					}
-				}
-			}
-
-			void CacheMemberData(IDetailLayoutBuilder& InDetailLayout)
-			{
-				TArray<TWeakObjectPtr<UObject>> Objects;
-				InDetailLayout.GetObjectsBeingCustomized(Objects);
-				if (!Objects.IsEmpty())
-				{
-					GraphMember = Cast<TMemberType>(Objects.Last().Get());
-
-					TSharedPtr<IPropertyHandle> LiteralHandle = InDetailLayout.GetProperty(GET_MEMBER_NAME_CHECKED(TMemberType, Literal));
-					if (ensure(GraphMember.IsValid()) && ensure(LiteralHandle.IsValid()))
-					{
-						// Always hide, even if no customization (LiteralObject isn't found) as this is the case
-						// where the default object is not required (i.e. Default Member is default constructed)
-						LiteralHandle->MarkHiddenByCustomization();
-
-						UObject* LiteralObject = nullptr;
-						if (LiteralHandle->GetValue(LiteralObject) == FPropertyAccess::Success)
-						{
-							MemberDefaultLiteral = Cast<UMetasoundEditorGraphMemberDefaultLiteral>(LiteralObject);
-						}
-					}
-				}
-			}
-
-			void CustomizeGeneralCategory(IDetailLayoutBuilder& InDetailLayout)
+			static IDetailCategoryBuilder& GetGeneralCategoryBuilder(IDetailLayoutBuilder& InDetailLayout)
 			{
-				IDetailCategoryBuilder& CategoryBuilder = InDetailLayout.EditCategory("General");
-
-				const bool bIsInterfaceMember = IsInterfaceMember();
-				const bool bIsGraphEditable = IsGraphEditable();
-
-				NameEditableTextBox = SNew(SEditableTextBox)
-					.Text(this, &TMetasoundGraphMemberDetailCustomization::GetName)
-					.OnTextChanged(this, &TMetasoundGraphMemberDetailCustomization::OnNameChanged)
-					.OnTextCommitted(this, &TMetasoundGraphMemberDetailCustomization::OnNameCommitted)
-					.IsReadOnly(bIsInterfaceMember || !bIsGraphEditable)
-					.SelectAllTextWhenFocused(true)
-					.Font(IDetailLayoutBuilder::GetDetailFont());
-
-				static const FText MemberNameToolTipFormat = LOCTEXT("GraphMember_NameDescriptionFormat", "Name used within the MetaSounds editor(s) and transacting systems (ex. Blueprints) if applicable to reference the given {0}.");
-				CategoryBuilder.AddCustomRow(LOCTEXT("GraphMember_NameProperty", "Name"))
-					.EditCondition(!bIsInterfaceMember && bIsGraphEditable, nullptr)
-					.NameContent()
-					[
-						SNew(STextBlock)
-						.Font(IDetailLayoutBuilder::GetDetailFontBold())
-						.Text(GraphMember->GetGraphMemberLabel())
-						.ToolTipText(FText::Format(MemberNameToolTipFormat, GraphMember->GetGraphMemberLabel()))
-					]
-					.ValueContent()
-					[
-						NameEditableTextBox.ToSharedRef()
-					];
-
-				static const FText MemberDisplayNameText = LOCTEXT("GraphMember_DisplayNameProperty", "Display Name");
-				static const FText MemberDisplayNameToolTipFormat = LOCTEXT("GraphMember_DisplayNameDescriptionFormat", "Optional, localized name used within the MetaSounds editor(s) to describe the given {0}.");
-				const FText MemberDisplayNameTooltipText = FText::Format(MemberDisplayNameToolTipFormat, GraphMember->GetGraphMemberLabel());
-				CategoryBuilder.AddCustomRow(MemberDisplayNameText)
-					.EditCondition(!bIsInterfaceMember && bIsGraphEditable, nullptr)
-					.NameContent()
-					[
-						SNew(STextBlock)
-						.Font(IDetailLayoutBuilder::GetDetailFontBold())
-						.Text(MemberDisplayNameText)
-						.ToolTipText(MemberDisplayNameTooltipText)
-					]
-					.ValueContent()
-					[
-						SNew(STextPropertyEditableTextBox, MakeShared<FGraphMemberEditableTextDisplayName>(GraphMember, MemberDisplayNameTooltipText))
-						.WrapTextAt(500)
-						.MinDesiredWidth(25.0f)
-						.MaxDesiredHeight(200)
-					];
-
-				static const FText MemberDescriptionText = LOCTEXT("Member_DescriptionPropertyName", "Description");
-				static const FText MemberDescriptionToolTipFormat = LOCTEXT("Member_DescriptionToolTipFormat", "Description for {0}. For example, used as a tooltip when displayed on another graph's referencing node.");
-				const FText MemberDescriptionToolTipText = FText::Format(MemberDescriptionToolTipFormat, GraphMember->GetGraphMemberLabel());
-				CategoryBuilder.AddCustomRow(MemberDescriptionText)
-					.EditCondition(!bIsInterfaceMember && bIsGraphEditable, nullptr)
-					.NameContent()
-					[
-						SNew(STextBlock)
-						.Font(IDetailLayoutBuilder::GetDetailFontBold())
-						.Text(MemberDescriptionText)
-						.ToolTipText(MemberDescriptionToolTipText)
-					]
-					.ValueContent()
-					[
-						SNew(STextPropertyEditableTextBox, MakeShared<FGraphMemberEditableTextDescription>(GraphMember, MemberDescriptionToolTipText))
-						.WrapTextAt(500)
-						.MinDesiredWidth(25.0f)
-						.MaxDesiredHeight(200)
-					];
-
-				DataTypeSelector->AddDataTypeSelector(InDetailLayout, MemberCustomizationStyle::DataTypeNameText, GraphMember, !bIsInterfaceMember && bIsGraphEditable);
+				return InDetailLayout.EditCategory("General");
 			}
 
-			TArray<IDetailPropertyRow*> CustomizeDefaultCategory(IDetailLayoutBuilder& InDetailLayout)
-			{
-				TArray<IDetailPropertyRow*> DefaultPropertyRows;
-
-				if (MemberDefaultLiteral.IsValid())
-				{
-					UpdateRenameDelegate(*MemberDefaultLiteral);
-
-					UClass* MemberClass = MemberDefaultLiteral->GetClass();
-					check(MemberClass);
-
-					IDetailCategoryBuilder& DefaultCategoryBuilder = GetDefaultCategoryBuilder(InDetailLayout);
-					IMetasoundEditorModule& EditorModule = FModuleManager::GetModuleChecked<IMetasoundEditorModule>("MetaSoundEditor");
-					TUniquePtr<FMetasoundDefaultLiteralCustomizationBase> LiteralCustomization = EditorModule.CreateMemberDefaultLiteralCustomization(*MemberClass, DefaultCategoryBuilder);
-					if (LiteralCustomization.IsValid())
-					{
-						DefaultPropertyRows = LiteralCustomization->CustomizeLiteral(*MemberDefaultLiteral, InDetailLayout);
-					}
-					else
-					{
-						IDetailPropertyRow* DefaultPropertyRow = DefaultCategoryBuilder.AddExternalObjectProperty(TArray<UObject*>({ MemberDefaultLiteral.Get() }), "Default");
-						ensureMsgf(DefaultPropertyRow, TEXT("Class '%s' missing expected 'Default' member."
-							"Either add/rename default member or register customization to display default value/opt out appropriately."),
-							*MemberClass->GetName());
-						DefaultPropertyRows.Add(DefaultPropertyRow);
-					}
-				}
-
-				return DefaultPropertyRows;
-			}
+			void UpdateRenameDelegate(UMetasoundEditorGraphMemberDefaultLiteral& InMemberDefaultLiteral);
+			void CacheMemberData(IDetailLayoutBuilder& InDetailLayout);
+			virtual void CustomizeGeneralCategory(IDetailLayoutBuilder& InDetailLayout);
+			virtual TArray<IDetailPropertyRow*> CustomizeDefaultCategory(IDetailLayoutBuilder& InDetailLayout);
 
 			virtual bool IsDefaultEditable() const
 			{
@@ -475,141 +327,47 @@ namespace Metasound
 			}
 
 			// IDetailCustomization interface
-			virtual void CustomizeDetails(IDetailLayoutBuilder& InDetailLayout) override
-			{
-				CacheMemberData(InDetailLayout);
-				if (GraphMember.IsValid())
-				{
-					CustomizeGeneralCategory(InDetailLayout);
-					CustomizeDefaultCategory(InDetailLayout);
-				}
-			}
+			virtual void CustomizeDetails(IDetailLayoutBuilder& InDetailLayout) override;
 			// End of IDetailCustomization interface
 
-			void OnNameChanged(const FText& InNewName)
-			{
-				using namespace Frontend;
-
-				bIsNameInvalid = false;
-				NameEditableTextBox->SetError(FText::GetEmpty());
-
-				if (!ensure(GraphMember.IsValid()))
-				{
-					return;
-				}
-
-				FText Error;
-				if (!GraphMember->CanRename(InNewName, Error))
-				{
-					bIsNameInvalid = true;
-					NameEditableTextBox->SetError(Error);
-				}
-			}
+			void OnNameChanged(const FText& InNewName);
+			FText GetName() const;
+			bool IsGraphEditable() const;
+			FText GetDisplayName() const;
+			void OnTooltipCommitted(const FText& InNewText, ETextCommit::Type InTextCommit);
+			FText GetTooltip() const;
+			void OnNameCommitted(const FText& InNewName, ETextCommit::Type InTextCommit);
 
-			FText GetName() const
-			{
-				if (GraphMember.IsValid())
-				{
-					return FText::FromName(GraphMember->GetMemberName());
-				}
-
-				return FText::GetEmpty();
-			}
-
-			bool IsGraphEditable() const
-			{
-				if (GraphMember.IsValid())
-				{
-					if (const UMetasoundEditorGraph* OwningGraph = GraphMember->GetOwningGraph())
-					{
-						return OwningGraph->IsEditable();
-					}
-				}
-
-				return false;
-			}
-
-			FText GetDisplayName() const
-			{
-				using namespace Frontend;
-
-				if (GraphMember.IsValid())
-				{
-					return GraphMember->GetDisplayName();
-				}
-
-				return FText::GetEmpty();
-			}
-
-			void OnTooltipCommitted(const FText& InNewText, ETextCommit::Type InTextCommit)
-			{
-				using namespace Frontend;
-
-				if (GraphMember.IsValid())
-				{
-					constexpr bool bPostTransaction = true;
-					GraphMember->SetDescription(InNewText, bPostTransaction);
-				}
-			}
-
-			FText GetTooltip() const
-			{
-				if (GraphMember.IsValid())
-				{
-					return GraphMember->GetDescription();
-				}
-
-				return FText::GetEmpty();
-			}
+		protected:
+			TWeakObjectPtr<UMetasoundEditorGraphMember> GraphMember;
 
-			void OnNameCommitted(const FText& InNewName, ETextCommit::Type InTextCommit)
-			{
-				using namespace Frontend;
+			TSharedPtr<SEditableTextBox> NameEditableTextBox;
+			FMetasoundDataTypeSelector DataTypeSelector;
 
-				if (!bIsNameInvalid && GraphMember.IsValid())
-				{
-					const FText TransactionLabel = FText::Format(LOCTEXT("RenameGraphMember_Format", "Set MetaSound {0}'s Name"), GraphMember->GetGraphMemberLabel());
-					const FScopedTransaction Transaction(TransactionLabel);
+			bool bIsNameInvalid = false;
 
-					constexpr bool bPostTransaction = false;
-					GraphMember->SetDisplayName(FText::GetEmpty(), bPostTransaction);
-					GraphMember->SetMemberName(*InNewName.ToString(), bPostTransaction);
-				}
+			FDelegateHandle RenameRequestedHandle;
+		};
 
-				NameEditableTextBox->SetError(FText::GetEmpty());
-				bIsNameInvalid = false;
-			}
+		using FMetasoundVariableDetailCustomization = FMetasoundMemberDetailCustomization;
 
-			ECheckBoxState OnGetPrivateCheckboxState() const
-			{
-				if (GraphMember.IsValid())
-				{
-					return GraphMember->GetNodeHandle()->GetNodeStyle().bIsPrivate ? ECheckBoxState::Checked : ECheckBoxState::Unchecked;
-				}
-
-				return ECheckBoxState::Unchecked;
-			}
+		class FMetasoundVertexDetailCustomization : public FMetasoundMemberDetailCustomization
+		{
+		public:
+			virtual ~FMetasoundVertexDetailCustomization() = default;
 
 		protected:
-			TWeakObjectPtr<UMetasoundEditorGraphMemberDefaultLiteral> MemberDefaultLiteral;
-
-			FDelegateHandle RenameRequestedHandle;
+			virtual void CustomizeGeneralCategory(IDetailLayoutBuilder& InDetailLayout) override;
 		};
 
-		class FMetasoundInputDetailCustomization : public TMetasoundGraphMemberDetailCustomization<UMetasoundEditorGraphInput>
+		class FMetasoundInputDetailCustomization : public FMetasoundVertexDetailCustomization
 		{
 		public:
-			FMetasoundInputDetailCustomization()
-				: TMetasoundGraphMemberDetailCustomization<UMetasoundEditorGraphInput>()
-			{
-			}
-
 			virtual ~FMetasoundInputDetailCustomization() = default;
 
 			virtual void CustomizeDetails(IDetailLayoutBuilder& InDetailLayout) override;
 			virtual bool IsDefaultEditable() const override;
 			virtual bool IsInterfaceMember() const override;
-			static const FText MemberNameText;
 
 		private:
 			Frontend::FDocumentHandle GetDocumentHandle() const;
@@ -618,16 +376,9 @@ namespace Metasound
 			void ClearInputInheritsDefault();
 		};
 
-		class FMetasoundOutputDetailCustomization : public TMetasoundGraphMemberDetailCustomization<UMetasoundEditorGraphOutput>
+		class FMetasoundOutputDetailCustomization : public FMetasoundVertexDetailCustomization
 		{
 		public:
-			static const FText MemberNameText;
-
-			FMetasoundOutputDetailCustomization()
-				: TMetasoundGraphMemberDetailCustomization<UMetasoundEditorGraphOutput>()
-			{
-			}
-
 			virtual ~FMetasoundOutputDetailCustomization() = default;
 
 			virtual bool IsInterfaceMember() const override;
diff --git a/Engine/Plugins/Runtime/Metasound/Source/MetasoundEditor/Private/MetasoundVariableDetailCustomization.cpp b/Engine/Plugins/Runtime/Metasound/Source/MetasoundEditor/Private/MetasoundVariableDetailCustomization.cpp
deleted file mode 100644
diff --git a/Engine/Plugins/Runtime/Metasound/Source/MetasoundEditor/Private/MetasoundVariableDetailCustomization.h b/Engine/Plugins/Runtime/Metasound/Source/MetasoundEditor/Private/MetasoundVariableDetailCustomization.h
deleted file mode 100644
diff --git a/Engine/Plugins/Runtime/Metasound/Source/MetasoundEditor/Public/MetasoundEditorGraph.h b/Engine/Plugins/Runtime/Metasound/Source/MetasoundEditor/Public/MetasoundEditorGraph.h
--- a/Engine/Plugins/Runtime/Metasound/Source/MetasoundEditor/Public/MetasoundEditorGraph.h
+++ b/Engine/Plugins/Runtime/Metasound/Source/MetasoundEditor/Public/MetasoundEditorGraph.h
@@ -10,6 +10,7 @@
 #include "MetasoundFrontendLiteral.h"
 #include "MetasoundSource.h"
 #include "MetasoundUObjectRegistry.h"
+#include "Misc/AssertionMacros.h"
 #include "UObject/ObjectMacros.h"
 #include "UObject/ObjectSaveContext.h"
 #include "UObject/ScriptInterface.h"
@@ -21,6 +22,8 @@
 struct FMetasoundFrontendDocument;
 struct FPropertyChangedChainEvent;
 struct FPropertyChangedEvent;
+
+class FMetasoundAssetBase;
 class ITargetPlatform;
 class UMetasoundEditorGraphInputNode;
 class UMetasoundEditorGraphNode;
@@ -159,11 +162,17 @@ public:
 	/** Returns the current data type */
 	FName GetDataType() const;
 
-	/** Returns literal associated with the given member */
-	virtual UMetasoundEditorGraphMemberDefaultLiteral* GetLiteral() const PURE_VIRTUAL(UMetasoundEditorGraphMember::GetLiteral, return nullptr; );
+	static FName GetLiteralPropertyName()
+	{
+		return GET_MEMBER_NAME_CHECKED(UMetasoundEditorGraphMember, Literal);
+	}
 
+	/** Returns literal associated with the given member */
+	UMetasoundEditorGraphMemberDefaultLiteral* GetLiteral() const { return Literal; }
 protected:
-	virtual void SetLiteral(UMetasoundEditorGraphMemberDefaultLiteral* InLiteral) PURE_VIRTUAL (UMetasoundEditorGraphMember::SetLiteral, );
+	/** Default literal value of member */
+	UPROPERTY()
+	UMetasoundEditorGraphMemberDefaultLiteral* Literal;
 
 	/** Conforms literal object type to member's DataType */
 	void ConformLiteralDataType();
@@ -218,7 +227,13 @@ public:
 	virtual bool IsInterfaceMember() const;
 
 	/** Returns the Metasound class type of the associated node */
-	virtual EMetasoundFrontendClassType GetClassType() const { return EMetasoundFrontendClassType::Invalid; }
+	virtual EMetasoundFrontendClassType GetClassType() const PURE_VIRTUAL(UMetasoundEditorGraphMember::GetClassType, return EMetasoundFrontendClassType::Invalid; )
+
+	/** Returns the SortOrderIndex assigned to this member. */
+	virtual int32 GetSortOrderIndex() const PURE_VIRTUAL(UMetasoundEditorGraphMember::GetSortOrderIndex, return 0; )
+
+	/** Sets the SortOrderIndex assigned to this member. */
+	virtual void SetSortOrderIndex(int32 InSortOrderIndex) PURE_VIRTUAL(UMetasoundEditorGraphMember::SetSortOrderIndex, )
 
 	/** Returns the node handle associated with the vertex. */
 	Metasound::Frontend::FNodeHandle GetNodeHandle();
@@ -226,12 +241,7 @@ public:
 	/** Returns the node handle associated with the vertex. */
 	Metasound::Frontend::FConstNodeHandle GetConstNodeHandle() const;
 
-	virtual UMetasoundEditorGraphMemberDefaultLiteral* GetLiteral() const PURE_VIRTUAL(UMetasoundEditorGraphMember::GetLiteral, return nullptr; );
-
 	virtual bool CanRename() const override;
-
-protected:
-	virtual void SetLiteral(UMetasoundEditorGraphMemberDefaultLiteral* InLiteral) PURE_VIRTUAL (UMetasoundEditorGraphMember::SetLiteral, );
 };
 
 UCLASS()
@@ -240,19 +250,15 @@ class METASOUNDEDITOR_API UMetasoundEditorGraphInput : public UMetasoundEditorGr
 	GENERATED_BODY()
 
 public:
-	/** Default literal value of member */
-	UPROPERTY(VisibleAnywhere, Category = DefaultValue)
-	UMetasoundEditorGraphMemberDefaultLiteral* Literal;
+	virtual int32 GetSortOrderIndex() const override;
+	virtual void SetSortOrderIndex(int32 InSortOrderIndex) override;
 
 	virtual const FText& GetGraphMemberLabel() const override;
 	virtual void ResetToClassDefault() override;
 	virtual void SetMemberName(const FName& InNewName, bool bPostTransaction) override;
 	virtual void UpdateFrontendDefaultLiteral(bool bPostTransaction) override;
 
-	virtual UMetasoundEditorGraphMemberDefaultLiteral* GetLiteral() const { return Literal; }
-
 protected:
-	virtual void SetLiteral(UMetasoundEditorGraphMemberDefaultLiteral* InLiteral) { Literal = InLiteral; }
 	virtual Metasound::Frontend::FNodeHandle AddNodeHandle(const FName& InNodeName, FName InDataType) override;
 	virtual EMetasoundFrontendClassType GetClassType() const override { return EMetasoundFrontendClassType::Input; }
 	virtual Metasound::Editor::ENodeSection GetSectionID() const override;
@@ -265,18 +271,13 @@ class METASOUNDEDITOR_API UMetasoundEditorGraphOutput : public UMetasoundEditorG
 	GENERATED_BODY()
 
 public:
-	/** Default literal value of member */
-	UPROPERTY(VisibleAnywhere, Category = DefaultValue)
-	UMetasoundEditorGraphMemberDefaultLiteral* Literal;
-
+	virtual int32 GetSortOrderIndex() const override;
+	virtual void SetSortOrderIndex(int32 InSortOrderIndex) override;
 	virtual const FText& GetGraphMemberLabel() const override;
 	virtual void ResetToClassDefault() override;
 	virtual void UpdateFrontendDefaultLiteral(bool bPostTransaction) override;
 
-	virtual UMetasoundEditorGraphMemberDefaultLiteral* GetLiteral() const { return Literal; }
-
 protected:
-	virtual void SetLiteral(UMetasoundEditorGraphMemberDefaultLiteral* InLiteral) { Literal = InLiteral; }
 	virtual Metasound::Frontend::FNodeHandle AddNodeHandle(const FName& InNodeName, FName InDataType) override;
 	virtual EMetasoundFrontendClassType GetClassType() const override { return EMetasoundFrontendClassType::Output; }
 	virtual Metasound::Editor::ENodeSection GetSectionID() const override;
@@ -291,10 +292,6 @@ class METASOUNDEDITOR_API UMetasoundEditorGraphVariable : public UMetasoundEdito
 	FGuid VariableID;
 
 public:
-	/** Default literal value of member */
-	UPROPERTY(VisibleAnywhere, Category = DefaultValue)
-	UMetasoundEditorGraphMemberDefaultLiteral* Literal;
-
 	void InitMember(FName InDataType, const FMetasoundFrontendLiteral& InDefaultLiteral, FGuid InVariableID);
 
 	/* ~Begin UMetasoundEditorGraphMember interface */
@@ -325,13 +322,8 @@ public:
 	Metasound::Frontend::FVariableHandle GetVariableHandle();
 	Metasound::Frontend::FConstVariableHandle GetConstVariableHandle() const;
 
-	virtual UMetasoundEditorGraphMemberDefaultLiteral* GetLiteral() const { return Literal; }
-
 	virtual bool CanRename() const override;
 
-protected:
-	virtual void SetLiteral(UMetasoundEditorGraphMemberDefaultLiteral* InLiteral) { Literal = InLiteral; }
-
 private:
 	struct FVariableEditorNodes
 	{
diff --git a/Engine/Plugins/Runtime/Metasound/Source/MetasoundEditor/Public/MetasoundEditorGraphNode.h b/Engine/Plugins/Runtime/Metasound/Source/MetasoundEditor/Public/MetasoundEditorGraphNode.h
--- a/Engine/Plugins/Runtime/Metasound/Source/MetasoundEditor/Public/MetasoundEditorGraphNode.h
+++ b/Engine/Plugins/Runtime/Metasound/Source/MetasoundEditor/Public/MetasoundEditorGraphNode.h
@@ -227,9 +227,6 @@ public:
 
 
 protected:
-	// Refreshes all Pin Metadata from the associated Frontend node's default class interface.
-	bool RefreshPinMetadata();
-
 	virtual void SetNodeID(FGuid InNodeID) override
 	{
 		NodeID = InNodeID;
diff --git a/Engine/Plugins/Runtime/Metasound/Source/MetasoundEngine/Private/MetasoundInterface.cpp b/Engine/Plugins/Runtime/Metasound/Source/MetasoundEngine/Private/MetasoundInterface.cpp
--- a/Engine/Plugins/Runtime/Metasound/Source/MetasoundEngine/Private/MetasoundInterface.cpp
+++ b/Engine/Plugins/Runtime/Metasound/Source/MetasoundEngine/Private/MetasoundInterface.cpp
@@ -132,62 +132,74 @@ namespace Metasound
 			FMetasoundFrontendInterface FrontendInterface;
 			FrontendInterface.Version = { Interface->GetName(), FMetasoundFrontendVersionNumber { Interface->GetVersion().Major, Interface->GetVersion().Minor } };
 
-			Algo::Transform(Interface->GetInputs(), FrontendInterface.Inputs, [&](const Audio::FParameterInterface::FInput& Input)
+			// Transfer all input data from AudioExtension interface struct to FrontendInterface
 			{
-				// Setup required inputs by telling the style that the input is required
-				// This will later be validated against.
-#if WITH_EDITOR
-				if (!Input.RequiredText.IsEmpty())
+				FMetasoundFrontendInterfaceStyle InputStyle;
+				Algo::Transform(Interface->GetInputs(), FrontendInterface.Inputs, [&](const Audio::FParameterInterface::FInput& Input)
 				{
-					FrontendInterface.AddRequiredInputToStyle(Input.InitValue.ParamName, Input.RequiredText);
-				}
-#endif // #if WITH_EDITOR
-
-				FMetasoundFrontendClassInput ClassInput;
-				ClassInput.Name = Input.InitValue.ParamName;
-				ClassInput.DefaultLiteral = FMetasoundFrontendLiteral(Input.InitValue);
-				ClassInput.TypeName = ResolveMemberDataType(Input.DataType, Input.InitValue.ParamType);
+					FMetasoundFrontendClassInput ClassInput;
+					ClassInput.Name = Input.InitValue.ParamName;
+					ClassInput.DefaultLiteral = FMetasoundFrontendLiteral(Input.InitValue);
+					ClassInput.TypeName = ResolveMemberDataType(Input.DataType, Input.InitValue.ParamType);
 
 #if WITH_EDITOR
-				// Interfaces should never serialize text to avoid desync between
-				// copied versions serialized in assets and those defined in code.
-				ClassInput.Metadata.SetSerializeText(false);
-				ClassInput.Metadata.SetDisplayName(Input.DisplayName);
-				ClassInput.Metadata.SetDescription(Input.Description);
+					// Interfaces should never serialize text to avoid desync between
+					// copied versions serialized in assets and those defined in code.
+					ClassInput.Metadata.SetSerializeText(false);
+					ClassInput.Metadata.SetDisplayName(Input.DisplayName);
+					ClassInput.Metadata.SetDescription(Input.Description);
+					ClassInput.Metadata.SortOrderIndex = Input.SortOrderIndex;
+					
+					InputStyle.DefaultSortOrder.Add(Input.SortOrderIndex);
+
+					// Setup required inputs by telling the style that the input is required
+					// This will later be validated against.
+					if (!Input.RequiredText.IsEmpty())
+					{
+						FrontendInterface.AddRequiredInputToStyle(Input.InitValue.ParamName, Input.RequiredText);
+					}
 #endif // WITH_EDITOR
 
-				ClassInput.VertexID = FGuid::NewGuid();
+					ClassInput.VertexID = FGuid::NewGuid();
 
-				return ClassInput;
-			});
+					return ClassInput;
+				});
+				FrontendInterface.SetInputStyle(InputStyle);
+			}
 
-			Algo::Transform(Interface->GetOutputs(), FrontendInterface.Outputs, [&](const Audio::FParameterInterface::FOutput& Output)
+			// Transfer all output data from AudioExtension interface struct to FrontendInterface
 			{
-				// Setup required outputs by telling the style that the output is required
-				// This will later be validated against.
-#if WITH_EDITOR
-				if (!Output.RequiredText.IsEmpty())
+				FMetasoundFrontendInterfaceStyle OutputStyle;
+				Algo::Transform(Interface->GetOutputs(), FrontendInterface.Outputs, [&](const Audio::FParameterInterface::FOutput& Output)
 				{
-					FrontendInterface.AddRequiredOutputToStyle(Output.ParamName, Output.RequiredText);
-				}
-#endif // #if WITH_EDITOR
-
-				FMetasoundFrontendClassOutput ClassOutput;
-				ClassOutput.Name = Output.ParamName;
-				ClassOutput.TypeName = ResolveMemberDataType(Output.DataType, Output.ParamType);
+					FMetasoundFrontendClassOutput ClassOutput;
+					ClassOutput.Name = Output.ParamName;
+					ClassOutput.TypeName = ResolveMemberDataType(Output.DataType, Output.ParamType);
 
 #if WITH_EDITOR
-				// Interfaces should never serialize text to avoid desync between
-				// copied versions serialized in assets and those defined in code.
-				ClassOutput.Metadata.SetSerializeText(false);
-				ClassOutput.Metadata.SetDisplayName(Output.DisplayName);
-				ClassOutput.Metadata.SetDescription(Output.Description);
+					// Interfaces should never serialize text to avoid desync between
+					// copied versions serialized in assets and those defined in code.
+					ClassOutput.Metadata.SetSerializeText(false);
+					ClassOutput.Metadata.SetDisplayName(Output.DisplayName);
+					ClassOutput.Metadata.SetDescription(Output.Description);
+					ClassOutput.Metadata.SortOrderIndex = Output.SortOrderIndex;
+					
+					OutputStyle.DefaultSortOrder.Add(Output.SortOrderIndex);
+
+					// Setup required outputs by telling the style that the output is required
+					// This will later be validated against.
+					if (!Output.RequiredText.IsEmpty())
+					{
+						FrontendInterface.AddRequiredOutputToStyle(Output.ParamName, Output.RequiredText);
+					}
 #endif // WITH_EDITOR
 
-				ClassOutput.VertexID = FGuid::NewGuid();
+					ClassOutput.VertexID = FGuid::NewGuid();
 
-				return ClassOutput;
-			});
+					return ClassOutput;
+				});
+				FrontendInterface.SetOutputStyle(OutputStyle);
+			}
 
 			Algo::Transform(Interface->GetEnvironment(), FrontendInterface.Environment, [&](const Audio::FParameterInterface::FEnvironmentVariable& Environment)
 			{
diff --git a/Engine/Plugins/Runtime/Metasound/Source/MetasoundFrontend/Private/MetasoundAssetBase.cpp b/Engine/Plugins/Runtime/Metasound/Source/MetasoundFrontend/Private/MetasoundAssetBase.cpp
--- a/Engine/Plugins/Runtime/Metasound/Source/MetasoundFrontend/Private/MetasoundAssetBase.cpp
+++ b/Engine/Plugins/Runtime/Metasound/Source/MetasoundFrontend/Private/MetasoundAssetBase.cpp
@@ -3,6 +3,7 @@
 #include "MetasoundAssetBase.h"
 
 #include "Algo/AnyOf.h"
+#include "Algo/ForEach.h"
 #include "Algo/Transform.h"
 #include "Containers/Set.h"
 #include "HAL/FileManager.h"
@@ -176,7 +177,6 @@ void FMetasoundAssetBase::RegisterGraphWithFrontend(Metasound::Frontend::FMetaSo
 			return nullptr;
 		}
 
-
 		virtual const FMetasoundFrontendClass& GetFrontendClass() const override
 		{
 			return FrontendClass;
@@ -429,6 +429,13 @@ void FMetasoundAssetBase::CacheRegistryMetadata()
 	}
 
 	// 2. Copy metadata for inputs/outputs managed by interfaces, removing them from maps generated
+	auto CacheInterfaceMetadata = [](const FMetasoundFrontendVertexMetadata & InRegistryMetadata, FMetasoundFrontendVertexMetadata& OutMetadata)
+	{
+		const int32 CachedSortOrderIndex = OutMetadata.SortOrderIndex;
+		OutMetadata = InRegistryMetadata;
+		OutMetadata.SortOrderIndex = CachedSortOrderIndex;
+	};
+
 	for (const FMetasoundFrontendVersion& Version : InterfaceVersions)
 	{
 		const FInterfaceRegistryKey InterfaceKey = GetInterfaceRegistryKey(Version);
@@ -440,8 +447,7 @@ void FMetasoundAssetBase::CacheRegistryMetadata()
 				const FNameDataTypePair NameDataTypePair = FNameDataTypePair(InterfaceInput.Name, InterfaceInput.TypeName);
 				if (FMetasoundFrontendClassInput* Input = Inputs.FindRef(NameDataTypePair))
 				{
-					// Interface members are set to not serialize text and thus is handled by copy, so explicit call omitted.
-					Input->Metadata = InterfaceInput.Metadata;
+					CacheInterfaceMetadata(InterfaceInput.Metadata, Input->Metadata);
 					Inputs.Remove(NameDataTypePair);
 				}
 			}
@@ -451,8 +457,7 @@ void FMetasoundAssetBase::CacheRegistryMetadata()
 				const FNameDataTypePair NameDataTypePair = FNameDataTypePair(InterfaceOutput.Name, InterfaceOutput.TypeName);
 				if (FMetasoundFrontendClassOutput* Output = Outputs.FindRef(NameDataTypePair))
 				{
-					// Interface members are set to not serialize text and thus is handled by copy, so explicit call omitted.
-					Output->Metadata = InterfaceOutput.Metadata;
+					CacheInterfaceMetadata(InterfaceOutput.Metadata, Output->Metadata);
 					Outputs.Remove(NameDataTypePair);
 				}
 			}
@@ -471,7 +476,26 @@ void FMetasoundAssetBase::CacheRegistryMetadata()
 		Pair.Value->Metadata.SetSerializeText(true);
 	}
 
-	// 4. Cache registry data on document dependencies
+	// 4. Refresh style as order of members could've changed
+	{
+		FMetasoundFrontendInterfaceStyle InputStyle;
+		Algo::ForEach(RootGraphClassInterface.Inputs, [&InputStyle](const FMetasoundFrontendClassInput& Input)
+		{
+			InputStyle.DefaultSortOrder.Add(Input.Metadata.SortOrderIndex);
+		});
+		RootGraphClassInterface.SetInputStyle(InputStyle);
+	}
+
+	{
+		FMetasoundFrontendInterfaceStyle OutputStyle;
+		Algo::ForEach(RootGraphClassInterface.Outputs, [&OutputStyle](const FMetasoundFrontendClassOutput& Output)
+		{
+			OutputStyle.DefaultSortOrder.Add(Output.Metadata.SortOrderIndex);
+		});
+		RootGraphClassInterface.SetOutputStyle(OutputStyle);
+	}
+
+	// 5. Cache registry data on document dependencies
 	for (FMetasoundFrontendClass& Dependency : Document->Dependencies)
 	{
 		if (!FMetasoundFrontendClass::CacheGraphDependencyMetadataFromRegistry(Dependency))
diff --git a/Engine/Plugins/Runtime/Metasound/Source/MetasoundFrontend/Private/MetasoundFrontendDocument.cpp b/Engine/Plugins/Runtime/Metasound/Source/MetasoundFrontend/Private/MetasoundFrontendDocument.cpp
--- a/Engine/Plugins/Runtime/Metasound/Source/MetasoundFrontend/Private/MetasoundFrontendDocument.cpp
+++ b/Engine/Plugins/Runtime/Metasound/Source/MetasoundFrontend/Private/MetasoundFrontendDocument.cpp
@@ -2,6 +2,7 @@
 
 #include "MetasoundFrontendDocument.h"
 
+#include "Algo/AnyOf.h"
 #include "Algo/ForEach.h"
 #include "Algo/Transform.h"
 #include "IAudioParameterInterfaceRegistry.h"
@@ -215,12 +216,12 @@ FMetasoundFrontendClassInterface FMetasoundFrontendClassInterface::GenerateClass
 			ClassInput.Metadata.bIsAdvancedDisplay = VertexMetadata.bIsAdvancedDisplay;
 
 			// Advanced display items are pushed to bottom of sort order
-			int32 OrderIndex = InputInterface.GetOrderIndex(InputTuple.Key);
+			ClassInput.Metadata.SortOrderIndex = InputInterface.GetSortOrderIndex(InputTuple.Key);
 			if (ClassInput.Metadata.bIsAdvancedDisplay)
 			{
-				OrderIndex += InputInterface.Num();
+				ClassInput.Metadata.SortOrderIndex += InputInterface.Num();
 			}
-			InputStyle.DefaultSortOrder.Add(OrderIndex);
+			InputStyle.DefaultSortOrder.Add(ClassInput.Metadata.SortOrderIndex);
 #endif // WITH_EDITOR
 
 			FLiteral DefaultLiteral = InputVertex.GetDefaultLiteral();
@@ -267,12 +268,12 @@ FMetasoundFrontendClassInterface FMetasoundFrontendClassInterface::GenerateClass
 			ClassOutput.Metadata.bIsAdvancedDisplay = VertexMetadata.bIsAdvancedDisplay;
 
 			// Advanced display items are pushed to bottom below non-advanced
-			int32 OrderIndex = OutputInterface.GetOrderIndex(OutputTuple.Key);
+			ClassOutput.Metadata.SortOrderIndex = OutputInterface.GetSortOrderIndex(OutputTuple.Key);
 			if (ClassOutput.Metadata.bIsAdvancedDisplay)
 			{
-				OrderIndex += OutputInterface.Num();
+				ClassOutput.Metadata.SortOrderIndex += OutputInterface.Num();
 			}
-			OutputStyle.DefaultSortOrder.Add(OrderIndex);
+			OutputStyle.DefaultSortOrder.Add(ClassOutput.Metadata.SortOrderIndex);
 #endif // WITH_EDITOR
 
 			ClassInterface.Outputs.Add(MoveTemp(ClassOutput));
@@ -415,34 +416,41 @@ bool FMetasoundFrontendClass::CacheGraphDependencyMetadataFromRegistry(FMetasoun
 			InOutDependency.Style = RegistryClass.Style;
 
 			using FNameTypeKey = TPair<FName, FName>;
-			TMap<FNameTypeKey, const FMetasoundFrontendVertexMetadata*> InterfaceMembers;
-
+			using FVertexMetadataMap = TMap<FNameTypeKey, const FMetasoundFrontendVertexMetadata*>;
 			auto MakePairFromVertex = [](const FMetasoundFrontendClassVertex& InVertex)
 			{
 				const FNameTypeKey Key(InVertex.Name, InVertex.TypeName);
 				return TPair<FNameTypeKey, const FMetasoundFrontendVertexMetadata*> { Key, &InVertex.Metadata };
 			};
 
-			auto CacheRegistryVertexMetadata = [&](FMetasoundFrontendClassVertex& OutVertex)
+			auto AddRegistryVertexMetadata = [](const FVertexMetadataMap& InInterfaceMembers, FMetasoundFrontendClassVertex& OutVertex, FMetasoundFrontendInterfaceStyle& OutNewStyle)
 			{
 				const FNameTypeKey Key(OutVertex.Name, OutVertex.TypeName);
-				if (const FMetasoundFrontendVertexMetadata* RegVertex = InterfaceMembers.FindRef(Key))
+				if (const FMetasoundFrontendVertexMetadata* RegVertex = InInterfaceMembers.FindRef(Key))
 				{
 					OutVertex.Metadata = *RegVertex;
 					OutVertex.Metadata.SetSerializeText(false);
 				}
+				OutNewStyle.DefaultSortOrder.Add(OutVertex.Metadata.SortOrderIndex);
 			};
 
-			Algo::Transform(RegistryClass.Interface.Inputs, InterfaceMembers, [&](const FMetasoundFrontendClassInput& Input) { return MakePairFromVertex(Input); });
-			Algo::ForEach(InOutDependency.Interface.Inputs, [&](FMetasoundFrontendClassInput& Input) { CacheRegistryVertexMetadata(Input); });
-
-			InterfaceMembers.Reset();
-
-			Algo::Transform(RegistryClass.Interface.Outputs, InterfaceMembers, [&](const FMetasoundFrontendClassOutput& Output) { return MakePairFromVertex(Output); });
-			Algo::ForEach(InOutDependency.Interface.Outputs, [&](FMetasoundFrontendClassOutput& Output) { CacheRegistryVertexMetadata(Output); });
-
-			InOutDependency.Interface.InputStyle = RegistryClass.Interface.InputStyle;
-			InOutDependency.Interface.OutputStyle = RegistryClass.Interface.OutputStyle;
+			FMetasoundFrontendInterfaceStyle InputStyle;
+			FVertexMetadataMap InputMembers;
+			Algo::Transform(RegistryClass.Interface.Inputs, InputMembers, [&](const FMetasoundFrontendClassInput& Input) { return MakePairFromVertex(Input); });
+			Algo::ForEach(InOutDependency.Interface.Inputs, [&](FMetasoundFrontendClassInput& Input)
+			{
+				AddRegistryVertexMetadata(InputMembers, Input, InputStyle);
+			});
+			InOutDependency.Interface.SetInputStyle(InputStyle);
+
+			FMetasoundFrontendInterfaceStyle OutputStyle;
+			FVertexMetadataMap OutputMembers;
+			Algo::Transform(RegistryClass.Interface.Outputs, OutputMembers, [&](const FMetasoundFrontendClassOutput& Output) { return MakePairFromVertex(Output); });
+			Algo::ForEach(InOutDependency.Interface.Outputs, [&](FMetasoundFrontendClassOutput& Output)
+			{
+				AddRegistryVertexMetadata(OutputMembers, Output, OutputStyle);
+			});
+			InOutDependency.Interface.SetOutputStyle(OutputStyle);
 
 			return true;
 		}
diff --git a/Engine/Plugins/Runtime/Metasound/Source/MetasoundFrontend/Private/MetasoundFrontendGraphController.cpp b/Engine/Plugins/Runtime/Metasound/Source/MetasoundFrontend/Private/MetasoundFrontendGraphController.cpp
--- a/Engine/Plugins/Runtime/Metasound/Source/MetasoundFrontend/Private/MetasoundFrontendGraphController.cpp
+++ b/Engine/Plugins/Runtime/Metasound/Source/MetasoundFrontend/Private/MetasoundFrontendGraphController.cpp
@@ -194,6 +194,22 @@ namespace Metasound
 				GraphClass->Graph.Style = InStyle;
 			}
 		}
+
+		void FGraphController::SetInputStyle(const FMetasoundFrontendInterfaceStyle& InStyle)
+		{
+			if (FMetasoundFrontendGraphClass* GraphClass = GraphClassPtr.Get())
+			{
+				GraphClass->Interface.SetInputStyle(InStyle);
+			}
+		}
+
+		void FGraphController::SetOutputStyle(const FMetasoundFrontendInterfaceStyle& InStyle)
+		{
+			if (FMetasoundFrontendGraphClass* GraphClass = GraphClassPtr.Get())
+			{
+				GraphClass->Interface.SetOutputStyle(InStyle);
+			}
+		}
 #endif // WITH_EDITOR
 
 		TArray<FNodeHandle> FGraphController::GetOutputNodes()
@@ -1174,6 +1190,26 @@ namespace Metasound
 			return FText::GetEmpty();
 		}
 
+		int32 FGraphController::GetSortOrderIndexForInput(const FVertexName& InName) const
+		{
+			if (const FMetasoundFrontendClassInput* Desc = FindInputDescriptionWithName(InName))
+			{
+				return Desc->Metadata.SortOrderIndex;
+			}
+
+			return 0;
+		}
+
+		int32 FGraphController::GetSortOrderIndexForOutput(const FVertexName& InName) const
+		{
+			if (const FMetasoundFrontendClassOutput* Desc = FindOutputDescriptionWithName(InName))
+			{
+				return Desc->Metadata.SortOrderIndex;
+			}
+
+			return 0;
+		}
+
 		void FGraphController::SetInputDisplayName(const FVertexName& InName, const FText& InDisplayName)
 		{
 			if (FMetasoundFrontendClassInput* Desc = FindInputDescriptionWithName(InName))
@@ -1205,6 +1241,42 @@ namespace Metasound
 				Desc->Metadata.SetDescription(InDescription);
 			}
 		}
+
+		void FGraphController::SetSortOrderIndexForInput(const FVertexName& InName, int32 InSortOrderIndex)
+		{
+			if (FMetasoundFrontendGraphClass* GraphClass = GraphClassPtr.Get())
+			{
+				FMetasoundFrontendInterfaceStyle Style = GraphClass->Interface.GetInputStyle();
+				Style.DefaultSortOrder.Reset();
+				for (FMetasoundFrontendClassInput& Input : GraphClass->Interface.Inputs)
+				{
+					if (Input.Name == InName)
+					{
+						Input.Metadata.SortOrderIndex = InSortOrderIndex;
+					}
+					Style.DefaultSortOrder.Add(Input.Metadata.SortOrderIndex);
+				}
+				GraphClass->Interface.SetInputStyle(Style);
+			}
+		}
+
+		void FGraphController::SetSortOrderIndexForOutput(const FVertexName& InName, int32 InSortOrderIndex)
+		{
+			if (FMetasoundFrontendGraphClass* GraphClass = GraphClassPtr.Get())
+			{
+				FMetasoundFrontendInterfaceStyle Style = GraphClass->Interface.GetOutputStyle();
+				Style.DefaultSortOrder.Reset();
+				for (FMetasoundFrontendClassOutput& Output : GraphClass->Interface.Outputs)
+				{
+					if (Output.Name == InName)
+					{
+						Output.Metadata.SortOrderIndex = InSortOrderIndex;
+					}
+					Style.DefaultSortOrder.Add(Output.Metadata.SortOrderIndex);
+				}
+				GraphClass->Interface.SetOutputStyle(Style);
+			}
+		}
 #endif // WITH_EDITOR
 
 		// This can be used to clear the current literal for a given input.
diff --git a/Engine/Plugins/Runtime/Metasound/Source/MetasoundFrontend/Private/MetasoundFrontendGraphController.h b/Engine/Plugins/Runtime/Metasound/Source/MetasoundFrontend/Private/MetasoundFrontendGraphController.h
--- a/Engine/Plugins/Runtime/Metasound/Source/MetasoundFrontend/Private/MetasoundFrontendGraphController.h
+++ b/Engine/Plugins/Runtime/Metasound/Source/MetasoundFrontend/Private/MetasoundFrontendGraphController.h
@@ -99,6 +99,8 @@ namespace Metasound
 #if WITH_EDITOR
 			const FMetasoundFrontendGraphStyle& GetGraphStyle() const override;
 			void SetGraphStyle(const FMetasoundFrontendGraphStyle& InStyle) override;
+			void SetInputStyle(const FMetasoundFrontendInterfaceStyle& InStyle) override;
+			void SetOutputStyle(const FMetasoundFrontendInterfaceStyle& InStyle) override;
 #endif // WITH_EDITOR
 
 			bool ContainsInputVertex(const FVertexName& InName, const FName& InTypeName) const override;
@@ -154,6 +156,11 @@ namespace Metasound
 
 			// Set the display name for the output with the given name
 			void SetOutputDisplayName(const FVertexName& InName, const FText& InDisplayName) override;
+
+			int32 GetSortOrderIndexForInput(const FVertexName& InName) const override;
+			int32 GetSortOrderIndexForOutput(const FVertexName& InName) const override;
+			void SetSortOrderIndexForInput(const FVertexName& InName, int32 InSortOrderIndex) override;
+			void SetSortOrderIndexForOutput(const FVertexName& InName, int32 InSortOrderIndex) override;
 #endif // WITH_EDITOR
 
 			// This can be used to clear the current literal for a given input.
diff --git a/Engine/Plugins/Runtime/Metasound/Source/MetasoundFrontend/Private/MetasoundFrontendInvalidController.h b/Engine/Plugins/Runtime/Metasound/Source/MetasoundFrontend/Private/MetasoundFrontendInvalidController.h
--- a/Engine/Plugins/Runtime/Metasound/Source/MetasoundFrontend/Private/MetasoundFrontendInvalidController.h
+++ b/Engine/Plugins/Runtime/Metasound/Source/MetasoundFrontend/Private/MetasoundFrontendInvalidController.h
@@ -333,6 +333,8 @@ namespace Metasound
 #if WITH_EDITOR
 			virtual const FMetasoundFrontendGraphStyle& GetGraphStyle() const override { return Invalid::GetInvalidGraphStyle(); }
 			virtual void SetGraphStyle(const FMetasoundFrontendGraphStyle& InStyle) override { }
+			virtual void SetInputStyle(const FMetasoundFrontendInterfaceStyle& InStyle) override { }
+			virtual void SetOutputStyle(const FMetasoundFrontendInterfaceStyle& InStyle) override { }
 #endif // WITH_EDITOR
 
 			virtual void ClearGraph() override { };
@@ -385,6 +387,11 @@ namespace Metasound
 			virtual void SetOutputDescription(const FVertexName& InName, const FText& InDescription) override { }
 			virtual void SetInputDisplayName(const FVertexName& InName, const FText& InDisplayName) override { }
 			virtual void SetOutputDisplayName(const FVertexName& InName, const FText& InDisplayName) override { }
+
+			virtual int32 GetSortOrderIndexForInput(const FVertexName& InName) const { return 0; }
+			virtual int32 GetSortOrderIndexForOutput(const FVertexName& InName) const { return 0; }
+			virtual void SetSortOrderIndexForInput(const FVertexName& InName, int32 InSortOrderIndex) { }
+			virtual void SetSortOrderIndexForOutput(const FVertexName& InName, int32 InSortOrderIndex) { }
 #endif // WITH_EDITOR
 
 			// This can be used to clear the current literal for a given input.
diff --git a/Engine/Plugins/Runtime/Metasound/Source/MetasoundFrontend/Private/MetasoundOutputFormatInterfaces.cpp b/Engine/Plugins/Runtime/Metasound/Source/MetasoundFrontend/Private/MetasoundOutputFormatInterfaces.cpp
--- a/Engine/Plugins/Runtime/Metasound/Source/MetasoundFrontend/Private/MetasoundOutputFormatInterfaces.cpp
+++ b/Engine/Plugins/Runtime/Metasound/Source/MetasoundFrontend/Private/MetasoundOutputFormatInterfaces.cpp
@@ -40,7 +40,10 @@ namespace Metasound
 								LOCTEXT("GeneratedAudioDisplayName", "Out Mono"),
 								LOCTEXT("GeneratedAudioDescription", "The resulting mono output from this source."),
 								GetMetasoundDataTypeName<FAudioBuffer>(),
-								Outputs::MonoOut
+								Outputs::MonoOut,
+								FText::GetEmpty(), // RequiredText
+								EAudioParameterType::None,
+								100
 							}
 						};
 					}
@@ -79,13 +82,19 @@ namespace Metasound
 								LOCTEXT("OutputFormatStereoInterface_GeneratedLeftDisplayName", "Out Left"),
 								LOCTEXT("OutputFormatStereoInterface_GeneratedLeftDescription", "The resulting left channel output audio from this source."),
 								GetMetasoundDataTypeName<FAudioBuffer>(),
-								Outputs::LeftOut
+								Outputs::LeftOut,
+								FText::GetEmpty(), // RequiredText
+								EAudioParameterType::None,
+								100
 							},
 							{
 								LOCTEXT("OutputFormatStereoInterface_GeneratedRightDisplayName", "Out Right"),
 								LOCTEXT("OutputFormatStereoInterface_GeneratedRightDescription", "The resulting right channel output audio from this source."),
 								GetMetasoundDataTypeName<FAudioBuffer>(),
-								Outputs::RightOut
+								Outputs::RightOut,
+								FText::GetEmpty(), // RequiredText
+								EAudioParameterType::None,
+								101
 							}
 						};
 					}
diff --git a/Engine/Plugins/Runtime/Metasound/Source/MetasoundFrontend/Public/MetasoundFrontendController.h b/Engine/Plugins/Runtime/Metasound/Source/MetasoundFrontend/Public/MetasoundFrontendController.h
--- a/Engine/Plugins/Runtime/Metasound/Source/MetasoundFrontend/Public/MetasoundFrontendController.h
+++ b/Engine/Plugins/Runtime/Metasound/Source/MetasoundFrontend/Public/MetasoundFrontendController.h
@@ -623,6 +623,12 @@ namespace Metasound
 			// Sets graph style.
 			virtual void SetGraphStyle(const FMetasoundFrontendGraphStyle& InStyle) = 0;
 
+			// Sets the input style for the graph.
+			virtual void SetInputStyle(const FMetasoundFrontendInterfaceStyle& InStyle) = 0;
+
+			// Sets the output style for the graph.
+			virtual void SetOutputStyle(const FMetasoundFrontendInterfaceStyle& InStyle) = 0;
+
 			/** Return the display name of the graph. */
 			virtual FText GetDisplayName() const = 0;
 #endif // WITH_EDITOR
@@ -874,6 +880,18 @@ namespace Metasound
 
 			/** Set the description for the output with the given name. */
 			virtual void SetOutputDescription(const FVertexName& InName, const FText& InDescription) = 0;
+
+			/** Returns the sort order index for the input with the given name. Returns 0 if not found or unset. */
+			virtual int32 GetSortOrderIndexForInput(const FVertexName& InName) const = 0;
+
+			/** Returns the sort order index for the input with the given name. Returns 0 if not found or unset. */
+			virtual int32 GetSortOrderIndexForOutput(const FVertexName& InName) const = 0;
+
+			/** Sets the sort order index for the input with the given name. No-ops if input not found. */
+			virtual void SetSortOrderIndexForInput(const FVertexName& InName, int32 InSortOrderIndex) = 0;
+
+			/** Sets the sort order index for the output with the given name. No-ops if output not found. */
+			virtual void SetSortOrderIndexForOutput(const FVertexName& InName, int32 InSortOrderIndex) = 0;
 #endif // WITH_EDITOR
 
 			/**
diff --git a/Engine/Plugins/Runtime/Metasound/Source/MetasoundFrontend/Public/MetasoundFrontendDocument.h b/Engine/Plugins/Runtime/Metasound/Source/MetasoundFrontend/Public/MetasoundFrontendDocument.h
--- a/Engine/Plugins/Runtime/Metasound/Source/MetasoundFrontend/Public/MetasoundFrontendDocument.h
+++ b/Engine/Plugins/Runtime/Metasound/Source/MetasoundFrontend/Public/MetasoundFrontendDocument.h
@@ -2,15 +2,16 @@
 
 #pragma once
 
+#include "Algo/Transform.h"
 #include "Containers/Array.h"
 #include "Containers/Map.h"
+#include "IAudioParameterInterfaceRegistry.h"
 #include "Internationalization/Text.h"
 #include "MetasoundAccessPtr.h"
 #include "MetasoundFrontendLiteral.h"
 #include "MetasoundNodeInterface.h"
 #include "MetasoundVertex.h"
 #include "Misc/Guid.h"
-#include "IAudioParameterInterfaceRegistry.h"
 #include "Templates/TypeHash.h"
 
 #include "MetasoundFrontendDocument.generated.h"
@@ -560,13 +561,9 @@ private:
 	FText DescriptionTransient;
 
 public:
-	// Keywords associated with the vertex
-	UPROPERTY()
-	TArray<FString> Keywords;
-
-	// Vertexes of the same group are generally placed together. 
+	// Order index of vertex member when shown as a node.
 	UPROPERTY()
-	FString Group;
+	int32 SortOrderIndex = 0;
 
 	// If true, vertex is shown for advanced display.
 	UPROPERTY()
@@ -769,69 +766,52 @@ struct METASOUNDFRONTEND_API FMetasoundFrontendClassEnvironmentVariable
 	bool bIsRequired = true;
 };
 
-// Layout mode for an interface.
-UENUM()
-enum class EMetasoundFrontendStyleInterfaceLayoutMode : uint8
-{
-	Default,
-	Inherited
-};
-
-
 // Style info of an interface.
 USTRUCT()
 struct FMetasoundFrontendInterfaceStyle
 {
 	GENERATED_BODY()
 
-	// Interface layout mode
-	UPROPERTY()
-	EMetasoundFrontendStyleInterfaceLayoutMode LayoutMode = EMetasoundFrontendStyleInterfaceLayoutMode::Inherited;
-
+#if WITH_EDITORONLY_DATA
 	// Default vertex sort order, where array index mirrors array interface index and value is display sort index.
 	UPROPERTY()
 	TArray<int32> DefaultSortOrder;
 
-#if WITH_EDITORONLY_DATA
 	// Map of member names with FText to be used as warnings if not hooked up
 	UPROPERTY()
 	TMap<FName, FText> RequiredMembers;
-#endif // #if WITH_EDITORONLY_DATA
 
 	template <typename HandleType>
-	TArray<HandleType> SortDefaults(const TArray<HandleType>& InHandles) const
+	void SortDefaults(TArray<HandleType>& OutHandles) const
 	{
-		TArray<HandleType> SortedHandles = InHandles;
-
-		// TODO: Hack for assets which aren't getting sort order set for inputs/outputs. Fix this & remove size check.
-		if (DefaultSortOrder.Num() > 0)
+		TMap<FGuid, int32> NodeIDToSortIndex;
+		int32 HighestSortOrder = TNumericLimits<int32>::Min();
+		for (int32 i = 0; i < OutHandles.Num(); ++i)
 		{
-			if (SortedHandles.Num() == DefaultSortOrder.Num())
+			const FGuid& HandleID = OutHandles[i]->GetID();
+			int32 SortIndex = 0;
+			if (DefaultSortOrder.IsValidIndex(i))
 			{
-				TMap<FGuid, int32> HandleIDToSortIndex;
-				for (int32 i = 0; i < DefaultSortOrder.Num(); ++i)
-				{
-					if (InHandles.IsValidIndex(i))
-					{
-						const int32 SortIndex = DefaultSortOrder[i];
-						HandleIDToSortIndex.Add(InHandles[i]->GetID(), SortIndex);
-					}
-				}
-
-				SortedHandles.Sort([&](const HandleType& HandleA, const HandleType& HandleB)
-				{
-					const FGuid HandleAID = HandleA->GetID();
-					const FGuid HandleBID = HandleB->GetID();
-					return HandleIDToSortIndex[HandleAID] < HandleIDToSortIndex[HandleBID];
-				});
+				SortIndex = DefaultSortOrder[i];
+				HighestSortOrder = FMath::Max(SortIndex, HighestSortOrder);
 			}
+			else
+			{
+				SortIndex = ++HighestSortOrder;
+			}
+			NodeIDToSortIndex.Add(HandleID, SortIndex);
 		}
 
-		return SortedHandles;
+		OutHandles.Sort([&NodeIDToSortIndex](const HandleType& HandleA, const HandleType& HandleB)
+		{
+			const FGuid HandleAID = HandleA->GetID();
+			const FGuid HandleBID = HandleB->GetID();
+			return NodeIDToSortIndex[HandleAID] < NodeIDToSortIndex[HandleBID];
+		});
 	}
+#endif // #if WITH_EDITORONLY_DATA
 };
 
-
 USTRUCT()
 struct METASOUNDFRONTEND_API FMetasoundFrontendClassInterface
 {
diff --git a/Engine/Plugins/Runtime/Metasound/Source/MetasoundGraphCore/Public/MetasoundVertex.h b/Engine/Plugins/Runtime/Metasound/Source/MetasoundGraphCore/Public/MetasoundVertex.h
--- a/Engine/Plugins/Runtime/Metasound/Source/MetasoundGraphCore/Public/MetasoundVertex.h
+++ b/Engine/Plugins/Runtime/Metasound/Source/MetasoundGraphCore/Public/MetasoundVertex.h
@@ -2,7 +2,6 @@
 
 #pragma once
 
-#include "CoreMinimal.h"
 #include "Containers/SortedMap.h"
 #include "MetasoundDataFactory.h"
 #include "MetasoundDataReference.h"
@@ -729,7 +728,7 @@ namespace Metasound
 				return Vertices.Find(InKey);
 			}
 
-			int32 GetOrderIndex(const FVertexName& InKey) const
+			int32 GetSortOrderIndex(const FVertexName& InKey) const
 			{
 				int32 OutIndex = INDEX_NONE;
 				OrderedKeys.Find(InKey, OutIndex);
diff --git a/Engine/Source/Runtime/AudioExtensions/Public/IAudioParameterInterfaceRegistry.h b/Engine/Source/Runtime/AudioExtensions/Public/IAudioParameterInterfaceRegistry.h
--- a/Engine/Source/Runtime/AudioExtensions/Public/IAudioParameterInterfaceRegistry.h
+++ b/Engine/Source/Runtime/AudioExtensions/Public/IAudioParameterInterfaceRegistry.h
@@ -30,6 +30,8 @@ namespace Audio
 
 			const FAudioParameter InitValue;
 			const FText RequiredText;
+
+			const int32 SortOrderIndex = 0;
 		};
 
 		struct FOutput
@@ -41,6 +43,8 @@ namespace Audio
 			const FName ParamName;
 			const FText RequiredText = FText();
 			const EAudioParameterType ParamType = EAudioParameterType::None;
+
+			const int32 SortOrderIndex = 0;
 		};
 
 		struct FEnvironmentVariable
