- Fix corrupted transaction stack when deleting a MetaSound member & then calling undo
- Guard against accessing null member entries when synchronizing MetaSound edgraph
#rnx

diff --git a/Engine/Plugins/Runtime/Metasound/Source/MetasoundEditor/Private/MetasoundEditor.cpp b/Engine/Plugins/Runtime/Metasound/Source/MetasoundEditor/Private/MetasoundEditor.cpp
--- a/Engine/Plugins/Runtime/Metasound/Source/MetasoundEditor/Private/MetasoundEditor.cpp
+++ b/Engine/Plugins/Runtime/Metasound/Source/MetasoundEditor/Private/MetasoundEditor.cpp
@@ -1754,6 +1754,7 @@ namespace Metasound
 					const FScopedTransaction Transaction(LOCTEXT("MetaSoundEditorDeleteSelectedMember", "Delete MetaSound Graph Member"));
 					Metasound->Modify();
 					Graph.Modify();
+					GraphMember->Modify();
 					Graph.RemoveMember(*GraphMember);
 				}
 
@@ -2666,7 +2667,7 @@ namespace Metasound
 				ActionList->AddAction(NewFuncAction);
 
 				UMetasoundEditorGraphInput * EdGraphInput = EdGraph.FindInput(NodeID);
-				if (ensure(EdGraphInput))
+				if (EdGraphInput)
 				{
 					if (FDelegateHandle* NameChangeDelegate = NameChangeDelegateHandles.Find(NodeID))
 					{
diff --git a/Engine/Plugins/Runtime/Metasound/Source/MetasoundEditor/Private/MetasoundEditorGraph.cpp b/Engine/Plugins/Runtime/Metasound/Source/MetasoundEditor/Private/MetasoundEditorGraph.cpp
--- a/Engine/Plugins/Runtime/Metasound/Source/MetasoundEditor/Private/MetasoundEditorGraph.cpp
+++ b/Engine/Plugins/Runtime/Metasound/Source/MetasoundEditor/Private/MetasoundEditorGraph.cpp
@@ -1313,7 +1313,12 @@ UMetasoundEditorGraphInput* UMetasoundEditorGraph::FindInput(FGuid InNodeID) con
 {
 	const TObjectPtr<UMetasoundEditorGraphInput>* Input = Inputs.FindByPredicate([InNodeID](const TObjectPtr<UMetasoundEditorGraphInput>& InInput)
 	{
-		return InInput->NodeID == InNodeID;
+		if (InInput)
+		{
+			return InInput->NodeID == InNodeID;
+		}
+
+		return false;
 	});
 	return Input ? Input->Get() : nullptr;
 }
@@ -1322,8 +1327,13 @@ UMetasoundEditorGraphInput* UMetasoundEditorGraph::FindInput(FName InName) const
 {
 	const TObjectPtr<UMetasoundEditorGraphInput>* Input = Inputs.FindByPredicate([InName](const TObjectPtr<UMetasoundEditorGraphInput>& InInput)
 	{
-		const FName NodeName = InInput->GetMemberName();
-		return NodeName == InName;
+		if (InInput)
+		{
+			const FName NodeName = InInput->GetMemberName();
+			return NodeName == InName;
+		}
+
+		return false;
 	});
 	return Input ? Input->Get() : nullptr;
 }
@@ -1370,7 +1380,11 @@ UMetasoundEditorGraphOutput* UMetasoundEditorGraph::FindOutput(FGuid InNodeID) c
 {
 	const TObjectPtr<UMetasoundEditorGraphOutput>* Output = Outputs.FindByPredicate([InNodeID](const TObjectPtr<UMetasoundEditorGraphOutput>& InOutput)
 	{
-		return InOutput->NodeID == InNodeID;
+		if (InOutput)
+		{
+			return InOutput->NodeID == InNodeID;
+		}
+		return false;
 	});
 	return Output ? Output->Get() : nullptr;
 }
@@ -1379,7 +1393,11 @@ UMetasoundEditorGraphOutput* UMetasoundEditorGraph::FindOutput(FName InName) con
 {
 	const TObjectPtr<UMetasoundEditorGraphOutput>* Output = Outputs.FindByPredicate([&InName](const TObjectPtr<UMetasoundEditorGraphOutput>& InOutput)
 	{
-		return InName == InOutput->GetMemberName();
+		if (InOutput)
+		{
+			return InName == InOutput->GetMemberName();
+		}
+		return false;
 	});
 	return Output ? Output->Get() : nullptr;
 }
@@ -1553,7 +1571,7 @@ void UMetasoundEditorGraph::IterateInputs(TUniqueFunction<void(UMetasoundEditorG
 {
 	for (UMetasoundEditorGraphInput* Input : Inputs)
 	{
-		if (ensure(Input))
+		if (Input)
 		{
 			InFunction(*Input);
 		}
diff --git a/Engine/Plugins/Runtime/Metasound/Source/MetasoundEditor/Private/MetasoundEditorGraphBuilder.cpp b/Engine/Plugins/Runtime/Metasound/Source/MetasoundEditor/Private/MetasoundEditorGraphBuilder.cpp
--- a/Engine/Plugins/Runtime/Metasound/Source/MetasoundEditor/Private/MetasoundEditorGraphBuilder.cpp
+++ b/Engine/Plugins/Runtime/Metasound/Source/MetasoundEditor/Private/MetasoundEditorGraphBuilder.cpp
@@ -1690,7 +1690,7 @@ namespace Metasound
 
 		bool FGraphBuilder::SynchronizeGraph(UObject& InMetaSound, bool bForceRefreshNodes)
 		{
-			bool bEditorGraphModified = SynchronizeGraphVertices(InMetaSound);
+			bool bEditorGraphModified = SynchronizeGraphMembers(InMetaSound);
 			bEditorGraphModified |= SynchronizeNodeMembers(InMetaSound);
 			bEditorGraphModified |= SynchronizeNodes(InMetaSound);
 			bEditorGraphModified |= SynchronizeConnections(InMetaSound);
@@ -2115,7 +2115,7 @@ namespace Metasound
 			return OldValue != InPin.DefaultValue;
 		}
 
-		bool FGraphBuilder::SynchronizeGraphVertices(UObject& InMetaSound)
+		bool FGraphBuilder::SynchronizeGraphMembers(UObject& InMetaSound)
 		{
 			using namespace Frontend;
 
@@ -2259,7 +2259,10 @@ namespace Metasound
 				}
 			}, EMetasoundFrontendClassType::Output);
 
-
+			// Remove empty entries
+			bEditorGraphModified |= Graph->Inputs.RemoveAllSwap([](const TObjectPtr<UMetasoundEditorGraphInput>& Input) { return !Input; }) > 0;
+			bEditorGraphModified |= Graph->Outputs.RemoveAllSwap([](const TObjectPtr<UMetasoundEditorGraphOutput>& Output) { return !Output; }) > 0;
+			bEditorGraphModified |= Graph->Variables.RemoveAllSwap([](const TObjectPtr<UMetasoundEditorGraphVariable>& Variable) { return !Variable; }) > 0;
 
 			return bEditorGraphModified;
 		}
diff --git a/Engine/Plugins/Runtime/Metasound/Source/MetasoundEditor/Private/MetasoundEditorGraphBuilder.h b/Engine/Plugins/Runtime/Metasound/Source/MetasoundEditor/Private/MetasoundEditorGraphBuilder.h
--- a/Engine/Plugins/Runtime/Metasound/Source/MetasoundEditor/Private/MetasoundEditorGraphBuilder.h
+++ b/Engine/Plugins/Runtime/Metasound/Source/MetasoundEditor/Private/MetasoundEditorGraphBuilder.h
@@ -249,10 +249,10 @@ namespace Metasound
 			// Synchronizes pin type for a given pin with that registered with the MetaSound editor module provided.
 			static bool SynchronizePinType(const IMetasoundEditorModule& InEditorModule, UEdGraphPin& InPin, const FName InDataType);
 
-			// Synchronizes inputs and outputs for the given MetaSound.
+			// Synchronizes inputs, variables, and outputs for the given MetaSound.
 			//
 			// @return True if the UEdGraph was altered. False otherwise.
-			static bool SynchronizeGraphVertices(UObject& InMetaSound);
+			static bool SynchronizeGraphMembers(UObject& InMetaSound);
 
 			// Returns true if the FInputHandle and UEdGraphPin match each other.
 			static bool IsMatchingInputHandleAndPin(const Frontend::FConstInputHandle& InInputHandle, const UEdGraphPin& InEditorPin);
diff --git a/Engine/Plugins/Runtime/Metasound/Source/MetasoundEditor/Private/MetasoundEditorGraphInputNode.cpp b/Engine/Plugins/Runtime/Metasound/Source/MetasoundEditor/Private/MetasoundEditorGraphInputNode.cpp
--- a/Engine/Plugins/Runtime/Metasound/Source/MetasoundEditor/Private/MetasoundEditorGraphInputNode.cpp
+++ b/Engine/Plugins/Runtime/Metasound/Source/MetasoundEditor/Private/MetasoundEditorGraphInputNode.cpp
@@ -114,5 +114,8 @@ bool UMetasoundEditorGraphInputNode::Validate(Metasound::Editor::FGraphNodeValid
 
 void UMetasoundEditorGraphInputNode::SetNodeID(FGuid InNodeID)
 {
-	Input->NodeID = InNodeID;
+	if (Input)
+	{
+		Input->NodeID = InNodeID;
+	}
 }
