Fix for metasound Random node failure to reinit shared state.

diff --git a/Engine/Plugins/Runtime/Metasound/Source/MetasoundFrontend/Private/MetasoundArrayRandomNode.cpp b/Engine/Plugins/Runtime/Metasound/Source/MetasoundFrontend/Private/MetasoundArrayRandomNode.cpp
--- a/Engine/Plugins/Runtime/Metasound/Source/MetasoundFrontend/Private/MetasoundArrayRandomNode.cpp
+++ b/Engine/Plugins/Runtime/Metasound/Source/MetasoundFrontend/Private/MetasoundArrayRandomNode.cpp
@@ -183,17 +183,32 @@ namespace Metasound
 		return RGM;
 	}
 
+	inline bool ShouldStompSharedState(const FArrayRandomGet& InRandomGet, const InitSharedStateArgs InArgs)
+	{
+		return
+			InArgs.bIsPreviewSound										// If it's a preview sound
+			|| InRandomGet.GetMaxIndex() != InArgs.NumElements			// Size of Array has changed
+			|| InRandomGet.GetNoRepeatOrder() != InArgs.NoRepeatOrder;	// Repeat order has changed
+	}
+
 	void FSharedStateRandomGetManager::InitSharedState(InitSharedStateArgs& InArgs)
 	{
 		FScopeLock Lock(&CritSect);
 
-		if (!InArgs.bIsPreviewSound && RandomGets.Contains(InArgs.SharedStateId))
+		// Look if there's already shared state. In certain cases we should stomp it by re-adding...
+		if (const TUniquePtr<FArrayRandomGet>* FoundExisting = RandomGets.Find(InArgs.SharedStateId))
 		{
-			// TODO: check the actual arguments here against what has been registered, and if the new args are different warn they'll be ignored.
-			return;
+			if (FoundExisting && FoundExisting->IsValid())
+			{
+				const FArrayRandomGet& Existing = *FoundExisting->Get();
+				if (!ShouldStompSharedState(Existing, InArgs))
+				{
+					return;
+				}
+			}
 		}
 
-		// For preview sounds, we'll stomp the state
+		// Add state/Stomp existing.
 		RandomGets.Add(InArgs.SharedStateId, MakeUnique<FArrayRandomGet>(InArgs.Seed, InArgs.NumElements, InArgs.Weights, InArgs.NoRepeatOrder));
 	}
 
diff --git a/Engine/Plugins/Runtime/Metasound/Source/MetasoundFrontend/Public/MetasoundArrayRandomNode.h b/Engine/Plugins/Runtime/Metasound/Source/MetasoundFrontend/Public/MetasoundArrayRandomNode.h
--- a/Engine/Plugins/Runtime/Metasound/Source/MetasoundFrontend/Public/MetasoundArrayRandomNode.h
+++ b/Engine/Plugins/Runtime/Metasound/Source/MetasoundFrontend/Public/MetasoundArrayRandomNode.h
@@ -53,6 +53,9 @@ namespace Metasound
 		void ResetSeed();
 		int32 NextValue();
 
+		int32 GetNoRepeatOrder() const { return NoRepeatOrder; }
+		int32 GetMaxIndex() const { return MaxIndex; }
+
 	private:
 		float ComputeTotalWeight();
 
