Log warning when node connection is dropped during auto-update due to removed vertex.
#rnx

diff --git a/Engine/Plugins/Runtime/Metasound/Source/MetasoundEditor/Private/MetasoundEditor.cpp b/Engine/Plugins/Runtime/Metasound/Source/MetasoundEditor/Private/MetasoundEditor.cpp
--- a/Engine/Plugins/Runtime/Metasound/Source/MetasoundEditor/Private/MetasoundEditor.cpp
+++ b/Engine/Plugins/Runtime/Metasound/Source/MetasoundEditor/Private/MetasoundEditor.cpp
@@ -2489,8 +2489,13 @@ namespace Metasound
 
 					if (bHasNewVersion || !bIsClassNative)
 					{
+						// These are ignored here when updating as the user is actively
+						// forcing an update.
+						constexpr TArray<INodeController::FVertexNameAndType>* DisconnectedInputs = nullptr;
+						constexpr TArray<INodeController::FVertexNameAndType>* DisconnectedOutputs = nullptr;
+
 						FNodeHandle ExistingNode = ExternalNode->GetNodeHandle();
-						FNodeHandle NewNode = ExistingNode->ReplaceWithVersion(HighestVersion);
+						FNodeHandle NewNode = ExistingNode->ReplaceWithVersion(HighestVersion, DisconnectedInputs, DisconnectedOutputs);
 						bReplacedNodes = true;
 					}
 				}
diff --git a/Engine/Plugins/Runtime/Metasound/Source/MetasoundEditor/Private/MetasoundEditorModule.cpp b/Engine/Plugins/Runtime/Metasound/Source/MetasoundEditor/Private/MetasoundEditorModule.cpp
--- a/Engine/Plugins/Runtime/Metasound/Source/MetasoundEditor/Private/MetasoundEditorModule.cpp
+++ b/Engine/Plugins/Runtime/Metasound/Source/MetasoundEditor/Private/MetasoundEditorModule.cpp
@@ -29,6 +29,7 @@
 #include "MetasoundFrontendRegistries.h"
 #include "MetasoundFrontendTransform.h"
 #include "MetasoundNodeDetailCustomization.h"
+#include "MetasoundSettings.h"
 #include "MetasoundSource.h"
 #include "MetasoundTime.h"
 #include "MetasoundTrigger.h"
@@ -165,14 +166,19 @@ namespace Metasound
 		{
 			void LoadAndRegisterAsset(const FAssetData& InAssetData)
 			{
+				Frontend::FMetaSoundAssetRegistrationOptions RegOptions;
+				RegOptions.bForceReregister = false;
+				if (const UMetaSoundSettings* Settings = GetDefault<UMetaSoundSettings>())
+				{
+					RegOptions.bAutoUpdateLogWarningOnDroppedConnection = Settings->bAutoUpdateLogWarningOnDroppedConnection;
+				}
+
 				if (InAssetData.IsAssetLoaded())
 				{
 					if (UObject* AssetObject = InAssetData.GetAsset())
 					{
 						FMetasoundAssetBase* MetaSoundAsset = IMetasoundUObjectRegistry::Get().GetObjectAsAssetBase(AssetObject);
 						check(MetaSoundAsset);
-						Frontend::FMetaSoundAssetRegistrationOptions RegOptions;
-						RegOptions.bForceReregister = false;
 						MetaSoundAsset->RegisterGraphWithFrontend(RegOptions);
 					}
 				}
@@ -191,7 +197,7 @@ namespace Metasound
 					ActiveAsyncAssetLoadRequests++;
 
 					FSoftObjectPath AssetPath = InAssetData.ToSoftObjectPath();
-					auto LoadAndRegister = [this, ObjectPath = AssetPath](const FName& PackageName, UPackage* Package, EAsyncLoadingResult::Type Result)
+					auto LoadAndRegister = [this, ObjectPath = AssetPath, RegOptions](const FName& PackageName, UPackage* Package, EAsyncLoadingResult::Type Result)
 					{
 						if (Result == EAsyncLoadingResult::Succeeded)
 						{
@@ -199,8 +205,6 @@ namespace Metasound
 							check(MetaSoundAsset);
 							if (!MetaSoundAsset->IsRegistered())
 							{
-								Frontend::FMetaSoundAssetRegistrationOptions RegOptions;
-								RegOptions.bForceReregister = false;
 								MetaSoundAsset->RegisterGraphWithFrontend(RegOptions);
 							}
 						}
diff --git a/Engine/Plugins/Runtime/Metasound/Source/MetasoundEngine/Private/MetasoundAssetSubsystem.cpp b/Engine/Plugins/Runtime/Metasound/Source/MetasoundEngine/Private/MetasoundAssetSubsystem.cpp
--- a/Engine/Plugins/Runtime/Metasound/Source/MetasoundEngine/Private/MetasoundAssetSubsystem.cpp
+++ b/Engine/Plugins/Runtime/Metasound/Source/MetasoundEngine/Private/MetasoundAssetSubsystem.cpp
@@ -460,7 +460,13 @@ void UMetaSoundAssetSubsystem::RegisterAssetClassesInDirectories(const TArray<FM
 		AddOrUpdateAsset(AssetData);
 		FMetasoundAssetBase* MetaSoundAsset = Metasound::IMetasoundUObjectRegistry::Get().GetObjectAsAssetBase(AssetData.GetAsset());
 		check(MetaSoundAsset);
-		MetaSoundAsset->RegisterGraphWithFrontend();
+
+		Metasound::Frontend::FMetaSoundAssetRegistrationOptions RegOptions;
+		if (const UMetaSoundSettings* Settings = GetDefault<UMetaSoundSettings>())
+		{
+			RegOptions.bAutoUpdateLogWarningOnDroppedConnection = Settings->bAutoUpdateLogWarningOnDroppedConnection;
+		}
+		MetaSoundAsset->RegisterGraphWithFrontend(RegOptions);
 	});
 }
 
diff --git a/Engine/Plugins/Runtime/Metasound/Source/MetasoundEngine/Private/MetasoundSource.cpp b/Engine/Plugins/Runtime/Metasound/Source/MetasoundEngine/Private/MetasoundSource.cpp
--- a/Engine/Plugins/Runtime/Metasound/Source/MetasoundEngine/Private/MetasoundSource.cpp
+++ b/Engine/Plugins/Runtime/Metasound/Source/MetasoundEngine/Private/MetasoundSource.cpp
@@ -24,6 +24,7 @@
 #include "MetasoundParameterTransmitter.h"
 #include "MetasoundPrimitives.h"
 #include "MetasoundReceiveNode.h"
+#include "MetasoundSettings.h"
 #include "MetasoundSourceInterface.h"
 #include "MetasoundTrace.h"
 #include "MetasoundTrigger.h"
@@ -451,6 +452,10 @@ void UMetaSoundSource::InitResources()
 
 	FMetaSoundAssetRegistrationOptions RegOptions;
 	RegOptions.bForceReregister = false;
+	if (const UMetaSoundSettings* Settings = GetDefault<UMetaSoundSettings>())
+	{
+		RegOptions.bAutoUpdateLogWarningOnDroppedConnection = Settings->bAutoUpdateLogWarningOnDroppedConnection;
+	}
 	RegisterGraphWithFrontend(RegOptions);
 }
 
diff --git a/Engine/Plugins/Runtime/Metasound/Source/MetasoundEngine/Public/MetasoundSettings.h b/Engine/Plugins/Runtime/Metasound/Source/MetasoundEngine/Public/MetasoundSettings.h
--- a/Engine/Plugins/Runtime/Metasound/Source/MetasoundEngine/Public/MetasoundSettings.h
+++ b/Engine/Plugins/Runtime/Metasound/Source/MetasoundEngine/Public/MetasoundSettings.h
@@ -52,6 +52,10 @@ public:
 	UPROPERTY(EditAnywhere, config, Category = AutoUpdate, meta = (DisplayName = "Asset DenyList", EditCondition = "bAutoUpdateEnabled"))
 	TArray<FDefaultMetaSoundAssetAutoUpdateSettings> AutoUpdateAssetDenylist;
 
+	/** If true, warnings will be logged if updating a node results in existing connections being discarded. */
+	UPROPERTY(EditAnywhere, config, Category = AutoUpdate, meta = (DisplayName = "Log Warning on Dropped Connection", EditCondition = "bAutoUpdateEnabled"))
+	bool bAutoUpdateLogWarningOnDroppedConnection = true;
+
 	/** Directories to scan & automatically register MetaSound post initial asset scan on engine start-up.
 	  * May speed up subsequent calls to playback MetaSounds post asset scan but increases application load time.
 	  * See 'MetaSoundAssetSubsystem::RegisterAssetClassesInDirectories' to dynamically register or 
diff --git a/Engine/Plugins/Runtime/Metasound/Source/MetasoundFrontend/Private/MetasoundAssetBase.cpp b/Engine/Plugins/Runtime/Metasound/Source/MetasoundFrontend/Private/MetasoundAssetBase.cpp
--- a/Engine/Plugins/Runtime/Metasound/Source/MetasoundFrontend/Private/MetasoundAssetBase.cpp
+++ b/Engine/Plugins/Runtime/Metasound/Source/MetasoundFrontend/Private/MetasoundAssetBase.cpp
@@ -112,7 +112,7 @@ void FMetasoundAssetBase::RegisterGraphWithFrontend(Metasound::Frontend::FMetaSo
 	if (InRegistrationOptions.bAutoUpdate)
 	{
 		FString OwningAssetName = GetOwningAssetName();
-		const bool bAutoUpdated = FAutoUpdateRootGraph(MoveTemp(OwningAssetName)).Transform(GetDocumentHandle());
+		const bool bAutoUpdated = FAutoUpdateRootGraph(MoveTemp(OwningAssetName), InRegistrationOptions.bAutoUpdateLogWarningOnDroppedConnection).Transform(GetDocumentHandle());
 #if WITH_EDITOR
 		if (bAutoUpdated || InRegistrationOptions.bForceViewSynchronization)
 		{
diff --git a/Engine/Plugins/Runtime/Metasound/Source/MetasoundFrontend/Private/MetasoundFrontendInvalidController.h b/Engine/Plugins/Runtime/Metasound/Source/MetasoundFrontend/Private/MetasoundFrontendInvalidController.h
--- a/Engine/Plugins/Runtime/Metasound/Source/MetasoundFrontend/Private/MetasoundFrontendInvalidController.h
+++ b/Engine/Plugins/Runtime/Metasound/Source/MetasoundFrontend/Private/MetasoundFrontendInvalidController.h
@@ -217,7 +217,8 @@ namespace Metasound
 
 			virtual void SetNodeName(const FVertexName& InName) override { }
 
-			virtual FNodeHandle ReplaceWithVersion(const FMetasoundFrontendVersionNumber& InNewVersion) override { return INodeController::GetInvalidHandle(); }
+			virtual FNodeHandle ReplaceWithVersion(const FMetasoundFrontendVersionNumber& InNewVersion, TArray<FVertexNameAndType>* OutDisconnectedInputs, TArray<FVertexNameAndType>* OutDisconnectedOutputs) override { return INodeController::GetInvalidHandle(); }
+
 
 			virtual bool CanAddInput(const FVertexName& InVertexName) const override { return false; }
 			virtual FInputHandle AddInput(const FVertexName& InVertexName, const FMetasoundFrontendLiteral* InDefault) override { return IInputController::GetInvalidHandle(); }
diff --git a/Engine/Plugins/Runtime/Metasound/Source/MetasoundFrontend/Private/MetasoundFrontendNodeController.cpp b/Engine/Plugins/Runtime/Metasound/Source/MetasoundFrontend/Private/MetasoundFrontendNodeController.cpp
--- a/Engine/Plugins/Runtime/Metasound/Source/MetasoundFrontend/Private/MetasoundFrontendNodeController.cpp
+++ b/Engine/Plugins/Runtime/Metasound/Source/MetasoundFrontend/Private/MetasoundFrontendNodeController.cpp
@@ -2,6 +2,7 @@
 
 #include "MetasoundFrontendNodeController.h"
 
+#include "Algo/AnyOf.h"
 #include "Algo/Transform.h"
 #include "HAL/IConsoleManager.h"
 #include "MetasoundFrontendDocumentAccessPtr.h"
@@ -726,7 +727,7 @@ namespace Metasound
 
 		}
 
-		FNodeHandle FBaseNodeController::ReplaceWithVersion(const FMetasoundFrontendVersionNumber& InNewVersion)
+		FNodeHandle FBaseNodeController::ReplaceWithVersion(const FMetasoundFrontendVersionNumber& InNewVersion, TArray<FVertexNameAndType>* OutDisconnectedInputs, TArray<FVertexNameAndType>* OutDisconnectedOutputs)
 		{
 			const FMetasoundFrontendClassMetadata Metadata = GetClassMetadata();
 			const TArray<FMetasoundFrontendClass> Versions = ISearchEngine::Get().FindClassesWithName(Metadata.GetClassName().ToNodeClassName(), false /* bInSortByVersion */);
@@ -746,11 +747,10 @@ namespace Metasound
 			FMetasoundFrontendNodeStyle Style = GetNodeStyle();
 #endif // WITH_EDITOR
 
-			using FConnectionKey = TPair<FVertexName, FName>;
-
 			struct FInputConnectionInfo
 			{
 				FOutputHandle ConnectedOutput;
+				FVertexName Name;
 				FName DataType;
 				FMetasoundFrontendLiteral DefaultValue;
 				bool bLiteralSet = false;
@@ -758,7 +758,7 @@ namespace Metasound
 
 			// Cache input/output connections by name to try so they can be
 			// hooked back up after swapping to the new class version.
-			TMap<FConnectionKey, FInputConnectionInfo> InputConnections;
+			TMap<FVertexNameAndType, FInputConnectionInfo> InputConnections;
 			IterateInputs([Connections = &InputConnections](FInputHandle InputHandle)
 			{
 				bool bLiteralSet = false;
@@ -774,10 +774,11 @@ namespace Metasound
 					}
 				}
 
-				const FConnectionKey ConnectionKey(InputHandle->GetName(), InputHandle->GetDataType());
+				const FVertexNameAndType ConnectionKey(InputHandle->GetName(), InputHandle->GetDataType());
 				Connections->Add(ConnectionKey, FInputConnectionInfo
 				{
 					InputHandle->GetConnectedOutput(),
+					InputHandle->GetName(),
 					InputHandle->GetDataType(),
 					MoveTemp(DefaultLiteral),
 					bLiteralSet
@@ -787,17 +788,19 @@ namespace Metasound
 			struct FOutputConnectionInfo
 			{
 				TArray<FInputHandle> ConnectedInputs;
+				FVertexName VertexName;
 				FName DataType;
 			};
 
-			TMap<FConnectionKey, FOutputConnectionInfo> OutputConnections;
+			TMap<FVertexNameAndType, FOutputConnectionInfo> OutputConnections;
 			IterateOutputs([Connections = &OutputConnections](FOutputHandle OutputHandle)
 			{
-				const FConnectionKey ConnectionKey(OutputHandle->GetName(), OutputHandle->GetDataType());
+				const FVertexNameAndType ConnectionKey(OutputHandle->GetName(), OutputHandle->GetDataType());
 				Connections->Add(ConnectionKey, FOutputConnectionInfo
 				{
 					OutputHandle->GetConnectedInputs(),
-					OutputHandle->GetDataType(),
+					OutputHandle->GetName(),
+					OutputHandle->GetDataType()
 				});
 			});
 
@@ -832,10 +835,10 @@ namespace Metasound
 			ReplacementNode->SetNodeStyle(Style);
 #endif // WITH_EDITOR
 
-			ReplacementNode->IterateInputs([Connections = &InputConnections](FInputHandle InputHandle)
+			ReplacementNode->IterateInputs([&InputConnections](FInputHandle InputHandle)
 			{
-				const FConnectionKey ConnectionKey(InputHandle->GetName(), InputHandle->GetDataType());
-				if (FInputConnectionInfo* ConnectionInfo = Connections->Find(ConnectionKey))
+				const FVertexNameAndType ConnectionKey(InputHandle->GetName(), InputHandle->GetDataType());
+				if (FInputConnectionInfo* ConnectionInfo = InputConnections.Find(ConnectionKey))
 				{
 					if (ConnectionInfo->bLiteralSet)
 					{
@@ -846,13 +849,30 @@ namespace Metasound
 					{
 						ensure(InputHandle->Connect(*ConnectionInfo->ConnectedOutput));
 					}
+
+					// Remove connection to track missing connections between 
+					// node versions.
+					InputConnections.Remove(ConnectionKey);
 				}
 			});
 
-			ReplacementNode->IterateOutputs([Connections = &OutputConnections](FOutputHandle OutputHandle)
+			// Track missing input connections
+			if (nullptr != OutDisconnectedInputs)
+			{
+				for (const auto& ConnectionInfoKV : InputConnections)
+				{
+					const FInputConnectionInfo& ConnectionInfo = ConnectionInfoKV.Value;
+					if (ConnectionInfo.ConnectedOutput->IsValid())
+					{
+						OutDisconnectedInputs->Add(FVertexNameAndType{ConnectionInfo.Name, ConnectionInfo.DataType});
+					}
+				}
+			}
+
+			ReplacementNode->IterateOutputs([&OutputConnections](FOutputHandle OutputHandle)
 			{
-				const FConnectionKey ConnectionKey(OutputHandle->GetName(), OutputHandle->GetDataType());
-				if (FOutputConnectionInfo* ConnectionInfo = Connections->Find(ConnectionKey))
+				const FVertexNameAndType ConnectionKey(OutputHandle->GetName(), OutputHandle->GetDataType());
+				if (FOutputConnectionInfo* ConnectionInfo = OutputConnections.Find(ConnectionKey))
 				{
 					for (FInputHandle InputHandle : ConnectionInfo->ConnectedInputs)
 					{
@@ -861,9 +881,27 @@ namespace Metasound
 							ensure(InputHandle->Connect(*OutputHandle));
 						}
 					}
+
+					// Remove connection to track missing connections between 
+					// node versions.
+					OutputConnections.Remove(ConnectionKey);
 				}
 			});
 
+			// Track missing output connections
+			if (nullptr != OutDisconnectedOutputs)
+			{
+				for (const auto& ConnectionInfoKV : OutputConnections)
+				{
+					const FOutputConnectionInfo& ConnectionInfo = ConnectionInfoKV.Value;
+					const bool bAnyConnectedInputs = Algo::AnyOf(ConnectionInfo.ConnectedInputs, [](const FInputHandle& Input) { return Input->IsValid(); });
+					if (bAnyConnectedInputs)
+					{
+						OutDisconnectedOutputs->Add(FVertexNameAndType{ConnectionInfo.VertexName, ConnectionInfo.DataType});
+					}
+				}
+			}
+
 			return ReplacementNode;
 		}
 
diff --git a/Engine/Plugins/Runtime/Metasound/Source/MetasoundFrontend/Private/MetasoundFrontendNodeController.h b/Engine/Plugins/Runtime/Metasound/Source/MetasoundFrontend/Private/MetasoundFrontendNodeController.h
--- a/Engine/Plugins/Runtime/Metasound/Source/MetasoundFrontend/Private/MetasoundFrontendNodeController.h
+++ b/Engine/Plugins/Runtime/Metasound/Source/MetasoundFrontend/Private/MetasoundFrontendNodeController.h
@@ -59,7 +59,7 @@ namespace Metasound
 			virtual bool DiffAgainstRegistryInterface(FClassInterfaceUpdates& OutInterfaceUpdates, bool bInUseHighestMinorVersion) const override;
 
 			virtual bool CanAutoUpdate(FClassInterfaceUpdates& OutInterfaceUpdates) const override;
-			virtual FNodeHandle ReplaceWithVersion(const FMetasoundFrontendVersionNumber& InNewVersion) override;
+			virtual FNodeHandle ReplaceWithVersion(const FMetasoundFrontendVersionNumber& InNewVersion, TArray<FVertexNameAndType>* OutDisconnectedInputs, TArray<FVertexNameAndType>* OutDisconnectedOutputs) override;
 			virtual FMetasoundFrontendVersionNumber FindHighestVersionInRegistry() const override;
 			virtual FMetasoundFrontendVersionNumber FindHighestMinorVersionInRegistry() const override;
 
diff --git a/Engine/Plugins/Runtime/Metasound/Source/MetasoundFrontend/Private/MetasoundFrontendTransform.cpp b/Engine/Plugins/Runtime/Metasound/Source/MetasoundFrontend/Private/MetasoundFrontendTransform.cpp
--- a/Engine/Plugins/Runtime/Metasound/Source/MetasoundFrontend/Private/MetasoundFrontendTransform.cpp
+++ b/Engine/Plugins/Runtime/Metasound/Source/MetasoundFrontend/Private/MetasoundFrontendTransform.cpp
@@ -599,12 +599,37 @@ namespace Metasound
 			}
 			else
 			{
+				using FVertexNameAndType = INodeController::FVertexNameAndType;
+
 				bDidEdit |= !NodesToUpdate.IsEmpty();
 				for (const TPair<FNodeHandle, FMetasoundFrontendVersionNumber>& Pair : NodesToUpdate)
 				{
 					FNodeHandle ExistingNode = Pair.Key;
 					FMetasoundFrontendVersionNumber InitialVersion = ExistingNode->GetClassMetadata().GetVersion();
-					FNodeHandle NewNode = ExistingNode->ReplaceWithVersion(Pair.Value);
+
+					TArray<FVertexNameAndType> DisconnectedInputs;
+					TArray<FVertexNameAndType> DisconnectedOutputs;
+					FNodeHandle NewNode = ExistingNode->ReplaceWithVersion(Pair.Value, &DisconnectedInputs, &DisconnectedOutputs);
+
+					// Log warnings for any disconnections
+					if (bLogWarningOnDroppedConnection)
+					{
+						if ((DisconnectedInputs.Num() > 0) || (DisconnectedOutputs.Num() > 0))
+						{
+							const FString NodeClassName = NewNode->GetClassMetadata().GetClassName().ToString();
+							const FString NewClassVersion = Pair.Value.ToString();
+
+							for (const FVertexNameAndType& InputPin : DisconnectedInputs)
+							{
+								UE_LOG(LogMetaSound, Warning, TEXT("Auto-Updating '%s' node class '%s (%s)': Previously connected input '%s' with data type '%s' no longer exists."), *DebugAssetPath, *NodeClassName, *NewClassVersion, *InputPin.Get<0>().ToString(), *InputPin.Get<1>().ToString());
+							}
+
+							for (const FVertexNameAndType& OutputPin : DisconnectedOutputs)
+							{
+								UE_LOG(LogMetaSound, Warning, TEXT("Auto-Updating '%s' node class '%s (%s)': Previously connected output '%s' with data type '%s' no longer exists."), *DebugAssetPath, *NodeClassName, *NewClassVersion, *OutputPin.Get<0>().ToString(), *OutputPin.Get<1>().ToString());
+							}
+						}
+					}
 				}
 
 				InDocument->RemoveUnreferencedDependencies();
diff --git a/Engine/Plugins/Runtime/Metasound/Source/MetasoundFrontend/Public/MetasoundAssetManager.h b/Engine/Plugins/Runtime/Metasound/Source/MetasoundFrontend/Public/MetasoundAssetManager.h
--- a/Engine/Plugins/Runtime/Metasound/Source/MetasoundFrontend/Public/MetasoundAssetManager.h
+++ b/Engine/Plugins/Runtime/Metasound/Source/MetasoundFrontend/Public/MetasoundAssetManager.h
@@ -48,6 +48,9 @@ namespace Metasound
 			// Will not respect being set to true if project-level MetaSoundSettings specify to not run auto-update.)
 			bool bAutoUpdate = true;
 
+			// If true, warnings will be logged if updating a node results in existing connections being discarded.
+			bool bAutoUpdateLogWarningOnDroppedConnection = false;
+
 			// Attempt to rebuild referenced class keys (only run if class not registered or set to force re-register)
 			bool bRebuildReferencedAssetClassKeys = true;
 		};
diff --git a/Engine/Plugins/Runtime/Metasound/Source/MetasoundFrontend/Public/MetasoundFrontendController.h b/Engine/Plugins/Runtime/Metasound/Source/MetasoundFrontend/Public/MetasoundFrontendController.h
--- a/Engine/Plugins/Runtime/Metasound/Source/MetasoundFrontend/Public/MetasoundFrontendController.h
+++ b/Engine/Plugins/Runtime/Metasound/Source/MetasoundFrontend/Public/MetasoundFrontendController.h
@@ -10,6 +10,7 @@
 #include "MetasoundFrontendRegistries.h"
 #include "MetasoundGraph.h"
 #include "MetasoundVertex.h"
+#include "Templates/Tuple.h"
 
 
 /* Metasound Controllers and Handles provide a object oriented interface for  manipulating Metasound Documents. 
@@ -409,6 +410,8 @@ namespace Metasound
 		public:
 			static FNodeHandle GetInvalidHandle();
 
+			using FVertexNameAndType = TTuple<FVertexName, FName>;
+
 			INodeController() = default;
 			virtual ~INodeController() = default;
 
@@ -480,7 +483,7 @@ namespace Metasound
 			  * rebuild edges where possible with matching vertex names that share the same DataType.
 			  * Returns a node handle to the new node.  If operation fails, returns a handle to this node.
 			  */
-			virtual FNodeHandle ReplaceWithVersion(const FMetasoundFrontendVersionNumber& InNewVersion) = 0;
+			virtual FNodeHandle ReplaceWithVersion(const FMetasoundFrontendVersionNumber& InNewVersion, TArray<FVertexNameAndType>* OutDisconnectedInputs, TArray<FVertexNameAndType>* OutDisconnectedOutputs) = 0;
 
 			/** Returns an input with the given id.
 			 *
diff --git a/Engine/Plugins/Runtime/Metasound/Source/MetasoundFrontend/Public/MetasoundFrontendTransform.h b/Engine/Plugins/Runtime/Metasound/Source/MetasoundFrontend/Public/MetasoundFrontendTransform.h
--- a/Engine/Plugins/Runtime/Metasound/Source/MetasoundFrontend/Public/MetasoundFrontendTransform.h
+++ b/Engine/Plugins/Runtime/Metasound/Source/MetasoundFrontend/Public/MetasoundFrontendTransform.h
@@ -143,14 +143,22 @@ namespace Metasound
 		class METASOUNDFRONTEND_API FAutoUpdateRootGraph : public IDocumentTransform
 		{
 		public:
-			FAutoUpdateRootGraph(FString&& InDebugAssetPath)
+			/** Construct an AutoUpdate transform
+			 *
+			 * @param InDebugAssetPath - Asset path used for debug logs on warnings and errors.
+			 * @param bInLogWarningOnDroppedConnections - If true, warnings will be logged if a node update results in a dropped connection.
+			 */
+			FAutoUpdateRootGraph(FString&& InDebugAssetPath, bool bInLogWarningOnDroppedConnection)
 				: DebugAssetPath(MoveTemp(InDebugAssetPath))
+				, bLogWarningOnDroppedConnection(bInLogWarningOnDroppedConnection)
 			{
 			}
 
 			bool Transform(FDocumentHandle InDocument) const override;
 
+		private:
 			const FString DebugAssetPath;
+			bool bLogWarningOnDroppedConnection;
 		};
 
 		/** Regenerates the class' name, effectively causing the class to get registered as
