Fix for ObjectArrayParams not being parsed properly by MetaSound plugin

#rnx

diff --git a/Engine/Plugins/Runtime/Metasound/Source/MetasoundEngine/Private/MetasoundSource.cpp b/Engine/Plugins/Runtime/Metasound/Source/MetasoundEngine/Private/MetasoundSource.cpp
--- a/Engine/Plugins/Runtime/Metasound/Source/MetasoundEngine/Private/MetasoundSource.cpp
+++ b/Engine/Plugins/Runtime/Metasound/Source/MetasoundEngine/Private/MetasoundSource.cpp
@@ -705,12 +705,13 @@ bool UMetaSoundSource::IsParameterValid(const FAudioParameter& InParameter, cons
 		{
 			bIsValid = true;
 
+			const FName ElementTypeName = CreateElementTypeNameFromArrayTypeName(*TypeName);
 			for (UObject* Object : InParameter.ArrayObjectParam)
 			{
 				FDataTypeRegistryInfo DataTypeInfo;
-				bIsValid = IDataTypeRegistry::Get().GetDataTypeInfo(InParameter.ObjectParam, DataTypeInfo);
+				bIsValid = IDataTypeRegistry::Get().GetDataTypeInfo(Object, DataTypeInfo);
 				bIsValid &= DataTypeInfo.bIsProxyParsable;
-				bIsValid &= DataTypeInfo.DataTypeName == *TypeName;
+				bIsValid &= DataTypeInfo.DataTypeName == ElementTypeName;
 
 				if (!bIsValid)
 				{
