Clamping metasound node input pins

diff --git a/Engine/Plugins/Runtime/Metasound/Source/MetasoundStandardNodes/Private/MetasoundBasicFilters.cpp b/Engine/Plugins/Runtime/Metasound/Source/MetasoundStandardNodes/Private/MetasoundBasicFilters.cpp
--- a/Engine/Plugins/Runtime/Metasound/Source/MetasoundStandardNodes/Private/MetasoundBasicFilters.cpp
+++ b/Engine/Plugins/Runtime/Metasound/Source/MetasoundStandardNodes/Private/MetasoundBasicFilters.cpp
@@ -877,7 +877,7 @@ namespace Metasound
 	{
 		const float CurrentFrequency = FMath::Clamp(*Frequency, 0.f, MaxCutoffFrequency);
 		const float CurrentBandwidth = FMath::Max(*Bandwidth, 0.f);
-		const float CurrentFilterGainDb = *FilterGainDb;
+		const float CurrentFilterGainDb = FMath::Clamp(*FilterGainDb, -90.0f, 20.0f);
 
 		if (!FMath::IsNearlyEqual(PreviousFrequency, CurrentFrequency))
 		{
diff --git a/Engine/Plugins/Runtime/Metasound/Source/MetasoundStandardNodes/Private/MetasoundDiffuserNode.cpp b/Engine/Plugins/Runtime/Metasound/Source/MetasoundStandardNodes/Private/MetasoundDiffuserNode.cpp
--- a/Engine/Plugins/Runtime/Metasound/Source/MetasoundStandardNodes/Private/MetasoundDiffuserNode.cpp
+++ b/Engine/Plugins/Runtime/Metasound/Source/MetasoundStandardNodes/Private/MetasoundDiffuserNode.cpp
@@ -102,7 +102,7 @@ namespace Metasound
 
 			ClampedDepth = FMath::Clamp(*InDepth, 1, MaxNumAPFs);
 
-			const float FeedbackGain = *InFeedback;
+			const float FeedbackGain = FMath::Clamp(*InFeedback, 0.0f, 1.0f - KINDA_SMALL_NUMBER);
 			const float SampleRatio =  InSettings.GetSampleRate() / 1000.f;
 
 			Delays.Reset(ClampedDepth);
