Fix for submix sends not working in metasound editor preview

diff --git a/Engine/Plugins/Runtime/Metasound/Source/MetasoundEngine/Public/MetasoundSource.h b/Engine/Plugins/Runtime/Metasound/Source/MetasoundEngine/Public/MetasoundSource.h
--- a/Engine/Plugins/Runtime/Metasound/Source/MetasoundEngine/Public/MetasoundSource.h
+++ b/Engine/Plugins/Runtime/Metasound/Source/MetasoundEngine/Public/MetasoundSource.h
@@ -174,7 +174,7 @@ public:
 	virtual bool IsParameterValid(const FAudioParameter& InParameter) const override;
 	virtual bool IsLooping() const override;
 	virtual bool IsOneShot() const override;
-
+	virtual bool IsMetaSound() const override { return true; }
 protected:
 	/** Gets all the default parameters for this Asset.  */
 	virtual bool GetAllDefaultParameters(TArray<FAudioParameter>& OutParameters) const override;
diff --git a/Engine/Source/Runtime/Engine/Classes/Sound/SoundBase.h b/Engine/Source/Runtime/Engine/Classes/Sound/SoundBase.h
--- a/Engine/Source/Runtime/Engine/Classes/Sound/SoundBase.h
+++ b/Engine/Source/Runtime/Engine/Classes/Sound/SoundBase.h
@@ -320,6 +320,9 @@ public:
 	/** Gets all the default parameters for this Asset.  */
 	virtual bool GetAllDefaultParameters(TArray<FAudioParameter>& OutParameters) const { return false; }
 
+	/** Whether or not this sound is a MetaSound. */
+	virtual bool IsMetaSound() const { return false; }
+
 #if WITH_EDITORONLY_DATA
 	void SetTimecodeOffset(const FSoundTimecodeOffset& InTimecodeOffset);
 	TOptional<FSoundTimecodeOffset> GetTimecodeOffset() const;
diff --git a/Engine/Source/Runtime/Engine/Private/SoundWave.cpp b/Engine/Source/Runtime/Engine/Private/SoundWave.cpp
--- a/Engine/Source/Runtime/Engine/Private/SoundWave.cpp
+++ b/Engine/Source/Runtime/Engine/Private/SoundWave.cpp
@@ -2779,6 +2779,11 @@ void USoundWave::Parse(FAudioDevice* AudioDevice, const UPTRINT NodeWaveInstance
 	}
 	else //if this is a preview sound, ignore Submix sends
 	{
+		// Only allow submix sends if this is a metasound preview
+		if (IsMetaSound())
+		{
+			WaveInstance->bEnableSubmixSends = ParseParams.bEnableSubmixSends;
+		}
 		WaveInstance->bEnableBaseSubmix = true;
 		WaveInstance->bEnableBusSends = ParseParams.bEnableBusSends;
 	}
