Fix missing drag/drop on MetaSound node object picker
#rnx

diff --git a/Engine/Plugins/Runtime/AudioModulation/Source/AudioModulation/Public/AudioModulationSettings.h b/Engine/Plugins/Runtime/AudioModulation/Source/AudioModulation/Public/AudioModulationSettings.h
--- a/Engine/Plugins/Runtime/AudioModulation/Source/AudioModulation/Public/AudioModulationSettings.h
+++ b/Engine/Plugins/Runtime/AudioModulation/Source/AudioModulation/Public/AudioModulationSettings.h
@@ -16,7 +16,7 @@ class AUDIOMODULATION_API UAudioModulationSettings : public UDeveloperSettings
 	GENERATED_BODY()
 
 public:
-	// Array of loaded Modulation Parameters
+	// Array of Modulation Parameters that are loaded on plugin startup
 	UPROPERTY(config, EditAnywhere, Category = "Parameters", meta = (AllowedClasses = "SoundModulationParameter"))
 	TArray<FSoftObjectPath> Parameters;
 
diff --git a/Engine/Plugins/Runtime/Metasound/Source/MetasoundEditor/Private/MetasoundEditorGraphNode.cpp b/Engine/Plugins/Runtime/Metasound/Source/MetasoundEditor/Private/MetasoundEditorGraphNode.cpp
--- a/Engine/Plugins/Runtime/Metasound/Source/MetasoundEditor/Private/MetasoundEditorGraphNode.cpp
+++ b/Engine/Plugins/Runtime/Metasound/Source/MetasoundEditor/Private/MetasoundEditorGraphNode.cpp
@@ -361,9 +361,16 @@ Metasound::Frontend::FDataTypeRegistryInfo UMetasoundEditorGraphNode::GetPinData
 
  	Metasound::Frontend::FDataTypeRegistryInfo DataTypeInfo;
 
-	FConstInputHandle Handle = FGraphBuilder::GetConstInputHandleFromPin(&InPin);
-	ensure(IDataTypeRegistry::Get().GetDataTypeInfo(Handle->GetDataType(), DataTypeInfo));
-
+	if (InPin.Direction == EGPD_Input)
+	{
+		FConstInputHandle Handle = FGraphBuilder::GetConstInputHandleFromPin(&InPin);
+		ensure(IDataTypeRegistry::Get().GetDataTypeInfo(Handle->GetDataType(), DataTypeInfo));
+	}
+	else // InPin.Direction == EGPD_Output
+	{
+		FConstOutputHandle Handle = FGraphBuilder::GetConstOutputHandleFromPin(&InPin);
+		ensure(IDataTypeRegistry::Get().GetDataTypeInfo(Handle->GetDataType(), DataTypeInfo));
+	}
 
 	return DataTypeInfo;
 }
diff --git a/Engine/Plugins/Runtime/Metasound/Source/MetasoundEditor/Private/MetasoundEditorGraphSchema.cpp b/Engine/Plugins/Runtime/Metasound/Source/MetasoundEditor/Private/MetasoundEditorGraphSchema.cpp
--- a/Engine/Plugins/Runtime/Metasound/Source/MetasoundEditor/Private/MetasoundEditorGraphSchema.cpp
+++ b/Engine/Plugins/Runtime/Metasound/Source/MetasoundEditor/Private/MetasoundEditorGraphSchema.cpp
@@ -89,7 +89,6 @@ namespace Metasound
 
 		namespace SchemaPrivate
 		{
-
 			static const FText CategoryDelim = LOCTEXT("MetaSoundActionsCategoryDelim", "|");
 			static const FText KeywordDelim = LOCTEXT("MetaSoundKeywordDelim", " ");
 
@@ -108,6 +107,35 @@ namespace Metasound
 			static const FText VariableMutatorDisplayNameFormat = LOCTEXT("DisplayNameAddVariableMutatorFormat", "Set {0}");
 			static const FText VariableMutatorTooltipFormat = LOCTEXT("TooltipAddVariableMutatorFormat", "Adds a setter for the variable '{0}' to the graph.");
 
+			bool DataTypeSupportsAssetTypes(const Metasound::Frontend::FDataTypeRegistryInfo& InRegistryInfo, const TArray<FAssetData>& InAssets)
+			{
+				if (InRegistryInfo.PreferredLiteralType != Metasound::ELiteralType::UObjectProxy)
+				{
+					return false;
+				}
+
+				const IMetasoundEditorModule& EditorModule = FModuleManager::GetModuleChecked<IMetasoundEditorModule>("MetaSoundEditor");
+				return Algo::AnyOf(InAssets, [&EditorModule, &InRegistryInfo](const FAssetData& Asset)
+				{
+					if (InRegistryInfo.ProxyGeneratorClass)
+					{
+						if (const UClass* Class = Asset.GetClass())
+						{
+							if (EditorModule.IsExplicitProxyClass(*InRegistryInfo.ProxyGeneratorClass))
+							{
+								return Class == InRegistryInfo.ProxyGeneratorClass;
+							}
+							else
+							{
+								return Class->IsChildOf(InRegistryInfo.ProxyGeneratorClass);
+							}
+						}
+					}
+
+					return false;
+				});
+			}
+
 			// Connects to first pin with the same DataType
 			bool TryConnectNewNodeToMatchingDataTypePin(UEdGraphNode& NewGraphNode, UEdGraphPin* FromPin)
 			{
@@ -1489,6 +1517,48 @@ void UMetasoundEditorGraphSchema::GetAssetsGraphHoverMessage(const TArray<FAsset
 	}
 }
 
+void UMetasoundEditorGraphSchema::GetAssetsPinHoverMessage(const TArray<FAssetData>& Assets, const UEdGraphPin* HoverPin, FString& OutTooltipText, bool& OutOkIcon) const
+{
+	using namespace Metasound::Editor;
+	using namespace Metasound::Frontend;
+
+	if (HoverPin && HoverPin->Direction == EGPD_Input)
+	{
+		if (const UEdGraphNode* Node = HoverPin->GetOwningNode())
+		{
+			if (const UMetasoundEditorGraphNode* MetaSoundNode = Cast<UMetasoundEditorGraphNode>(HoverPin->GetOwningNode()))
+			{
+				if (Assets.Num() == 1)
+				{
+					FDataTypeRegistryInfo RegistryInfo = MetaSoundNode->GetPinDataTypeInfo(*HoverPin);
+					const bool bAssetTypesMatch = SchemaPrivate::DataTypeSupportsAssetTypes(RegistryInfo, Assets);
+					if (bAssetTypesMatch)
+					{
+						OutTooltipText = FString::Format(TEXT("Set to '{0}'"), { *Assets[0].AssetName.ToString() });
+						OutOkIcon = true;
+						return;
+					}
+
+					OutTooltipText = FString::Format(TEXT("'{0}': Invalid Type"), { *Assets[0].AssetName.ToString() });
+					OutOkIcon = false;
+					return;
+				}
+
+				OutTooltipText = TEXT("Cannot drop multiple assets on single pin.");
+				OutOkIcon = false;
+				return;
+			}
+
+			OutTooltipText = FString::Format(TEXT("Node '{0}' does not support drag/drop"), { *Node->GetName() });
+			OutOkIcon = false;
+			return;
+		}
+	}
+
+	OutTooltipText = FString();
+	OutOkIcon = false;
+}
+
 void UMetasoundEditorGraphSchema::DroppedAssetsOnGraph(const TArray<FAssetData>& Assets, const FVector2D& GraphPosition, UEdGraph* Graph) const
 {
 	using namespace Metasound;
@@ -1565,6 +1635,38 @@ void UMetasoundEditorGraphSchema::DroppedAssetsOnNode(const TArray<FAssetData>&
 	// Still needed?
 }
 
+void UMetasoundEditorGraphSchema::DroppedAssetsOnPin(const TArray<FAssetData>& Assets, const FVector2D& GraphPosition, UEdGraphPin* Pin) const
+{
+	using namespace Metasound::Editor;
+	using namespace Metasound::Frontend;
+
+	if (!Pin)
+	{
+		return;
+	}
+
+	if (UMetasoundEditorGraphNode* Node = Cast<UMetasoundEditorGraphNode>(Pin->GetOwningNode()))
+	{
+		if (Assets.Num() == 1)
+		{
+			FDataTypeRegistryInfo RegistryInfo = Node->GetPinDataTypeInfo(*Pin);
+			const bool bAssetTypesMatch = SchemaPrivate::DataTypeSupportsAssetTypes(RegistryInfo, Assets);
+			if (bAssetTypesMatch)
+			{
+				UObject* Object = Assets.Last().GetAsset();
+				const FText TransactionText = FText::Format(LOCTEXT("ChangeDefaultObjectTransaction", "Set {0} to '{1}'"),
+					Pin->GetDisplayName(),
+					FText::FromName(Object->GetFName()));
+				const FScopedTransaction Transaction(TransactionText);
+				Node->Modify();
+
+				constexpr bool bMarkAsModified = true;
+				TrySetDefaultObject(*Pin, Object, bMarkAsModified);
+			}
+		}
+	}
+}
+
 void UMetasoundEditorGraphSchema::GetConversionActions(FGraphActionMenuBuilder& ActionMenuBuilder, Metasound::Editor::FActionClassFilters InFilters, bool bShowSelectedActions) const
 {
 	using namespace Metasound;
diff --git a/Engine/Plugins/Runtime/Metasound/Source/MetasoundEditor/Public/MetasoundEditorGraphSchema.h b/Engine/Plugins/Runtime/Metasound/Source/MetasoundEditor/Public/MetasoundEditorGraphSchema.h
--- a/Engine/Plugins/Runtime/Metasound/Source/MetasoundEditor/Public/MetasoundEditorGraphSchema.h
+++ b/Engine/Plugins/Runtime/Metasound/Source/MetasoundEditor/Public/MetasoundEditorGraphSchema.h
@@ -100,8 +100,8 @@ struct METASOUNDEDITOR_API FMetasoundGraphSchemaAction_NodeWithMultipleOutputs :
 	{}
 
 	//~ Begin FEdGraphSchemaAction Interface
-	virtual UEdGraphNode* PerformAction(class UEdGraph* ParentGraph, TArray<UEdGraphPin*>& FromPins, const FVector2D Location, bool bSelectNewNode = true) override;
-	virtual UEdGraphNode* PerformAction(class UEdGraph* ParentGraph, UEdGraphPin* FromPin, const FVector2D Location, bool bSelectNewNode = true) { return nullptr; }
+	virtual UEdGraphNode* PerformAction(UEdGraph* ParentGraph, TArray<UEdGraphPin*>& FromPins, const FVector2D Location, bool bSelectNewNode = true) override;
+	virtual UEdGraphNode* PerformAction(UEdGraph* ParentGraph, UEdGraphPin* FromPin, const FVector2D Location, bool bSelectNewNode = true) { return nullptr; }
 	//~ End FEdGraphSchemaAction Interface
 };
 
@@ -348,7 +348,7 @@ struct METASOUNDEDITOR_API FMetasoundGraphSchemaAction_NewFromSelected : public
 	{}
 
 	//~ Begin FEdGraphSchemaAction Interface
-	virtual UEdGraphNode* PerformAction(class UEdGraph* ParentGraph, UEdGraphPin* FromPin, const FVector2D Location, bool bSelectNewNode = true) override;
+	virtual UEdGraphNode* PerformAction(UEdGraph* ParentGraph, UEdGraphPin* FromPin, const FVector2D Location, bool bSelectNewNode = true) override;
 	//~ End FEdGraphSchemaAction Interface
 };
 
@@ -373,7 +373,7 @@ struct METASOUNDEDITOR_API FMetasoundGraphSchemaAction_NewComment : public FMeta
 	//~ End FMetasoundGraphSchemaAction Interface
 
 	//~ Begin FEdGraphSchemaAction Interface
-	virtual UEdGraphNode* PerformAction(class UEdGraph* ParentGraph, UEdGraphPin* FromPin, const FVector2D Location, bool bSelectNewNode = true) override;
+	virtual UEdGraphNode* PerformAction(UEdGraph* ParentGraph, UEdGraphPin* FromPin, const FVector2D Location, bool bSelectNewNode = true) override;
 	//~ End FEdGraphSchemaAction Interface
 };
 
@@ -392,7 +392,7 @@ struct METASOUNDEDITOR_API FMetasoundGraphSchemaAction_Paste : public FMetasound
 	{}
 
 	//~ Begin FEdGraphSchemaAction Interface
-	virtual UEdGraphNode* PerformAction(class UEdGraph* ParentGraph, UEdGraphPin* FromPin, const FVector2D Location, bool bSelectNewNode = true) override;
+	virtual UEdGraphNode* PerformAction(UEdGraph* ParentGraph, UEdGraphPin* FromPin, const FVector2D Location, bool bSelectNewNode = true) override;
 	//~ End FEdGraphSchemaAction Interface
 };
 
@@ -409,7 +409,7 @@ class UMetasoundEditorGraphSchema : public UEdGraphSchema
 
 	//~ Begin EdGraphSchema Interface
 	virtual void GetGraphContextActions(FGraphContextMenuBuilder& ContextMenuBuilder) const override;
-	virtual void GetContextMenuActions(class UToolMenu* Menu, class UGraphNodeContextMenuContext* Context) const override;
+	virtual void GetContextMenuActions(UToolMenu* Menu, UGraphNodeContextMenuContext* Context) const override;
 	virtual FText GetPinDisplayName(const UEdGraphPin* Pin) const override;
 	virtual void CreateDefaultNodesForGraph(UEdGraph& Graph) const override;
 	virtual const FPinConnectionResponse CanCreateConnection(const UEdGraphPin* A, const UEdGraphPin* B) const override;
@@ -422,8 +422,10 @@ class UMetasoundEditorGraphSchema : public UEdGraphSchema
 	virtual void BreakPinLinks(UEdGraphPin& TargetPin, bool bSendsNodeNotifcation) const override;
 	virtual void BreakSinglePinLink(UEdGraphPin* SourcePin, UEdGraphPin* TargetPin) const override;
 	virtual void GetAssetsGraphHoverMessage(const TArray<FAssetData>& Assets, const UEdGraph* HoverGraph, FString& OutTooltipText, bool& OutOkIcon) const override;
-	virtual void DroppedAssetsOnGraph(const TArray<struct FAssetData>& Assets, const FVector2D& GraphPosition, UEdGraph* Graph) const override;
+	virtual void GetAssetsPinHoverMessage(const TArray<FAssetData>& Assets, const UEdGraphPin* HoverPin, FString& OutTooltipText, bool& OutOkIcon) const override;
+	virtual void DroppedAssetsOnGraph(const TArray<FAssetData>& Assets, const FVector2D& GraphPosition, UEdGraph* Graph) const override;
 	virtual void DroppedAssetsOnNode(const TArray<FAssetData>& Assets, const FVector2D& GraphPosition, UEdGraphNode* Node) const override;
+	virtual void DroppedAssetsOnPin(const TArray<FAssetData>& Assets, const FVector2D& GraphPosition, UEdGraphPin* Pin) const override;
 	virtual int32 GetNodeSelectionCount(const UEdGraph* Graph) const override;
 	virtual TSharedPtr<FEdGraphSchemaAction> GetCreateCommentAction() const override;
 	//~ End EdGraphSchema Interface
