Fix missing MetaSound interface warning and interfaces not updating referencing graphs properly

#rnx

diff --git a/Engine/Plugins/Runtime/Metasound/Source/MetasoundEditor/Private/MetasoundDetailCustomization.cpp b/Engine/Plugins/Runtime/Metasound/Source/MetasoundEditor/Private/MetasoundDetailCustomization.cpp
--- a/Engine/Plugins/Runtime/Metasound/Source/MetasoundEditor/Private/MetasoundDetailCustomization.cpp
+++ b/Engine/Plugins/Runtime/Metasound/Source/MetasoundEditor/Private/MetasoundDetailCustomization.cpp
@@ -374,6 +374,7 @@ namespace Metasound
 						UpdateInterfaceNames();
 						InterfaceComboBox->RefreshOptions();
 						MetaSoundAsset->SetSynchronizationRequired();
+						FGraphBuilder::RegisterGraphWithFrontend(*MetaSound.Get());
 					}
 				})
 				.Content()
@@ -429,6 +430,8 @@ namespace Metasound
 						InterfaceComboBox->RefreshOptions();
 						MetaSoundAsset->SetUpdateDetailsOnSynchronization();
 						MetaSoundAsset->SetSynchronizationRequired();
+						FGraphBuilder::RegisterGraphWithFrontend(*MetaSound.Get());
+
 					}), LOCTEXT("RemoveInterfaceTooltip1", "Removes all interfaces from the given MetaSound."))
 				];
 
diff --git a/Engine/Plugins/Runtime/Metasound/Source/MetasoundEditor/Private/SMetasoundGraphNode.cpp b/Engine/Plugins/Runtime/Metasound/Source/MetasoundEditor/Private/SMetasoundGraphNode.cpp
--- a/Engine/Plugins/Runtime/Metasound/Source/MetasoundEditor/Private/SMetasoundGraphNode.cpp
+++ b/Engine/Plugins/Runtime/Metasound/Source/MetasoundEditor/Private/SMetasoundGraphNode.cpp
@@ -842,7 +842,7 @@ namespace Metasound
 			ContentBox->AddSlot()
 			.HAlign(HAlign_Left)
 			.VAlign(VAlign_Top)
-			.AutoWidth()
+			.FillWidth(1.0f)
 			.Padding(0.0f, 0.0f, LeftNodeGrabPadding, 0.0f)
 			[
 				SAssignNew(LeftNodeBox, SVerticalBox)
diff --git a/Engine/Plugins/Runtime/Metasound/Source/MetasoundEngine/Private/MetasoundInterface.cpp b/Engine/Plugins/Runtime/Metasound/Source/MetasoundEngine/Private/MetasoundInterface.cpp
--- a/Engine/Plugins/Runtime/Metasound/Source/MetasoundEngine/Private/MetasoundInterface.cpp
+++ b/Engine/Plugins/Runtime/Metasound/Source/MetasoundEngine/Private/MetasoundInterface.cpp
@@ -134,7 +134,6 @@ namespace Metasound
 
 			// Transfer all input data from AudioExtension interface struct to FrontendInterface
 			{
-				FMetasoundFrontendInterfaceStyle InputStyle;
 				Algo::Transform(Interface->GetInputs(), FrontendInterface.Inputs, [&](const Audio::FParameterInterface::FInput& Input)
 				{
 					FMetasoundFrontendClassInput ClassInput;
@@ -150,7 +149,7 @@ namespace Metasound
 					ClassInput.Metadata.SetDescription(Input.Description);
 					ClassInput.Metadata.SortOrderIndex = Input.SortOrderIndex;
 					
-					InputStyle.DefaultSortOrder.Add(Input.SortOrderIndex);
+					FrontendInterface.AddSortOrderToInputStyle(Input.SortOrderIndex);
 
 					// Setup required inputs by telling the style that the input is required
 					// This will later be validated against.
@@ -164,12 +163,10 @@ namespace Metasound
 
 					return ClassInput;
 				});
-				FrontendInterface.SetInputStyle(InputStyle);
 			}
 
 			// Transfer all output data from AudioExtension interface struct to FrontendInterface
 			{
-				FMetasoundFrontendInterfaceStyle OutputStyle;
 				Algo::Transform(Interface->GetOutputs(), FrontendInterface.Outputs, [&](const Audio::FParameterInterface::FOutput& Output)
 				{
 					FMetasoundFrontendClassOutput ClassOutput;
@@ -184,7 +181,7 @@ namespace Metasound
 					ClassOutput.Metadata.SetDescription(Output.Description);
 					ClassOutput.Metadata.SortOrderIndex = Output.SortOrderIndex;
 					
-					OutputStyle.DefaultSortOrder.Add(Output.SortOrderIndex);
+					FrontendInterface.AddSortOrderToOutputStyle(Output.SortOrderIndex);
 
 					// Setup required outputs by telling the style that the output is required
 					// This will later be validated against.
@@ -198,7 +195,6 @@ namespace Metasound
 
 					return ClassOutput;
 				});
-				FrontendInterface.SetOutputStyle(OutputStyle);
 			}
 
 			Algo::Transform(Interface->GetEnvironment(), FrontendInterface.Environment, [&](const Audio::FParameterInterface::FEnvironmentVariable& Environment)
diff --git a/Engine/Plugins/Runtime/Metasound/Source/MetasoundFrontend/Public/MetasoundFrontendDocument.h b/Engine/Plugins/Runtime/Metasound/Source/MetasoundFrontend/Public/MetasoundFrontendDocument.h
--- a/Engine/Plugins/Runtime/Metasound/Source/MetasoundFrontend/Public/MetasoundFrontendDocument.h
+++ b/Engine/Plugins/Runtime/Metasound/Source/MetasoundFrontend/Public/MetasoundFrontendDocument.h
@@ -876,11 +876,13 @@ public:
 	void AddRequiredInputToStyle(const FName& InInputName, const FText& InRequiredText)
 	{
 		InputStyle.RequiredMembers.Add(InInputName, InRequiredText);
+		ChangeID = FGuid::NewGuid();
 	}
 
 	void AddRequiredOutputToStyle(const FName& InOutputName, const FText& InRequiredText)
 	{
 		OutputStyle.RequiredMembers.Add(InOutputName, InRequiredText);
+		ChangeID = FGuid::NewGuid();
 	}
 
 	bool IsMemberInputRequired(const FName& InInputName, FText& OutRequiredText)
@@ -902,6 +904,18 @@ public:
 		}
 		return false;
 	}
+
+	void AddSortOrderToInputStyle(const int32 InSortOrder)
+	{
+		InputStyle.DefaultSortOrder.Add(InSortOrder);
+		ChangeID = FGuid::NewGuid();
+	}
+
+	void AddSortOrderToOutputStyle(const int32 InSortOrder)
+	{
+		OutputStyle.DefaultSortOrder.Add(InSortOrder);
+		ChangeID = FGuid::NewGuid();
+	}
 #endif // #if WITH_EDITORONLY_DATA
 
 
