Use loc text for MetaSound node vertex names (Part 1)
- New macro for defining node parameters to generate localized metadata struct
- Simple find and replace to use new macro

diff --git a/Engine/Plugins/Runtime/Metasound/Source/MetasoundEngine/Private/MetasoundWaveInfoNode.cpp b/Engine/Plugins/Runtime/Metasound/Source/MetasoundEngine/Private/MetasoundWaveInfoNode.cpp
--- a/Engine/Plugins/Runtime/Metasound/Source/MetasoundEngine/Private/MetasoundWaveInfoNode.cpp
+++ b/Engine/Plugins/Runtime/Metasound/Source/MetasoundEngine/Private/MetasoundWaveInfoNode.cpp
@@ -104,10 +104,10 @@ namespace Metasound
 		using namespace WaveInfoNodeParameterNames;
 		static const FVertexInterface Interface(
 			FInputVertexInterface(
-				TInputDataVertexModel<FWaveAsset>(METASOUND_GET_PARAM_NAME(ParamWaveAsset), METASOUND_GET_PARAM_TT(ParamWaveAsset))
+				TInputDataVertexModel<FWaveAsset>(METASOUND_GET_PARAM_NAME_AND_METADATA(ParamWaveAsset))
 			),
 			FOutputVertexInterface(
-				TOutputDataVertexModel<FTime>(METASOUND_GET_PARAM_NAME(ParamDurationSeconds), METASOUND_GET_PARAM_TT(ParamDurationSeconds))
+				TOutputDataVertexModel<FTime>(METASOUND_GET_PARAM_NAME_AND_METADATA(ParamDurationSeconds))
 			)
 		);
 
diff --git a/Engine/Plugins/Runtime/Metasound/Source/MetasoundGraphCore/Public/MetasoundParamHelper.h b/Engine/Plugins/Runtime/Metasound/Source/MetasoundGraphCore/Public/MetasoundParamHelper.h
new file mode 100644
--- /dev/null
+++ b/Engine/Plugins/Runtime/Metasound/Source/MetasoundGraphCore/Public/MetasoundParamHelper.h
@@ -0,0 +1,74 @@
+// Copyright Epic Games, Inc. All Rights Reserved.
+
+#pragma once
+
+#include "MetasoundVertex.h"
+
+/*
+
+	 Macros to make parameter definitions a little easier
+	 Use in a namespace
+
+	namespace Metasound
+	{
+		namespace MyNodeVertexNames
+		{
+			METASOUND_PARAM(InputParam1, "Param 1", "Tooltip for param1.");
+			METASOUND_PARAM(InputParam2, "Param 2", "Tooltip for param2.");
+			METASOUND_PARAM(OutputParam1, "Out Param1", "Tooltip for output param 2.");
+			METASOUND_PARAM(OutputParam2, "Out Param2", "Tooltip for output param 2.");
+		}
+	}
+
+	Then to retrieve:
+
+	using namespace MyNodeVertexNames;
+
+	METASOUND_GET_PARAM_NAME(InputParam1)
+	METASOUND_GET_PARAM_TT(InputParam1)
+
+	To ensure display names and tooltips are localized: 
+	METASOUND_GET_PARAM_NAME_AND_METADATA(InputParam1)
+
+	For variable parameters (e.g. templated on number), do the following:
+
+		namespace TriggerAccumulatorVertexNames
+		{
+			METASOUND_PARAM(InputParam, "Param {0}", "Do {0} tooltip.");
+		}
+
+	Then to retrieve the name/tooltip use:
+
+	METASOUND_GET_PARAM_NAME_WITH_INDEX(InputParam, NUMBER);
+
+	Where NUMBER is a number (or whatever) to slot into the format specifier you defined in the param def.
+
+*/
+
+#if WITH_EDITOR
+#define METASOUND_PARAM(NAME, NAME_TEXT, TOOLTIP_TEXT) \
+	static const TCHAR* NAME##Name = TEXT(NAME_TEXT); \
+	static const FText NAME##Tooltip = LOCTEXT(#NAME "Tooltip", TOOLTIP_TEXT); \
+	static const FText NAME##DisplayName = LOCTEXT(#NAME "DisplayName", NAME_TEXT); 
+#else 
+#define METASOUND_PARAM(NAME, NAME_TEXT, TOOLTIP_TEXT) \
+	static const TCHAR* NAME##Name = TEXT(NAME_TEXT); \
+	static const FText NAME##Tooltip = FText::GetEmpty(); \
+	static const FText NAME##DisplayName = FText::GetEmpty();
+#endif // WITH_EDITOR
+
+#define METASOUND_GET_PARAM_NAME(NAME) NAME##Name
+#define METASOUND_GET_PARAM_TT(NAME) NAME##Tooltip
+#define METASOUND_GET_PARAM_NAME_AND_TT(NAME) NAME##Name, NAME##Tooltip
+#define METASOUND_GET_PARAM_METADATA(NAME) FDataVertexMetadata { NAME##Tooltip, NAME##DisplayName }
+#define METASOUND_GET_PARAM_DISPLAYNAME(NAME) NAME##DisplayName
+#define METASOUND_GET_PARAM_NAME_AND_METADATA(NAME) METASOUND_GET_PARAM_NAME(NAME), METASOUND_GET_PARAM_METADATA(NAME)
+
+#define METASOUND_GET_PARAM_NAME_WITH_INDEX(NAME, INDEX) *FString::Format(NAME##Name, {INDEX})
+#if WITH_EDITOR
+#define METASOUND_GET_PARAM_NAME_WITH_INDEX_AND_TT(NAME, INDEX)  *FString::Format(NAME##Name, {INDEX}), FText::Format(NAME##Tooltip, INDEX)
+#define METASOUND_GET_PARAM_TT_WITH_INDEX(NAME, INDEX)  FText::Format(NAME##Tooltip, INDEX)
+#else 
+#define METASOUND_GET_PARAM_NAME_WITH_INDEX_AND_TT(NAME, INDEX)  *FString::Format(NAME##Name, {INDEX}), FText::GetEmpty()
+#define METASOUND_GET_PARAM_TT_WITH_INDEX(NAME, INDEX)  FText::GetEmpty();
+#endif // WITH_EDITOR
diff --git a/Engine/Plugins/Runtime/Metasound/Source/MetasoundStandardNodes/Private/MetasoundADEnvelopeNode.cpp b/Engine/Plugins/Runtime/Metasound/Source/MetasoundStandardNodes/Private/MetasoundADEnvelopeNode.cpp
--- a/Engine/Plugins/Runtime/Metasound/Source/MetasoundStandardNodes/Private/MetasoundADEnvelopeNode.cpp
+++ b/Engine/Plugins/Runtime/Metasound/Source/MetasoundStandardNodes/Private/MetasoundADEnvelopeNode.cpp
@@ -206,17 +206,17 @@ namespace Metasound
 
 			static const FVertexInterface DefaultInterface(
 				FInputVertexInterface(
-					TInputDataVertexModel<FTrigger>(METASOUND_GET_PARAM_NAME_AND_TT(InputTrigger)),
-					TInputDataVertexModel<FTime>(METASOUND_GET_PARAM_NAME_AND_TT(InputAttackTime), 0.01f),
-					TInputDataVertexModel<FTime>(METASOUND_GET_PARAM_NAME_AND_TT(InputDecayTime), 1.0f),
-					TInputDataVertexModel<float>(METASOUND_GET_PARAM_NAME_AND_TT(InputAttackCurve), 1.0f),
-					TInputDataVertexModel<float>(METASOUND_GET_PARAM_NAME_AND_TT(InputDecayCurve), 1.0f),
-					TInputDataVertexModel<bool>(METASOUND_GET_PARAM_NAME_AND_TT(InputLooping), false)
+					TInputDataVertexModel<FTrigger>(METASOUND_GET_PARAM_NAME_AND_METADATA(InputTrigger)),
+					TInputDataVertexModel<FTime>(METASOUND_GET_PARAM_NAME_AND_METADATA(InputAttackTime), 0.01f),
+					TInputDataVertexModel<FTime>(METASOUND_GET_PARAM_NAME_AND_METADATA(InputDecayTime), 1.0f),
+					TInputDataVertexModel<float>(METASOUND_GET_PARAM_NAME_AND_METADATA(InputAttackCurve), 1.0f),
+					TInputDataVertexModel<float>(METASOUND_GET_PARAM_NAME_AND_METADATA(InputDecayCurve), 1.0f),
+					TInputDataVertexModel<bool>(METASOUND_GET_PARAM_NAME_AND_METADATA(InputLooping), false)
 				),
 				FOutputVertexInterface(
-					TOutputDataVertexModel<FTrigger>(METASOUND_GET_PARAM_NAME_AND_TT(OutputOnTrigger)),
-					TOutputDataVertexModel<FTrigger>(METASOUND_GET_PARAM_NAME_AND_TT(OutputOnDone)),
-					TOutputDataVertexModel<ValueType>(METASOUND_GET_PARAM_NAME_AND_TT(OutputEnvelopeValue))
+					TOutputDataVertexModel<FTrigger>(METASOUND_GET_PARAM_NAME_AND_METADATA(OutputOnTrigger)),
+					TOutputDataVertexModel<FTrigger>(METASOUND_GET_PARAM_NAME_AND_METADATA(OutputOnDone)),
+					TOutputDataVertexModel<ValueType>(METASOUND_GET_PARAM_NAME_AND_METADATA(OutputEnvelopeValue))
 				)
 			);
 
diff --git a/Engine/Plugins/Runtime/Metasound/Source/MetasoundStandardNodes/Private/MetasoundADSREnvelopeNode.cpp b/Engine/Plugins/Runtime/Metasound/Source/MetasoundStandardNodes/Private/MetasoundADSREnvelopeNode.cpp
--- a/Engine/Plugins/Runtime/Metasound/Source/MetasoundStandardNodes/Private/MetasoundADSREnvelopeNode.cpp
+++ b/Engine/Plugins/Runtime/Metasound/Source/MetasoundStandardNodes/Private/MetasoundADSREnvelopeNode.cpp
@@ -330,23 +330,23 @@ namespace Metasound
 
 			static const FVertexInterface DefaultInterface(
 				FInputVertexInterface(
-					TInputDataVertexModel<FTrigger>(METASOUND_GET_PARAM_NAME_AND_TT(InputAttackTrigger)),
-					TInputDataVertexModel<FTrigger>(METASOUND_GET_PARAM_NAME_AND_TT(InputReleaseTrigger)),
-					TInputDataVertexModel<FTime>(METASOUND_GET_PARAM_NAME_AND_TT(InputAttackTime), 0.01f),
-					TInputDataVertexModel<FTime>(METASOUND_GET_PARAM_NAME_AND_TT(InputDecayTime), 0.2f),
-					TInputDataVertexModel<float>(METASOUND_GET_PARAM_NAME_AND_TT(InputSustainLevel), 0.5f),
-					TInputDataVertexModel<FTime>(METASOUND_GET_PARAM_NAME_AND_TT(InputReleaseTime), 1.0f),
-					TInputDataVertexModel<float>(METASOUND_GET_PARAM_NAME_AND_TT(InputAttackCurve), 1.0f),
-					TInputDataVertexModel<float>(METASOUND_GET_PARAM_NAME_AND_TT(InputDecayCurve), 1.0f),
-					TInputDataVertexModel<float>(METASOUND_GET_PARAM_NAME_AND_TT(InputReleaseCurve), 1.0f)
+					TInputDataVertexModel<FTrigger>(METASOUND_GET_PARAM_NAME_AND_METADATA(InputAttackTrigger)),
+					TInputDataVertexModel<FTrigger>(METASOUND_GET_PARAM_NAME_AND_METADATA(InputReleaseTrigger)),
+					TInputDataVertexModel<FTime>(METASOUND_GET_PARAM_NAME_AND_METADATA(InputAttackTime), 0.01f),
+					TInputDataVertexModel<FTime>(METASOUND_GET_PARAM_NAME_AND_METADATA(InputDecayTime), 0.2f),
+					TInputDataVertexModel<float>(METASOUND_GET_PARAM_NAME_AND_METADATA(InputSustainLevel), 0.5f),
+					TInputDataVertexModel<FTime>(METASOUND_GET_PARAM_NAME_AND_METADATA(InputReleaseTime), 1.0f),
+					TInputDataVertexModel<float>(METASOUND_GET_PARAM_NAME_AND_METADATA(InputAttackCurve), 1.0f),
+					TInputDataVertexModel<float>(METASOUND_GET_PARAM_NAME_AND_METADATA(InputDecayCurve), 1.0f),
+					TInputDataVertexModel<float>(METASOUND_GET_PARAM_NAME_AND_METADATA(InputReleaseCurve), 1.0f)
 				),
 				FOutputVertexInterface(
-					TOutputDataVertexModel<FTrigger>(METASOUND_GET_PARAM_NAME_AND_TT(OutputOnAttackTrigger)),
-					TOutputDataVertexModel<FTrigger>(METASOUND_GET_PARAM_NAME_AND_TT(OutputOnDecayTrigger)),
-					TOutputDataVertexModel<FTrigger>(METASOUND_GET_PARAM_NAME_AND_TT(OutputOnSustainTrigger)),
-					TOutputDataVertexModel<FTrigger>(METASOUND_GET_PARAM_NAME_AND_TT(OutputOnReleaseTrigger)),
-					TOutputDataVertexModel<FTrigger>(METASOUND_GET_PARAM_NAME_AND_TT(OutputOnDone)),
-					TOutputDataVertexModel<ValueType>(METASOUND_GET_PARAM_NAME_AND_TT(OutputEnvelopeValue))
+					TOutputDataVertexModel<FTrigger>(METASOUND_GET_PARAM_NAME_AND_METADATA(OutputOnAttackTrigger)),
+					TOutputDataVertexModel<FTrigger>(METASOUND_GET_PARAM_NAME_AND_METADATA(OutputOnDecayTrigger)),
+					TOutputDataVertexModel<FTrigger>(METASOUND_GET_PARAM_NAME_AND_METADATA(OutputOnSustainTrigger)),
+					TOutputDataVertexModel<FTrigger>(METASOUND_GET_PARAM_NAME_AND_METADATA(OutputOnReleaseTrigger)),
+					TOutputDataVertexModel<FTrigger>(METASOUND_GET_PARAM_NAME_AND_METADATA(OutputOnDone)),
+					TOutputDataVertexModel<ValueType>(METASOUND_GET_PARAM_NAME_AND_METADATA(OutputEnvelopeValue))
 				)
 			);
 
diff --git a/Engine/Plugins/Runtime/Metasound/Source/MetasoundStandardNodes/Private/MetasoundAdditiveSynthNode.cpp b/Engine/Plugins/Runtime/Metasound/Source/MetasoundStandardNodes/Private/MetasoundAdditiveSynthNode.cpp
--- a/Engine/Plugins/Runtime/Metasound/Source/MetasoundStandardNodes/Private/MetasoundAdditiveSynthNode.cpp
+++ b/Engine/Plugins/Runtime/Metasound/Source/MetasoundStandardNodes/Private/MetasoundAdditiveSynthNode.cpp
@@ -64,15 +64,15 @@ namespace Metasound
 
 			static const FVertexInterface Interface(
 				FInputVertexInterface(
-					TInputDataVertexModel<float>(METASOUND_GET_PARAM_NAME_AND_TT(BaseFrequency), 440.0f),
-					TInputDataVertexModel<FInputFloatArrayType>(METASOUND_GET_PARAM_NAME_AND_TT(HarmonicMultipliers)),
-					TInputDataVertexModel<FInputFloatArrayType>(METASOUND_GET_PARAM_NAME_AND_TT(Amplitudes)),
-					TInputDataVertexModel<FInputFloatArrayType>(METASOUND_GET_PARAM_NAME_AND_TT(Phases)),
-					TInputDataVertexModel<FInputFloatArrayType>(METASOUND_GET_PARAM_NAME_AND_TT(PanAmounts))
+					TInputDataVertexModel<float>(METASOUND_GET_PARAM_NAME_AND_METADATA(BaseFrequency), 440.0f),
+					TInputDataVertexModel<FInputFloatArrayType>(METASOUND_GET_PARAM_NAME_AND_METADATA(HarmonicMultipliers)),
+					TInputDataVertexModel<FInputFloatArrayType>(METASOUND_GET_PARAM_NAME_AND_METADATA(Amplitudes)),
+					TInputDataVertexModel<FInputFloatArrayType>(METASOUND_GET_PARAM_NAME_AND_METADATA(Phases)),
+					TInputDataVertexModel<FInputFloatArrayType>(METASOUND_GET_PARAM_NAME_AND_METADATA(PanAmounts))
 				),
 				FOutputVertexInterface(
-					TOutputDataVertexModel<FAudioBuffer>(METASOUND_GET_PARAM_NAME_AND_TT(LeftAudioOut)),
-					TOutputDataVertexModel<FAudioBuffer>(METASOUND_GET_PARAM_NAME_AND_TT(RightAudioOut))
+					TOutputDataVertexModel<FAudioBuffer>(METASOUND_GET_PARAM_NAME_AND_METADATA(LeftAudioOut)),
+					TOutputDataVertexModel<FAudioBuffer>(METASOUND_GET_PARAM_NAME_AND_METADATA(RightAudioOut))
 				)
 			);
 
diff --git a/Engine/Plugins/Runtime/Metasound/Source/MetasoundStandardNodes/Private/MetasoundBasicFilters.cpp b/Engine/Plugins/Runtime/Metasound/Source/MetasoundStandardNodes/Private/MetasoundBasicFilters.cpp
--- a/Engine/Plugins/Runtime/Metasound/Source/MetasoundStandardNodes/Private/MetasoundBasicFilters.cpp
+++ b/Engine/Plugins/Runtime/Metasound/Source/MetasoundStandardNodes/Private/MetasoundBasicFilters.cpp
@@ -182,12 +182,12 @@ namespace Metasound
 	{
 		static const FVertexInterface Interface(
 			FInputVertexInterface(
-				TInputDataVertexModel<FAudioBuffer>(METASOUND_GET_PARAM_NAME(ParamAudioInput), METASOUND_GET_PARAM_TT(ParamAudioInput)),
-				TInputDataVertexModel<float>(METASOUND_GET_PARAM_NAME(ParamCutoffFrequency), METASOUND_GET_PARAM_TT(ParamCutoffFrequency), 20000.f),
-				TInputDataVertexModel<float>(METASOUND_GET_PARAM_NAME(ParamResonance), METASOUND_GET_PARAM_TT(ParamResonance), 1.f)
+				TInputDataVertexModel<FAudioBuffer>(METASOUND_GET_PARAM_NAME_AND_METADATA(ParamAudioInput)),
+				TInputDataVertexModel<float>(METASOUND_GET_PARAM_NAME_AND_METADATA(ParamCutoffFrequency), 20000.f),
+				TInputDataVertexModel<float>(METASOUND_GET_PARAM_NAME_AND_METADATA(ParamResonance), 1.f)
 			),
 			FOutputVertexInterface(
-				TOutputDataVertexModel<FAudioBuffer>(METASOUND_GET_PARAM_NAME(ParamAudioOutput), METASOUND_GET_PARAM_TT(ParamAudioOutput))
+				TOutputDataVertexModel<FAudioBuffer>(METASOUND_GET_PARAM_NAME_AND_METADATA(ParamAudioOutput))
 			)
 		);
 
@@ -364,16 +364,16 @@ namespace Metasound
 	{
 		static const FVertexInterface Interface(
 			FInputVertexInterface(
-				TInputDataVertexModel<FAudioBuffer>(METASOUND_GET_PARAM_NAME(ParamAudioInput), METASOUND_GET_PARAM_TT(ParamAudioInput)),
-				TInputDataVertexModel<float>(METASOUND_GET_PARAM_NAME(ParamCutoffFrequency), METASOUND_GET_PARAM_TT(ParamCutoffFrequency), 20000.f),
-				TInputDataVertexModel<float>(METASOUND_GET_PARAM_NAME(ParamResonance), METASOUND_GET_PARAM_TT(ParamResonance), 0.f),
-				TInputDataVertexModel<float>(METASOUND_GET_PARAM_NAME(ParamBandStopControl), METASOUND_GET_PARAM_TT(ParamBandStopControl), 0.f)
+				TInputDataVertexModel<FAudioBuffer>(METASOUND_GET_PARAM_NAME_AND_METADATA(ParamAudioInput)),
+				TInputDataVertexModel<float>(METASOUND_GET_PARAM_NAME_AND_METADATA(ParamCutoffFrequency), 20000.f),
+				TInputDataVertexModel<float>(METASOUND_GET_PARAM_NAME_AND_METADATA(ParamResonance), 0.f),
+				TInputDataVertexModel<float>(METASOUND_GET_PARAM_NAME_AND_METADATA(ParamBandStopControl), 0.f)
 			),
 			FOutputVertexInterface(
-				TOutputDataVertexModel<FAudioBuffer>(METASOUND_GET_PARAM_NAME(ParamLowPassOutput), METASOUND_GET_PARAM_TT(ParamLowPassOutput)),
-				TOutputDataVertexModel<FAudioBuffer>(METASOUND_GET_PARAM_NAME(ParamHighPassOutput), METASOUND_GET_PARAM_TT(ParamHighPassOutput)),
-				TOutputDataVertexModel<FAudioBuffer>(METASOUND_GET_PARAM_NAME(ParamBandPassOutput), METASOUND_GET_PARAM_TT(ParamBandPassOutput)),
-				TOutputDataVertexModel<FAudioBuffer>(METASOUND_GET_PARAM_NAME(ParamBandStopOutput), METASOUND_GET_PARAM_TT(ParamBandStopOutput))
+				TOutputDataVertexModel<FAudioBuffer>(METASOUND_GET_PARAM_NAME_AND_METADATA(ParamLowPassOutput)),
+				TOutputDataVertexModel<FAudioBuffer>(METASOUND_GET_PARAM_NAME_AND_METADATA(ParamHighPassOutput)),
+				TOutputDataVertexModel<FAudioBuffer>(METASOUND_GET_PARAM_NAME_AND_METADATA(ParamBandPassOutput)),
+				TOutputDataVertexModel<FAudioBuffer>(METASOUND_GET_PARAM_NAME_AND_METADATA(ParamBandStopOutput))
 			)
 		);
 
@@ -536,11 +536,11 @@ namespace Metasound
 	{
 		static const FVertexInterface Interface(
 			FInputVertexInterface(
-				TInputDataVertexModel<FAudioBuffer>(METASOUND_GET_PARAM_NAME(ParamAudioInput), METASOUND_GET_PARAM_TT(ParamAudioInput)),
-				TInputDataVertexModel<float>(METASOUND_GET_PARAM_NAME(ParamCutoffFrequency), METASOUND_GET_PARAM_TT(ParamCutoffFrequency), 20000.f)
+				TInputDataVertexModel<FAudioBuffer>(METASOUND_GET_PARAM_NAME_AND_METADATA(ParamAudioInput)),
+				TInputDataVertexModel<float>(METASOUND_GET_PARAM_NAME_AND_METADATA(ParamCutoffFrequency), 20000.f)
 			),
 			FOutputVertexInterface(
-				TOutputDataVertexModel<FAudioBuffer>(METASOUND_GET_PARAM_NAME(ParamAudioOutput), METASOUND_GET_PARAM_TT(ParamAudioOutput))
+				TOutputDataVertexModel<FAudioBuffer>(METASOUND_GET_PARAM_NAME_AND_METADATA(ParamAudioOutput))
 			)
 		);
 
@@ -673,11 +673,11 @@ namespace Metasound
 	{
 		static const FVertexInterface Interface(
 			FInputVertexInterface(
-				TInputDataVertexModel<FAudioBuffer>(METASOUND_GET_PARAM_NAME(ParamAudioInput), METASOUND_GET_PARAM_TT(ParamAudioInput)),
-				TInputDataVertexModel<float>(METASOUND_GET_PARAM_NAME(ParamCutoffFrequency), METASOUND_GET_PARAM_TT(ParamCutoffFrequency), 10.f)
+				TInputDataVertexModel<FAudioBuffer>(METASOUND_GET_PARAM_NAME_AND_METADATA(ParamAudioInput)),
+				TInputDataVertexModel<float>(METASOUND_GET_PARAM_NAME_AND_METADATA(ParamCutoffFrequency), 10.f)
 			),
 			FOutputVertexInterface(
-				TOutputDataVertexModel<FAudioBuffer>(METASOUND_GET_PARAM_NAME(ParamAudioOutput), METASOUND_GET_PARAM_TT(ParamAudioOutput))
+				TOutputDataVertexModel<FAudioBuffer>(METASOUND_GET_PARAM_NAME_AND_METADATA(ParamAudioOutput))
 			)
 		);
 
@@ -835,14 +835,14 @@ namespace Metasound
 
 		static const FVertexInterface Interface(
 			FInputVertexInterface(
-				TInputDataVertexModel<FAudioBuffer>(METASOUND_GET_PARAM_NAME(ParamAudioInput), METASOUND_GET_PARAM_TT(ParamAudioInput)),
-				TInputDataVertexModel<float>(METASOUND_GET_PARAM_NAME(ParamCutoffFrequency), METASOUND_GET_PARAM_TT(ParamCutoffFrequency), 20000.f),
-				TInputDataVertexModel<float>(METASOUND_GET_PARAM_NAME(ParamBandwidth), METASOUND_GET_PARAM_TT(ParamBandwidth), 1.f),
+				TInputDataVertexModel<FAudioBuffer>(METASOUND_GET_PARAM_NAME_AND_METADATA(ParamAudioInput)),
+				TInputDataVertexModel<float>(METASOUND_GET_PARAM_NAME_AND_METADATA(ParamCutoffFrequency), 20000.f),
+				TInputDataVertexModel<float>(METASOUND_GET_PARAM_NAME_AND_METADATA(ParamBandwidth), 1.f),
 				TInputDataVertexModel<float>(METASOUND_GET_PARAM_NAME(ParamGainDb), GainPinMetaData, 0.f),
-				TInputDataVertexModel<FEnumEBiquadFilterType>(METASOUND_GET_PARAM_NAME(ParamFilterType), METASOUND_GET_PARAM_TT(ParamFilterType))
+				TInputDataVertexModel<FEnumEBiquadFilterType>(METASOUND_GET_PARAM_NAME_AND_METADATA(ParamFilterType))
 			),
 			FOutputVertexInterface(
-				TOutputDataVertexModel<FAudioBuffer>(METASOUND_GET_PARAM_NAME(ParamAudioOutput), METASOUND_GET_PARAM_TT(ParamAudioOutput))
+				TOutputDataVertexModel<FAudioBuffer>(METASOUND_GET_PARAM_NAME_AND_METADATA(ParamAudioOutput))
 			)
 		);
 
diff --git a/Engine/Plugins/Runtime/Metasound/Source/MetasoundStandardNodes/Private/MetasoundBitcrusherNode.cpp b/Engine/Plugins/Runtime/Metasound/Source/MetasoundStandardNodes/Private/MetasoundBitcrusherNode.cpp
--- a/Engine/Plugins/Runtime/Metasound/Source/MetasoundStandardNodes/Private/MetasoundBitcrusherNode.cpp
+++ b/Engine/Plugins/Runtime/Metasound/Source/MetasoundStandardNodes/Private/MetasoundBitcrusherNode.cpp
@@ -82,12 +82,12 @@ namespace Metasound
 
 			static const FVertexInterface Interface(
 				FInputVertexInterface(
-					TInputDataVertexModel<FAudioBuffer>(METASOUND_GET_PARAM_NAME_AND_TT(InputAudio)),
-					TInputDataVertexModel<float>(METASOUND_GET_PARAM_NAME_AND_TT(InputSampleRate), 8000.0f),
-					TInputDataVertexModel<float>(METASOUND_GET_PARAM_NAME_AND_TT(InputBitDepth), 8.0f)
+					TInputDataVertexModel<FAudioBuffer>(METASOUND_GET_PARAM_NAME_AND_METADATA(InputAudio)),
+					TInputDataVertexModel<float>(METASOUND_GET_PARAM_NAME_AND_METADATA(InputSampleRate), 8000.0f),
+					TInputDataVertexModel<float>(METASOUND_GET_PARAM_NAME_AND_METADATA(InputBitDepth), 8.0f)
 				),
 				FOutputVertexInterface(
-					TOutputDataVertexModel<FAudioBuffer>(METASOUND_GET_PARAM_NAME_AND_TT(OutputAudio))
+					TOutputDataVertexModel<FAudioBuffer>(METASOUND_GET_PARAM_NAME_AND_METADATA(OutputAudio))
 				)
 			);
 
diff --git a/Engine/Plugins/Runtime/Metasound/Source/MetasoundStandardNodes/Private/MetasoundClampNode.cpp b/Engine/Plugins/Runtime/Metasound/Source/MetasoundStandardNodes/Private/MetasoundClampNode.cpp
--- a/Engine/Plugins/Runtime/Metasound/Source/MetasoundStandardNodes/Private/MetasoundClampNode.cpp
+++ b/Engine/Plugins/Runtime/Metasound/Source/MetasoundStandardNodes/Private/MetasoundClampNode.cpp
@@ -157,12 +157,12 @@ namespace Metasound
 
 			static const FVertexInterface DefaultInterface(
 				FInputVertexInterface(
-					TInputDataVertexModel<ValueType>(METASOUND_GET_PARAM_NAME_AND_TT(InputValue)),
-					TInputDataVertexModel<ValueType>(METASOUND_GET_PARAM_NAME_AND_TT(InputMinValue)),
-					TInputDataVertexModel<ValueType>(METASOUND_GET_PARAM_NAME_AND_TT(InputMaxValue))
+					TInputDataVertexModel<ValueType>(METASOUND_GET_PARAM_NAME_AND_METADATA(InputValue)),
+					TInputDataVertexModel<ValueType>(METASOUND_GET_PARAM_NAME_AND_METADATA(InputMinValue)),
+					TInputDataVertexModel<ValueType>(METASOUND_GET_PARAM_NAME_AND_METADATA(InputMaxValue))
 				),
 				FOutputVertexInterface(
-					TOutputDataVertexModel<ValueType>(METASOUND_GET_PARAM_NAME_AND_TT(OutputValue))
+					TOutputDataVertexModel<ValueType>(METASOUND_GET_PARAM_NAME_AND_METADATA(OutputValue))
 				)
 			);
 
diff --git a/Engine/Plugins/Runtime/Metasound/Source/MetasoundStandardNodes/Private/MetasoundCompressorNode.cpp b/Engine/Plugins/Runtime/Metasound/Source/MetasoundStandardNodes/Private/MetasoundCompressorNode.cpp
--- a/Engine/Plugins/Runtime/Metasound/Source/MetasoundStandardNodes/Private/MetasoundCompressorNode.cpp
+++ b/Engine/Plugins/Runtime/Metasound/Source/MetasoundStandardNodes/Private/MetasoundCompressorNode.cpp
@@ -145,22 +145,22 @@ namespace Metasound
 
 			static const FVertexInterface Interface(
 				FInputVertexInterface(
-					TInputDataVertexModel<FAudioBuffer>(METASOUND_GET_PARAM_NAME_AND_TT(InputAudio)),
-					TInputDataVertexModel<float>(METASOUND_GET_PARAM_NAME_AND_TT(InputRatio), 1.5f),
-					TInputDataVertexModel<float>(METASOUND_GET_PARAM_NAME_AND_TT(InputThreshold), -6.0f),
-					TInputDataVertexModel<FTime>(METASOUND_GET_PARAM_NAME_AND_TT(InputAttackTime), 0.01f),
-					TInputDataVertexModel<FTime>(METASOUND_GET_PARAM_NAME_AND_TT(InputReleaseTime), 0.1f),
-					TInputDataVertexModel<FTime>(METASOUND_GET_PARAM_NAME_AND_TT(InputLookaheadTime), 0.01f),
-					TInputDataVertexModel<float>(METASOUND_GET_PARAM_NAME_AND_TT(InputKnee), 10.0f),
-					TInputDataVertexModel<FAudioBuffer>(METASOUND_GET_PARAM_NAME_AND_TT(InputSidechain)),
-					TInputDataVertexModel<FEnumEnvelopePeakMode>(METASOUND_GET_PARAM_NAME_AND_TT(InputEnvelopeMode)),
-					TInputDataVertexModel<bool>(METASOUND_GET_PARAM_NAME_AND_TT(InputIsAnalog), true),
-					TInputDataVertexModel<bool>(METASOUND_GET_PARAM_NAME_AND_TT(InputIsUpwards), false),
-					TInputDataVertexModel<float>(METASOUND_GET_PARAM_NAME_AND_TT(InputWetDryMix), 1.0f)
+					TInputDataVertexModel<FAudioBuffer>(METASOUND_GET_PARAM_NAME_AND_METADATA(InputAudio)),
+					TInputDataVertexModel<float>(METASOUND_GET_PARAM_NAME_AND_METADATA(InputRatio), 1.5f),
+					TInputDataVertexModel<float>(METASOUND_GET_PARAM_NAME_AND_METADATA(InputThreshold), -6.0f),
+					TInputDataVertexModel<FTime>(METASOUND_GET_PARAM_NAME_AND_METADATA(InputAttackTime), 0.01f),
+					TInputDataVertexModel<FTime>(METASOUND_GET_PARAM_NAME_AND_METADATA(InputReleaseTime), 0.1f),
+					TInputDataVertexModel<FTime>(METASOUND_GET_PARAM_NAME_AND_METADATA(InputLookaheadTime), 0.01f),
+					TInputDataVertexModel<float>(METASOUND_GET_PARAM_NAME_AND_METADATA(InputKnee), 10.0f),
+					TInputDataVertexModel<FAudioBuffer>(METASOUND_GET_PARAM_NAME_AND_METADATA(InputSidechain)),
+					TInputDataVertexModel<FEnumEnvelopePeakMode>(METASOUND_GET_PARAM_NAME_AND_METADATA(InputEnvelopeMode)),
+					TInputDataVertexModel<bool>(METASOUND_GET_PARAM_NAME_AND_METADATA(InputIsAnalog), true),
+					TInputDataVertexModel<bool>(METASOUND_GET_PARAM_NAME_AND_METADATA(InputIsUpwards), false),
+					TInputDataVertexModel<float>(METASOUND_GET_PARAM_NAME_AND_METADATA(InputWetDryMix), 1.0f)
 				),
 				FOutputVertexInterface(
-					TOutputDataVertexModel<FAudioBuffer>(METASOUND_GET_PARAM_NAME_AND_TT(OutputAudio)),
-					TOutputDataVertexModel<FAudioBuffer>(METASOUND_GET_PARAM_NAME_AND_TT(OutputEnvelope))
+					TOutputDataVertexModel<FAudioBuffer>(METASOUND_GET_PARAM_NAME_AND_METADATA(OutputAudio)),
+					TOutputDataVertexModel<FAudioBuffer>(METASOUND_GET_PARAM_NAME_AND_METADATA(OutputEnvelope))
 				)
 			);
 
diff --git a/Engine/Plugins/Runtime/Metasound/Source/MetasoundStandardNodes/Private/MetasoundDecibelsToLinearGain.cpp b/Engine/Plugins/Runtime/Metasound/Source/MetasoundStandardNodes/Private/MetasoundDecibelsToLinearGain.cpp
--- a/Engine/Plugins/Runtime/Metasound/Source/MetasoundStandardNodes/Private/MetasoundDecibelsToLinearGain.cpp
+++ b/Engine/Plugins/Runtime/Metasound/Source/MetasoundStandardNodes/Private/MetasoundDecibelsToLinearGain.cpp
@@ -82,10 +82,10 @@ namespace Metasound
 
 		static const FVertexInterface Interface(
 			FInputVertexInterface(
-				TInputDataVertexModel<float>(METASOUND_GET_PARAM_NAME_AND_TT(InputDecibelGain), 0.0f)
+				TInputDataVertexModel<float>(METASOUND_GET_PARAM_NAME_AND_METADATA(InputDecibelGain), 0.0f)
 			),
 			FOutputVertexInterface(
-				TOutputDataVertexModel<float>(METASOUND_GET_PARAM_NAME_AND_TT(OutputLinearGain))
+				TOutputDataVertexModel<float>(METASOUND_GET_PARAM_NAME_AND_METADATA(OutputLinearGain))
 			)
 		);
 
diff --git a/Engine/Plugins/Runtime/Metasound/Source/MetasoundStandardNodes/Private/MetasoundDiffuserNode.cpp b/Engine/Plugins/Runtime/Metasound/Source/MetasoundStandardNodes/Private/MetasoundDiffuserNode.cpp
--- a/Engine/Plugins/Runtime/Metasound/Source/MetasoundStandardNodes/Private/MetasoundDiffuserNode.cpp
+++ b/Engine/Plugins/Runtime/Metasound/Source/MetasoundStandardNodes/Private/MetasoundDiffuserNode.cpp
@@ -63,12 +63,12 @@ namespace Metasound
 
 			static const FVertexInterface Interface(
 				FInputVertexInterface(
-					TInputDataVertexModel<FAudioBuffer>(METASOUND_GET_PARAM_NAME_AND_TT(InputAudio)),
-					TInputDataVertexModel<int32>(METASOUND_GET_PARAM_NAME_AND_TT(InputDiffusionDepth), 4),
-					TInputDataVertexModel<float>(METASOUND_GET_PARAM_NAME_AND_TT(InputFeedbackGain), 0.707f)
+					TInputDataVertexModel<FAudioBuffer>(METASOUND_GET_PARAM_NAME_AND_METADATA(InputAudio)),
+					TInputDataVertexModel<int32>(METASOUND_GET_PARAM_NAME_AND_METADATA(InputDiffusionDepth), 4),
+					TInputDataVertexModel<float>(METASOUND_GET_PARAM_NAME_AND_METADATA(InputFeedbackGain), 0.707f)
 				),
 				FOutputVertexInterface(
-					TOutputDataVertexModel<FAudioBuffer>(METASOUND_GET_PARAM_NAME_AND_TT(OutputAudio))
+					TOutputDataVertexModel<FAudioBuffer>(METASOUND_GET_PARAM_NAME_AND_METADATA(OutputAudio))
 				)
 			);
 
diff --git a/Engine/Plugins/Runtime/Metasound/Source/MetasoundStandardNodes/Private/MetasoundDynamicFilterNode.cpp b/Engine/Plugins/Runtime/Metasound/Source/MetasoundStandardNodes/Private/MetasoundDynamicFilterNode.cpp
--- a/Engine/Plugins/Runtime/Metasound/Source/MetasoundStandardNodes/Private/MetasoundDynamicFilterNode.cpp
+++ b/Engine/Plugins/Runtime/Metasound/Source/MetasoundStandardNodes/Private/MetasoundDynamicFilterNode.cpp
@@ -161,27 +161,27 @@ namespace Metasound
 
 			static const FVertexInterface Interface(
 				FInputVertexInterface(
-					TInputDataVertexModel<FAudioBuffer>(METASOUND_GET_PARAM_NAME_AND_TT(InputAudio)),
-					TInputDataVertexModel<FAudioBuffer>(METASOUND_GET_PARAM_NAME_AND_TT(InputSidechain)),
-					TInputDataVertexModel<FEnumEDynamicFilterType>(METASOUND_GET_PARAM_NAME_AND_TT(InputFilterType)),
-					TInputDataVertexModel<float>(METASOUND_GET_PARAM_NAME_AND_TT(InputFrequency), 1000.0f),
-					TInputDataVertexModel<float>(METASOUND_GET_PARAM_NAME_AND_TT(InputQ), 1.0f),
-					TInputDataVertexModel<float>(METASOUND_GET_PARAM_NAME_AND_TT(InputThreshold), -12.0f),
-					TInputDataVertexModel<float>(METASOUND_GET_PARAM_NAME_AND_TT(InputRatio), 4.0f),
-					TInputDataVertexModel<float>(METASOUND_GET_PARAM_NAME_AND_TT(InputKnee), 12.0f),
-					TInputDataVertexModel<float>(METASOUND_GET_PARAM_NAME_AND_TT(InputRange), -60.0f),
+					TInputDataVertexModel<FAudioBuffer>(METASOUND_GET_PARAM_NAME_AND_METADATA(InputAudio)),
+					TInputDataVertexModel<FAudioBuffer>(METASOUND_GET_PARAM_NAME_AND_METADATA(InputSidechain)),
+					TInputDataVertexModel<FEnumEDynamicFilterType>(METASOUND_GET_PARAM_NAME_AND_METADATA(InputFilterType)),
+					TInputDataVertexModel<float>(METASOUND_GET_PARAM_NAME_AND_METADATA(InputFrequency), 1000.0f),
+					TInputDataVertexModel<float>(METASOUND_GET_PARAM_NAME_AND_METADATA(InputQ), 1.0f),
+					TInputDataVertexModel<float>(METASOUND_GET_PARAM_NAME_AND_METADATA(InputThreshold), -12.0f),
+					TInputDataVertexModel<float>(METASOUND_GET_PARAM_NAME_AND_METADATA(InputRatio), 4.0f),
+					TInputDataVertexModel<float>(METASOUND_GET_PARAM_NAME_AND_METADATA(InputKnee), 12.0f),
+					TInputDataVertexModel<float>(METASOUND_GET_PARAM_NAME_AND_METADATA(InputRange), -60.0f),
 					TInputDataVertexModel<float>(METASOUND_GET_PARAM_NAME(InputGain),
 					{ 
 						METASOUND_GET_PARAM_TT(InputGain),
 						METASOUND_LOCTEXT("InputGainName", "Gain (dB)")
 					}, 0.0f),
-					TInputDataVertexModel<FTime>(METASOUND_GET_PARAM_NAME_AND_TT(InputAttackTime), 0.01f),
-					TInputDataVertexModel<FTime>(METASOUND_GET_PARAM_NAME_AND_TT(InputReleaseTime), 0.1f),
-					TInputDataVertexModel<FEnumEnvelopePeakMode>(METASOUND_GET_PARAM_NAME_AND_TT(InputEnvelopeMode), 1.0f),
-					TInputDataVertexModel<bool>(METASOUND_GET_PARAM_NAME_AND_TT(InputAnalogMode), true)
+					TInputDataVertexModel<FTime>(METASOUND_GET_PARAM_NAME_AND_METADATA(InputAttackTime), 0.01f),
+					TInputDataVertexModel<FTime>(METASOUND_GET_PARAM_NAME_AND_METADATA(InputReleaseTime), 0.1f),
+					TInputDataVertexModel<FEnumEnvelopePeakMode>(METASOUND_GET_PARAM_NAME_AND_METADATA(InputEnvelopeMode), 1.0f),
+					TInputDataVertexModel<bool>(METASOUND_GET_PARAM_NAME_AND_METADATA(InputAnalogMode), true)
 				),
 				FOutputVertexInterface(
-					TOutputDataVertexModel<FAudioBuffer>(METASOUND_GET_PARAM_NAME_AND_TT(OutputAudio))
+					TOutputDataVertexModel<FAudioBuffer>(METASOUND_GET_PARAM_NAME_AND_METADATA(OutputAudio))
 				)
 			);
 
diff --git a/Engine/Plugins/Runtime/Metasound/Source/MetasoundStandardNodes/Private/MetasoundFlangerNode.cpp b/Engine/Plugins/Runtime/Metasound/Source/MetasoundStandardNodes/Private/MetasoundFlangerNode.cpp
--- a/Engine/Plugins/Runtime/Metasound/Source/MetasoundStandardNodes/Private/MetasoundFlangerNode.cpp
+++ b/Engine/Plugins/Runtime/Metasound/Source/MetasoundStandardNodes/Private/MetasoundFlangerNode.cpp
@@ -61,14 +61,14 @@ namespace Metasound
 
 			static const FVertexInterface Interface(
 				FInputVertexInterface(
-					TInputDataVertexModel<FAudioBuffer>(METASOUND_GET_PARAM_NAME_AND_TT(InputAudio)), 
-					TInputDataVertexModel<float>(METASOUND_GET_PARAM_NAME_AND_TT(InputModulationRate), 0.5f),
-					TInputDataVertexModel<float>(METASOUND_GET_PARAM_NAME_AND_TT(InputModulationDepth), 0.5f),
-					TInputDataVertexModel<float>(METASOUND_GET_PARAM_NAME_AND_TT(InputCenterDelay), 0.5f),
-					TInputDataVertexModel<float>(METASOUND_GET_PARAM_NAME_AND_TT(InputMixLevel), 0.5f)
+					TInputDataVertexModel<FAudioBuffer>(METASOUND_GET_PARAM_NAME_AND_METADATA(InputAudio)), 
+					TInputDataVertexModel<float>(METASOUND_GET_PARAM_NAME_AND_METADATA(InputModulationRate), 0.5f),
+					TInputDataVertexModel<float>(METASOUND_GET_PARAM_NAME_AND_METADATA(InputModulationDepth), 0.5f),
+					TInputDataVertexModel<float>(METASOUND_GET_PARAM_NAME_AND_METADATA(InputCenterDelay), 0.5f),
+					TInputDataVertexModel<float>(METASOUND_GET_PARAM_NAME_AND_METADATA(InputMixLevel), 0.5f)
 				),
 				FOutputVertexInterface(
-					TOutputDataVertexModel<FAudioBuffer>(METASOUND_GET_PARAM_NAME_AND_TT(OutputAudio))
+					TOutputDataVertexModel<FAudioBuffer>(METASOUND_GET_PARAM_NAME_AND_METADATA(OutputAudio))
 				)
 			);
 
diff --git a/Engine/Plugins/Runtime/Metasound/Source/MetasoundStandardNodes/Private/MetasoundFreqMultiplierToSemitoneNode.cpp b/Engine/Plugins/Runtime/Metasound/Source/MetasoundStandardNodes/Private/MetasoundFreqMultiplierToSemitoneNode.cpp
--- a/Engine/Plugins/Runtime/Metasound/Source/MetasoundStandardNodes/Private/MetasoundFreqMultiplierToSemitoneNode.cpp
+++ b/Engine/Plugins/Runtime/Metasound/Source/MetasoundStandardNodes/Private/MetasoundFreqMultiplierToSemitoneNode.cpp
@@ -91,10 +91,10 @@ namespace Metasound
 
 		static const FVertexInterface Interface(
 			FInputVertexInterface(
-				TInputDataVertexModel<float>(METASOUND_GET_PARAM_NAME_AND_TT(InputFrequencyMultiplier), 1.0f)
+				TInputDataVertexModel<float>(METASOUND_GET_PARAM_NAME_AND_METADATA(InputFrequencyMultiplier), 1.0f)
 			),
 			FOutputVertexInterface(
-				TOutputDataVertexModel<float>(METASOUND_GET_PARAM_NAME_AND_TT(OutputSemitone))
+				TOutputDataVertexModel<float>(METASOUND_GET_PARAM_NAME_AND_METADATA(OutputSemitone))
 			)
 		);
 
diff --git a/Engine/Plugins/Runtime/Metasound/Source/MetasoundStandardNodes/Private/MetasoundFreqToMIDINode.cpp b/Engine/Plugins/Runtime/Metasound/Source/MetasoundStandardNodes/Private/MetasoundFreqToMIDINode.cpp
--- a/Engine/Plugins/Runtime/Metasound/Source/MetasoundStandardNodes/Private/MetasoundFreqToMIDINode.cpp
+++ b/Engine/Plugins/Runtime/Metasound/Source/MetasoundStandardNodes/Private/MetasoundFreqToMIDINode.cpp
@@ -92,10 +92,10 @@ namespace Metasound
 
 		static const FVertexInterface Interface(
 			FInputVertexInterface(
-				TInputDataVertexModel<float>(METASOUND_GET_PARAM_NAME_AND_TT(InputFreq), 440.0f)
+				TInputDataVertexModel<float>(METASOUND_GET_PARAM_NAME_AND_METADATA(InputFreq), 440.0f)
 			),
 			FOutputVertexInterface(
-				TOutputDataVertexModel<float>(METASOUND_GET_PARAM_NAME_AND_TT(OutputMidi))
+				TOutputDataVertexModel<float>(METASOUND_GET_PARAM_NAME_AND_METADATA(OutputMidi))
 			)
 		);
 
diff --git a/Engine/Plugins/Runtime/Metasound/Source/MetasoundStandardNodes/Private/MetasoundLimiterNode.cpp b/Engine/Plugins/Runtime/Metasound/Source/MetasoundStandardNodes/Private/MetasoundLimiterNode.cpp
--- a/Engine/Plugins/Runtime/Metasound/Source/MetasoundStandardNodes/Private/MetasoundLimiterNode.cpp
+++ b/Engine/Plugins/Runtime/Metasound/Source/MetasoundStandardNodes/Private/MetasoundLimiterNode.cpp
@@ -123,14 +123,14 @@ namespace Metasound
 
 			static const FVertexInterface Interface(
 				FInputVertexInterface(
-					TInputDataVertexModel<FAudioBuffer>(METASOUND_GET_PARAM_NAME_AND_TT(InputAudio)),
-					TInputDataVertexModel<float>(METASOUND_GET_PARAM_NAME_AND_TT(InputInGainDb), 0.0f),
-					TInputDataVertexModel<float>(METASOUND_GET_PARAM_NAME_AND_TT(InputThresholdDb), 0.0f),
-					TInputDataVertexModel<FTime>(METASOUND_GET_PARAM_NAME_AND_TT(InputReleaseTime), 0.1f),
-					TInputDataVertexModel<FEnumKneeMode>(METASOUND_GET_PARAM_NAME_AND_TT(InputKneeMode))
+					TInputDataVertexModel<FAudioBuffer>(METASOUND_GET_PARAM_NAME_AND_METADATA(InputAudio)),
+					TInputDataVertexModel<float>(METASOUND_GET_PARAM_NAME_AND_METADATA(InputInGainDb), 0.0f),
+					TInputDataVertexModel<float>(METASOUND_GET_PARAM_NAME_AND_METADATA(InputThresholdDb), 0.0f),
+					TInputDataVertexModel<FTime>(METASOUND_GET_PARAM_NAME_AND_METADATA(InputReleaseTime), 0.1f),
+					TInputDataVertexModel<FEnumKneeMode>(METASOUND_GET_PARAM_NAME_AND_METADATA(InputKneeMode))
 				),
 				FOutputVertexInterface(
-					TOutputDataVertexModel<FAudioBuffer>(METASOUND_GET_PARAM_NAME_AND_TT(OutputAudio))
+					TOutputDataVertexModel<FAudioBuffer>(METASOUND_GET_PARAM_NAME_AND_METADATA(OutputAudio))
 				)
 			);
 
diff --git a/Engine/Plugins/Runtime/Metasound/Source/MetasoundStandardNodes/Private/MetasoundLinearGainToDecibels.cpp b/Engine/Plugins/Runtime/Metasound/Source/MetasoundStandardNodes/Private/MetasoundLinearGainToDecibels.cpp
--- a/Engine/Plugins/Runtime/Metasound/Source/MetasoundStandardNodes/Private/MetasoundLinearGainToDecibels.cpp
+++ b/Engine/Plugins/Runtime/Metasound/Source/MetasoundStandardNodes/Private/MetasoundLinearGainToDecibels.cpp
@@ -87,10 +87,10 @@ namespace Metasound
 
 		static const FVertexInterface Interface(
 			FInputVertexInterface(
-				TInputDataVertexModel<float>(METASOUND_GET_PARAM_NAME_AND_TT(InputLinearGain), 1.0f)
+				TInputDataVertexModel<float>(METASOUND_GET_PARAM_NAME_AND_METADATA(InputLinearGain), 1.0f)
 			),
 			FOutputVertexInterface(
-				TOutputDataVertexModel<float>(METASOUND_GET_PARAM_NAME_AND_TT(OutputDecibels))
+				TOutputDataVertexModel<float>(METASOUND_GET_PARAM_NAME_AND_METADATA(OutputDecibels))
 			)
 		);
 
diff --git a/Engine/Plugins/Runtime/Metasound/Source/MetasoundStandardNodes/Private/MetasoundLinearToLogFrequencyNode.cpp b/Engine/Plugins/Runtime/Metasound/Source/MetasoundStandardNodes/Private/MetasoundLinearToLogFrequencyNode.cpp
--- a/Engine/Plugins/Runtime/Metasound/Source/MetasoundStandardNodes/Private/MetasoundLinearToLogFrequencyNode.cpp
+++ b/Engine/Plugins/Runtime/Metasound/Source/MetasoundStandardNodes/Private/MetasoundLinearToLogFrequencyNode.cpp
@@ -61,14 +61,14 @@ namespace Metasound
 
 			static const FVertexInterface Interface(
 				FInputVertexInterface(
-					TInputDataVertexModel<float>(METASOUND_GET_PARAM_NAME_AND_TT(InputValue),     0.0f),
-					TInputDataVertexModel<float>(METASOUND_GET_PARAM_NAME_AND_TT(InputDomainMin), 0.0f),
-					TInputDataVertexModel<float>(METASOUND_GET_PARAM_NAME_AND_TT(InputDomainMax), 1.0f),
-					TInputDataVertexModel<float>(METASOUND_GET_PARAM_NAME_AND_TT(InputRangeMin),  20.0f),
-					TInputDataVertexModel<float>(METASOUND_GET_PARAM_NAME_AND_TT(InputRangeMax),  20000.0f)
+					TInputDataVertexModel<float>(METASOUND_GET_PARAM_NAME_AND_METADATA(InputValue),     0.0f),
+					TInputDataVertexModel<float>(METASOUND_GET_PARAM_NAME_AND_METADATA(InputDomainMin), 0.0f),
+					TInputDataVertexModel<float>(METASOUND_GET_PARAM_NAME_AND_METADATA(InputDomainMax), 1.0f),
+					TInputDataVertexModel<float>(METASOUND_GET_PARAM_NAME_AND_METADATA(InputRangeMin),  20.0f),
+					TInputDataVertexModel<float>(METASOUND_GET_PARAM_NAME_AND_METADATA(InputRangeMax),  20000.0f)
 					),
 				FOutputVertexInterface(
-					TOutputDataVertexModel<float>(METASOUND_GET_PARAM_NAME_AND_TT(OutputFreq))
+					TOutputDataVertexModel<float>(METASOUND_GET_PARAM_NAME_AND_METADATA(OutputFreq))
 				)
 			);
 
diff --git a/Engine/Plugins/Runtime/Metasound/Source/MetasoundStandardNodes/Private/MetasoundMIDIToFreqNode.cpp b/Engine/Plugins/Runtime/Metasound/Source/MetasoundStandardNodes/Private/MetasoundMIDIToFreqNode.cpp
--- a/Engine/Plugins/Runtime/Metasound/Source/MetasoundStandardNodes/Private/MetasoundMIDIToFreqNode.cpp
+++ b/Engine/Plugins/Runtime/Metasound/Source/MetasoundStandardNodes/Private/MetasoundMIDIToFreqNode.cpp
@@ -133,10 +133,10 @@ namespace Metasound
 
 		static const FVertexInterface Interface(
 			FInputVertexInterface(
-				TInputDataVertexModel<ValueType>(METASOUND_GET_PARAM_NAME_AND_TT(InputMidi), (ValueType)60.0f)
+				TInputDataVertexModel<ValueType>(METASOUND_GET_PARAM_NAME_AND_METADATA(InputMidi), (ValueType)60.0f)
 			),
 			FOutputVertexInterface(
-				TOutputDataVertexModel<float>(METASOUND_GET_PARAM_NAME_AND_TT(OutputFreq))
+				TOutputDataVertexModel<float>(METASOUND_GET_PARAM_NAME_AND_METADATA(OutputFreq))
 			)
 		);
 
diff --git a/Engine/Plugins/Runtime/Metasound/Source/MetasoundStandardNodes/Private/MetasoundMapRangeNode.cpp b/Engine/Plugins/Runtime/Metasound/Source/MetasoundStandardNodes/Private/MetasoundMapRangeNode.cpp
--- a/Engine/Plugins/Runtime/Metasound/Source/MetasoundStandardNodes/Private/MetasoundMapRangeNode.cpp
+++ b/Engine/Plugins/Runtime/Metasound/Source/MetasoundStandardNodes/Private/MetasoundMapRangeNode.cpp
@@ -49,15 +49,15 @@ namespace Metasound
 			{
 				static const FVertexInterface DefaultInterface(
 					FInputVertexInterface(
-						TInputDataVertexModel<int32>(METASOUND_GET_PARAM_NAME_AND_TT(InputValueName)),
-						TInputDataVertexModel<int32>(METASOUND_GET_PARAM_NAME_AND_TT(InputInRangeAName), 0),
-						TInputDataVertexModel<int32>(METASOUND_GET_PARAM_NAME_AND_TT(InputInRangeBName), 100),
-						TInputDataVertexModel<int32>(METASOUND_GET_PARAM_NAME_AND_TT(InputOutRangeAName), 0),
-						TInputDataVertexModel<int32>(METASOUND_GET_PARAM_NAME_AND_TT(InputOutRangeBName), 100),
-						TInputDataVertexModel<bool>(METASOUND_GET_PARAM_NAME_AND_TT(InputClampedName), true)
+						TInputDataVertexModel<int32>(METASOUND_GET_PARAM_NAME_AND_METADATA(InputValueName)),
+						TInputDataVertexModel<int32>(METASOUND_GET_PARAM_NAME_AND_METADATA(InputInRangeAName), 0),
+						TInputDataVertexModel<int32>(METASOUND_GET_PARAM_NAME_AND_METADATA(InputInRangeBName), 100),
+						TInputDataVertexModel<int32>(METASOUND_GET_PARAM_NAME_AND_METADATA(InputOutRangeAName), 0),
+						TInputDataVertexModel<int32>(METASOUND_GET_PARAM_NAME_AND_METADATA(InputOutRangeBName), 100),
+						TInputDataVertexModel<bool>(METASOUND_GET_PARAM_NAME_AND_METADATA(InputClampedName), true)
 					),
 					FOutputVertexInterface(
-						TOutputDataVertexModel<int32>(METASOUND_GET_PARAM_NAME_AND_TT(OutputValueName))
+						TOutputDataVertexModel<int32>(METASOUND_GET_PARAM_NAME_AND_METADATA(OutputValueName))
 					)
 				);
 				return DefaultInterface;
@@ -133,15 +133,15 @@ namespace Metasound
 			{
 				static const FVertexInterface DefaultInterface(
 					FInputVertexInterface(
-						TInputDataVertexModel<float>(METASOUND_GET_PARAM_NAME_AND_TT(InputValueName)),
-						TInputDataVertexModel<float>(METASOUND_GET_PARAM_NAME_AND_TT(InputInRangeAName), 0.0f),
-						TInputDataVertexModel<float>(METASOUND_GET_PARAM_NAME_AND_TT(InputInRangeBName), 1.0f),
-						TInputDataVertexModel<float>(METASOUND_GET_PARAM_NAME_AND_TT(InputOutRangeAName), 0.0f),
-						TInputDataVertexModel<float>(METASOUND_GET_PARAM_NAME_AND_TT(InputOutRangeBName), 1.0f),
-						TInputDataVertexModel<bool>(METASOUND_GET_PARAM_NAME_AND_TT(InputClampedName), true)
+						TInputDataVertexModel<float>(METASOUND_GET_PARAM_NAME_AND_METADATA(InputValueName)),
+						TInputDataVertexModel<float>(METASOUND_GET_PARAM_NAME_AND_METADATA(InputInRangeAName), 0.0f),
+						TInputDataVertexModel<float>(METASOUND_GET_PARAM_NAME_AND_METADATA(InputInRangeBName), 1.0f),
+						TInputDataVertexModel<float>(METASOUND_GET_PARAM_NAME_AND_METADATA(InputOutRangeAName), 0.0f),
+						TInputDataVertexModel<float>(METASOUND_GET_PARAM_NAME_AND_METADATA(InputOutRangeBName), 1.0f),
+						TInputDataVertexModel<bool>(METASOUND_GET_PARAM_NAME_AND_METADATA(InputClampedName), true)
 					),
 					FOutputVertexInterface(
-						TOutputDataVertexModel<float>(METASOUND_GET_PARAM_NAME_AND_TT(OutputValueName))
+						TOutputDataVertexModel<float>(METASOUND_GET_PARAM_NAME_AND_METADATA(OutputValueName))
 					)
 				);
 				return DefaultInterface;
@@ -217,15 +217,15 @@ namespace Metasound
 			{
 				static const FVertexInterface DefaultInterface(
 					FInputVertexInterface(
-						TInputDataVertexModel<FAudioBuffer>(METASOUND_GET_PARAM_NAME_AND_TT(InputValueName)),
-						TInputDataVertexModel<float>(METASOUND_GET_PARAM_NAME_AND_TT(InputInRangeAName), -1.0f),
-						TInputDataVertexModel<float>(METASOUND_GET_PARAM_NAME_AND_TT(InputInRangeBName), 1.0f),
-						TInputDataVertexModel<float>(METASOUND_GET_PARAM_NAME_AND_TT(InputOutRangeAName), -1.0f),
-						TInputDataVertexModel<float>(METASOUND_GET_PARAM_NAME_AND_TT(InputOutRangeBName), 1.0f),
-						TInputDataVertexModel<bool>(METASOUND_GET_PARAM_NAME_AND_TT(InputClampedName), true)
+						TInputDataVertexModel<FAudioBuffer>(METASOUND_GET_PARAM_NAME_AND_METADATA(InputValueName)),
+						TInputDataVertexModel<float>(METASOUND_GET_PARAM_NAME_AND_METADATA(InputInRangeAName), -1.0f),
+						TInputDataVertexModel<float>(METASOUND_GET_PARAM_NAME_AND_METADATA(InputInRangeBName), 1.0f),
+						TInputDataVertexModel<float>(METASOUND_GET_PARAM_NAME_AND_METADATA(InputOutRangeAName), -1.0f),
+						TInputDataVertexModel<float>(METASOUND_GET_PARAM_NAME_AND_METADATA(InputOutRangeBName), 1.0f),
+						TInputDataVertexModel<bool>(METASOUND_GET_PARAM_NAME_AND_METADATA(InputClampedName), true)
 					),
 					FOutputVertexInterface(
-						TOutputDataVertexModel<FAudioBuffer>(METASOUND_GET_PARAM_NAME_AND_TT(OutputValueName))
+						TOutputDataVertexModel<FAudioBuffer>(METASOUND_GET_PARAM_NAME_AND_METADATA(OutputValueName))
 					)
 				);
 
diff --git a/Engine/Plugins/Runtime/Metasound/Source/MetasoundStandardNodes/Private/MetasoundMaxNode.cpp b/Engine/Plugins/Runtime/Metasound/Source/MetasoundStandardNodes/Private/MetasoundMaxNode.cpp
--- a/Engine/Plugins/Runtime/Metasound/Source/MetasoundStandardNodes/Private/MetasoundMaxNode.cpp
+++ b/Engine/Plugins/Runtime/Metasound/Source/MetasoundStandardNodes/Private/MetasoundMaxNode.cpp
@@ -139,11 +139,11 @@ namespace Metasound
 
 			static const FVertexInterface DefaultInterface(
 				FInputVertexInterface(
-					TInputDataVertexModel<ValueType>(METASOUND_GET_PARAM_NAME_AND_TT(InputAValue)),
-					TInputDataVertexModel<ValueType>(METASOUND_GET_PARAM_NAME_AND_TT(InputBValue))
+					TInputDataVertexModel<ValueType>(METASOUND_GET_PARAM_NAME_AND_METADATA(InputAValue)),
+					TInputDataVertexModel<ValueType>(METASOUND_GET_PARAM_NAME_AND_METADATA(InputBValue))
 				),
 				FOutputVertexInterface(
-					TOutputDataVertexModel<ValueType>(METASOUND_GET_PARAM_NAME_AND_TT(OutputValue))
+					TOutputDataVertexModel<ValueType>(METASOUND_GET_PARAM_NAME_AND_METADATA(OutputValue))
 				)
 			);
 
diff --git a/Engine/Plugins/Runtime/Metasound/Source/MetasoundStandardNodes/Private/MetasoundMidSideNodes.cpp b/Engine/Plugins/Runtime/Metasound/Source/MetasoundStandardNodes/Private/MetasoundMidSideNodes.cpp
--- a/Engine/Plugins/Runtime/Metasound/Source/MetasoundStandardNodes/Private/MetasoundMidSideNodes.cpp
+++ b/Engine/Plugins/Runtime/Metasound/Source/MetasoundStandardNodes/Private/MetasoundMidSideNodes.cpp
@@ -87,14 +87,14 @@ namespace Metasound
 
 			static const FVertexInterface Interface(
 				FInputVertexInterface(
-					TInputDataVertexModel<FAudioBuffer>(METASOUND_GET_PARAM_NAME_AND_TT(InputAudioLeft)),
-					TInputDataVertexModel<FAudioBuffer>(METASOUND_GET_PARAM_NAME_AND_TT(InputAudioRight)),
-					TInputDataVertexModel<float>(METASOUND_GET_PARAM_NAME_AND_TT(InputSpreadAmount), 0.5f),
-					TInputDataVertexModel<bool>(METASOUND_GET_PARAM_NAME_AND_TT(InputEqualPower), false)
+					TInputDataVertexModel<FAudioBuffer>(METASOUND_GET_PARAM_NAME_AND_METADATA(InputAudioLeft)),
+					TInputDataVertexModel<FAudioBuffer>(METASOUND_GET_PARAM_NAME_AND_METADATA(InputAudioRight)),
+					TInputDataVertexModel<float>(METASOUND_GET_PARAM_NAME_AND_METADATA(InputSpreadAmount), 0.5f),
+					TInputDataVertexModel<bool>(METASOUND_GET_PARAM_NAME_AND_METADATA(InputEqualPower), false)
 				),
 				FOutputVertexInterface(
-					TOutputDataVertexModel<FAudioBuffer>(METASOUND_GET_PARAM_NAME_AND_TT(OutputAudioMid)),
-					TOutputDataVertexModel<FAudioBuffer>(METASOUND_GET_PARAM_NAME_AND_TT(OutputAudioSide))
+					TOutputDataVertexModel<FAudioBuffer>(METASOUND_GET_PARAM_NAME_AND_METADATA(OutputAudioMid)),
+					TOutputDataVertexModel<FAudioBuffer>(METASOUND_GET_PARAM_NAME_AND_METADATA(OutputAudioSide))
 				)
 			);
 
@@ -311,14 +311,14 @@ namespace Metasound
 
 			static const FVertexInterface Interface(
 				FInputVertexInterface(
-					TInputDataVertexModel<FAudioBuffer>(METASOUND_GET_PARAM_NAME_AND_TT(InputAudioMid)),
-					TInputDataVertexModel<FAudioBuffer>(METASOUND_GET_PARAM_NAME_AND_TT(InputAudioSide)),
-					TInputDataVertexModel<float>(METASOUND_GET_PARAM_NAME_AND_TT(InputSpreadAmount), 0.5f),
-					TInputDataVertexModel<bool>(METASOUND_GET_PARAM_NAME_AND_TT(InputEqualPower), false)
+					TInputDataVertexModel<FAudioBuffer>(METASOUND_GET_PARAM_NAME_AND_METADATA(InputAudioMid)),
+					TInputDataVertexModel<FAudioBuffer>(METASOUND_GET_PARAM_NAME_AND_METADATA(InputAudioSide)),
+					TInputDataVertexModel<float>(METASOUND_GET_PARAM_NAME_AND_METADATA(InputSpreadAmount), 0.5f),
+					TInputDataVertexModel<bool>(METASOUND_GET_PARAM_NAME_AND_METADATA(InputEqualPower), false)
 				),
 				FOutputVertexInterface(
-					TOutputDataVertexModel<FAudioBuffer>(METASOUND_GET_PARAM_NAME_AND_TT(OutputAudioLeft)),
-					TOutputDataVertexModel<FAudioBuffer>(METASOUND_GET_PARAM_NAME_AND_TT(OutputAudioRight))
+					TOutputDataVertexModel<FAudioBuffer>(METASOUND_GET_PARAM_NAME_AND_METADATA(OutputAudioLeft)),
+					TOutputDataVertexModel<FAudioBuffer>(METASOUND_GET_PARAM_NAME_AND_METADATA(OutputAudioRight))
 				)
 			);
 
diff --git a/Engine/Plugins/Runtime/Metasound/Source/MetasoundStandardNodes/Private/MetasoundMidiNoteQuantizerNode.cpp b/Engine/Plugins/Runtime/Metasound/Source/MetasoundStandardNodes/Private/MetasoundMidiNoteQuantizerNode.cpp
--- a/Engine/Plugins/Runtime/Metasound/Source/MetasoundStandardNodes/Private/MetasoundMidiNoteQuantizerNode.cpp
+++ b/Engine/Plugins/Runtime/Metasound/Source/MetasoundStandardNodes/Private/MetasoundMidiNoteQuantizerNode.cpp
@@ -133,12 +133,12 @@ namespace Metasound
 	{
 		static const FVertexInterface Interface(
 			FInputVertexInterface(
-				TInputDataVertexModel<float>(METASOUND_GET_PARAM_NAME(ParamNoteIn), METASOUND_GET_PARAM_TT(ParamNoteIn), 60.0f),
-				TInputDataVertexModel<float>(METASOUND_GET_PARAM_NAME(ParamRootNote), METASOUND_GET_PARAM_TT(ParamRootNote), 0.0f),
-				TInputDataVertexModel<ScaleDegreeArrayType>(METASOUND_GET_PARAM_NAME(ParamScaleDegrees), METASOUND_GET_PARAM_TT(ParamScaleDegrees))
+				TInputDataVertexModel<float>(METASOUND_GET_PARAM_NAME_AND_METADATA(ParamNoteIn), 60.0f),
+				TInputDataVertexModel<float>(METASOUND_GET_PARAM_NAME_AND_METADATA(ParamRootNote), 0.0f),
+				TInputDataVertexModel<ScaleDegreeArrayType>(METASOUND_GET_PARAM_NAME_AND_METADATA(ParamScaleDegrees))
 			),
 			FOutputVertexInterface(
-				TOutputDataVertexModel<float>(METASOUND_GET_PARAM_NAME(ParamNoteOutput), METASOUND_GET_PARAM_TT(ParamNoteOutput))
+				TOutputDataVertexModel<float>(METASOUND_GET_PARAM_NAME_AND_METADATA(ParamNoteOutput))
 				)
 			);
 
diff --git a/Engine/Plugins/Runtime/Metasound/Source/MetasoundStandardNodes/Private/MetasoundMidiScaleToArrayNode.cpp b/Engine/Plugins/Runtime/Metasound/Source/MetasoundStandardNodes/Private/MetasoundMidiScaleToArrayNode.cpp
--- a/Engine/Plugins/Runtime/Metasound/Source/MetasoundStandardNodes/Private/MetasoundMidiScaleToArrayNode.cpp
+++ b/Engine/Plugins/Runtime/Metasound/Source/MetasoundStandardNodes/Private/MetasoundMidiScaleToArrayNode.cpp
@@ -127,11 +127,11 @@ namespace Metasound
 	{
 		static const FVertexInterface Interface(
 			FInputVertexInterface(
-				TInputDataVertexModel<FEnumEMusicalScale>(METASOUND_GET_PARAM_NAME(ParamScaleDegreesPreset), METASOUND_GET_PARAM_TT(ParamScaleDegreesPreset)),
-				TInputDataVertexModel<bool>(METASOUND_GET_PARAM_NAME(ParamChordTonesOnly), METASOUND_GET_PARAM_TT(ParamChordTonesOnly))
+				TInputDataVertexModel<FEnumEMusicalScale>(METASOUND_GET_PARAM_NAME_AND_METADATA(ParamScaleDegreesPreset)),
+				TInputDataVertexModel<bool>(METASOUND_GET_PARAM_NAME_AND_METADATA(ParamChordTonesOnly))
 			),
 			FOutputVertexInterface(
-				TOutputDataVertexModel<ScaleDegreeArrayType>(METASOUND_GET_PARAM_NAME(ParamNoteArrayOutput), METASOUND_GET_PARAM_TT(ParamNoteArrayOutput))
+				TOutputDataVertexModel<ScaleDegreeArrayType>(METASOUND_GET_PARAM_NAME_AND_METADATA(ParamNoteArrayOutput))
 				)
 			);
 
diff --git a/Engine/Plugins/Runtime/Metasound/Source/MetasoundStandardNodes/Private/MetasoundMinNode.cpp b/Engine/Plugins/Runtime/Metasound/Source/MetasoundStandardNodes/Private/MetasoundMinNode.cpp
--- a/Engine/Plugins/Runtime/Metasound/Source/MetasoundStandardNodes/Private/MetasoundMinNode.cpp
+++ b/Engine/Plugins/Runtime/Metasound/Source/MetasoundStandardNodes/Private/MetasoundMinNode.cpp
@@ -138,11 +138,11 @@ namespace Metasound
 
 			static const FVertexInterface DefaultInterface(
 				FInputVertexInterface(
-					TInputDataVertexModel<ValueType>(METASOUND_GET_PARAM_NAME_AND_TT(InputAValue)),
-					TInputDataVertexModel<ValueType>(METASOUND_GET_PARAM_NAME_AND_TT(InputBValue))
+					TInputDataVertexModel<ValueType>(METASOUND_GET_PARAM_NAME_AND_METADATA(InputAValue)),
+					TInputDataVertexModel<ValueType>(METASOUND_GET_PARAM_NAME_AND_METADATA(InputBValue))
 				),
 				FOutputVertexInterface(
-					TOutputDataVertexModel<ValueType>(METASOUND_GET_PARAM_NAME_AND_TT(OutputValue))
+					TOutputDataVertexModel<ValueType>(METASOUND_GET_PARAM_NAME_AND_METADATA(OutputValue))
 				)
 			);
 
diff --git a/Engine/Plugins/Runtime/Metasound/Source/MetasoundStandardNodes/Private/MetasoundNoiseGenerator.cpp b/Engine/Plugins/Runtime/Metasound/Source/MetasoundStandardNodes/Private/MetasoundNoiseGenerator.cpp
--- a/Engine/Plugins/Runtime/Metasound/Source/MetasoundStandardNodes/Private/MetasoundNoiseGenerator.cpp
+++ b/Engine/Plugins/Runtime/Metasound/Source/MetasoundStandardNodes/Private/MetasoundNoiseGenerator.cpp
@@ -8,6 +8,7 @@
 #include "MetasoundNodeRegistrationMacro.h"
 #include "MetasoundDataTypeRegistrationMacro.h"
 #include "MetasoundOperatorSettings.h"
+#include "MetasoundParamHelper.h"
 #include "MetasoundPrimitives.h"
 #include "MetasoundStandardNodesNames.h"
 #include "MetasoundTrigger.h"
@@ -33,6 +34,13 @@ namespace Metasound
 		DEFINE_METASOUND_ENUM_ENTRY(ENoiseType::White, "WhiteDescription", "White Noise", "WhiteDescriptionTT", "A random signal having equal intensity at different frequencies"),
 	DEFINE_METASOUND_ENUM_END()
 
+	namespace NoiseGeneratorVertexNames
+	{
+		METASOUND_PARAM(OutAudio, "Audio", "Audio output.")
+		METASOUND_PARAM(InputType, "Type", "Type of Noise to Generate.")
+		METASOUND_PARAM(InputSeed, "Seed", "Seed for seeding the Random Number Generator, -1 (default) will use current time.")
+	}
+
 	class FNoiseOperator : public IOperator
 	{
 	public:
@@ -86,10 +94,6 @@ namespace Metasound
 				*WritePtr++ = InGenerator.Generate();
 			}
 		}
-
-		static constexpr const TCHAR* OutPinName = TEXT("Audio");
-		static constexpr const TCHAR* TriggerTypePinName = TEXT("Type");
-		static constexpr const TCHAR* SeedPinName = TEXT("Seed");
 	};
 
 	constexpr int32 FNoiseOperator::DefaultSeed;
@@ -146,27 +150,33 @@ namespace Metasound
 
 	FDataReferenceCollection FNoiseOperator::GetInputs() const
 	{
+		using namespace NoiseGeneratorVertexNames;
+
 		FDataReferenceCollection InputDataReferences;
-		InputDataReferences.AddDataReadReference(SeedPinName, FInt32ReadRef(Seed));
+		InputDataReferences.AddDataReadReference(METASOUND_GET_PARAM_NAME(InputSeed), FInt32ReadRef(Seed));
 		return InputDataReferences;
 	}
 
 	FDataReferenceCollection FNoiseOperator::GetOutputs() const
 	{
+		using namespace NoiseGeneratorVertexNames;
+
 		FDataReferenceCollection OutputDataReferences;
-		OutputDataReferences.AddDataReadReference(OutPinName, FAudioBufferReadRef(Out));
+		OutputDataReferences.AddDataReadReference(METASOUND_GET_PARAM_NAME(OutAudio), FAudioBufferReadRef(Out));
 		return OutputDataReferences;
 	}
 
 	FVertexInterface FNoiseOperator::DeclareVertexInterface()
 	{
+		using namespace NoiseGeneratorVertexNames; 
+
 		static const FVertexInterface Interface(
 			FInputVertexInterface(
-				TInputDataVertexModel<int32>(SeedPinName, METASOUND_LOCTEXT("SeedPinDescription", "Seed for seeding the Random Number Generator, -1 (default) will use current time"), FNoiseOperator::DefaultSeed),
-				TInputDataVertexModel<FEnumNoiseType>(TriggerTypePinName, METASOUND_LOCTEXT("NoiseTypeDescription", "Type of Noise to Generate"))
+				TInputDataVertexModel<int32>(METASOUND_GET_PARAM_NAME_AND_METADATA(InputSeed), FNoiseOperator::DefaultSeed),
+				TInputDataVertexModel<FEnumNoiseType>(METASOUND_GET_PARAM_NAME_AND_METADATA(InputType))
 			),
 			FOutputVertexInterface(
-				TOutputDataVertexModel<FAudioBuffer>(OutPinName, METASOUND_LOCTEXT("AudioOutDescription", "Audio output"))
+				TOutputDataVertexModel<FAudioBuffer>(METASOUND_GET_PARAM_NAME_AND_METADATA(OutAudio))
 			)
 		);
 		return Interface;
@@ -203,16 +213,18 @@ namespace Metasound
 
 	TUniquePtr<IOperator> FNoiseOperator::CreateOperator(const FCreateOperatorParams& InParams, FBuildErrorArray& OutErrors)
 	{
+		using namespace NoiseGeneratorVertexNames; 
+
 		const FNoiseNode& Node = static_cast<const FNoiseNode&>(InParams.Node);
 		const FDataReferenceCollection& InputCol = InParams.InputDataReferences;
 		const FOperatorSettings& Settings = InParams.OperatorSettings;
 		const FInputVertexInterface& InputInterface = DeclareVertexInterface().GetInputInterface();
 
 		// Static property pin, only used for factory.
-		FEnumNoiseTypeReadRef Type = InputCol.GetDataReadReferenceOrConstruct<FEnumNoiseType>(TriggerTypePinName);
+		FEnumNoiseTypeReadRef Type = InputCol.GetDataReadReferenceOrConstruct<FEnumNoiseType>(METASOUND_GET_PARAM_NAME(InputType));
 
 		// Seed.
-		FInt32ReadRef Seed = InputCol.GetDataReadReferenceOrConstruct<int32>(TriggerTypePinName, Node.GetDefaultSeed());
+		FInt32ReadRef Seed = InputCol.GetDataReadReferenceOrConstruct<int32>(METASOUND_GET_PARAM_NAME(InputSeed), Node.GetDefaultSeed());
 
 		switch (*Type)
 		{
diff --git a/Engine/Plugins/Runtime/Metasound/Source/MetasoundStandardNodes/Private/MetasoundSampleAndHoldNode.cpp b/Engine/Plugins/Runtime/Metasound/Source/MetasoundStandardNodes/Private/MetasoundSampleAndHoldNode.cpp
--- a/Engine/Plugins/Runtime/Metasound/Source/MetasoundStandardNodes/Private/MetasoundSampleAndHoldNode.cpp
+++ b/Engine/Plugins/Runtime/Metasound/Source/MetasoundStandardNodes/Private/MetasoundSampleAndHoldNode.cpp
@@ -64,12 +64,12 @@ namespace Metasound
 
 			static const FVertexInterface Interface(
 				FInputVertexInterface(
-					TInputDataVertexModel<FTrigger>(METASOUND_GET_PARAM_NAME_AND_TT(InputTriggerSampleAndHold)),
-					TInputDataVertexModel<FAudioBuffer>(METASOUND_GET_PARAM_NAME_AND_TT(InputAudio))
+					TInputDataVertexModel<FTrigger>(METASOUND_GET_PARAM_NAME_AND_METADATA(InputTriggerSampleAndHold)),
+					TInputDataVertexModel<FAudioBuffer>(METASOUND_GET_PARAM_NAME_AND_METADATA(InputAudio))
 				),
 				FOutputVertexInterface(
-					TOutputDataVertexModel<FTrigger>(METASOUND_GET_PARAM_NAME_AND_TT(OutputOnSampleAndHold)),
-					TOutputDataVertexModel<FAudioBuffer>(METASOUND_GET_PARAM_NAME_AND_TT(OutputAudio))
+					TOutputDataVertexModel<FTrigger>(METASOUND_GET_PARAM_NAME_AND_METADATA(OutputOnSampleAndHold)),
+					TOutputDataVertexModel<FAudioBuffer>(METASOUND_GET_PARAM_NAME_AND_METADATA(OutputAudio))
 				)
 			);
 
diff --git a/Engine/Plugins/Runtime/Metasound/Source/MetasoundStandardNodes/Private/MetasoundSemitoneToFreqMultiplierNode.cpp b/Engine/Plugins/Runtime/Metasound/Source/MetasoundStandardNodes/Private/MetasoundSemitoneToFreqMultiplierNode.cpp
--- a/Engine/Plugins/Runtime/Metasound/Source/MetasoundStandardNodes/Private/MetasoundSemitoneToFreqMultiplierNode.cpp
+++ b/Engine/Plugins/Runtime/Metasound/Source/MetasoundStandardNodes/Private/MetasoundSemitoneToFreqMultiplierNode.cpp
@@ -93,10 +93,10 @@ namespace Metasound
 
 		static const FVertexInterface Interface(
 			FInputVertexInterface(
-				TInputDataVertexModel<float>(METASOUND_GET_PARAM_NAME_AND_TT(InputSemitone), 0.0f)
+				TInputDataVertexModel<float>(METASOUND_GET_PARAM_NAME_AND_METADATA(InputSemitone), 0.0f)
 			),
 			FOutputVertexInterface(
-				TOutputDataVertexModel<float>(METASOUND_GET_PARAM_NAME_AND_TT(OutputFrequencyMultiplier))
+				TOutputDataVertexModel<float>(METASOUND_GET_PARAM_NAME_AND_METADATA(OutputFrequencyMultiplier))
 			)
 		);
 
diff --git a/Engine/Plugins/Runtime/Metasound/Source/MetasoundStandardNodes/Private/MetasoundTriggerCoinNode.cpp b/Engine/Plugins/Runtime/Metasound/Source/MetasoundStandardNodes/Private/MetasoundTriggerCoinNode.cpp
--- a/Engine/Plugins/Runtime/Metasound/Source/MetasoundStandardNodes/Private/MetasoundTriggerCoinNode.cpp
+++ b/Engine/Plugins/Runtime/Metasound/Source/MetasoundStandardNodes/Private/MetasoundTriggerCoinNode.cpp
@@ -177,14 +177,14 @@ namespace Metasound
 
 		static const FVertexInterface Interface(
 			FInputVertexInterface(
-				TInputDataVertexModel<FTrigger>(METASOUND_GET_PARAM_NAME_AND_TT(InputTrigger)),
-				TInputDataVertexModel<FTrigger>(METASOUND_GET_PARAM_NAME_AND_TT(InputReset)),
-				TInputDataVertexModel<int32>(METASOUND_GET_PARAM_NAME_AND_TT(InputSeed), -1),
-				TInputDataVertexModel<float>(METASOUND_GET_PARAM_NAME_AND_TT(InputProbability), 1.0f)
+				TInputDataVertexModel<FTrigger>(METASOUND_GET_PARAM_NAME_AND_METADATA(InputTrigger)),
+				TInputDataVertexModel<FTrigger>(METASOUND_GET_PARAM_NAME_AND_METADATA(InputReset)),
+				TInputDataVertexModel<int32>(METASOUND_GET_PARAM_NAME_AND_METADATA(InputSeed), -1),
+				TInputDataVertexModel<float>(METASOUND_GET_PARAM_NAME_AND_METADATA(InputProbability), 1.0f)
 			),
 			FOutputVertexInterface(
-				TOutputDataVertexModel<FTrigger>(METASOUND_GET_PARAM_NAME_AND_TT(OutputTrueTrigger)),
-				TOutputDataVertexModel<FTrigger>(METASOUND_GET_PARAM_NAME_AND_TT(OutputFalseTrigger))
+				TOutputDataVertexModel<FTrigger>(METASOUND_GET_PARAM_NAME_AND_METADATA(OutputTrueTrigger)),
+				TOutputDataVertexModel<FTrigger>(METASOUND_GET_PARAM_NAME_AND_METADATA(OutputFalseTrigger))
 			)
 		);
 
diff --git a/Engine/Plugins/Runtime/Metasound/Source/MetasoundStandardNodes/Private/MetasoundTriggerControlNode.cpp b/Engine/Plugins/Runtime/Metasound/Source/MetasoundStandardNodes/Private/MetasoundTriggerControlNode.cpp
--- a/Engine/Plugins/Runtime/Metasound/Source/MetasoundStandardNodes/Private/MetasoundTriggerControlNode.cpp
+++ b/Engine/Plugins/Runtime/Metasound/Source/MetasoundStandardNodes/Private/MetasoundTriggerControlNode.cpp
@@ -182,14 +182,14 @@ namespace Metasound
 
 		static const FVertexInterface Interface(
 			FInputVertexInterface(
-				TInputDataVertexModel<FTrigger>(METASOUND_GET_PARAM_NAME_AND_TT(InputEnter)),
-				TInputDataVertexModel<FTrigger>(METASOUND_GET_PARAM_NAME_AND_TT(InputOpen)),
-				TInputDataVertexModel<FTrigger>(METASOUND_GET_PARAM_NAME_AND_TT(InputClose)),
-				TInputDataVertexModel<FTrigger>(METASOUND_GET_PARAM_NAME_AND_TT(InputToggle)),
-				TInputDataVertexModel<bool>(METASOUND_GET_PARAM_NAME_AND_TT(InputStartClosed), true)
+				TInputDataVertexModel<FTrigger>(METASOUND_GET_PARAM_NAME_AND_METADATA(InputEnter)),
+				TInputDataVertexModel<FTrigger>(METASOUND_GET_PARAM_NAME_AND_METADATA(InputOpen)),
+				TInputDataVertexModel<FTrigger>(METASOUND_GET_PARAM_NAME_AND_METADATA(InputClose)),
+				TInputDataVertexModel<FTrigger>(METASOUND_GET_PARAM_NAME_AND_METADATA(InputToggle)),
+				TInputDataVertexModel<bool>(METASOUND_GET_PARAM_NAME_AND_METADATA(InputStartClosed), true)
 			),
 			FOutputVertexInterface(
-				TOutputDataVertexModel<FTrigger>(METASOUND_GET_PARAM_NAME_AND_TT(OutputExit))
+				TOutputDataVertexModel<FTrigger>(METASOUND_GET_PARAM_NAME_AND_METADATA(OutputExit))
 			)
 		);
 
diff --git a/Engine/Plugins/Runtime/Metasound/Source/MetasoundStandardNodes/Private/MetasoundTriggerCounterNode.cpp b/Engine/Plugins/Runtime/Metasound/Source/MetasoundStandardNodes/Private/MetasoundTriggerCounterNode.cpp
--- a/Engine/Plugins/Runtime/Metasound/Source/MetasoundStandardNodes/Private/MetasoundTriggerCounterNode.cpp
+++ b/Engine/Plugins/Runtime/Metasound/Source/MetasoundStandardNodes/Private/MetasoundTriggerCounterNode.cpp
@@ -191,17 +191,17 @@ namespace Metasound
 
 		static const FVertexInterface Interface(
 			FInputVertexInterface(
-				TInputDataVertexModel<FTrigger>(METASOUND_GET_PARAM_NAME_AND_TT(InputInTrigger)),
-				TInputDataVertexModel<FTrigger>(METASOUND_GET_PARAM_NAME_AND_TT(InputReset)),
-				TInputDataVertexModel<float>(METASOUND_GET_PARAM_NAME_AND_TT(InputStartValue), 0.0f),
-				TInputDataVertexModel<float>(METASOUND_GET_PARAM_NAME_AND_TT(InputStepSize), 1.0f),
-				TInputDataVertexModel<int32>(METASOUND_GET_PARAM_NAME_AND_TT(InputAutoResetCount), 0)
+				TInputDataVertexModel<FTrigger>(METASOUND_GET_PARAM_NAME_AND_METADATA(InputInTrigger)),
+				TInputDataVertexModel<FTrigger>(METASOUND_GET_PARAM_NAME_AND_METADATA(InputReset)),
+				TInputDataVertexModel<float>(METASOUND_GET_PARAM_NAME_AND_METADATA(InputStartValue), 0.0f),
+				TInputDataVertexModel<float>(METASOUND_GET_PARAM_NAME_AND_METADATA(InputStepSize), 1.0f),
+				TInputDataVertexModel<int32>(METASOUND_GET_PARAM_NAME_AND_METADATA(InputAutoResetCount), 0)
 			),
 			FOutputVertexInterface(
-				TOutputDataVertexModel<FTrigger>(METASOUND_GET_PARAM_NAME_AND_TT(OutputOnTrigger)),
-				TOutputDataVertexModel<FTrigger>(METASOUND_GET_PARAM_NAME_AND_TT(OutputOnReset)),
-				TOutputDataVertexModel<int32>(METASOUND_GET_PARAM_NAME_AND_TT(OutputCount)),
-				TOutputDataVertexModel<float>(METASOUND_GET_PARAM_NAME_AND_TT(OutputValue))
+				TOutputDataVertexModel<FTrigger>(METASOUND_GET_PARAM_NAME_AND_METADATA(OutputOnTrigger)),
+				TOutputDataVertexModel<FTrigger>(METASOUND_GET_PARAM_NAME_AND_METADATA(OutputOnReset)),
+				TOutputDataVertexModel<int32>(METASOUND_GET_PARAM_NAME_AND_METADATA(OutputCount)),
+				TOutputDataVertexModel<float>(METASOUND_GET_PARAM_NAME_AND_METADATA(OutputValue))
 			)
 		);
 
diff --git a/Engine/Plugins/Runtime/Metasound/Source/MetasoundStandardNodes/Private/MetasoundTriggerDelayNode.cpp b/Engine/Plugins/Runtime/Metasound/Source/MetasoundStandardNodes/Private/MetasoundTriggerDelayNode.cpp
--- a/Engine/Plugins/Runtime/Metasound/Source/MetasoundStandardNodes/Private/MetasoundTriggerDelayNode.cpp
+++ b/Engine/Plugins/Runtime/Metasound/Source/MetasoundStandardNodes/Private/MetasoundTriggerDelayNode.cpp
@@ -155,12 +155,12 @@ namespace Metasound
 
 		static const FVertexInterface Interface(
 			FInputVertexInterface(
-				TInputDataVertexModel<FTrigger>(METASOUND_GET_PARAM_NAME_AND_TT(InputInTriggerDelay)),
-				TInputDataVertexModel<FTrigger>(METASOUND_GET_PARAM_NAME_AND_TT(InputResetDelay)),
-				TInputDataVertexModel<FTime>(METASOUND_GET_PARAM_NAME_AND_TT(InputDelayTime), 1.0f)
+				TInputDataVertexModel<FTrigger>(METASOUND_GET_PARAM_NAME_AND_METADATA(InputInTriggerDelay)),
+				TInputDataVertexModel<FTrigger>(METASOUND_GET_PARAM_NAME_AND_METADATA(InputResetDelay)),
+				TInputDataVertexModel<FTime>(METASOUND_GET_PARAM_NAME_AND_METADATA(InputDelayTime), 1.0f)
 			),
 			FOutputVertexInterface(
-				TOutputDataVertexModel<FTrigger>(METASOUND_GET_PARAM_NAME_AND_TT(OutputOnTrigger))
+				TOutputDataVertexModel<FTrigger>(METASOUND_GET_PARAM_NAME_AND_METADATA(OutputOnTrigger))
 		)
 		);
 
diff --git a/Engine/Plugins/Runtime/Metasound/Source/MetasoundStandardNodes/Private/MetasoundTriggerOnValueChangeNode.cpp b/Engine/Plugins/Runtime/Metasound/Source/MetasoundStandardNodes/Private/MetasoundTriggerOnValueChangeNode.cpp
--- a/Engine/Plugins/Runtime/Metasound/Source/MetasoundStandardNodes/Private/MetasoundTriggerOnValueChangeNode.cpp
+++ b/Engine/Plugins/Runtime/Metasound/Source/MetasoundStandardNodes/Private/MetasoundTriggerOnValueChangeNode.cpp
@@ -132,10 +132,10 @@ namespace Metasound
 
 		static const FVertexInterface Interface(
 			FInputVertexInterface(
-				TInputDataVertexModel<ValueType>(METASOUND_GET_PARAM_NAME_AND_TT(InputValue), 1)
+				TInputDataVertexModel<ValueType>(METASOUND_GET_PARAM_NAME_AND_METADATA(InputValue), 1)
 			),
 			FOutputVertexInterface(
-				TOutputDataVertexModel<FTrigger>(METASOUND_GET_PARAM_NAME_AND_TT(OutputOnChange))
+				TOutputDataVertexModel<FTrigger>(METASOUND_GET_PARAM_NAME_AND_METADATA(OutputOnChange))
 			)
 		);
 
diff --git a/Engine/Plugins/Runtime/Metasound/Source/MetasoundStandardNodes/Private/MetasoundTriggerOnceNode.cpp b/Engine/Plugins/Runtime/Metasound/Source/MetasoundStandardNodes/Private/MetasoundTriggerOnceNode.cpp
--- a/Engine/Plugins/Runtime/Metasound/Source/MetasoundStandardNodes/Private/MetasoundTriggerOnceNode.cpp
+++ b/Engine/Plugins/Runtime/Metasound/Source/MetasoundStandardNodes/Private/MetasoundTriggerOnceNode.cpp
@@ -146,12 +146,12 @@ namespace Metasound
 
 		static const FVertexInterface Interface(
 			FInputVertexInterface(
-				TInputDataVertexModel<FTrigger>(METASOUND_GET_PARAM_NAME_AND_TT(InputEnter)),
-				TInputDataVertexModel<FTrigger>(METASOUND_GET_PARAM_NAME_AND_TT(InputReset)),
-				TInputDataVertexModel<bool>(METASOUND_GET_PARAM_NAME_AND_TT(InputStartClosed), false)
+				TInputDataVertexModel<FTrigger>(METASOUND_GET_PARAM_NAME_AND_METADATA(InputEnter)),
+				TInputDataVertexModel<FTrigger>(METASOUND_GET_PARAM_NAME_AND_METADATA(InputReset)),
+				TInputDataVertexModel<bool>(METASOUND_GET_PARAM_NAME_AND_METADATA(InputStartClosed), false)
 			),
 			FOutputVertexInterface(
-				TOutputDataVertexModel<FTrigger>(METASOUND_GET_PARAM_NAME_AND_TT(OutputExit))
+				TOutputDataVertexModel<FTrigger>(METASOUND_GET_PARAM_NAME_AND_METADATA(OutputExit))
 			)
 		);
 
diff --git a/Engine/Plugins/Runtime/Metasound/Source/MetasoundStandardNodes/Private/MetasoundTriggerPipeNode.cpp b/Engine/Plugins/Runtime/Metasound/Source/MetasoundStandardNodes/Private/MetasoundTriggerPipeNode.cpp
--- a/Engine/Plugins/Runtime/Metasound/Source/MetasoundStandardNodes/Private/MetasoundTriggerPipeNode.cpp
+++ b/Engine/Plugins/Runtime/Metasound/Source/MetasoundStandardNodes/Private/MetasoundTriggerPipeNode.cpp
@@ -155,12 +155,12 @@ namespace Metasound
 
 		static const FVertexInterface Interface(
 			FInputVertexInterface(
-				TInputDataVertexModel<FTrigger>(METASOUND_GET_PARAM_NAME_AND_TT(InputInTrigger)),
-				TInputDataVertexModel<FTrigger>(METASOUND_GET_PARAM_NAME_AND_TT(InputReset)),
-				TInputDataVertexModel<FTime>(METASOUND_GET_PARAM_NAME_AND_TT(InputDelayTime), 1.0f)
+				TInputDataVertexModel<FTrigger>(METASOUND_GET_PARAM_NAME_AND_METADATA(InputInTrigger)),
+				TInputDataVertexModel<FTrigger>(METASOUND_GET_PARAM_NAME_AND_METADATA(InputReset)),
+				TInputDataVertexModel<FTime>(METASOUND_GET_PARAM_NAME_AND_METADATA(InputDelayTime), 1.0f)
 				),
 			FOutputVertexInterface(
-				TOutputDataVertexModel<FTrigger>(METASOUND_GET_PARAM_NAME_AND_TT(OutputOutTrigger))
+				TOutputDataVertexModel<FTrigger>(METASOUND_GET_PARAM_NAME_AND_METADATA(OutputOutTrigger))
 			)
 		);
 
diff --git a/Engine/Plugins/Runtime/Metasound/Source/MetasoundStandardNodes/Private/MetasoundTriggerRepeatNode.cpp b/Engine/Plugins/Runtime/Metasound/Source/MetasoundStandardNodes/Private/MetasoundTriggerRepeatNode.cpp
--- a/Engine/Plugins/Runtime/Metasound/Source/MetasoundStandardNodes/Private/MetasoundTriggerRepeatNode.cpp
+++ b/Engine/Plugins/Runtime/Metasound/Source/MetasoundStandardNodes/Private/MetasoundTriggerRepeatNode.cpp
@@ -152,12 +152,12 @@ namespace Metasound
 
 		static const FVertexInterface Interface(
 			FInputVertexInterface(
-				TInputDataVertexModel<FTrigger>(METASOUND_GET_PARAM_NAME_AND_TT(InputStart)),
-				TInputDataVertexModel<FTrigger>(METASOUND_GET_PARAM_NAME_AND_TT(InputStop)),
-				TInputDataVertexModel<FTime>(METASOUND_GET_PARAM_NAME_AND_TT(InputPeriod), 0.2f)
+				TInputDataVertexModel<FTrigger>(METASOUND_GET_PARAM_NAME_AND_METADATA(InputStart)),
+				TInputDataVertexModel<FTrigger>(METASOUND_GET_PARAM_NAME_AND_METADATA(InputStop)),
+				TInputDataVertexModel<FTime>(METASOUND_GET_PARAM_NAME_AND_METADATA(InputPeriod), 0.2f)
 			),
 			FOutputVertexInterface(
-				TOutputDataVertexModel<FTrigger>(METASOUND_GET_PARAM_NAME_AND_TT(RepeatOutputOnTrigger))
+				TOutputDataVertexModel<FTrigger>(METASOUND_GET_PARAM_NAME_AND_METADATA(RepeatOutputOnTrigger))
 			)
 		);
 
diff --git a/Engine/Plugins/Runtime/Metasound/Source/MetasoundStandardNodes/Private/MetasoundTriggerSequenceNode.cpp b/Engine/Plugins/Runtime/Metasound/Source/MetasoundStandardNodes/Private/MetasoundTriggerSequenceNode.cpp
--- a/Engine/Plugins/Runtime/Metasound/Source/MetasoundStandardNodes/Private/MetasoundTriggerSequenceNode.cpp
+++ b/Engine/Plugins/Runtime/Metasound/Source/MetasoundStandardNodes/Private/MetasoundTriggerSequenceNode.cpp
@@ -171,9 +171,9 @@ namespace Metasound
 		auto CreateVertexInterface = []() -> FVertexInterface
 		{
 			FInputVertexInterface InputInterface;
-			InputInterface.Add(TInputDataVertexModel<FTrigger>(METASOUND_GET_PARAM_NAME_AND_TT(InputTrigger)));
-			InputInterface.Add(TInputDataVertexModel<FTrigger>(METASOUND_GET_PARAM_NAME_AND_TT(InputReset)));
-			InputInterface.Add(TInputDataVertexModel<bool>(METASOUND_GET_PARAM_NAME_AND_TT(InputLoop), true));
+			InputInterface.Add(TInputDataVertexModel<FTrigger>(METASOUND_GET_PARAM_NAME_AND_METADATA(InputTrigger)));
+			InputInterface.Add(TInputDataVertexModel<FTrigger>(METASOUND_GET_PARAM_NAME_AND_METADATA(InputReset)));
+			InputInterface.Add(TInputDataVertexModel<bool>(METASOUND_GET_PARAM_NAME_AND_METADATA(InputLoop), true));
 			
 			FOutputVertexInterface OutputInterface;
 			for (uint32 i = 0; i < NumOutputs; ++i)
diff --git a/Engine/Plugins/Runtime/Metasound/Source/MetasoundStandardNodes/Public/MetasoundParamHelper.h b/Engine/Plugins/Runtime/Metasound/Source/MetasoundStandardNodes/Public/MetasoundParamHelper.h
deleted file mode 100644
diff --git a/Engine/Plugins/Runtime/Metasound/Source/MetasoundStandardNodes/Public/MetasoundTriggerAccumulatorNode.h b/Engine/Plugins/Runtime/Metasound/Source/MetasoundStandardNodes/Public/MetasoundTriggerAccumulatorNode.h
--- a/Engine/Plugins/Runtime/Metasound/Source/MetasoundStandardNodes/Public/MetasoundTriggerAccumulatorNode.h
+++ b/Engine/Plugins/Runtime/Metasound/Source/MetasoundStandardNodes/Public/MetasoundTriggerAccumulatorNode.h
@@ -44,10 +44,10 @@ namespace Metasound
 					InputInterface.Add(TInputDataVertexModel<FTrigger>(METASOUND_GET_PARAM_NAME_WITH_INDEX_AND_TT(InputTrigger, i)));
 				}
 
-				InputInterface.Add(TInputDataVertexModel<bool>(METASOUND_GET_PARAM_NAME_AND_TT(InputAutoReset)));
+				InputInterface.Add(TInputDataVertexModel<bool>(METASOUND_GET_PARAM_NAME_AND_METADATA(InputAutoReset)));
 
 				FOutputVertexInterface OutputInterface;
-				OutputInterface.Add(TOutputDataVertexModel<FTrigger>(METASOUND_GET_PARAM_NAME_AND_TT(AccumulateOutputOnTrigger)));
+				OutputInterface.Add(TOutputDataVertexModel<FTrigger>(METASOUND_GET_PARAM_NAME_AND_METADATA(AccumulateOutputOnTrigger)));
 
 				return FVertexInterface(InputInterface, OutputInterface);
 			};
diff --git a/Engine/Plugins/Runtime/Metasound/Source/MetasoundStandardNodes/Public/MetasoundTriggerCompareNode.h b/Engine/Plugins/Runtime/Metasound/Source/MetasoundStandardNodes/Public/MetasoundTriggerCompareNode.h
--- a/Engine/Plugins/Runtime/Metasound/Source/MetasoundStandardNodes/Public/MetasoundTriggerCompareNode.h
+++ b/Engine/Plugins/Runtime/Metasound/Source/MetasoundStandardNodes/Public/MetasoundTriggerCompareNode.h
@@ -71,14 +71,14 @@ namespace Metasound
 
 			static const FVertexInterface DefaultInterface(
 				FInputVertexInterface(
-					TInputDataVertexModel<FTrigger>(METASOUND_GET_PARAM_NAME_AND_TT(InputCompare)),
-					TInputDataVertexModel<ValueType>(METASOUND_GET_PARAM_NAME_AND_TT(InputParamA)),
-					TInputDataVertexModel<ValueType>(METASOUND_GET_PARAM_NAME_AND_TT(InputParamB)),
-					TInputDataVertexModel<FEnumTriggerComparisonType>(METASOUND_GET_PARAM_NAME_AND_TT(InputCompareType))
+					TInputDataVertexModel<FTrigger>(METASOUND_GET_PARAM_NAME_AND_METADATA(InputCompare)),
+					TInputDataVertexModel<ValueType>(METASOUND_GET_PARAM_NAME_AND_METADATA(InputParamA)),
+					TInputDataVertexModel<ValueType>(METASOUND_GET_PARAM_NAME_AND_METADATA(InputParamB)),
+					TInputDataVertexModel<FEnumTriggerComparisonType>(METASOUND_GET_PARAM_NAME_AND_METADATA(InputCompareType))
 				),
 				FOutputVertexInterface(
-					TOutputDataVertexModel<FTrigger>(METASOUND_GET_PARAM_NAME_AND_TT(OutputOnTrue)),
-					TOutputDataVertexModel<FTrigger>(METASOUND_GET_PARAM_NAME_AND_TT(OutputOnFalse))
+					TOutputDataVertexModel<FTrigger>(METASOUND_GET_PARAM_NAME_AND_METADATA(OutputOnTrue)),
+					TOutputDataVertexModel<FTrigger>(METASOUND_GET_PARAM_NAME_AND_METADATA(OutputOnFalse))
 				)
 			);
 
diff --git a/Engine/Plugins/Runtime/Metasound/Source/MetasoundStandardNodes/Public/MetasoundValueNode.h b/Engine/Plugins/Runtime/Metasound/Source/MetasoundStandardNodes/Public/MetasoundValueNode.h
--- a/Engine/Plugins/Runtime/Metasound/Source/MetasoundStandardNodes/Public/MetasoundValueNode.h
+++ b/Engine/Plugins/Runtime/Metasound/Source/MetasoundStandardNodes/Public/MetasoundValueNode.h
@@ -42,15 +42,15 @@ namespace Metasound
 
 			static const FVertexInterface DefaultInterface(
 				FInputVertexInterface(
-					TInputDataVertexModel<FTrigger>(METASOUND_GET_PARAM_NAME_AND_TT(InputSetTrigger)),
-					TInputDataVertexModel<FTrigger>(METASOUND_GET_PARAM_NAME_AND_TT(InputResetTrigger)),
-					TInputDataVertexModel<ValueType>(METASOUND_GET_PARAM_NAME_AND_TT(InputInitValue)),
-					TInputDataVertexModel<ValueType>(METASOUND_GET_PARAM_NAME_AND_TT(InputTargetValue))
+					TInputDataVertexModel<FTrigger>(METASOUND_GET_PARAM_NAME_AND_METADATA(InputSetTrigger)),
+					TInputDataVertexModel<FTrigger>(METASOUND_GET_PARAM_NAME_AND_METADATA(InputResetTrigger)),
+					TInputDataVertexModel<ValueType>(METASOUND_GET_PARAM_NAME_AND_METADATA(InputInitValue)),
+					TInputDataVertexModel<ValueType>(METASOUND_GET_PARAM_NAME_AND_METADATA(InputTargetValue))
 				),
 				FOutputVertexInterface(
-					TOutputDataVertexModel<FTrigger>(METASOUND_GET_PARAM_NAME_AND_TT(OutputOnSet)),
-					TOutputDataVertexModel<FTrigger>(METASOUND_GET_PARAM_NAME_AND_TT(OutputOnReset)),
-					TOutputDataVertexModel<ValueType>(METASOUND_GET_PARAM_NAME_AND_TT(OutputValue))
+					TOutputDataVertexModel<FTrigger>(METASOUND_GET_PARAM_NAME_AND_METADATA(OutputOnSet)),
+					TOutputDataVertexModel<FTrigger>(METASOUND_GET_PARAM_NAME_AND_METADATA(OutputOnReset)),
+					TOutputDataVertexModel<ValueType>(METASOUND_GET_PARAM_NAME_AND_METADATA(OutputValue))
 				)
 			);
 
