- Fix sort order not propagating to presets
- Fix for interfaces not updating/some manipulators being incorrectly editable on MetaSound Presets
- Clarify Set vs Synchronize location on MSEdGraphNodes
- Split up RefreshDetails to RefreshDetails/Interfaces & move RefreshInterface to RefreshGraphMemberMenu
#rnx

diff --git a/Engine/Plugins/Runtime/Metasound/Source/MetasoundEditor/Private/MetasoundDetailCustomization.cpp b/Engine/Plugins/Runtime/Metasound/Source/MetasoundEditor/Private/MetasoundDetailCustomization.cpp
--- a/Engine/Plugins/Runtime/Metasound/Source/MetasoundEditor/Private/MetasoundDetailCustomization.cpp
+++ b/Engine/Plugins/Runtime/Metasound/Source/MetasoundEditor/Private/MetasoundDetailCustomization.cpp
@@ -434,6 +434,7 @@ namespace Metasound
 
 					}), LOCTEXT("RemoveInterfaceTooltip1", "Removes all interfaces from the given MetaSound."))
 				];
+			InterfaceUtilities->SetEnabled(IsGraphEditableAttribute);
 
 			const FText HeaderName = LOCTEXT("InterfacesGroupDisplayName", "Interfaces");
 			IDetailCategoryBuilder& InterfaceCategory = DetailLayout.EditCategory("Interfaces", HeaderName);
@@ -443,7 +444,7 @@ namespace Metasound
 				InterfaceUtilities
 			];
 
-			auto CreateInterfaceEntryWidget = [&](FName InInterfaceName) -> TSharedPtr<SWidget>
+			auto CreateInterfaceEntryWidget = [&](FName InInterfaceName) -> TSharedRef<SWidget>
 			{
 				using namespace Frontend;
 
@@ -478,7 +479,7 @@ namespace Metasound
 					MetaSoundAsset->SetSynchronizationRequired();
 				}), LOCTEXT("RemoveInterfaceTooltip2", "Removes the associated interface from the MetaSound."));
 
-				return SNew(SHorizontalBox)
+				TSharedRef<SWidget> EntryWidget = SNew(SHorizontalBox)
 					+ SHorizontalBox::Slot()
 					.Padding(2.0f)
 					.HAlign(HAlign_Center)
@@ -496,6 +497,9 @@ namespace Metasound
 					[
 						RemoveButtonWidget
 					];
+
+				EntryWidget->SetEnabled(IsGraphEditableAttribute);
+				return EntryWidget;
 			};
 
 			TArray<FName> InterfaceNames = ImplementedInterfaceNames.Array();
@@ -504,7 +508,7 @@ namespace Metasound
 			{
 				InterfaceCategory.AddCustomRow(FText::FromName(InterfaceName))
 				[
-					CreateInterfaceEntryWidget(InterfaceName)->AsShared()
+					CreateInterfaceEntryWidget(InterfaceName)
 				];
 			}
 		}
diff --git a/Engine/Plugins/Runtime/Metasound/Source/MetasoundEditor/Private/MetasoundEditor.cpp b/Engine/Plugins/Runtime/Metasound/Source/MetasoundEditor/Private/MetasoundEditor.cpp
--- a/Engine/Plugins/Runtime/Metasound/Source/MetasoundEditor/Private/MetasoundEditor.cpp
+++ b/Engine/Plugins/Runtime/Metasound/Source/MetasoundEditor/Private/MetasoundEditor.cpp
@@ -2447,14 +2447,17 @@ namespace Metasound
 			{
 				MetasoundDetails->ForceRefresh();
 			}
+		}
 
+		void FEditor::RefreshInterfaces()
+		{
 			if (InterfacesDetails.IsValid())
 			{
 				InterfacesDetails->ForceRefresh();
 			}
 		}
 
-		void FEditor::RefreshInterface()
+		void FEditor::RefreshGraphMemberMenu()
 		{
 			if (GraphMembersMenu.IsValid())
 			{
@@ -3008,13 +3011,23 @@ namespace Metasound
 				const bool bShouldRefreshDetails = MetasoundAsset->GetSynchronizationUpdateDetails();
 				FGraphBuilder::SynchronizeGraph(*Metasound);
 
-				if (bShouldRefreshDetails)
+				// Presets always update interfaces
+				const FMetasoundFrontendGraphClass& RootGraphClass = MetasoundAsset->GetDocumentHandle()->GetRootGraphClass();
+				const bool bIsPreset = RootGraphClass.PresetOptions.bIsPreset;
+				if (bIsPreset)
+				{
+					RefreshInterfaces();
+				}
+
+				if (bShouldRefreshDetails || bIsPreset)
 				{
+					// TODO: Break up this synchronization flag
 					RefreshDetails();
+					RefreshInterfaces();
 				}
 				else
 				{
-					// Also refresh if the object in the panel has gone invalid
+					// Also refresh details if the object in the panel has gone invalid
 					auto ShouldRefreshDetails = [](TWeakObjectPtr<UObject> Obj)
 					{
 						if (!Obj.IsValid() || !IsValid(Obj.Get()))
diff --git a/Engine/Plugins/Runtime/Metasound/Source/MetasoundEditor/Private/MetasoundEditor.h b/Engine/Plugins/Runtime/Metasound/Source/MetasoundEditor/Private/MetasoundEditor.h
--- a/Engine/Plugins/Runtime/Metasound/Source/MetasoundEditor/Private/MetasoundEditor.h
+++ b/Engine/Plugins/Runtime/Metasound/Source/MetasoundEditor/Private/MetasoundEditor.h
@@ -151,7 +151,10 @@ namespace Metasound
 			void RenameSelected();
 
 			/** Forces all UX pertaining to the root graph's interface to be refreshed. */
-			void RefreshInterface();
+			void RefreshGraphMemberMenu();
+
+			/** Forces refresh of interfaces list. */
+			void RefreshInterfaces();
 
 			/** Updates selected node classes to highest class found in the MetaSound Class Registry. */
 			void UpdateSelectedNodeClasses();
diff --git a/Engine/Plugins/Runtime/Metasound/Source/MetasoundEditor/Private/MetasoundEditorGraphBuilder.cpp b/Engine/Plugins/Runtime/Metasound/Source/MetasoundEditor/Private/MetasoundEditorGraphBuilder.cpp
--- a/Engine/Plugins/Runtime/Metasound/Source/MetasoundEditor/Private/MetasoundEditorGraphBuilder.cpp
+++ b/Engine/Plugins/Runtime/Metasound/Source/MetasoundEditor/Private/MetasoundEditorGraphBuilder.cpp
@@ -286,8 +286,7 @@ namespace Metasound
 
 			NodeCreator.Finalize();
 			InitGraphNode(InNodeHandle, NewGraphNode, InMetaSound);
-
-			SynchronizeNodeLocation(InLocation, InNodeHandle, *NewGraphNode);
+			NewGraphNode->SetNodeLocation(InLocation);
 
 			// Adding external node may introduce referenced asset so rebuild referenced keys.
 			MetaSoundAsset->RebuildReferencedAssetClassKeys();
@@ -339,8 +338,7 @@ namespace Metasound
 					NodeCreator.Finalize();
 
 					InitGraphNode(InNodeHandle, NewGraphNode, InMetaSound);
-
-					SynchronizeNodeLocation(InLocation, InNodeHandle, *NewGraphNode);
+					NewGraphNode->SetNodeLocation(InLocation);
 				}
 			}
 
@@ -369,8 +367,7 @@ namespace Metasound
 
 			NodeCreator.Finalize();
 			InitGraphNode(InNodeHandle, NewGraphNode, InMetaSound);
-
-			SynchronizeNodeLocation(InLocation, InNodeHandle, *NewGraphNode);
+			NewGraphNode->SetNodeLocation(InLocation);
 
 			return NewGraphNode;
 		}
@@ -432,7 +429,7 @@ namespace Metasound
 
 			if (MetaSoundEditor.IsValid())
 			{
-				MetaSoundEditor->RefreshInterface();
+				MetaSoundEditor->RefreshGraphMemberMenu();
 			}
 
 			if (bMarkDirty)
@@ -630,14 +627,32 @@ namespace Metasound
 			return false;
 		}
 
-		void FGraphBuilder::SynchronizeNodeLocation(FVector2D InLocation, Frontend::FNodeHandle InNodeHandle, UMetasoundEditorGraphNode& InNode)
+		bool FGraphBuilder::SynchronizeNodeLocation(UMetasoundEditorGraphNode& InNode)
 		{
-			InNode.NodePosX = InLocation.X;
-			InNode.NodePosY = InLocation.Y;
+			bool bModified = false;
+
+			const FMetasoundFrontendNodeStyle& Style = InNode.GetConstNodeHandle()->GetNodeStyle();
+
+			const FVector2D* Location = Style.Display.Locations.Find(InNode.NodeGuid);
+			if (!Location)
+			{
+				// If no specific location found, use default location if provided (zero guid
+				// for example, provided by preset defaults.)
+				Location = Style.Display.Locations.Find({ });
+			}
+
+			if (Location)
+			{
+				const FVector2D Delta = *Location - FVector2D(InNode.NodePosX, InNode.NodePosY);
+				if (!Delta.IsNearlyZero(SMALL_NUMBER))
+				{
+					InNode.NodePosX = Location->X;
+					InNode.NodePosY = Location->Y;
+					bModified = true;
+				}
+			}
 
-			FMetasoundFrontendNodeStyle Style = InNodeHandle->GetNodeStyle();
-			Style.Display.Locations.FindOrAdd(InNode.NodeGuid) = InLocation;
-			InNodeHandle->SetNodeStyle(Style);
+			return bModified;
 		}
 
 		UMetasoundEditorGraphInputNode* FGraphBuilder::AddInputNode(UObject& InMetaSound, Frontend::FNodeHandle InNodeHandle, FVector2D InLocation, bool bInSelectNewNode)
@@ -656,8 +671,7 @@ namespace Metasound
 			UMetasoundEditorGraphInputNode* NewGraphNode = MetasoundGraph->CreateInputNode(InNodeHandle, bInSelectNewNode);
 			if (ensure(NewGraphNode))
 			{
-				SynchronizeNodeLocation(InLocation, InNodeHandle, *NewGraphNode);
-
+				NewGraphNode->SetNodeLocation(InLocation);
 				RebuildNodePins(*NewGraphNode);
 				return NewGraphNode;
 			}
@@ -1820,11 +1834,6 @@ namespace Metasound
 			for (int32 i = FrontendNodes.Num() - 1; i >= 0; i--)
 			{
 				FNodeHandle Node = FrontendNodes[i];
-				auto IsEditorNodeWithSameNodeID = [&](const UMetasoundEditorGraphNode* InEditorNode)
-				{
-					return InEditorNode->GetNodeID() == Node->GetID();
-				};
-
 				bool bFoundEditorNode = false;
 				for (int32 j = EditorNodes.Num() - 1; j >= 0; --j)
 				{
@@ -1842,6 +1851,7 @@ namespace Metasound
 							AssociatedNodeData.Node = Node;
 						}
 
+						bEditorGraphModified |= SynchronizeNodeLocation(*EditorNode);
 						AssociatedNodeData.EditorNodes.Add(EditorNode);
 						EditorNodes.RemoveAtSwap(j, 1, false /* bAllowShrinking */);
 					}
diff --git a/Engine/Plugins/Runtime/Metasound/Source/MetasoundEditor/Private/MetasoundEditorGraphBuilder.h b/Engine/Plugins/Runtime/Metasound/Source/MetasoundEditor/Private/MetasoundEditorGraphBuilder.h
--- a/Engine/Plugins/Runtime/Metasound/Source/MetasoundEditor/Private/MetasoundEditorGraphBuilder.h
+++ b/Engine/Plugins/Runtime/Metasound/Source/MetasoundEditor/Private/MetasoundEditorGraphBuilder.h
@@ -130,9 +130,6 @@ namespace Metasound
 			// Adds an variable node with the given node handle to the editor graph.
 			static UMetasoundEditorGraphVariableNode* AddVariableNode(UObject& InMetaSound, Frontend::FNodeHandle& InNodeHandle, FVector2D InLocation, bool bInSelectNewNode = true);
 
-			// Synchronizes node location data
-			static void SynchronizeNodeLocation(FVector2D InLocation, Frontend::FNodeHandle InNodeHandle, UMetasoundEditorGraphNode& InNode);
-
 			// Adds an output node to the editor graph that corresponds to the provided node handle.
 			static UMetasoundEditorGraphOutputNode* AddOutputNode(UObject& InMetaSound, Frontend::FNodeHandle& InNodeHandle, FVector2D InLocation, bool bInSelectNewNode = true);
 
@@ -232,6 +229,11 @@ namespace Metasound
 			// @return True if the UMetasoundEditorGraphNode was altered. False otherwise.
 			static bool SynchronizeNodeMembers(UObject& InMetaSound);
 
+			// Synchronizes node location data using handle's data
+			//
+			// @return True if the UMetasoundEditorGraphNode was altered. False otherwise.
+			static bool SynchronizeNodeLocation(UMetasoundEditorGraphNode& InNode);
+
 			// Adds and removes pins so that the UMetasoundEditorGraphNode matches the InNode.
 			//
 			// @return True if the UMetasoundEditorGraphNode was altered. False otherwise.
diff --git a/Engine/Plugins/Runtime/Metasound/Source/MetasoundEditor/Private/MetasoundEditorGraphNode.cpp b/Engine/Plugins/Runtime/Metasound/Source/MetasoundEditor/Private/MetasoundEditorGraphNode.cpp
--- a/Engine/Plugins/Runtime/Metasound/Source/MetasoundEditor/Private/MetasoundEditorGraphNode.cpp
+++ b/Engine/Plugins/Runtime/Metasound/Source/MetasoundEditor/Private/MetasoundEditorGraphNode.cpp
@@ -68,11 +68,25 @@ Metasound::Editor::FGraphNodeValidationResult UMetasoundEditorGraphMemberNode::C
 	}
 	return OutResult;
 }
+
 UMetasoundEditorGraphNode::UMetasoundEditorGraphNode(const FObjectInitializer& ObjectInitializer)
 	: Super(ObjectInitializer)
 {
 }
 
+void UMetasoundEditorGraphNode::SetNodeLocation(const FVector2D& InLocation)
+{
+	using namespace Metasound::Frontend;
+
+	NodePosX = InLocation.X;
+	NodePosY = InLocation.Y;
+
+	FNodeHandle NodeHandle = GetNodeHandle();
+	FMetasoundFrontendNodeStyle Style = NodeHandle->GetNodeStyle();
+	Style.Display.Locations.FindOrAdd(NodeGuid) = InLocation;
+	NodeHandle->SetNodeStyle(Style);
+}
+
 void UMetasoundEditorGraphNode::PostLoad()
 {
 	Super::PostLoad();
@@ -484,16 +498,6 @@ void UMetasoundEditorGraphNode::PostEditUndo()
 	}
 }
 
-void UMetasoundEditorGraphNode::UpdatePosition()
-{
-	GetMetasoundChecked().Modify();
-
-	Metasound::Frontend::FNodeHandle NodeHandle = GetNodeHandle();
-	FMetasoundFrontendNodeStyle Style = NodeHandle->GetNodeStyle();
-	Style.Display.Locations.FindOrAdd(NodeGuid) = FVector2D(NodePosX, NodePosY);
-	GetNodeHandle()->SetNodeStyle(Style);
-}
-
 void UMetasoundEditorGraphNode::PostDuplicate(bool bDuplicateForPIE)
 {
 	Super::PostDuplicate(bDuplicateForPIE);
diff --git a/Engine/Plugins/Runtime/Metasound/Source/MetasoundEditor/Private/SMetasoundGraphNode.cpp b/Engine/Plugins/Runtime/Metasound/Source/MetasoundEditor/Private/SMetasoundGraphNode.cpp
--- a/Engine/Plugins/Runtime/Metasound/Source/MetasoundEditor/Private/SMetasoundGraphNode.cpp
+++ b/Engine/Plugins/Runtime/Metasound/Source/MetasoundEditor/Private/SMetasoundGraphNode.cpp
@@ -511,7 +511,9 @@ namespace Metasound
 		{
 			SGraphNode::MoveTo(NewPosition, NodeFilter, bMarkDirty);
 
-			GetMetaSoundNode().UpdatePosition();
+			UMetasoundEditorGraphNode& Node = GetMetaSoundNode();
+			Node.GetMetasoundChecked().Modify();
+			Node.SetNodeLocation(NewPosition);
 		}
 
 		const FSlateBrush* SMetaSoundGraphNode::GetNodeBodyBrush() const
diff --git a/Engine/Plugins/Runtime/Metasound/Source/MetasoundEditor/Public/MetasoundEditorGraphNode.h b/Engine/Plugins/Runtime/Metasound/Source/MetasoundEditor/Public/MetasoundEditorGraphNode.h
--- a/Engine/Plugins/Runtime/Metasound/Source/MetasoundEditor/Public/MetasoundEditorGraphNode.h
+++ b/Engine/Plugins/Runtime/Metasound/Source/MetasoundEditor/Public/MetasoundEditorGraphNode.h
@@ -90,7 +90,8 @@ public:
 	UObject& GetMetasoundChecked();
 	const UObject& GetMetasoundChecked() const;
 
-	void UpdatePosition();
+	// Sets the node's location, both on this graph member node and on the frontend handle
+	void SetNodeLocation(const FVector2D& InLocation);
 
 	Metasound::Frontend::FGraphHandle GetRootGraphHandle() const;
 	Metasound::Frontend::FConstGraphHandle GetConstRootGraphHandle() const;
diff --git a/Engine/Plugins/Runtime/Metasound/Source/MetasoundFrontend/Private/MetasoundFrontendGraphController.cpp b/Engine/Plugins/Runtime/Metasound/Source/MetasoundFrontend/Private/MetasoundFrontendGraphController.cpp
--- a/Engine/Plugins/Runtime/Metasound/Source/MetasoundFrontend/Private/MetasoundFrontendGraphController.cpp
+++ b/Engine/Plugins/Runtime/Metasound/Source/MetasoundFrontend/Private/MetasoundFrontendGraphController.cpp
@@ -187,6 +187,26 @@ namespace Metasound
 			return Invalid::GetInvalidGraphStyle();
 		}
 
+		const FMetasoundFrontendInterfaceStyle& FGraphController::GetInputStyle() const
+		{
+			if (FMetasoundFrontendGraphClass* GraphClass = GraphClassPtr.Get())
+			{
+				return GraphClass->Interface.GetInputStyle();
+			}
+
+			return Invalid::GetInvalidInterfaceStyle();
+		}
+
+		const FMetasoundFrontendInterfaceStyle& FGraphController::GetOutputStyle() const
+		{
+			if (FMetasoundFrontendGraphClass* GraphClass = GraphClassPtr.Get())
+			{
+				return GraphClass->Interface.GetOutputStyle();
+			}
+
+			return Invalid::GetInvalidInterfaceStyle();
+		}
+
 		void FGraphController::SetGraphStyle(const FMetasoundFrontendGraphStyle& InStyle)
 		{
 			if (FMetasoundFrontendGraphClass* GraphClass = GraphClassPtr.Get())
@@ -199,7 +219,16 @@ namespace Metasound
 		{
 			if (FMetasoundFrontendGraphClass* GraphClass = GraphClassPtr.Get())
 			{
-				GraphClass->Interface.SetInputStyle(InStyle);
+				FMetasoundFrontendInterfaceStyle Style = InStyle;
+
+				TArray<FMetasoundFrontendClassInput>& Inputs = GraphClass->Interface.Inputs;
+				Style.DefaultSortOrder.SetNumZeroed(Inputs.Num());
+
+				for (int32 i = 0; i < Inputs.Num(); ++i)
+				{
+					Inputs[i].Metadata.SortOrderIndex = InStyle.DefaultSortOrder[i];
+				}
+				GraphClass->Interface.SetInputStyle(Style);
 			}
 		}
 
@@ -207,7 +236,16 @@ namespace Metasound
 		{
 			if (FMetasoundFrontendGraphClass* GraphClass = GraphClassPtr.Get())
 			{
-				GraphClass->Interface.SetOutputStyle(InStyle);
+				FMetasoundFrontendInterfaceStyle Style = InStyle;
+
+				TArray<FMetasoundFrontendClassOutput>& Outputs = GraphClass->Interface.Outputs;
+				Style.DefaultSortOrder.SetNumZeroed(Outputs.Num());
+
+				for (int32 i = 0; i < Outputs.Num(); ++i)
+				{
+					Outputs[i].Metadata.SortOrderIndex = InStyle.DefaultSortOrder[i];
+				}
+				GraphClass->Interface.SetOutputStyle(Style);
 			}
 		}
 #endif // WITH_EDITOR
diff --git a/Engine/Plugins/Runtime/Metasound/Source/MetasoundFrontend/Private/MetasoundFrontendGraphController.h b/Engine/Plugins/Runtime/Metasound/Source/MetasoundFrontend/Private/MetasoundFrontendGraphController.h
--- a/Engine/Plugins/Runtime/Metasound/Source/MetasoundFrontend/Private/MetasoundFrontendGraphController.h
+++ b/Engine/Plugins/Runtime/Metasound/Source/MetasoundFrontend/Private/MetasoundFrontendGraphController.h
@@ -97,10 +97,12 @@ namespace Metasound
 			TArray<FConstNodeHandle> GetConstOutputNodes() const override;
 
 #if WITH_EDITOR
-			const FMetasoundFrontendGraphStyle& GetGraphStyle() const override;
-			void SetGraphStyle(const FMetasoundFrontendGraphStyle& InStyle) override;
-			void SetInputStyle(const FMetasoundFrontendInterfaceStyle& InStyle) override;
-			void SetOutputStyle(const FMetasoundFrontendInterfaceStyle& InStyle) override;
+			virtual const FMetasoundFrontendGraphStyle& GetGraphStyle() const override;
+			virtual const FMetasoundFrontendInterfaceStyle& GetInputStyle() const override;
+			virtual const FMetasoundFrontendInterfaceStyle& GetOutputStyle() const override;
+			virtual void SetGraphStyle(const FMetasoundFrontendGraphStyle& InStyle) override;
+			virtual void SetInputStyle(const FMetasoundFrontendInterfaceStyle& InStyle) override;
+			virtual void SetOutputStyle(const FMetasoundFrontendInterfaceStyle& InStyle) override;
 #endif // WITH_EDITOR
 
 			bool ContainsInputVertex(const FVertexName& InName, const FName& InTypeName) const override;
diff --git a/Engine/Plugins/Runtime/Metasound/Source/MetasoundFrontend/Private/MetasoundFrontendInvalidController.h b/Engine/Plugins/Runtime/Metasound/Source/MetasoundFrontend/Private/MetasoundFrontendInvalidController.h
--- a/Engine/Plugins/Runtime/Metasound/Source/MetasoundFrontend/Private/MetasoundFrontendInvalidController.h
+++ b/Engine/Plugins/Runtime/Metasound/Source/MetasoundFrontend/Private/MetasoundFrontendInvalidController.h
@@ -334,6 +334,8 @@ namespace Metasound
 #if WITH_EDITOR
 			virtual const FMetasoundFrontendGraphStyle& GetGraphStyle() const override { return Invalid::GetInvalidGraphStyle(); }
 			virtual void SetGraphStyle(const FMetasoundFrontendGraphStyle& InStyle) override { }
+			virtual const FMetasoundFrontendInterfaceStyle& GetInputStyle() const override { return Invalid::GetInvalidInterfaceStyle(); }
+			virtual const FMetasoundFrontendInterfaceStyle& GetOutputStyle() const override { return Invalid::GetInvalidInterfaceStyle(); }
 			virtual void SetInputStyle(const FMetasoundFrontendInterfaceStyle& InStyle) override { }
 			virtual void SetOutputStyle(const FMetasoundFrontendInterfaceStyle& InStyle) override { }
 #endif // WITH_EDITOR
diff --git a/Engine/Plugins/Runtime/Metasound/Source/MetasoundFrontend/Private/MetasoundFrontendTransform.cpp b/Engine/Plugins/Runtime/Metasound/Source/MetasoundFrontend/Private/MetasoundFrontendTransform.cpp
--- a/Engine/Plugins/Runtime/Metasound/Source/MetasoundFrontend/Private/MetasoundFrontendTransform.cpp
+++ b/Engine/Plugins/Runtime/Metasound/Source/MetasoundFrontend/Private/MetasoundFrontendTransform.cpp
@@ -699,7 +699,7 @@ namespace Metasound
 			FMetasoundFrontendNodeStyle RefNodeStyle;
 
 			// Offset to be to the right of input nodes
-			RefNodeStyle.Display.Locations.Add(FGuid::NewGuid(), DisplayStyle::NodeLayout::DefaultOffsetX);
+			RefNodeStyle.Display.Locations.Add(FGuid(), DisplayStyle::NodeLayout::DefaultOffsetX);
 			ReferencedNodeHandle->SetNodeStyle(RefNodeStyle);
 #endif // WITH_EDITOR
 
@@ -717,27 +717,38 @@ namespace Metasound
 			// Add inputs and space appropriately
 			FVector2D InputNodeLocation = FVector2D::ZeroVector;
 
+			FConstGraphHandle ReferencedGraphHandle = ReferencedDocument->GetRootGraph();
+
+			TArray<FNodeHandle> NodeHandles;
 			for (const FMetasoundFrontendClassInput& ClassInput : InClassInputs)
 			{
 				FNodeHandle InputNode = InPresetGraphHandle->AddInputVertex(ClassInput);
-
 				if (ensure(InputNode->IsValid()))
 				{
-#if WITH_EDITOR
-					// Set input node location
-					FMetasoundFrontendNodeStyle NodeStyle;
-					NodeStyle.Display.Locations.Add(FGuid::NewGuid(), InputNodeLocation);
-					InputNode->SetNodeStyle(NodeStyle);
-					InputNodeLocation += DisplayStyle::NodeLayout::DefaultOffsetY;
-#endif // WITH_EDITOR
-
 					// Connect input node to corresponding referencing node.
 					FOutputHandle OutputToConnect = InputNode->GetOutputWithVertexName(ClassInput.Name);
 					FInputHandle InputToConnect = InReferencedNode->GetInputWithVertexName(ClassInput.Name);
-					bool bSuccess = OutputToConnect->Connect(*InputToConnect);
-					check(bSuccess);
+					ensure(OutputToConnect->Connect(*InputToConnect));
+					NodeHandles.Add(MoveTemp(InputNode));
 				}
 			}
+
+#if WITH_EDITOR
+			// Sort before adding nodes to graph layout & copy to preset (must be done after all
+			// inputs/outputs are added but before setting locations to propagate effectively)
+			FMetasoundFrontendInterfaceStyle Style = ReferencedGraphHandle->GetInputStyle();
+			InPresetGraphHandle->SetInputStyle(Style);
+			Style.SortDefaults(NodeHandles);
+
+			for (const FNodeHandle& InputNode : NodeHandles)
+			{
+				// Set input node location
+				FMetasoundFrontendNodeStyle NodeStyle;
+				NodeStyle.Display.Locations.Add(FGuid(), InputNodeLocation);
+				InputNode->SetNodeStyle(NodeStyle);
+				InputNodeLocation += DisplayStyle::NodeLayout::DefaultOffsetY;
+			}
+#endif // WITH_EDITOR
 		}
 
 		void FRebuildPresetRootGraph::AddAndConnectOutputs(const TArray<FMetasoundFrontendClassOutput>& InClassOutputs, FGraphHandle& InPresetGraphHandle, FNodeHandle& InReferencedNode) const
@@ -745,27 +756,38 @@ namespace Metasound
 			// Add outputs and space appropriately
 			FVector2D OutputNodeLocation = (2 * DisplayStyle::NodeLayout::DefaultOffsetX);
 
+			FConstGraphHandle ReferencedGraphHandle = ReferencedDocument->GetRootGraph();
+
+			TArray<FNodeHandle> NodeHandles;
 			for (const FMetasoundFrontendClassOutput& ClassOutput : InClassOutputs)
 			{
 				FNodeHandle OutputNode = InPresetGraphHandle->AddOutputVertex(ClassOutput);
-				
 				if (ensure(OutputNode->IsValid()))
 				{
-#if WITH_EDITOR
-					// Set input node location
-					FMetasoundFrontendNodeStyle NodeStyle;
-					NodeStyle.Display.Locations.Add(FGuid::NewGuid(), OutputNodeLocation);
-					OutputNode->SetNodeStyle(NodeStyle);
-					OutputNodeLocation += DisplayStyle::NodeLayout::DefaultOffsetY;
-#endif // WITH_EDITOR
-
 					// Connect input node to corresponding referenced node. 
 					FInputHandle InputToConnect = OutputNode->GetInputWithVertexName(ClassOutput.Name);
 					FOutputHandle OutputToConnect = InReferencedNode->GetOutputWithVertexName(ClassOutput.Name);
-					bool bSuccess = InputToConnect->Connect(*OutputToConnect);
-					check(bSuccess);
+					ensure(InputToConnect->Connect(*OutputToConnect));
+					NodeHandles.Add(MoveTemp(OutputNode));
 				}
 			}
+
+#if WITH_EDITOR
+			// Sort before adding nodes to graph layout & copy to preset (must be done after all
+			// inputs/outputs are added but before setting locations to propagate effectively)
+			FMetasoundFrontendInterfaceStyle Style = ReferencedGraphHandle->GetOutputStyle();
+			InPresetGraphHandle->SetOutputStyle(Style);
+			Style.SortDefaults(NodeHandles);
+
+			for (const FNodeHandle& OutputNode : NodeHandles)
+			{
+				// Set input node location
+				FMetasoundFrontendNodeStyle NodeStyle;
+				NodeStyle.Display.Locations.Add(FGuid(), OutputNodeLocation);
+				OutputNode->SetNodeStyle(NodeStyle);
+				OutputNodeLocation += DisplayStyle::NodeLayout::DefaultOffsetY;
+			}
+#endif // WITH_EDITOR
 		}
 
 		TArray<FMetasoundFrontendClassInput> FRebuildPresetRootGraph::GenerateRequiredClassInputs(const FConstGraphHandle& InPresetGraph, TSet<FName>& OutInputsInheritingDefault) const
diff --git a/Engine/Plugins/Runtime/Metasound/Source/MetasoundFrontend/Public/MetasoundFrontendController.h b/Engine/Plugins/Runtime/Metasound/Source/MetasoundFrontend/Public/MetasoundFrontendController.h
--- a/Engine/Plugins/Runtime/Metasound/Source/MetasoundFrontend/Public/MetasoundFrontendController.h
+++ b/Engine/Plugins/Runtime/Metasound/Source/MetasoundFrontend/Public/MetasoundFrontendController.h
@@ -623,6 +623,9 @@ namespace Metasound
 			// Returns graph style.
 			virtual const FMetasoundFrontendGraphStyle& GetGraphStyle() const = 0;
 
+			virtual const FMetasoundFrontendInterfaceStyle& GetInputStyle() const = 0;
+			virtual const FMetasoundFrontendInterfaceStyle& GetOutputStyle() const = 0;
+
 			// Sets graph style.
 			virtual void SetGraphStyle(const FMetasoundFrontendGraphStyle& InStyle) = 0;
 
