Hide default value if MetaSound member is hooked up.
- Clean-up some MetaSoundMember edgraphnode members
- Minor fix for hover when drag/dropping members not onto graph showing no error state
#rnx

diff --git a/Engine/Plugins/Runtime/Metasound/Source/MetasoundEditor/Private/MetasoundEditor.cpp b/Engine/Plugins/Runtime/Metasound/Source/MetasoundEditor/Private/MetasoundEditor.cpp
--- a/Engine/Plugins/Runtime/Metasound/Source/MetasoundEditor/Private/MetasoundEditor.cpp
+++ b/Engine/Plugins/Runtime/Metasound/Source/MetasoundEditor/Private/MetasoundEditor.cpp
@@ -32,7 +32,6 @@
 #include "MetasoundEditorGraph.h"
 #include "MetasoundEditorGraphBuilder.h"
 #include "MetasoundEditorGraphInputNode.h"
-#include "MetasoundEditorGraphNode.h"
 #include "MetasoundEditorGraphSchema.h"
 #include "MetasoundEditorGraphValidation.h"
 #include "MetasoundEditorModule.h"
@@ -50,7 +49,6 @@
 #include "PropertyCustomizationHelpers.h"
 #include "PropertyEditorModule.h"
 #include "ScopedTransaction.h"
-#include "SGraphPanel.h"
 #include "SMetasoundActionMenu.h"
 #include "SMetasoundPalette.h"
 #include "SNodePanel.h"
@@ -166,8 +164,6 @@ namespace Metasound
 				CursorDecoratorWindow = SWindow::MakeCursorDecorator();
 				constexpr bool bShowImmediately = false;
 				FSlateApplication::Get().AddWindow(CursorDecoratorWindow.ToSharedRef(), bShowImmediately);
-
-				HoverTargetChanged();
 			}
 
 			DRAG_DROP_OPERATOR_TYPE(FMetaSoundDragDropMemberAction, FGraphSchemaActionDragDropAction)
@@ -201,7 +197,7 @@ namespace Metasound
 
 				if (UMetasoundEditorGraphOutput* Output = Cast<UMetasoundEditorGraphOutput>(GraphMember.Get()))
 				{
-					TArray<UMetasoundEditorGraphNode*> Nodes = Output->GetNodes();
+					TArray<UMetasoundEditorGraphMemberNode*> Nodes = Output->GetNodes();
 					if (Nodes.IsEmpty())
 					{
 						const FScopedTransaction Transaction(LOCTEXT("DropAddNewOutputNode", "Drop New MetaSound Output Node"));
@@ -245,15 +241,15 @@ namespace Metasound
 						{
 							if (Editor.IsValid())
 							{
-								auto IsMutatorNode = [&MutatorNodeHandle](const UMetasoundEditorGraphNode* Node)
+								auto IsMutatorNode = [&MutatorNodeHandle](const UMetasoundEditorGraphMemberNode* Node)
 								{
 									return Node->GetNodeID() == MutatorNodeHandle->GetID();
 								};
-								TArray<UMetasoundEditorGraphNode*> Nodes = Variable->GetNodes();
-								if (UMetasoundEditorGraphNode** MutatorNode = Nodes.FindByPredicate(IsMutatorNode))
+								TArray<UMetasoundEditorGraphMemberNode*> Nodes = Variable->GetNodes();
+								if (UMetasoundEditorGraphMemberNode** MutatorNode = Nodes.FindByPredicate(IsMutatorNode))
 								{
 									check(*MutatorNode);
-									Editor->JumpToNodes({ *MutatorNode });
+									Editor->JumpToNodes<UMetasoundEditorGraphMemberNode>({ *MutatorNode });
 									return FReply::Handled();
 								}
 							}
@@ -268,7 +264,7 @@ namespace Metasound
 						const bool bJumpToGetters = FSlateApplication::Get().GetModifierKeys().AreModifersDown(EModifierKey::Control);
 						if (bJumpToGetters)
 						{
-							TArray<UMetasoundEditorGraphNode*> Nodes = Variable->GetNodes();
+							TArray<UMetasoundEditorGraphMemberNode*> Nodes = Variable->GetNodes();
 							for (int32 i = Nodes.Num() - 1; i >= 0; --i)
 							{
 								const UMetasoundEditorGraphVariableNode* VariableNode = CastChecked<UMetasoundEditorGraphVariableNode>(Nodes[i]);
@@ -317,7 +313,7 @@ namespace Metasound
 			{
 				using namespace Frontend;
 
-				bDropTargetValid = true;
+				bDropTargetValid = false;
 
 				const FSlateBrush* PrimarySymbol = nullptr;
 				const FSlateBrush* SecondarySymbol = nullptr;
@@ -328,93 +324,109 @@ namespace Metasound
 				FText Message;
 				if (GraphMember.IsValid())
 				{
+					UMetasoundEditorGraph* OwningGraph = GraphMember->GetOwningGraph();
 					Message = GraphMember->GetDisplayName();
-					if (GetHoveredGraph() == GraphMember->GetOwningGraph())
+					if (GetHoveredGraph() && OwningGraph)
 					{
-						if (UMetasoundEditorGraphInput* Input = Cast<UMetasoundEditorGraphInput>(GraphMember.Get()))
+						if (GetHoveredGraph() == OwningGraph)
 						{
-							if (const ISlateStyle* MetasoundStyle = FSlateStyleRegistry::FindSlateStyle("MetaSoundStyle"))
-							{
-								PrimarySymbol = MetasoundStyle->GetBrush("MetasoundEditor.Graph.Node.Class.Input");
-								SecondarySymbol = nullptr;
-							}
+							FConstDocumentHandle DocumentHandle= OwningGraph->GetDocumentHandle();
+							const FMetasoundFrontendGraphClass& RootGraphClass = DocumentHandle->GetRootGraphClass();
+							const bool bIsPreset = RootGraphClass.PresetOptions.bIsPreset;
 
-							if (const UMetasoundEditorSettings* EditorSettings = GetDefault<UMetasoundEditorSettings>())
+							if (bIsPreset)
 							{
-								PrimaryColor = EditorSettings->InputNodeTitleColor;
-								SecondaryColor = EditorSettings->InputNodeTitleColor;
-							}
-						}
-
-						if (UMetasoundEditorGraphOutput* Output = Cast<UMetasoundEditorGraphOutput>(GraphMember.Get()))
-						{
-							if (!Output->GetNodes().IsEmpty())
-							{
-								PrimarySymbol = FEditorStyle::GetBrush(TEXT("Graph.ConnectorFeedback.ShowNode"));
-								SecondarySymbol = nullptr;
-								Message = FText::Format(LOCTEXT("DropTargetShowOutput", "Show '{0}' (One per graph)"), GraphMember->GetDisplayName());
+								Message = FText::Format(LOCTEXT("DropTargetFailIsPreset", "'{0}': Graph is Preset"), GraphMember->GetDisplayName());
 							}
-							else
+							else if (UMetasoundEditorGraphInput* Input = Cast<UMetasoundEditorGraphInput>(GraphMember.Get()))
 							{
+								bDropTargetValid = true;
+
 								if (const ISlateStyle* MetasoundStyle = FSlateStyleRegistry::FindSlateStyle("MetaSoundStyle"))
 								{
-									PrimarySymbol = MetasoundStyle->GetBrush("MetasoundEditor.Graph.Node.Class.Output");
+									PrimarySymbol = MetasoundStyle->GetBrush("MetasoundEditor.Graph.Node.Class.Input");
 									SecondarySymbol = nullptr;
 								}
 
 								if (const UMetasoundEditorSettings* EditorSettings = GetDefault<UMetasoundEditorSettings>())
 								{
-									PrimaryColor = EditorSettings->OutputNodeTitleColor;
-									SecondaryColor = EditorSettings->OutputNodeTitleColor;
+									PrimaryColor = EditorSettings->InputNodeTitleColor;
+									SecondaryColor = EditorSettings->InputNodeTitleColor;
 								}
 							}
-						}
-
-						if (UMetasoundEditorGraphVariable* Variable = Cast<UMetasoundEditorGraphVariable>(GraphMember.Get()))
-						{
-							PrimarySymbol = FEditorStyle::GetBrush(TEXT("Graph.ConnectorFeedback.ShowNode"));
-
-							if (const ISlateStyle* MetasoundStyle = FSlateStyleRegistry::FindSlateStyle("MetaSoundStyle"))
+							else if (UMetasoundEditorGraphOutput* Output = Cast<UMetasoundEditorGraphOutput>(GraphMember.Get()))
 							{
-								PrimarySymbol = MetasoundStyle->GetBrush("MetasoundEditor.Graph.Node.Class.Variable");
-								SecondarySymbol = nullptr;
-							}
+								bDropTargetValid = true;
 
-							if (const UMetasoundEditorSettings* EditorSettings = GetDefault<UMetasoundEditorSettings>())
-							{
-								PrimaryColor = EditorSettings->VariableNodeTitleColor;
-								SecondaryColor = EditorSettings->VariableNodeTitleColor;
-							}
+								if (!Output->GetNodes().IsEmpty())
+								{
+									PrimarySymbol = FEditorStyle::GetBrush(TEXT("Graph.ConnectorFeedback.ShowNode"));
+									SecondarySymbol = nullptr;
+									Message = FText::Format(LOCTEXT("DropTargetShowOutput", "Show '{0}' (One per graph)"), GraphMember->GetDisplayName());
+								}
+								else
+								{
+									if (const ISlateStyle* MetasoundStyle = FSlateStyleRegistry::FindSlateStyle("MetaSoundStyle"))
+									{
+										PrimarySymbol = MetasoundStyle->GetBrush("MetasoundEditor.Graph.Node.Class.Output");
+										SecondarySymbol = nullptr;
+									}
 
-							const FText DisplayName = GraphMember->GetDisplayName();
-							const FText GetterToolTip = FText::Format(LOCTEXT("DropTargetGetterVariableToolTipFormat", "{0}\nAdd:\n* Get (Drop)\n* Get Delayed (Alt+Drop)\n"), DisplayName);
-							static const FText GetJumpToToolTip = LOCTEXT("JumpToGettersToolTip", "Get (Ctrl+Drop)");
-							static const FText AddOrJumpToSetToolTip = LOCTEXT("AddOrJumpToSetToolTip", "");
-							FConstNodeHandle MutatorNodeHandle = Variable->GetConstVariableHandle()->FindMutatorNode();
-							if (MutatorNodeHandle->IsValid())
-							{
-								Message = FText::Format(LOCTEXT("DropTargetVariableJumpToFormat", "{0}\nJump To:\n* {1}\n* Set (Shift+Drop, One per graph)"), GetterToolTip, GetJumpToToolTip);
+									if (const UMetasoundEditorSettings* EditorSettings = GetDefault<UMetasoundEditorSettings>())
+									{
+										PrimaryColor = EditorSettings->OutputNodeTitleColor;
+										SecondaryColor = EditorSettings->OutputNodeTitleColor;
+									}
+								}
 							}
-							else
+							else if (UMetasoundEditorGraphVariable* Variable = Cast<UMetasoundEditorGraphVariable>(GraphMember.Get()))
 							{
-								TArray<FConstNodeHandle> AccessorNodeHandles = Variable->GetConstVariableHandle()->FindAccessorNodes();
+								bDropTargetValid = true;
+
+								PrimarySymbol = FEditorStyle::GetBrush(TEXT("Graph.ConnectorFeedback.ShowNode"));
+
+								if (const ISlateStyle* MetasoundStyle = FSlateStyleRegistry::FindSlateStyle("MetaSoundStyle"))
+								{
+									PrimarySymbol = MetasoundStyle->GetBrush("MetasoundEditor.Graph.Node.Class.Variable");
+									SecondarySymbol = nullptr;
+								}
+
+								if (const UMetasoundEditorSettings* EditorSettings = GetDefault<UMetasoundEditorSettings>())
+								{
+									PrimaryColor = EditorSettings->VariableNodeTitleColor;
+									SecondaryColor = EditorSettings->VariableNodeTitleColor;
+								}
 
-								if (AccessorNodeHandles.IsEmpty())
+								const FText DisplayName = GraphMember->GetDisplayName();
+								const FText GetterToolTip = FText::Format(LOCTEXT("DropTargetGetterVariableToolTipFormat", "{0}\nAdd:\n* Get (Drop)\n* Get Delayed (Alt+Drop)\n"), DisplayName);
+								static const FText GetJumpToToolTip = LOCTEXT("JumpToGettersToolTip", "Get (Ctrl+Drop)");
+								static const FText AddOrJumpToSetToolTip = LOCTEXT("AddOrJumpToSetToolTip", "");
+								FConstNodeHandle MutatorNodeHandle = Variable->GetConstVariableHandle()->FindMutatorNode();
+								if (MutatorNodeHandle->IsValid())
 								{
-									Message = FText::Format(LOCTEXT("DropTargetVariableAddSetGetFormat", "{0}* Set (Shift+Drop)"), GetterToolTip);
+									Message = FText::Format(LOCTEXT("DropTargetVariableJumpToFormat", "{0}\nJump To:\n* {1}\n* Set (Shift+Drop, One per graph)"), GetterToolTip, GetJumpToToolTip);
 								}
 								else
 								{
-									Message = FText::Format(LOCTEXT("DropTargetVariableAddSetJumpToGetFormat", "{0}* Set (Shift+Drop)\n\nJump To:\n* {1}"), GetterToolTip, GetJumpToToolTip);
+									TArray<FConstNodeHandle> AccessorNodeHandles = Variable->GetConstVariableHandle()->FindAccessorNodes();
+
+									if (AccessorNodeHandles.IsEmpty())
+									{
+										Message = FText::Format(LOCTEXT("DropTargetVariableAddSetGetFormat", "{0}* Set (Shift+Drop)"), GetterToolTip);
+									}
+									else
+									{
+										Message = FText::Format(LOCTEXT("DropTargetVariableAddSetJumpToGetFormat", "{0}* Set (Shift+Drop)\n\nJump To:\n* {1}"), GetterToolTip, GetJumpToToolTip);
+									}
 								}
 							}
 						}
+						else
+						{
+							Message = FText::Format(LOCTEXT("DropTargetFailNotParentGraph", "'{0}': Graph is not parent of member."), GraphMember->GetDisplayName());
+						}
 					}
 				}
-				else
-				{
-					bDropTargetValid = false;
-				}
 
 				SetSimpleFeedbackMessage(PrimarySymbol, PrimaryColor, Message, SecondarySymbol, SecondaryColor);
 			}
@@ -2810,42 +2822,7 @@ namespace Metasound
 			}
 		}
 
-		void FEditor::JumpToNodes(const TArray<UMetasoundEditorGraphNode*>& InNodes)
-		{
-			if (!MetasoundGraphEditor.IsValid())
-			{
-				return;
-			}
-
-			MetasoundGraphEditor->ClearSelectionSet();
-			const UMetasoundEditorGraph& Graph = GetMetaSoundGraphChecked();
-			if (!InNodes.IsEmpty())
-			{
-				if (SGraphPanel* GraphPanel = MetasoundGraphEditor->GetGraphPanel())
-				{
-					FVector2D BottomLeft = { TNumericLimits<float>::Max(), TNumericLimits<float>::Max() };
-					FVector2D TopRight = { TNumericLimits<float>::Min(), TNumericLimits<float>::Min() };
-					for (UMetasoundEditorGraphNode* Node : InNodes)
-					{
-						if (!Node || Node->GetGraph() != &Graph)
-						{
-							continue;
-						}
-
-						constexpr bool bSelected = true;
-						MetasoundGraphEditor->SetNodeSelection(Node, bSelected);
-						BottomLeft.X = FMath::Min(BottomLeft.X, Node->NodePosX);
-						BottomLeft.Y = FMath::Min(BottomLeft.Y, Node->NodePosY);
-						TopRight.X = FMath::Max(TopRight.X, Node->NodePosX + Node->EstimateNodeWidth());
-						TopRight.Y = FMath::Max(TopRight.Y, Node->NodePosY);
-					}
-
-					GraphPanel->JumpToRect(BottomLeft, TopRight);
-				}
-			}
-		}
-
-		bool FEditor::CanJumpToNodesForSelectedInterfaceItem() const 
+		bool FEditor::CanJumpToNodesForSelectedInterfaceItem() const
 		{
 			if (!GraphMembersMenu.IsValid())
 			{
@@ -2863,7 +2840,7 @@ namespace Metasound
 					{
 						if (const UMetasoundEditorGraphMember* GraphMember = MetasoundAction->GetGraphMember())
 						{
-							TArray<UMetasoundEditorGraphNode*> Nodes = GraphMember->GetNodes();
+							TArray<UMetasoundEditorGraphMemberNode*> Nodes = GraphMember->GetNodes();
 							if (!Nodes.IsEmpty())
 							{
 								return true;
@@ -2982,7 +2959,7 @@ namespace Metasound
 
 			return MenuBuilder.MakeWidget();
 		}
-			
+
 		void FEditor::Tick(float DeltaTime)
 		{
 			if (!Metasound)
diff --git a/Engine/Plugins/Runtime/Metasound/Source/MetasoundEditor/Private/MetasoundEditor.h b/Engine/Plugins/Runtime/Metasound/Source/MetasoundEditor/Private/MetasoundEditor.h
--- a/Engine/Plugins/Runtime/Metasound/Source/MetasoundEditor/Private/MetasoundEditor.h
+++ b/Engine/Plugins/Runtime/Metasound/Source/MetasoundEditor/Private/MetasoundEditor.h
@@ -3,17 +3,21 @@
 
 #include "AudioMeterStyle.h"
 #include "AudioSynesthesia/Classes/Meter.h"
+#include "EdGraph/EdGraphNode.h"
 #include "EditorUndoClient.h"
 #include "Framework/Commands/UICommandList.h"
 #include "GraphEditor.h"
 #include "IDetailsView.h"
 #include "IMetasoundEditor.h"
+#include "Math/UnrealMathUtility.h"
+#include "MetasoundEditorGraphNode.h"
 #include "MetasoundEditorMeter.h"
 #include "MetasoundFrontend.h"
 #include "MetasoundFrontendController.h"
 #include "Misc/NotifyHook.h"
 #include "SAudioMeter.h"
 #include "SGraphActionMenu.h"
+#include "SGraphPanel.h"
 #include "SMetasoundPalette.h"
 #include "Sound/AudioBus.h"
 #include "Textures/SlateIcon.h"
@@ -147,8 +151,6 @@ namespace Metasound
 			void PasteNodes(const FVector2D* InLocation = nullptr);
 			void PasteNodes(const FVector2D* InLocation, const FText& InTransactionText);
 
-			/** Rename selected items (nodes or interface items). */
-			void RenameSelected();
 
 			/** Forces all UX pertaining to the root graph's interface to be refreshed. */
 			void RefreshGraphMemberMenu();
@@ -177,8 +179,42 @@ namespace Metasound
 			/** Destroys analyzers */
 			void DestroyAnalyzers();
 
-			/** Jumps to the given nodes on the graph */
-			void JumpToNodes(const TArray<UMetasoundEditorGraphNode*>& InNodes);
+			/** Jumps to the given nodes on the graph (templated to support arrays of various MetaSound graph node types) */
+			template<typename TNodeType = UMetasoundEditorGraphNode>
+			void JumpToNodes(const TArray<TNodeType*>& InNodes)
+			{
+				if (!MetasoundGraphEditor.IsValid())
+				{
+					return;
+				}
+
+				MetasoundGraphEditor->ClearSelectionSet();
+				const UMetasoundEditorGraph& Graph = GetMetaSoundGraphChecked();
+				if (!InNodes.IsEmpty())
+				{
+					if (SGraphPanel* GraphPanel = MetasoundGraphEditor->GetGraphPanel())
+					{
+						FVector2D BottomLeft = { TNumericLimits<float>::Max(), TNumericLimits<float>::Max() };
+						FVector2D TopRight = { TNumericLimits<float>::Min(), TNumericLimits<float>::Min() };
+						for (TNodeType* Node : InNodes)
+						{
+							if (!Node || Node->GetGraph() != &Graph)
+							{
+								continue;
+							}
+
+							constexpr bool bSelected = true;
+							MetasoundGraphEditor->SetNodeSelection(Node, bSelected);
+							BottomLeft.X = FMath::Min(BottomLeft.X, Node->NodePosX);
+							BottomLeft.Y = FMath::Min(BottomLeft.Y, Node->NodePosY);
+							TopRight.X = FMath::Max(TopRight.X, Node->NodePosX + Node->EstimateNodeWidth());
+							TopRight.Y = FMath::Max(TopRight.Y, Node->NodePosY);
+						}
+
+						GraphPanel->JumpToRect(BottomLeft, TopRight);
+					}
+				}
+			}
 
 		protected:
 			// Callbacks for action tree
diff --git a/Engine/Plugins/Runtime/Metasound/Source/MetasoundEditor/Private/MetasoundEditorGraph.cpp b/Engine/Plugins/Runtime/Metasound/Source/MetasoundEditor/Private/MetasoundEditorGraph.cpp
--- a/Engine/Plugins/Runtime/Metasound/Source/MetasoundEditor/Private/MetasoundEditorGraph.cpp
+++ b/Engine/Plugins/Runtime/Metasound/Source/MetasoundEditor/Private/MetasoundEditorGraph.cpp
@@ -46,6 +46,35 @@ namespace Metasound
 					}
 				}
 			}
+
+			void SetLiteralOrClearIfMatchesDefault(Frontend::FInputHandle& InInputHandle, const FMetasoundFrontendLiteral& InDefaultLiteral)
+			{
+				using namespace Frontend;
+
+				// Avoids member literal setting the node literal if its not required (which in turn
+				// avoids 'Reset To Default' action from being enabled when the default is equal)
+				bool bClearLiteral = false;
+				if (const FMetasoundFrontendLiteral* ClassDefault = InInputHandle->GetClassDefaultLiteral())
+				{
+					bClearLiteral = ClassDefault->IsEqual(InDefaultLiteral);
+				}
+
+				if (!bClearLiteral)
+				{
+					FMetasoundFrontendLiteral DefaultTypeLiteral;
+					DefaultTypeLiteral.SetFromLiteral(IDataTypeRegistry::Get().CreateDefaultLiteral(InInputHandle->GetDataType()));
+					bClearLiteral = InDefaultLiteral.IsEqual(DefaultTypeLiteral);
+				}
+
+				if (bClearLiteral)
+				{
+					InInputHandle->ClearLiteral();
+				}
+				else
+				{
+					InInputHandle->SetLiteral(InDefaultLiteral);
+				}
+			}
 		} // namespace GraphPrivate
 	} // namespace Editor
 } // namespace Metasound
@@ -133,14 +162,14 @@ void UMetasoundEditorGraphVertex::InitMember(FName InDataType, const FMetasoundF
 	}
 }
 
-TArray<UMetasoundEditorGraphNode*> UMetasoundEditorGraphVertex::GetNodes() const
+TArray<UMetasoundEditorGraphMemberNode*> UMetasoundEditorGraphVertex::GetNodes() const
 {
-	TArray<UMetasoundEditorGraphNode*> Nodes;
+	TArray<UMetasoundEditorGraphMemberNode*> Nodes;
 
 	const UMetasoundEditorGraph* Graph = GetOwningGraph();
 	if (ensure(Graph))
 	{
-		Graph->GetNodesOfClassEx<UMetasoundEditorGraphNode>(Nodes);
+		Graph->GetNodesOfClassEx<UMetasoundEditorGraphMemberNode>(Nodes);
 		for (int32 i = Nodes.Num() -1; i >= 0; --i)
 		{
 			UMetasoundEditorGraphNode* Node = Nodes[i];
@@ -245,8 +274,8 @@ void UMetasoundEditorGraphVertex::SetDisplayName(const FText& InNewName, bool bP
 		FNodeHandle NodeHandle = GetNodeHandle();
 		NodeHandle->SetDisplayName(InNewName);
 
-		const TArray<UMetasoundEditorGraphNode*> Nodes = GetNodes();
-		for (UMetasoundEditorGraphNode* Node : Nodes)
+		const TArray<UMetasoundEditorGraphMemberNode*> Nodes = GetNodes();
+		for (UMetasoundEditorGraphMemberNode* Node : Nodes)
 		{
 			const TArray<UEdGraphPin*>& Pins = Node->GetAllPins();
 			ensure(Pins.Num() == 1);
@@ -283,9 +312,9 @@ void UMetasoundEditorGraphVertex::SetDataType(FName InNewType, bool bPostTransac
 	Modify();
 
 	// 1. Cache current editor input node reference positions & delete nodes.
-	TArray<UMetasoundEditorGraphNode*> InputNodes = GetNodes();
+	TArray<UMetasoundEditorGraphMemberNode*> InputNodes = GetNodes();
 	TArray<FVector2D> NodeLocations;
-	for (UMetasoundEditorGraphNode* Node : InputNodes)
+	for (UMetasoundEditorGraphMemberNode* Node : InputNodes)
 	{
 		if (ensure(Node))
 		{
@@ -656,8 +685,8 @@ void UMetasoundEditorGraphOutput::ResetToClassDefault()
 	using namespace Metasound;
 	using namespace Metasound::Frontend;
 
-	TArray<UMetasoundEditorGraphNode*> Nodes = GetNodes();
-	for (UMetasoundEditorGraphNode* Node : Nodes)
+	TArray<UMetasoundEditorGraphMemberNode*> Nodes = GetNodes();
+	for (UMetasoundEditorGraphMemberNode* Node : Nodes)
 	{
 		TArray<FInputHandle> Inputs = Node->GetNodeHandle()->GetInputs();
 		if (ensure(!Inputs.IsEmpty()))
@@ -681,6 +710,7 @@ void UMetasoundEditorGraphOutput::ResetToClassDefault()
 void UMetasoundEditorGraphOutput::UpdateFrontendDefaultLiteral(bool bPostTransaction)
 {
 	using namespace Metasound;
+	using namespace Metasound::Editor;
 	using namespace Metasound::Frontend;
 
 	UObject* Metasound = nullptr;
@@ -705,13 +735,13 @@ void UMetasoundEditorGraphOutput::UpdateFrontendDefaultLiteral(bool bPostTransac
 
 	const FMetasoundFrontendLiteral DefaultLiteral = Literal->GetDefault();
 
-	TArray<UMetasoundEditorGraphNode*> Nodes = GetNodes();
-	for (UMetasoundEditorGraphNode* Node : Nodes)
+	TArray<UMetasoundEditorGraphMemberNode*> Nodes = GetNodes();
+	for (UMetasoundEditorGraphMemberNode* Node : Nodes)
 	{
 		TArray<FInputHandle> Inputs = Node->GetNodeHandle()->GetInputs();
 		if (ensure(!Inputs.IsEmpty()))
 		{
-			Inputs.Last()->SetLiteral(DefaultLiteral);
+			GraphPrivate::SetLiteralOrClearIfMatchesDefault(Inputs.Last(), DefaultLiteral);
 		}
 	}
 }
@@ -864,9 +894,9 @@ bool UMetasoundEditorGraphVariable::CanRename(const FText& InNewText, FText& Out
 	return true;
 }
 
-TArray<UMetasoundEditorGraphNode*> UMetasoundEditorGraphVariable::GetNodes() const
+TArray<UMetasoundEditorGraphMemberNode*> UMetasoundEditorGraphVariable::GetNodes() const
 {
-	TArray<UMetasoundEditorGraphNode*> Nodes;
+	TArray<UMetasoundEditorGraphMemberNode*> Nodes;
 
 	FVariableEditorNodes EditorNodes = GetVariableNodes();
 	if (nullptr != EditorNodes.MutatorNode)
@@ -988,26 +1018,26 @@ UMetasoundEditorGraphVariable::FVariableEditorNodes UMetasoundEditorGraphVariabl
 	using namespace Metasound::Frontend;
 
 	FVariableEditorNodes VariableNodes;
-	TArray<UMetasoundEditorGraphNode*> AllMetasoundNodes;
+	TArray<UMetasoundEditorGraphMemberNode*> AllMetasoundNodes;
 
 	const UMetasoundEditorGraph* Graph = GetOwningGraph();
 	if (ensure(Graph))
 	{
-		Graph->GetNodesOfClassEx<UMetasoundEditorGraphNode>(AllMetasoundNodes);
+		Graph->GetNodesOfClassEx<UMetasoundEditorGraphMemberNode>(AllMetasoundNodes);
 		FConstVariableHandle FrontendVariable = GetConstVariableHandle();
 
-		// Find the mutator node if it exists. 
+		// Find the mutator node if it exists.
 		{
 			FConstNodeHandle FrontendMutatorNode = FrontendVariable->FindMutatorNode();
 			if (FrontendMutatorNode->IsValid())
 			{
 				const FGuid& MutatorNodeID = FrontendMutatorNode->GetID();
-				auto IsNodeWithID = [&MutatorNodeID](const UMetasoundEditorGraphNode* InNode)
+				auto IsNodeWithID = [&MutatorNodeID](const UMetasoundEditorGraphMemberNode* InNode)
 				{
 					return (nullptr != InNode) && (MutatorNodeID == InNode->GetNodeID());
 				};
 
-				if (UMetasoundEditorGraphNode** FoundMutatorNode = AllMetasoundNodes.FindByPredicate(IsNodeWithID))
+				if (UMetasoundEditorGraphMemberNode** FoundMutatorNode = AllMetasoundNodes.FindByPredicate(IsNodeWithID))
 				{
 					VariableNodes.MutatorNode = *FoundMutatorNode;
 				}
@@ -1021,7 +1051,7 @@ UMetasoundEditorGraphVariable::FVariableEditorNodes UMetasoundEditorGraphVariabl
 			{
 				AccessorNodeIDs.Add(FrontendAccessorNode->GetID());
 			}
-			auto IsNodeInAccessorSet = [&AccessorNodeIDs](const UMetasoundEditorGraphNode* InNode)
+			auto IsNodeInAccessorSet = [&AccessorNodeIDs](const UMetasoundEditorGraphMemberNode* InNode)
 			{
 				return (nullptr != InNode) && AccessorNodeIDs.Contains(InNode->GetNodeID());
 			};
@@ -1035,7 +1065,7 @@ UMetasoundEditorGraphVariable::FVariableEditorNodes UMetasoundEditorGraphVariabl
 			{
 				DeferredAccessorNodeIDs.Add(FrontendAccessorNode->GetID());
 			}
-			auto IsNodeInDeferredAccessorSet = [&DeferredAccessorNodeIDs](const UMetasoundEditorGraphNode* InNode)
+			auto IsNodeInDeferredAccessorSet = [&DeferredAccessorNodeIDs](const UMetasoundEditorGraphMemberNode* InNode)
 			{
 				return (nullptr != InNode) && DeferredAccessorNodeIDs.Contains(InNode->GetNodeID());
 			};
@@ -1052,7 +1082,7 @@ UMetasoundEditorGraphVariable::FVariableNodeLocations UMetasoundEditorGraphVaria
 	FVariableNodeLocations Locations;
 	// Cache current node positions 
 	FVariableEditorNodes EditorNodes = GetVariableNodes();
-	auto GetNodeLocation = [](const UMetasoundEditorGraphNode* InNode) { return FVector2D(InNode->NodePosX, InNode->NodePosY); };
+	auto GetNodeLocation = [](const UMetasoundEditorGraphMemberNode* InNode) { return FVector2D(InNode->NodePosX, InNode->NodePosY); };
 
 	if (nullptr != EditorNodes.MutatorNode)
 	{
@@ -1130,6 +1160,7 @@ void UMetasoundEditorGraphVariable::ResetToClassDefault()
 void UMetasoundEditorGraphVariable::UpdateFrontendDefaultLiteral(bool bPostTransaction)
 {
 	using namespace Metasound;
+	using namespace Metasound::Editor;
 	using namespace Metasound::Frontend;
 
 	UObject* Metasound = nullptr;
@@ -1162,7 +1193,7 @@ void UMetasoundEditorGraphVariable::UpdateFrontendDefaultLiteral(bool bPostTrans
 		FInputHandle InputHandle = MutatorNode->GetInputWithVertexName(VariableNames::GetInputDataName());
 		if (ensure(InputHandle->IsValid()))
 		{
-			InputHandle->SetLiteral(DefaultLiteral);
+			GraphPrivate::SetLiteralOrClearIfMatchesDefault(InputHandle, DefaultLiteral);
 		}
 	}
 }
diff --git a/Engine/Plugins/Runtime/Metasound/Source/MetasoundEditor/Private/MetasoundEditorGraphBuilder.cpp b/Engine/Plugins/Runtime/Metasound/Source/MetasoundEditor/Private/MetasoundEditorGraphBuilder.cpp
--- a/Engine/Plugins/Runtime/Metasound/Source/MetasoundEditor/Private/MetasoundEditorGraphBuilder.cpp
+++ b/Engine/Plugins/Runtime/Metasound/Source/MetasoundEditor/Private/MetasoundEditorGraphBuilder.cpp
@@ -278,18 +278,20 @@ namespace Metasound
 			FGraphNodeCreator<UMetasoundEditorGraphExternalNode> NodeCreator(Graph);
 
 			NewGraphNode = NodeCreator.CreateNode(bInSelectNewNode);
+			if (ensure(NewGraphNode))
+			{
+				const FNodeRegistryKey RegistryKey = NodeRegistryKey::CreateKey(InNodeHandle->GetClassMetadata());
+				NewGraphNode->bIsClassNative = FMetasoundFrontendRegistryContainer::Get()->IsNodeNative(RegistryKey);
+				NewGraphNode->ClassName = InNodeHandle->GetClassMetadata().GetClassName();
+				NewGraphNode->CacheTitle();
 
-			const FNodeRegistryKey RegistryKey = NodeRegistryKey::CreateKey(InNodeHandle->GetClassMetadata());
-			NewGraphNode->bIsClassNative = FMetasoundFrontendRegistryContainer::Get()->IsNodeNative(RegistryKey);
-			NewGraphNode->ClassName = InNodeHandle->GetClassMetadata().GetClassName();
-			NewGraphNode->CacheTitle();
-
-			NodeCreator.Finalize();
-			InitGraphNode(InNodeHandle, NewGraphNode, InMetaSound);
-			NewGraphNode->SetNodeLocation(InLocation);
+				NodeCreator.Finalize();
+				InitGraphNode(InNodeHandle, NewGraphNode, InMetaSound);
+				NewGraphNode->SetNodeLocation(InLocation);
 
-			// Adding external node may introduce referenced asset so rebuild referenced keys.
-			MetaSoundAsset->RebuildReferencedAssetClassKeys();
+				// Adding external node may introduce referenced asset so rebuild referenced keys.
+				MetaSoundAsset->RebuildReferencedAssetClassKeys();
+			}
 
 			return NewGraphNode;
 		}
@@ -332,13 +334,26 @@ namespace Metasound
 					FGraphNodeCreator<UMetasoundEditorGraphVariableNode> NodeCreator(*MetasoundGraph);
 
 					NewGraphNode = NodeCreator.CreateNode(bInSelectNewNode);
-					NewGraphNode->ClassName = InNodeHandle->GetClassMetadata().GetClassName();
-					NewGraphNode->ClassType = ClassType;
-					NewGraphNode->Variable = MetasoundGraph->FindOrAddVariable(FrontendVariable);
-					NodeCreator.Finalize();
+					if (ensure(NewGraphNode))
+					{
+						NewGraphNode->ClassName = InNodeHandle->GetClassMetadata().GetClassName();
+						NewGraphNode->ClassType = ClassType;
+						NodeCreator.Finalize();
+						InitGraphNode(InNodeHandle, NewGraphNode, InMetaSound);
 
-					InitGraphNode(InNodeHandle, NewGraphNode, InMetaSound);
-					NewGraphNode->SetNodeLocation(InLocation);
+						UMetasoundEditorGraphVariable* Variable = MetasoundGraph->FindOrAddVariable(FrontendVariable);
+						if (ensure(Variable))
+						{
+							NewGraphNode->Variable = Variable;
+
+							// Ensures the variable node value is synced with the editor literal value should it be set
+							constexpr bool bPostTransaction = false;
+							Variable->UpdateFrontendDefaultLiteral(bPostTransaction);
+						}
+
+						MetasoundGraph->SetSynchronizationRequired();
+						NewGraphNode->SetNodeLocation(InLocation);
+					}
 				}
 			}
 
@@ -361,22 +376,34 @@ namespace Metasound
 			FGraphNodeCreator<UMetasoundEditorGraphOutputNode> NodeCreator(Graph);
 
 			NewGraphNode = NodeCreator.CreateNode(bInSelectNewNode);
-			UMetasoundEditorGraph* MetasoundGraph = CastChecked<UMetasoundEditorGraph>(&Graph);
-			NewGraphNode->Output = MetasoundGraph->FindOrAddOutput(InNodeHandle);
-			NewGraphNode->CacheTitle();
+			if (ensure(NewGraphNode))
+			{
+				UMetasoundEditorGraph* MetasoundGraph = CastChecked<UMetasoundEditorGraph>(&Graph);
 
-			NodeCreator.Finalize();
-			InitGraphNode(InNodeHandle, NewGraphNode, InMetaSound);
-			NewGraphNode->SetNodeLocation(InLocation);
+				UMetasoundEditorGraphOutput* Output = MetasoundGraph->FindOrAddOutput(InNodeHandle);
+				if (ensure(Output))
+				{
+					NewGraphNode->Output = Output;
+					NodeCreator.Finalize();
+					InitGraphNode(InNodeHandle, NewGraphNode, InMetaSound);
+
+					// Ensures the output node value is synced with the editor literal value should it be set
+					constexpr bool bPostTransaction = false;
+					Output->UpdateFrontendDefaultLiteral(bPostTransaction);
+
+					MetasoundGraph->SetSynchronizationRequired();
+				}
+
+				NewGraphNode->CacheTitle();
+				NewGraphNode->SetNodeLocation(InLocation);
+			}
 
 			return NewGraphNode;
 		}
 
 		void FGraphBuilder::InitGraphNode(Frontend::FNodeHandle& InNodeHandle, UMetasoundEditorGraphNode* NewGraphNode, UObject& InMetaSound)
 		{
-			NewGraphNode->CreateNewGuid();
 			NewGraphNode->SetNodeID(InNodeHandle->GetID());
-
 			RebuildNodePins(*NewGraphNode);
 		}
 
@@ -415,7 +442,6 @@ namespace Metasound
 				Result.Node->CacheTitle();
 				const bool bTitleUpdated = !Title.IdenticalTo(Result.Node->GetCachedTitle());
 
-
 				if (Result.bIsDirty || bTitleUpdated || bMetadataChange || bInterfaceChange || bStyleChange || bForceRefreshNodes)
 				{
 					Result.Node->SyncChangeIDs();
@@ -673,6 +699,7 @@ namespace Metasound
 			{
 				NewGraphNode->SetNodeLocation(InLocation);
 				RebuildNodePins(*NewGraphNode);
+				MetasoundGraph->SetSynchronizationRequired();
 				return NewGraphNode;
 			}
 
diff --git a/Engine/Plugins/Runtime/Metasound/Source/MetasoundEditor/Private/MetasoundEditorGraphSchema.cpp b/Engine/Plugins/Runtime/Metasound/Source/MetasoundEditor/Private/MetasoundEditorGraphSchema.cpp
--- a/Engine/Plugins/Runtime/Metasound/Source/MetasoundEditor/Private/MetasoundEditorGraphSchema.cpp
+++ b/Engine/Plugins/Runtime/Metasound/Source/MetasoundEditor/Private/MetasoundEditorGraphSchema.cpp
@@ -489,11 +489,9 @@ UEdGraphNode* FMetasoundGraphSchemaAction_NewInput::PerformAction(UEdGraph* Pare
 	if (UMetasoundEditorGraphInputNode* NewGraphNode = FGraphBuilder::AddInputNode(ParentMetasound, NodeHandle, InLocation))
 	{
 		NewGraphNode->Modify();
-		UEdGraphNode* EdGraphNode = CastChecked<UEdGraphNode>(NewGraphNode);
-		SchemaPrivate::TryConnectNewNodeToMatchingDataTypePin(*EdGraphNode, FromPin);
-
+		SchemaPrivate::TryConnectNewNodeToMatchingDataTypePin(*NewGraphNode, FromPin);
 		FGraphBuilder::RegisterGraphWithFrontend(ParentMetasound);
-		return EdGraphNode;
+		return NewGraphNode;
 	}
 
 	return nullptr;
@@ -796,15 +794,7 @@ UEdGraphNode* FMetasoundGraphSchemaAction_NewVariableNode::PerformAction(UEdGrap
 					if (UMetasoundEditorGraphVariableNode* NewGraphNode = FGraphBuilder::AddVariableNode(*ParentMetasound, FrontendNode, Location, bSelectNewNode))
 					{
 						NewGraphNode->Modify();
-
-						// Ensures the setter node value is synced with the editor literal value
-						constexpr bool bPostTransaction = false;
-						Variable->UpdateFrontendDefaultLiteral(bPostTransaction);
-
 						SchemaPrivate::TryConnectNewNodeToMatchingDataTypePin(*NewGraphNode, FromPin);
-
-						MetasoundGraph->SetSynchronizationRequired();
-
 						return NewGraphNode;
 					}
 				}
diff --git a/Engine/Plugins/Runtime/Metasound/Source/MetasoundEditor/Private/MetasoundNodeDetailCustomization.cpp b/Engine/Plugins/Runtime/Metasound/Source/MetasoundEditor/Private/MetasoundNodeDetailCustomization.cpp
--- a/Engine/Plugins/Runtime/Metasound/Source/MetasoundEditor/Private/MetasoundNodeDetailCustomization.cpp
+++ b/Engine/Plugins/Runtime/Metasound/Source/MetasoundEditor/Private/MetasoundNodeDetailCustomization.cpp
@@ -897,6 +897,17 @@ namespace Metasound
 				}
 			}
 
+			for (IDetailPropertyRow* Row : DefaultPropertyRows)
+			{
+				if (ensure(Row))
+				{
+					Row->Visibility(TAttribute<EVisibility>::CreateLambda([this]()
+					{
+						return GetDefaultVisibility();
+					}));
+				}
+			}
+
 			return DefaultPropertyRows;
 		}
 
@@ -1013,6 +1024,19 @@ namespace Metasound
 			return FText::GetEmpty();
 		}
 
+		Frontend::FDocumentHandle FMetasoundMemberDetailCustomization::GetDocumentHandle() const
+		{
+			if (GraphMember.IsValid())
+			{
+				if (UMetasoundEditorGraph* Graph = GraphMember->GetOwningGraph())
+				{
+					return Graph->GetDocumentHandle();
+				}
+			}
+
+			return Frontend::IDocumentController::GetInvalidHandle();
+		}
+
 		bool FMetasoundMemberDetailCustomization::IsGraphEditable() const
 		{
 			if (GraphMember.IsValid())
@@ -1059,6 +1083,27 @@ namespace Metasound
 			return FText::GetEmpty();
 		}
 
+		EVisibility FMetasoundVertexDetailCustomization::GetDefaultVisibility() const
+		{
+			using namespace Frontend;
+
+			if (GraphMember.IsValid())
+			{
+				bool bIsInputConnected = false;
+				FConstNodeHandle NodeHandle = CastChecked<UMetasoundEditorGraphVertex>(GraphMember)->GetConstNodeHandle();
+				if (NodeHandle->IsValid())
+				{
+					NodeHandle->IterateConstInputs([&bIsInputConnected](FConstInputHandle InputHandle)
+					{
+						bIsInputConnected |= InputHandle->IsConnectionUserModifiable() && InputHandle->IsConnected();
+					});
+				}
+				return bIsInputConnected ? EVisibility::Collapsed : EVisibility::Visible;
+			}
+
+			return EVisibility::Collapsed;
+		}
+
 		void FMetasoundMemberDetailCustomization::OnNameCommitted(const FText& InNewName, ETextCommit::Type InTextCommit)
 		{
 			using namespace Frontend;
@@ -1135,17 +1180,14 @@ namespace Metasound
 			}
 		}
 
-		Frontend::FDocumentHandle FMetasoundInputDetailCustomization::GetDocumentHandle() const
+		bool FMetasoundVertexDetailCustomization::IsInterfaceMember() const
 		{
 			if (GraphMember.IsValid())
 			{
-				if (UMetasoundEditorGraph* Graph = GraphMember->GetOwningGraph())
-				{
-					return Graph->GetDocumentHandle();
-				}
+				return CastChecked<UMetasoundEditorGraphVertex>(GraphMember)->IsInterfaceMember();
 			}
 
-			return Frontend::IDocumentController::GetInvalidHandle();
+			return false;
 		}
 
 		bool FMetasoundInputDetailCustomization::GetInputInheritsDefault() const
@@ -1297,25 +1339,28 @@ namespace Metasound
 			return !GetInputInheritsDefault();
 		}
 
-		bool FMetasoundInputDetailCustomization::IsInterfaceMember() const
+		EVisibility FMetasoundVariableDetailCustomization::GetDefaultVisibility() const
 		{
-			if (GraphMember.IsValid())
-			{
-				return CastChecked<UMetasoundEditorGraphVertex>(GraphMember)->IsInterfaceMember();
-			}
-
-			return false;
-		}
+			using namespace Frontend;
 
-		bool FMetasoundOutputDetailCustomization::IsInterfaceMember() const
-		{
 			if (GraphMember.IsValid())
 			{
-				return CastChecked<UMetasoundEditorGraphVertex>(GraphMember)->IsInterfaceMember();
+				bool bIsInputConnected = false;
+				const UMetasoundEditorGraphVariable* Variable = CastChecked<UMetasoundEditorGraphVariable>(GraphMember);
+				FConstNodeHandle NodeHandle = Variable->GetConstVariableHandle()->FindMutatorNode();
+				if (NodeHandle->IsValid())
+				{
+					NodeHandle->IterateConstInputs([&bIsInputConnected](FConstInputHandle InputHandle)
+					{
+						bIsInputConnected |= InputHandle->IsConnectionUserModifiable() && InputHandle->IsConnected();
+					});
+				}
+				return bIsInputConnected ? EVisibility::Collapsed : EVisibility::Visible;
 			}
 
-			return false;
+			return EVisibility::Collapsed;
 		}
+
 	} // namespace Editor
 } // namespace Metasound
 #undef LOCTEXT_NAMESPACE
diff --git a/Engine/Plugins/Runtime/Metasound/Source/MetasoundEditor/Private/MetasoundNodeDetailCustomization.h b/Engine/Plugins/Runtime/Metasound/Source/MetasoundEditor/Private/MetasoundNodeDetailCustomization.h
--- a/Engine/Plugins/Runtime/Metasound/Source/MetasoundEditor/Private/MetasoundNodeDetailCustomization.h
+++ b/Engine/Plugins/Runtime/Metasound/Source/MetasoundEditor/Private/MetasoundNodeDetailCustomization.h
@@ -166,7 +166,7 @@ namespace Metasound
 			static constexpr float DetailsTitleWrapPadding = 32.0f;
 
 			static const FText DataTypeNameText = LOCTEXT("Node_DataTypeName", "Type");
-			static const FText DefaultPropertyText = LOCTEXT("Node_DefaultPropertyName", "Default Value");
+			static const FText DefaultPropertyText = LOCTEXT("Node_DefaultPropertyName", "Default");
 		} // namespace MemberCustomizationStyle
 
 		class FMetasoundFloatLiteralCustomization : public FMetasoundDefaultLiteralCustomizationBase
@@ -316,21 +316,16 @@ namespace Metasound
 			virtual void CustomizeGeneralCategory(IDetailLayoutBuilder& InDetailLayout);
 			virtual TArray<IDetailPropertyRow*> CustomizeDefaultCategory(IDetailLayoutBuilder& InDetailLayout);
 
-			virtual bool IsDefaultEditable() const
-			{
-				return true;
-			}
-
-			virtual bool IsInterfaceMember() const
-			{
-				return false;
-			}
+			virtual EVisibility GetDefaultVisibility() const { return EVisibility::Visible; }
+			virtual bool IsDefaultEditable() const { return true; }
+			virtual bool IsInterfaceMember() const { return false; }
 
 			// IDetailCustomization interface
 			virtual void CustomizeDetails(IDetailLayoutBuilder& InDetailLayout) override;
 			// End of IDetailCustomization interface
 
 			void OnNameChanged(const FText& InNewName);
+			Frontend::FDocumentHandle GetDocumentHandle() const;
 			FText GetName() const;
 			bool IsGraphEditable() const;
 			FText GetDisplayName() const;
@@ -349,8 +344,6 @@ namespace Metasound
 			FDelegateHandle RenameRequestedHandle;
 		};
 
-		using FMetasoundVariableDetailCustomization = FMetasoundMemberDetailCustomization;
-
 		class FMetasoundVertexDetailCustomization : public FMetasoundMemberDetailCustomization
 		{
 		public:
@@ -358,6 +351,8 @@ namespace Metasound
 
 		protected:
 			virtual void CustomizeGeneralCategory(IDetailLayoutBuilder& InDetailLayout) override;
+			virtual EVisibility GetDefaultVisibility() const;
+			virtual bool IsInterfaceMember() const override;
 		};
 
 		class FMetasoundInputDetailCustomization : public FMetasoundVertexDetailCustomization
@@ -367,22 +362,19 @@ namespace Metasound
 
 			virtual void CustomizeDetails(IDetailLayoutBuilder& InDetailLayout) override;
 			virtual bool IsDefaultEditable() const override;
-			virtual bool IsInterfaceMember() const override;
 
 		private:
-			Frontend::FDocumentHandle GetDocumentHandle() const;
 			bool GetInputInheritsDefault() const;
 			void SetInputInheritsDefault();
 			void ClearInputInheritsDefault();
 		};
 
-		class FMetasoundOutputDetailCustomization : public FMetasoundVertexDetailCustomization
+		class FMetasoundVariableDetailCustomization : public FMetasoundMemberDetailCustomization
 		{
-		public:
-			virtual ~FMetasoundOutputDetailCustomization() = default;
-
-			virtual bool IsInterfaceMember() const override;
+			EVisibility GetDefaultVisibility() const override;
 		};
+
+		using FMetasoundOutputDetailCustomization = FMetasoundVertexDetailCustomization;
 	} // namespace Editor
 } // namespace Metasound
 #undef LOCTEXT_NAMESPACE
diff --git a/Engine/Plugins/Runtime/Metasound/Source/MetasoundEditor/Public/MetasoundEditorGraph.h b/Engine/Plugins/Runtime/Metasound/Source/MetasoundEditor/Public/MetasoundEditorGraph.h
--- a/Engine/Plugins/Runtime/Metasound/Source/MetasoundEditor/Public/MetasoundEditorGraph.h
+++ b/Engine/Plugins/Runtime/Metasound/Source/MetasoundEditor/Public/MetasoundEditorGraph.h
@@ -103,7 +103,7 @@ public:
 	virtual Metasound::Editor::ENodeSection GetSectionID() const PURE_VIRTUAL(UMetasoundEditorGraphMember::GetSectionID, return Metasound::Editor::ENodeSection::None; );
 
 	/** Return the nodes associated with this member */
-	virtual TArray<UMetasoundEditorGraphNode*> GetNodes() const PURE_VIRTUAL(UMetasoundEditorGraphMember::GetNodes, return TArray<UMetasoundEditorGraphNode*>(); );
+	virtual TArray<UMetasoundEditorGraphMemberNode*> GetNodes() const PURE_VIRTUAL(UMetasoundEditorGraphMember::GetNodes, return TArray<UMetasoundEditorGraphMemberNode*>(); );
 
 	/** Sets the datatype on the member. */
 	virtual void SetDataType(FName InNewType, bool bPostTransaction = true) PURE_VIRTUAL(UMetasoundEditorGraphMember::SetDataType, );
@@ -213,7 +213,7 @@ public:
 	virtual FText GetDisplayName() const override;
 
 	virtual bool CanRename(const FText& InNewName, FText& OutError) const override;
-	virtual TArray<UMetasoundEditorGraphNode*> GetNodes() const override;
+	virtual TArray<UMetasoundEditorGraphMemberNode*> GetNodes() const override;
 	virtual void SetDescription(const FText& InDescription, bool bPostTransaction) override;
 	virtual void SetMemberName(const FName& InNewName, bool bPostTransaction) override;
 	virtual void SetDisplayName(const FText& InNewName, bool bPostTransaction) override;
@@ -312,7 +312,7 @@ public:
 	virtual void ResetToClassDefault() override;
 	virtual void UpdateFrontendDefaultLiteral(bool bPostTransaction) override;
 
-	virtual TArray<UMetasoundEditorGraphNode*> GetNodes() const override;
+	virtual TArray<UMetasoundEditorGraphMemberNode*> GetNodes() const override;
 
 	virtual const FText& GetGraphMemberLabel() const override;
 	/* ~EndUMetasoundEditorGraphMember interface */
@@ -327,9 +327,9 @@ public:
 private:
 	struct FVariableEditorNodes
 	{
-		UMetasoundEditorGraphNode* MutatorNode = nullptr;
-		TArray<UMetasoundEditorGraphNode*> AccessorNodes;
-		TArray<UMetasoundEditorGraphNode*> DeferredAccessorNodes;
+		UMetasoundEditorGraphMemberNode* MutatorNode = nullptr;
+		TArray<UMetasoundEditorGraphMemberNode*> AccessorNodes;
+		TArray<UMetasoundEditorGraphMemberNode*> DeferredAccessorNodes;
 	};
 
 	struct FVariableNodeLocations
diff --git a/Engine/Plugins/Runtime/Metasound/Source/MetasoundEditor/Public/MetasoundEditorGraphNode.h b/Engine/Plugins/Runtime/Metasound/Source/MetasoundEditor/Public/MetasoundEditorGraphNode.h
--- a/Engine/Plugins/Runtime/Metasound/Source/MetasoundEditor/Public/MetasoundEditorGraphNode.h
+++ b/Engine/Plugins/Runtime/Metasound/Source/MetasoundEditor/Public/MetasoundEditorGraphNode.h
@@ -2,7 +2,6 @@
 
 #pragma once
 
-#include "CoreMinimal.h"
 #include "EdGraph/EdGraphNode.h"
 #include "MetasoundFrontend.h"
 #include "MetasoundFrontendController.h"
