Add logging for failures when matching MetaSound parameters, prevent ensure on enum-string mismatch

#rnx

diff --git a/Engine/Plugins/Runtime/Metasound/Source/MetasoundEngine/Private/MetasoundSource.cpp b/Engine/Plugins/Runtime/Metasound/Source/MetasoundEngine/Private/MetasoundSource.cpp
--- a/Engine/Plugins/Runtime/Metasound/Source/MetasoundEngine/Private/MetasoundSource.cpp
+++ b/Engine/Plugins/Runtime/Metasound/Source/MetasoundEngine/Private/MetasoundSource.cpp
@@ -425,22 +425,34 @@ void UMetaSoundSource::InitParameters(TArray<FAudioParameter>& InParametersToIni
 
 	for (int32 i = InParametersToInit.Num() - 1; i >= 0; --i)
 	{
-		FAudioParameter& Parameter = InParametersToInit[i];
+		FAudioParameter& Parameter = InParametersToInit[i];		
+		
+#if !NO_LOGGING
+		// For logging in case of failure
+		FString AssetName;
+		GetName(AssetName); 
+#endif // !NO_LOGGING
+		
 		if (Sanitize(Parameter))
 		{
 			if (IsParameterValid(Parameter, InputNameTypeMap))
 			{
 				ConstructProxies(Parameter);
 			}
+#if !NO_LOGGING
 			else
 			{
-				UE_LOG(LogMetaSound, Error, TEXT("Failed to set invalid parameter '%s': Either does not exist or is unsupported type"), *Parameter.ParamName.ToString());
+				UE_LOG(LogMetaSound, Error, TEXT("Failed to set invalid parameter '%s' in asset '%s': Either does not exist or is unsupported type"), *Parameter.ParamName.ToString(), *AssetName);
 			}
+#endif // !NO_LOGGING
 		}
 		else
 		{
 			constexpr bool bAllowShrinking = false;
 			InParametersToInit.RemoveAtSwap(i, 1, bAllowShrinking);
+#if !NO_LOGGING
+			UE_LOG(LogMetaSound, Error, TEXT("Failed to set parameter '%s' in asset '%s': No name specified, no matching input found, or type mismatch."), *Parameter.ParamName.ToString(), *AssetName);
+#endif // !NO_LOGGING
 		}
 	}
 }
diff --git a/Engine/Plugins/Runtime/Metasound/Source/MetasoundFrontend/Public/MetasoundEnum.h b/Engine/Plugins/Runtime/Metasound/Source/MetasoundFrontend/Public/MetasoundEnum.h
--- a/Engine/Plugins/Runtime/Metasound/Source/MetasoundFrontend/Public/MetasoundEnum.h
+++ b/Engine/Plugins/Runtime/Metasound/Source/MetasoundFrontend/Public/MetasoundEnum.h
@@ -3,6 +3,7 @@
 #pragma once
 
 #include "CoreMinimal.h"
+#include "MetasoundLog.h"
 
 #include <type_traits>
 
@@ -131,12 +132,20 @@ namespace Metasound
 		// From Name
 		explicit TEnum(FName InValueName)		
 		{
-			// Try and convert from Name to Value and ensure that we succeed
+			// Try and convert from Name to Value 
 			TOptional<EnumType> Converted = NameToEnum(InValueName);
-			if (ensure(Converted))
+			if (Converted)
 			{
 				EnumValue = *Converted;
 			}
+			else
+			{
+				if (!bHasWarnedNameToEnumConversionFailure)
+				{
+					UE_LOG(LogMetaSound, Warning, TEXT("Cannot create valid enum value from string '%s'"), *InValueName.ToString());
+					bHasWarnedNameToEnumConversionFailure = true;
+				}
+			}
 		}
 
 		// Slow, construct from FString to FName.
@@ -188,8 +197,13 @@ namespace Metasound
 	private:
 		// Keep the type in its fully typed form for debugging.
 		EnumType EnumValue = DefaultValue;
+
+		static bool bHasWarnedNameToEnumConversionFailure;
 	};
 
+	template<typename EnumType, EnumType DefaultValue>
+	bool TEnum<EnumType, DefaultValue>::bHasWarnedNameToEnumConversionFailure = false;
+
 	template<typename T>
 	struct TEnumTraits
 	{
