Fix for ensure on corrupt MetaSound asset in Lyra
#rnx

diff --git a/Engine/Plugins/Runtime/Metasound/Source/MetasoundEngine/Private/MetasoundSource.cpp b/Engine/Plugins/Runtime/Metasound/Source/MetasoundEngine/Private/MetasoundSource.cpp
--- a/Engine/Plugins/Runtime/Metasound/Source/MetasoundEngine/Private/MetasoundSource.cpp
+++ b/Engine/Plugins/Runtime/Metasound/Source/MetasoundEngine/Private/MetasoundSource.cpp
@@ -514,7 +514,6 @@ ISoundGeneratorPtr UMetaSoundSource::CreateSoundGenerator(const FSoundGeneratorI
 	TSharedPtr<const IGraph, ESPMode::ThreadSafe> MetasoundGraph = GetMetasoundCoreGraph();
 	if (!MetasoundGraph.IsValid())
 	{
-		UE_LOG(LogMetaSound, Error, TEXT("Cannot create UMetaSoundSource SoundGenerator [Name:%s]. Failed to create MetaSound Graph"), *GetName());
 		return ISoundGeneratorPtr(nullptr);
 	}
 
diff --git a/Engine/Plugins/Runtime/Metasound/Source/MetasoundFrontend/Private/MetasoundAssetBase.cpp b/Engine/Plugins/Runtime/Metasound/Source/MetasoundFrontend/Private/MetasoundAssetBase.cpp
--- a/Engine/Plugins/Runtime/Metasound/Source/MetasoundFrontend/Private/MetasoundAssetBase.cpp
+++ b/Engine/Plugins/Runtime/Metasound/Source/MetasoundFrontend/Private/MetasoundAssetBase.cpp
@@ -159,7 +159,7 @@ void FMetasoundAssetBase::RegisterGraphWithFrontend(Metasound::Frontend::FMetaSo
 
 		virtual TUniquePtr<INode> CreateNode(const FNodeInitData&) const override
 		{
-			return FFrontendGraphBuilder().CreateGraph(Document);
+			return FFrontendGraphBuilder().CreateGraph(Document, Name);
 		}
 
 		virtual TUniquePtr<INode> CreateNode(FDefaultLiteralNodeConstructorParams&&) const override
@@ -546,8 +546,8 @@ TSharedPtr<Metasound::IGraph, ESPMode::ThreadSafe> FMetasoundAssetBase::BuildMet
 
 	METASOUND_TRACE_CPUPROFILER_EVENT_SCOPE(MetaSoundAssetBase::BuildMetasoundDocument);
 
-	// Create graph which can spawn instances. TODO: cache graph.
-	TUniquePtr<FFrontendGraph> FrontendGraph = FFrontendGraphBuilder::CreateGraph(GetDocumentChecked());
+	// Create graph which can spawn instances. 
+	TUniquePtr<FFrontendGraph> FrontendGraph = FFrontendGraphBuilder::CreateGraph(GetDocumentChecked(), GetOwningAssetName());
 	if (FrontendGraph.IsValid())
 	{
 		const TArray<const FMetasoundFrontendClassInput*> TransmittableInputs = GetTransmittableClassInputs();
@@ -561,6 +561,10 @@ TSharedPtr<Metasound::IGraph, ESPMode::ThreadSafe> FMetasoundAssetBase::BuildMet
 			UE_LOG(LogMetaSound, Error, TEXT("Error while injecting async communication hooks. Instance communication may not function properly [Name:%s]."), *GetOwningAssetName());
 		}
 	}
+	else
+	{
+		UE_LOG(LogMetaSound, Error, TEXT("Failed to build MetaSound graph in asset '%s'"), *GetOwningAssetName());
+	}
 
 	TSharedPtr<Metasound::IGraph, ESPMode::ThreadSafe> SharedGraph(FrontendGraph.Release());
 
diff --git a/Engine/Plugins/Runtime/Metasound/Source/MetasoundFrontend/Private/MetasoundFrontendGraph.cpp b/Engine/Plugins/Runtime/Metasound/Source/MetasoundFrontend/Private/MetasoundFrontendGraph.cpp
--- a/Engine/Plugins/Runtime/Metasound/Source/MetasoundFrontend/Private/MetasoundFrontendGraph.cpp
+++ b/Engine/Plugins/Runtime/Metasound/Source/MetasoundFrontend/Private/MetasoundFrontendGraph.cpp
@@ -394,8 +394,7 @@ namespace Metasound
 		return Literal;
 	}
 
-	// TODO: add errors here. Most will be a "PromptIfMissing"...
-	void FFrontendGraphBuilder::AddNodesToGraph(FBuildGraphContext& InGraphContext)
+	bool FFrontendGraphBuilder::AddNodesToGraph(FBuildGraphContext& InGraphContext)
 	{
 		TSet<FNodeIDVertexID> GraphEdgeDestinations;
 		const TArray<FMetasoundFrontendEdge>& GraphEdges = InGraphContext.GraphClass.Graph.Edges;
@@ -425,7 +424,8 @@ namespace Metasound
 						else
 						{
 							const FString GraphClassIDString = InGraphContext.GraphClass.ID.ToString();
-							UE_LOG(LogMetaSound, Error, TEXT("Failed to match input node [NodeID:%s, NodeName:%s] to owning graph [ClassID:%s] output."), *Node.GetID().ToString(), *Node.Name.ToString(), *GraphClassIDString);
+							UE_LOG(LogMetaSound, Error, TEXT("MetaSound '%s': Failed to match input node [NodeID:%s, NodeName:%s] to owning graph [ClassID:%s] output."), *InGraphContext.BuildContext.DebugAssetName, *Node.GetID().ToString(), *Node.Name.ToString(), *GraphClassIDString);
+							return false;
 						}
 					}
 					break;
@@ -442,7 +442,8 @@ namespace Metasound
 						else
 						{
 							const FString GraphClassIDString = InGraphContext.GraphClass.ID.ToString();
-							UE_LOG(LogMetaSound, Error, TEXT("Failed to match output node [NodeID:%s, NodeName:%s] to owning graph [ClassID:%s] output."), *Node.GetID().ToString(), *Node.Name.ToString(), *GraphClassIDString);
+							UE_LOG(LogMetaSound, Error, TEXT("MetaSound '%s': Failed to match output node [NodeID:%s, NodeName:%s] to owning graph [ClassID:%s] output."), *InGraphContext.BuildContext.DebugAssetName, *Node.GetID().ToString(), *Node.Name.ToString(), *GraphClassIDString);
+							return false;
 						}
 					}
 					break;
@@ -457,7 +458,8 @@ namespace Metasound
 						}
 						else
 						{
-							UE_LOG(LogMetaSound, Error, TEXT("Failed to find subgraph for node [NodeID:%s, NodeName:%s, ClassID:%s]"), *Node.GetID().ToString(), *Node.Name.ToString(), *Node.ClassID.ToString());
+							UE_LOG(LogMetaSound, Error, TEXT("MetaSound '%s': Failed to find subgraph for node [NodeID:%s, NodeName:%s, ClassID:%s]"), *InGraphContext.BuildContext.DebugAssetName, *Node.GetID().ToString(), *Node.Name.ToString(), *Node.ClassID.ToString());
+							return false;
 						}
 					}
 					break;
@@ -465,6 +467,7 @@ namespace Metasound
 					case EMetasoundFrontendClassType::Literal:
 					{
 						checkNoEntry(); // Unsupported.
+						return false;
 					}
 
 					case EMetasoundFrontendClassType::Variable:
@@ -489,9 +492,11 @@ namespace Metasound
 				}
 			}
 		}
+
+		return true;
 	}
 
-	void FFrontendGraphBuilder::AddEdgesToGraph(FBuildGraphContext& InGraphContext)
+	bool FFrontendGraphBuilder::AddEdgesToGraph(FBuildGraphContext& InGraphContext)
 	{
 		// Pair of frontend node and core node. The frontend node can be one of
 		// several types.
@@ -510,8 +515,8 @@ namespace Metasound
 			const INode* CoreNode = InGraphContext.Graph->FindNode(Node.GetID());
 			if (nullptr == CoreNode)
 			{
-				UE_LOG(LogMetaSound, Warning, TEXT("Could not find referenced node [Name:%s, NodeID:%s]"), *Node.Name.ToString(), *Node.GetID().ToString());
-				continue;
+				UE_LOG(LogMetaSound, Warning, TEXT("Metasound '%s': Could not find referenced node [Name:%s, NodeID:%s]"), *InGraphContext.BuildContext.DebugAssetName, *Node.Name.ToString(), *Node.GetID().ToString());
+				return false;
 			}
 
 			for (const FMetasoundFrontendVertex& Vertex : Node.Interface.Inputs)
@@ -532,16 +537,14 @@ namespace Metasound
 
 			if (nullptr == DestinationNodeAndVertex)
 			{
-				// TODO: bubble up error
-				UE_LOG(LogMetaSound, Error, TEXT("Failed to add edge. Could not find destination [NodeID:%s, VertexID:%s]"), *Edge.ToNodeID.ToString(), *Edge.ToVertexID.ToString());
-				continue;
+				UE_LOG(LogMetaSound, Error, TEXT("MetaSound '%s': Failed to add edge. Could not find destination [NodeID:%s, VertexID:%s]"), *InGraphContext.BuildContext.DebugAssetName, *Edge.ToNodeID.ToString(), *Edge.ToVertexID.ToString());
+				return false;
 			}
 
 			if (nullptr == DestinationNodeAndVertex->Node)
 			{
-				// TODO: bubble up error
-				UE_LOG(LogMetaSound, Warning, TEXT("Skipping edge. Null destination node [NodeID:%s]"), *Edge.ToNodeID.ToString());
-				continue;
+				UE_LOG(LogMetaSound, Warning, TEXT("MetaSound '%s': 'Failed to add edge. Null destination node [NodeID:%s]"), *InGraphContext.BuildContext.DebugAssetName, *Edge.ToNodeID.ToString());
+				return false;
 			}
 
 			const FNodeIDVertexID SourceKey(Edge.FromNodeID, Edge.FromVertexID);
@@ -549,15 +552,14 @@ namespace Metasound
 
 			if (nullptr == SourceNodeAndVertex)
 			{
-				UE_LOG(LogMetaSound, Error, TEXT("Failed to add edge. Could not find source [NodeID:%s, VertexID:%s]"), *Edge.FromNodeID.ToString(), *Edge.FromVertexID.ToString());
-				continue;
+				UE_LOG(LogMetaSound, Error, TEXT("MetaSound '%s': Failed to add edge. Could not find source [NodeID:%s, VertexID:%s]"), *InGraphContext.BuildContext.DebugAssetName, *Edge.FromNodeID.ToString(), *Edge.FromVertexID.ToString());
+				return false;
 			}
 
 			if (nullptr == SourceNodeAndVertex->Node)
 			{
-				// TODO: bubble up error
-				UE_LOG(LogMetaSound, Warning, TEXT("Skipping edge. Null source node [NodeID:%s]"), *Edge.FromNodeID.ToString());
-				continue;
+				UE_LOG(LogMetaSound, Warning, TEXT("MetaSound '%s': Skipping edge. Null source node [NodeID:%s]"), *InGraphContext.BuildContext.DebugAssetName, *Edge.FromNodeID.ToString());
+				return false;
 			}
 
 			const INode* FromNode = SourceNodeAndVertex->Node;
@@ -576,12 +578,15 @@ namespace Metasound
 			}
 			else
 			{
-				UE_LOG(LogMetaSound, Error, TEXT("Failed to connect edge from [NodeID:%s, VertexID:%s] to [NodeID:%s, VertexID:%s]"), *Edge.FromNodeID.ToString(), *Edge.FromVertexID.ToString(), *Edge.ToNodeID.ToString(), *Edge.ToVertexID.ToString());
+				UE_LOG(LogMetaSound, Error, TEXT("MetaSound '%s': Failed to connect edge from [NodeID:%s, VertexID:%s] to [NodeID:%s, VertexID:%s]"), *InGraphContext.BuildContext.DebugAssetName, *Edge.FromNodeID.ToString(), *Edge.FromVertexID.ToString(), *Edge.ToNodeID.ToString(), *Edge.ToVertexID.ToString());
+				return false;
 			}
 		}
+
+		return true;
 	}
 
-	void FFrontendGraphBuilder::AddDefaultInputLiterals(FBuildGraphContext& InGraphContext)
+	bool FFrontendGraphBuilder::AddDefaultInputLiterals(FBuildGraphContext& InGraphContext)
 	{
 		using namespace Metasound::Frontend;
 
@@ -600,28 +605,32 @@ namespace Metasound
 
 			// 2. Connect the default variable to the expected input
 			const INode* FromNode = InGraphContext.Graph->FindNode(LiteralNodeID);
-			if (!ensure(FromNode))
+			if (nullptr == FromNode)
 			{
-				continue;
+				UE_LOG(LogMetaSound, Error, TEXT("MetaSound '%s': Failed to find node in graph [NodeID:%s]"), *InGraphContext.BuildContext.DebugAssetName, *LiteralNodeID.ToString());
+				return false;
 			}
 			const FVertexName& FromVertexName = LiteralNodeNames::GetOutputDataName();
 
 			const INode* ToNode = InGraphContext.Graph->FindNode(LiteralData.DestinationNodeID);
-			if (!ensure(ToNode))
+			if (nullptr == ToNode)
 			{
-				continue;
+				UE_LOG(LogMetaSound, Error, TEXT("MetaSound '%s': Failed to node in graph [NodeID:%s]"), *InGraphContext.BuildContext.DebugAssetName, *LiteralData.DestinationNodeID.ToString());
+				return false;
 			}
 			const FVertexName& ToVertexName = LiteralData.DestinationVertexKey;
 
 			bool bSuccess = InGraphContext.Graph->AddDataEdge(*FromNode, FromVertexName, *ToNode, ToVertexName);
 			if (!bSuccess)
 			{
-				UE_LOG(LogMetaSound, Error, TEXT("Failed to connect default variable edge: from '%s' to '%s'"), *FromVertexName.ToString(), *ToVertexName.ToString());
+				UE_LOG(LogMetaSound, Error, TEXT("MetaSound '%s': Failed to connect default variable edge: from '%s' to '%s'"), *InGraphContext.BuildContext.DebugAssetName, *FromVertexName.ToString(), *ToVertexName.ToString());
+				return false;
 			}
 		}
 
 		// Clear default inputs because literals have been moved out of default input map.
 		InGraphContext.DefaultInputs.Reset();
+		return true;
 	}
 
 	TArray<FFrontendGraphBuilder::FDefaultLiteralData> FFrontendGraphBuilder::GetInputDefaultLiteralData(const FMetasoundFrontendNode& InNode, const FNodeInitData& InInitData, const TSet<FNodeIDVertexID>& InEdgeDestinations)
@@ -781,19 +790,33 @@ namespace Metasound
 			InContext
 		};
 
-		// TODO: will likely want to bubble up errors here for case where
-		// a datatype or node is not registered.
-		AddNodesToGraph(BuildGraphContext);
-		AddEdgesToGraph(BuildGraphContext);
-		AddDefaultInputLiterals(BuildGraphContext);
+		bool bSuccess = AddNodesToGraph(BuildGraphContext);
+
+		if (bSuccess)
+		{
+			bSuccess = AddEdgesToGraph(BuildGraphContext);
+		}
+
+		if (bSuccess)
+		{
+			bSuccess = AddDefaultInputLiterals(BuildGraphContext);
+		}
 
-		check(BuildGraphContext.Graph->OwnsAllReferencedNodes());
-		return MoveTemp(BuildGraphContext.Graph);
+		if (bSuccess)
+		{
+			check(BuildGraphContext.Graph->OwnsAllReferencedNodes());
+			return MoveTemp(BuildGraphContext.Graph);
+		}
+		else
+		{
+			return TUniquePtr<FFrontendGraph>(nullptr);
+		}
 	}
 
-	TUniquePtr<FFrontendGraph> FFrontendGraphBuilder::CreateGraph(const FMetasoundFrontendGraphClass& InGraph, const TArray<FMetasoundFrontendGraphClass>& InSubgraphs, const TArray<FMetasoundFrontendClass>& InDependencies)
+	TUniquePtr<FFrontendGraph> FFrontendGraphBuilder::CreateGraph(const FMetasoundFrontendGraphClass& InGraph, const TArray<FMetasoundFrontendGraphClass>& InSubgraphs, const TArray<FMetasoundFrontendClass>& InDependencies, const FString& InDebugAssetName)
 	{
 		FBuildContext Context;
+		Context.DebugAssetName = InDebugAssetName;
 
 		// Gather all references to node classes from external dependencies and subgraphs.
 		for (const FMetasoundFrontendClass& ExtClass : InDependencies)
@@ -812,7 +835,7 @@ namespace Metasound
 		bool bSuccess = SortSubgraphDependencies(FrontendSubgraphPtrs);
 		if (!bSuccess)
 		{
-			UE_LOG(LogMetaSound, Error, TEXT("Failed to create graph due to failed subgraph ordering."));
+			UE_LOG(LogMetaSound, Error, TEXT("Failed to create graph due to failed subgraph ordering in asset '%s'."), *InDebugAssetName);
 			return TUniquePtr<FFrontendGraph>(nullptr);
 		}
 
@@ -822,7 +845,7 @@ namespace Metasound
 			TSharedPtr<const INode> Subgraph(CreateGraph(Context, *FrontendSubgraphPtr).Release());
 			if (!Subgraph.IsValid())
 			{
-				UE_LOG(LogMetaSound, Warning, TEXT("Failed to create subgraph [SubgraphName: %s]"), *FrontendSubgraphPtr->Metadata.GetClassName().ToString());
+				UE_LOG(LogMetaSound, Warning, TEXT("Failed to create subgraph [SubgraphName: %s] in asset '%s'"), *FrontendSubgraphPtr->Metadata.GetClassName().ToString(), *InDebugAssetName);
 			}
 			else
 			{
@@ -836,8 +859,8 @@ namespace Metasound
 	}
 	
 	/* Metasound document should be inflated by now. */
-	TUniquePtr<FFrontendGraph> FFrontendGraphBuilder::CreateGraph(const FMetasoundFrontendDocument& InDocument)
+	TUniquePtr<FFrontendGraph> FFrontendGraphBuilder::CreateGraph(const FMetasoundFrontendDocument& InDocument, const FString& InDebugAssetName)
 	{
-		return CreateGraph(InDocument.RootGraph, InDocument.Subgraphs, InDocument.Dependencies);
+		return CreateGraph(InDocument.RootGraph, InDocument.Subgraphs, InDocument.Dependencies, InDebugAssetName);
 	}
 }
diff --git a/Engine/Plugins/Runtime/Metasound/Source/MetasoundFrontend/Private/MetasoundFrontendGraphController.cpp b/Engine/Plugins/Runtime/Metasound/Source/MetasoundFrontend/Private/MetasoundFrontendGraphController.cpp
--- a/Engine/Plugins/Runtime/Metasound/Source/MetasoundFrontend/Private/MetasoundFrontendGraphController.cpp
+++ b/Engine/Plugins/Runtime/Metasound/Source/MetasoundFrontend/Private/MetasoundFrontendGraphController.cpp
@@ -1509,7 +1509,8 @@ namespace Metasound
 				const TArray<FMetasoundFrontendGraphClass>& Subgraphs = OwningDocument->GetSubgraphs();
 				const TArray<FMetasoundFrontendClass>& Dependencies = OwningDocument->GetDependencies();
 
-				TUniquePtr<FFrontendGraph> Graph = FFrontendGraphBuilder::CreateGraph(*GraphClass, Subgraphs, Dependencies);
+				FString UnknownAsset = TEXT("UnknownAsset");
+				TUniquePtr<FFrontendGraph> Graph = FFrontendGraphBuilder::CreateGraph(*GraphClass, Subgraphs, Dependencies, UnknownAsset);
 
 				if (!Graph.IsValid())
 				{
diff --git a/Engine/Plugins/Runtime/Metasound/Source/MetasoundFrontend/Public/MetasoundFrontendGraph.h b/Engine/Plugins/Runtime/Metasound/Source/MetasoundFrontend/Public/MetasoundFrontendGraph.h
--- a/Engine/Plugins/Runtime/Metasound/Source/MetasoundFrontend/Public/MetasoundFrontendGraph.h
+++ b/Engine/Plugins/Runtime/Metasound/Source/MetasoundFrontend/Public/MetasoundFrontendGraph.h
@@ -109,10 +109,10 @@ namespace Metasound
 		static bool IsFlat(const FMetasoundFrontendGraphClass& InRoot, const TArray<FMetasoundFrontendClass>& InDependencies);
 
 		/* Metasound document should be in order to create this graph. */
-		static TUniquePtr<FFrontendGraph> CreateGraph(const FMetasoundFrontendDocument& InDocument);
+		static TUniquePtr<FFrontendGraph> CreateGraph(const FMetasoundFrontendDocument& InDocument, const FString& InDebugAssetName);
 
 		/* Metasound document should be in order to create this graph. */
-		static TUniquePtr<FFrontendGraph> CreateGraph(const FMetasoundFrontendGraphClass& InGraphClass, const TArray<FMetasoundFrontendGraphClass>& InSubgraphs, const TArray<FMetasoundFrontendClass>& InDependencies);
+		static TUniquePtr<FFrontendGraph> CreateGraph(const FMetasoundFrontendGraphClass& InGraphClass, const TArray<FMetasoundFrontendGraphClass>& InSubgraphs, const TArray<FMetasoundFrontendClass>& InDependencies, const FString& InDebugAssetName);
 
 	private:
 		struct FDefaultLiteralData
@@ -134,6 +134,7 @@ namespace Metasound
 		// (for both a root and nested subgraphs)
 		struct FBuildContext
 		{
+			FString DebugAssetName;
 			FDependencyByIDMap FrontendClasses;
 			FSharedNodeByIDMap Graphs;
 		};
@@ -163,9 +164,13 @@ namespace Metasound
 		static TUniquePtr<INode> CreateVariableNode(const FMetasoundFrontendNode& InNode, const FMetasoundFrontendGraph& InGraph);
 		static TUniquePtr<INode> CreateExternalNode(const FMetasoundFrontendNode& InNode, const FMetasoundFrontendClass& InClass, FBuildGraphContext& InGraphContext, const TSet<FNodeIDVertexID>& InEdgeDestinations);
 
-		// TODO: add errors here. Most will be a "PromptIfMissing"...
-		static void AddNodesToGraph(FBuildGraphContext& InGraphContext);
-		static void AddEdgesToGraph(FBuildGraphContext& InGraphContext);
-		static void AddDefaultInputLiterals(FBuildGraphContext& InGraphContext);
+		// Returns false on error
+		static bool AddNodesToGraph(FBuildGraphContext& InGraphContext);
+
+		// Returns false on error
+		static bool AddEdgesToGraph(FBuildGraphContext& InGraphContext);
+
+		// Returns false on error
+		static bool AddDefaultInputLiterals(FBuildGraphContext& InGraphContext);
 	};
 }
