Fix MetaSound assets always redirtying after save
- Unaccounted for integer conversion causing the dirty flag to always get set on synchronization even when data hadn't changed
#rnx

diff --git a/Engine/Plugins/Runtime/Metasound/Source/MetasoundEditor/Private/MetasoundEditorGraphBuilder.cpp b/Engine/Plugins/Runtime/Metasound/Source/MetasoundEditor/Private/MetasoundEditorGraphBuilder.cpp
--- a/Engine/Plugins/Runtime/Metasound/Source/MetasoundEditor/Private/MetasoundEditorGraphBuilder.cpp
+++ b/Engine/Plugins/Runtime/Metasound/Source/MetasoundEditor/Private/MetasoundEditorGraphBuilder.cpp
@@ -669,11 +669,14 @@ namespace Metasound
 
 			if (Location)
 			{
-				const FVector2D Delta = *Location - FVector2D(InNode.NodePosX, InNode.NodePosY);
-				if (!Delta.IsNearlyZero(SMALL_NUMBER))
-				{
-					InNode.NodePosX = Location->X;
-					InNode.NodePosY = Location->Y;
+				const int32 LocX = FMath::TruncToInt(Location->X);
+				const int32 LocY = FMath::TruncToInt(Location->Y);
+				const bool bXChanged = static_cast<bool>(LocX - InNode.NodePosX);
+				const bool bYChanged = static_cast<bool>(LocY - InNode.NodePosY);
+				if (bXChanged || bYChanged)
+				{
+					InNode.NodePosX = LocX;
+					InNode.NodePosY = LocY;
 					bModified = true;
 				}
 			}
diff --git a/Engine/Plugins/Runtime/Metasound/Source/MetasoundEditor/Private/MetasoundEditorGraphNode.cpp b/Engine/Plugins/Runtime/Metasound/Source/MetasoundEditor/Private/MetasoundEditorGraphNode.cpp
--- a/Engine/Plugins/Runtime/Metasound/Source/MetasoundEditor/Private/MetasoundEditorGraphNode.cpp
+++ b/Engine/Plugins/Runtime/Metasound/Source/MetasoundEditor/Private/MetasoundEditorGraphNode.cpp
@@ -78,8 +78,8 @@ void UMetasoundEditorGraphNode::SetNodeLocation(const FVector2D& InLocation)
 {
 	using namespace Metasound::Frontend;
 
-	NodePosX = InLocation.X;
-	NodePosY = InLocation.Y;
+	NodePosX = FMath::TruncToInt(InLocation.X);
+	NodePosY = FMath::TruncToInt(InLocation.Y);
 
 	FNodeHandle NodeHandle = GetNodeHandle();
 	FMetasoundFrontendNodeStyle Style = NodeHandle->GetNodeStyle();
