AssetClassID generation fix-up (collab fixes w/Helen)
0 Guid MetaSound AssetRegistry entry fix-up
#rnx

diff --git a/Engine/Plugins/Runtime/Metasound/Source/MetasoundEditor/Private/MetasoundEditorGraphSchema.cpp b/Engine/Plugins/Runtime/Metasound/Source/MetasoundEditor/Private/MetasoundEditorGraphSchema.cpp
--- a/Engine/Plugins/Runtime/Metasound/Source/MetasoundEditor/Private/MetasoundEditorGraphSchema.cpp
+++ b/Engine/Plugins/Runtime/Metasound/Source/MetasoundEditor/Private/MetasoundEditorGraphSchema.cpp
@@ -1589,7 +1589,7 @@ void UMetasoundEditorGraphSchema::DroppedAssetsOnGraph(const TArray<FAssetData>&
 			// This may not be necessary as dropping an asset on the graph may load it, thus triggering the registration from the MetaSoundAssetManager.
 			const FNodeClassInfo ClassInfo = DroppedMetaSoundAsset->GetAssetClassInfo();
 			const FNodeRegistryKey RegistryKey = NodeRegistryKey::CreateKey(ClassInfo);
-			if (ensure(IsValidNodeRegistryKey(RegistryKey)))
+			if (ensure(NodeRegistryKey::IsValid(RegistryKey)))
 			{
 				FMetaSoundAssetRegistrationOptions RegOptions;
 				RegOptions.bForceReregister = false;
diff --git a/Engine/Plugins/Runtime/Metasound/Source/MetasoundEditor/Private/MetasoundEditorModule.cpp b/Engine/Plugins/Runtime/Metasound/Source/MetasoundEditor/Private/MetasoundEditorModule.cpp
--- a/Engine/Plugins/Runtime/Metasound/Source/MetasoundEditor/Private/MetasoundEditorModule.cpp
+++ b/Engine/Plugins/Runtime/Metasound/Source/MetasoundEditor/Private/MetasoundEditorModule.cpp
@@ -233,11 +233,16 @@ namespace Metasound
 				check(AssetSubsystem);
 
 				const FNodeRegistryKey RegistryKey = AssetSubsystem->AddOrUpdateAsset(InAssetData);
-				const bool bPrimeRequested = AssetPrimeStatus > EAssetPrimeStatus::NotRequested;
-				const bool bIsRegistered = FMetasoundFrontendRegistryContainer::Get()->IsNodeRegistered(RegistryKey);
-				if (bPrimeRequested && !bIsRegistered)
+
+				// Can be invalid if being called for the first time on an asset before FRenameRootGraphClass is called
+				if (NodeRegistryKey::IsValid(RegistryKey))
 				{
-					LoadAndRegisterAsset(InAssetData);
+					const bool bPrimeRequested = AssetPrimeStatus > EAssetPrimeStatus::NotRequested;
+					const bool bIsRegistered = FMetasoundFrontendRegistryContainer::Get()->IsNodeRegistered(RegistryKey);
+					if (bPrimeRequested && !bIsRegistered)
+					{
+						LoadAndRegisterAsset(InAssetData);
+					}
 				}
 			}
 
diff --git a/Engine/Plugins/Runtime/Metasound/Source/MetasoundEngine/Private/Metasound.cpp b/Engine/Plugins/Runtime/Metasound/Source/MetasoundEngine/Private/Metasound.cpp
--- a/Engine/Plugins/Runtime/Metasound/Source/MetasoundEngine/Private/Metasound.cpp
+++ b/Engine/Plugins/Runtime/Metasound/Source/MetasoundEngine/Private/Metasound.cpp
@@ -46,7 +46,8 @@ void UMetaSound::PostDuplicate(EDuplicateMode::Type InDuplicateMode)
 	// asset duplicated from but should not be registered as such.
 	if (InDuplicateMode == EDuplicateMode::Normal)
 	{
-		Metasound::Frontend::FRegenerateAssetClassName().Transform(GetDocumentHandle());
+		AssetClassID = FGuid::NewGuid();
+		Metasound::Frontend::FRenameRootGraphClass::Generate(GetDocumentHandle(), AssetClassID);
 	}
 }
 
diff --git a/Engine/Plugins/Runtime/Metasound/Source/MetasoundEngine/Private/MetasoundAssetSubsystem.cpp b/Engine/Plugins/Runtime/Metasound/Source/MetasoundEngine/Private/MetasoundAssetSubsystem.cpp
--- a/Engine/Plugins/Runtime/Metasound/Source/MetasoundEngine/Private/MetasoundAssetSubsystem.cpp
+++ b/Engine/Plugins/Runtime/Metasound/Source/MetasoundEngine/Private/MetasoundAssetSubsystem.cpp
@@ -184,7 +184,11 @@ Metasound::Frontend::FNodeRegistryKey UMetaSoundAssetSubsystem::AddOrUpdateAsset
 
 	FNodeClassInfo ClassInfo = MetaSoundAsset->GetAssetClassInfo();
 	const FNodeRegistryKey RegistryKey = NodeRegistryKey::CreateKey(ClassInfo);
-	PathMap.FindOrAdd(RegistryKey) = InObject.GetPathName();
+
+	if (NodeRegistryKey::IsValid(RegistryKey))
+	{
+		PathMap.FindOrAdd(RegistryKey) = InObject.GetPathName();
+	}
 
 	return RegistryKey;
 }
@@ -221,10 +225,20 @@ Metasound::Frontend::FNodeRegistryKey UMetaSoundAssetSubsystem::AddOrUpdateAsset
 		}
 	}
 
-	const FNodeRegistryKey RegistryKey = NodeRegistryKey::CreateKey(ClassInfo);
-	PathMap.FindOrAdd(RegistryKey) = InAssetData.ObjectPath;
+	if (ClassInfo.AssetClassID.IsValid())
+	{
+		const FNodeRegistryKey RegistryKey = NodeRegistryKey::CreateKey(ClassInfo);
+		if (NodeRegistryKey::IsValid(RegistryKey))
+		{
+			PathMap.FindOrAdd(RegistryKey) = InAssetData.ObjectPath;
+		}
 
-	return RegistryKey;
+		return RegistryKey;
+	}
+
+	// Invalid ClassID means the node could not be registered.
+	// Let caller report or ensure as necessary.
+	return NodeRegistryKey::GetInvalid();
 }
 
 bool UMetaSoundAssetSubsystem::CanAutoUpdate(const FMetasoundFrontendClassName& InClassName) const
diff --git a/Engine/Plugins/Runtime/Metasound/Source/MetasoundEngine/Private/MetasoundSource.cpp b/Engine/Plugins/Runtime/Metasound/Source/MetasoundEngine/Private/MetasoundSource.cpp
--- a/Engine/Plugins/Runtime/Metasound/Source/MetasoundEngine/Private/MetasoundSource.cpp
+++ b/Engine/Plugins/Runtime/Metasound/Source/MetasoundEngine/Private/MetasoundSource.cpp
@@ -106,7 +106,8 @@ void UMetaSoundSource::PostDuplicate(EDuplicateMode::Type InDuplicateMode)
 	// asset duplicated from but should not be registered as such.
 	if (InDuplicateMode == EDuplicateMode::Normal)
 	{
-		Metasound::Frontend::FRegenerateAssetClassName().Transform(GetDocumentHandle());
+		AssetClassID = FGuid::NewGuid();
+		Metasound::Frontend::FRenameRootGraphClass::Generate(GetDocumentHandle(), AssetClassID);
 	}
 }
 
diff --git a/Engine/Plugins/Runtime/Metasound/Source/MetasoundFrontend/Private/MetasoundFrontendRegistries.cpp b/Engine/Plugins/Runtime/Metasound/Source/MetasoundFrontend/Private/MetasoundFrontendRegistries.cpp
--- a/Engine/Plugins/Runtime/Metasound/Source/MetasoundFrontend/Private/MetasoundFrontendRegistries.cpp
+++ b/Engine/Plugins/Runtime/Metasound/Source/MetasoundFrontend/Private/MetasoundFrontendRegistries.cpp
@@ -288,7 +288,7 @@ namespace Metasound
 
 			bool FRegistryContainerImpl::UnregisterNode(const FNodeRegistryKey& InKey)
 			{
-				if (IsValidNodeRegistryKey(InKey))
+				if (NodeRegistryKey::IsValid(InKey))
 				{
 					if (const INodeRegistryEntry* Entry = FindNodeEntry(InKey))
 					{
@@ -467,15 +467,19 @@ namespace Metasound
 			FNodeRegistryKey CreateKey(EMetasoundFrontendClassType InType, const FString& InFullClassName, int32 InMajorVersion, int32 InMinorVersion)
 			{
 				using namespace MetasoundFrontendRegistryPrivate;
-
-				FString RegistryKey = FString::Format(TEXT("{0}_{1}_{2}.{3}"), {*GetClassTypeString(InType), *InFullClassName, InMajorVersion, InMinorVersion});
-
+				const FString RegistryKey = FString::Format(TEXT("{0}_{1}_{2}.{3}"), {*GetClassTypeString(InType), *InFullClassName, InMajorVersion, InMinorVersion});
 				return RegistryKey;
 			}
 
+			const FNodeRegistryKey& GetInvalid()
+			{
+				static const FNodeRegistryKey InvalidKey;
+				return InvalidKey;
+			}
+
 			bool IsValid(const FNodeRegistryKey& InKey)
 			{
-				return !InKey.IsEmpty();
+				return InKey != GetInvalid();
 			}
 
 			bool IsEqual(const FNodeRegistryKey& InLHS, const FNodeRegistryKey& InRHS)
@@ -529,11 +533,6 @@ namespace Metasound
 			}
 		}
 
-		bool IsValidNodeRegistryKey(const FNodeRegistryKey& InKey)
-		{
-			return NodeRegistryKey::IsValid(InKey);
-		}
-
 		FNodeClassInfo::FNodeClassInfo(const FMetasoundFrontendClassMetadata& InMetadata)
 			: ClassName(InMetadata.GetClassName())
 			, Type(InMetadata.GetType())
diff --git a/Engine/Plugins/Runtime/Metasound/Source/MetasoundFrontend/Private/MetasoundFrontendTransform.cpp b/Engine/Plugins/Runtime/Metasound/Source/MetasoundFrontend/Private/MetasoundFrontendTransform.cpp
--- a/Engine/Plugins/Runtime/Metasound/Source/MetasoundFrontend/Private/MetasoundFrontendTransform.cpp
+++ b/Engine/Plugins/Runtime/Metasound/Source/MetasoundFrontend/Private/MetasoundFrontendTransform.cpp
@@ -902,13 +902,14 @@ namespace Metasound
 			return ClassOutputs;
 		}
 
-		bool FRegenerateAssetClassName::Transform(FDocumentHandle InDocument) const
+		bool FRenameRootGraphClass::Transform(FDocumentHandle InDocument) const
 		{
-			FMetasoundFrontendClassMetadata Metadata = InDocument->GetRootGraph()->GetGraphMetadata();
-			FMetasoundFrontendClassName NewName = Metadata.GetClassName();
-			NewName.Name = *FGuid::NewGuid().ToString();
-			Metadata.SetClassName(NewName);
-			InDocument->GetRootGraph()->SetGraphMetadata(Metadata);
+			FGraphHandle RootGraph = InDocument->GetRootGraph();
+
+			FMetasoundFrontendClassMetadata Metadata = RootGraph->GetGraphMetadata();
+			Metadata.SetClassName(NewClassName);
+			RootGraph->SetGraphMetadata(Metadata);
+
 			return true;
 		}
 
@@ -1154,9 +1155,7 @@ namespace Metasound
 		class FVersionDocument_1_6 : public FVersionDocumentTransform
 		{
 		public:
-			FVersionDocument_1_6()
-			{
-			}
+			FVersionDocument_1_6() = default;
 
 			FMetasoundFrontendVersionNumber GetTargetVersion() const override
 			{
@@ -1165,7 +1164,8 @@ namespace Metasound
 
 			void TransformInternal(FDocumentHandle InDocument) const override
 			{
-				FRegenerateAssetClassName().Transform(InDocument);
+				const FGuid NewAssetClassID = FGuid::NewGuid();
+				FRenameRootGraphClass::Generate(InDocument, NewAssetClassID);
 			}
 		};
 
diff --git a/Engine/Plugins/Runtime/Metasound/Source/MetasoundFrontend/Public/MetasoundFrontendRegistries.h b/Engine/Plugins/Runtime/Metasound/Source/MetasoundFrontend/Public/MetasoundFrontendRegistries.h
--- a/Engine/Plugins/Runtime/Metasound/Source/MetasoundFrontend/Public/MetasoundFrontendRegistries.h
+++ b/Engine/Plugins/Runtime/Metasound/Source/MetasoundFrontend/Public/MetasoundFrontendRegistries.h
@@ -1,7 +1,6 @@
 // Copyright Epic Games, Inc. All Rights Reserved.
 #pragma once
 
-#include "CoreMinimal.h"
 #include "IAudioProxyInitializer.h"
 #include "MetasoundDataReference.h"
 #include "MetasoundEnum.h"
@@ -52,12 +51,6 @@ namespace Metasound
 	{
 		using FNodeRegistryKey = FString;
 
-
-		// Returns true if the registry key is a valid key. 
-		//
-		// This does *not* connote that the registry key exists in the registry.
-		METASOUNDFRONTEND_API bool IsValidNodeRegistryKey(const FNodeRegistryKey& InKey);
-
 		/** FNodeClassInfo contains a minimal set of information needed to find
 		 * and query node classes. 
 		 */
@@ -236,7 +229,10 @@ namespace Metasound
 
 		namespace NodeRegistryKey
 		{
-			// Returns true if the registry key is a valid key. 
+			// Returns the invalid NodeRegistryKey.
+			METASOUNDFRONTEND_API const FNodeRegistryKey& GetInvalid();
+
+			// Returns true if the registry key is a valid key.
 			//
 			// This does *not* connote that the registry key exists in the registry.
 			METASOUNDFRONTEND_API bool IsValid(const FNodeRegistryKey& InKey);
@@ -250,7 +246,7 @@ namespace Metasound
 			// Returns true if the class info and key represent the same entry in the node registry.
 			METASOUNDFRONTEND_API bool IsEqual(const FNodeClassInfo& InLHS, const FNodeRegistryKey& InRHS);
 
-			// Returns true if the class metadatas represent the same entry in the node registry.
+			// Returns true if the class metadata represent the same entry in the node registry.
 			METASOUNDFRONTEND_API bool IsEqual(const FMetasoundFrontendClassMetadata& InLHS, const FMetasoundFrontendClassMetadata& InRHS);
 
 			// Returns true if the class info and class metadata represent the same entry in the node registry.
diff --git a/Engine/Plugins/Runtime/Metasound/Source/MetasoundFrontend/Public/MetasoundFrontendTransform.h b/Engine/Plugins/Runtime/Metasound/Source/MetasoundFrontend/Public/MetasoundFrontendTransform.h
--- a/Engine/Plugins/Runtime/Metasound/Source/MetasoundFrontend/Public/MetasoundFrontendTransform.h
+++ b/Engine/Plugins/Runtime/Metasound/Source/MetasoundFrontend/Public/MetasoundFrontendTransform.h
@@ -161,13 +161,26 @@ namespace Metasound
 			bool bLogWarningOnDroppedConnection;
 		};
 
-		/** Regenerates the class' name, effectively causing the class to get registered as
-		  * a new class (useful for when an asset is duplicated). */
-		class METASOUNDFRONTEND_API FRegenerateAssetClassName : public IDocumentTransform
+		/** Sets the document's graph class, optionally updating the namespace and variant. */
+		class METASOUNDFRONTEND_API FRenameRootGraphClass : public IDocumentTransform
 		{
+			const FMetasoundFrontendClassName NewClassName;
 
 		public:
-			FRegenerateAssetClassName() = default;
+			/* Generates and assigns the docment's root graph class a unique name using the provided guid as the name field and
+			 * (optionally) namespace and variant.
+			 * Returns true if transform succeeded, false if not.
+			 */
+			static bool Generate(FDocumentHandle InDocument, const FGuid& InGuid, const FName Namespace = { }, const FName Variant = { })
+			{
+				const FMetasoundFrontendClassName GeneratedClassName = { Namespace, *InGuid.ToString(), Variant };
+				return FRenameRootGraphClass(GeneratedClassName).Transform(InDocument);
+			}
+
+			FRenameRootGraphClass(const FMetasoundFrontendClassName InClassName)
+				: NewClassName(InClassName)
+			{
+			}
 
 			bool Transform(FDocumentHandle InDocument) const override;
 		};
diff --git a/Engine/Plugins/Runtime/Metasound/Source/MetasoundFrontend/Public/MetasoundNodeRegistrationMacro.h b/Engine/Plugins/Runtime/Metasound/Source/MetasoundFrontend/Public/MetasoundNodeRegistrationMacro.h
--- a/Engine/Plugins/Runtime/Metasound/Source/MetasoundFrontend/Public/MetasoundNodeRegistrationMacro.h
+++ b/Engine/Plugins/Runtime/Metasound/Source/MetasoundFrontend/Public/MetasoundNodeRegistrationMacro.h
@@ -101,7 +101,7 @@ namespace Metasound
 
 
 		Frontend::FNodeRegistryKey Key = FMetasoundFrontendRegistryContainer::Get()->RegisterNode(MakeUnique<FNodeRegistryEntry>(InMetadata));
-		const bool bSuccessfullyRegisteredNode = Frontend::IsValidNodeRegistryKey(Key);
+		const bool bSuccessfullyRegisteredNode = Frontend::NodeRegistryKey::IsValid(Key);
 		ensureAlwaysMsgf(bSuccessfullyRegisteredNode, TEXT("Registering node class failed. Please check the logs."));
 
 		return bSuccessfullyRegisteredNode;
