Do not register during PreSave if cooking to fix non-deterministic cook output (auto-update called during registration modifies asset, and is currently too aggressive)
#rnx

diff --git a/Engine/Plugins/Runtime/Metasound/Source/MetasoundEngine/Public/Metasound.h b/Engine/Plugins/Runtime/Metasound/Source/MetasoundEngine/Public/Metasound.h
--- a/Engine/Plugins/Runtime/Metasound/Source/MetasoundEngine/Public/Metasound.h
+++ b/Engine/Plugins/Runtime/Metasound/Source/MetasoundEngine/Public/Metasound.h
@@ -65,37 +65,17 @@ namespace Metasound
 	{
 		using namespace Metasound::Frontend;
 
-		// If not editor and cooking, no need to force re-registering as if its already registered,
-		// no edits have been made between last registration and this one if a parent asset has requested
-		// it to register.
-		FMetaSoundAssetRegistrationOptions NonEditorRegistrationOptions;
-		NonEditorRegistrationOptions.bForceReregister = false;
-
 #if WITH_EDITORONLY_DATA
 		if (UMetasoundEditorGraphBase* MetaSoundGraph = Cast<UMetasoundEditorGraphBase>(InMetaSound.GetGraph()))
 		{
-			if (InSaveContext.IsCooking())
-			{
-				// Non-editor builds use explicit options above as opposed to using
-				// editor variant via the MetaSoundEdGraph.
-				if (GIsEditor)
-				{
-					MetaSoundGraph->RegisterGraphWithFrontend();
-				}
-				else
-				{
-					InMetaSound.RegisterGraphWithFrontend(NonEditorRegistrationOptions);
-				}
-				MetaSoundGraph->SetSynchronizationRequired();
-			}
-			else
+			// Cooked data must be deterministic, so do not call register graph as this can
+			// initiate an auto-update and/or local registry data cache and modify serialized data.
+			if (!InSaveContext.IsCooking())
 			{
 				MetaSoundGraph->RegisterGraphWithFrontend();
 				MetaSoundGraph->SetSynchronizationRequired();
 			}
 		}
-#else
-		InMetaSound.RegisterGraphWithFrontend(NonEditorRegistrationOptions);
 #endif // WITH_EDITORONLY_DATA
 	}
 
