Fallback to sorting display name alphabetically if SortOrder not provided or results in a tie.

diff --git a/Engine/Plugins/Runtime/Metasound/Source/MetasoundEditor/Private/MetasoundEditorGraphBuilder.cpp b/Engine/Plugins/Runtime/Metasound/Source/MetasoundEditor/Private/MetasoundEditorGraphBuilder.cpp
--- a/Engine/Plugins/Runtime/Metasound/Source/MetasoundEditor/Private/MetasoundEditorGraphBuilder.cpp
+++ b/Engine/Plugins/Runtime/Metasound/Source/MetasoundEditor/Private/MetasoundEditorGraphBuilder.cpp
@@ -1260,6 +1260,7 @@ namespace Metasound
 			check(ReferencedAsset);
 
 			FRebuildPresetRootGraph(ReferencedAsset->GetDocumentHandle()).Transform(PresetAsset->GetDocumentHandle());
+
 			PresetAsset->ConformObjectDataToInterfaces();
 		}
 
@@ -1377,7 +1378,7 @@ namespace Metasound
 			}
 
 			TArray<FInputHandle> InputHandles = NodeHandle->GetInputs();
-			NodeHandle->GetInputStyle().SortDefaults(InputHandles);
+			NodeHandle->GetInputStyle().SortDefaults(InputHandles, [](const Frontend::FInputHandle& Handle) { return FGraphBuilder::GetDisplayName(*Handle); });
 			for (const FInputHandle& InputHandle : InputHandles)
 			{
 				// Only add pins of the node if the connection is user modifiable. 
@@ -1389,7 +1390,7 @@ namespace Metasound
 			}
 
 			TArray<FOutputHandle> OutputHandles = NodeHandle->GetOutputs();
-			NodeHandle->GetOutputStyle().SortDefaults(OutputHandles);
+			NodeHandle->GetOutputStyle().SortDefaults(OutputHandles, [](const Frontend::FOutputHandle& Handle) { return FGraphBuilder::GetDisplayName(*Handle); });
 			for (const FOutputHandle& OutputHandle : OutputHandles)
 			{
 				// Only add pins of the node if the connection is user modifiable. 
@@ -1959,20 +1960,22 @@ namespace Metasound
 
 		bool FGraphBuilder::SynchronizeNodePins(UMetasoundEditorGraphNode& InEditorNode, Frontend::FConstNodeHandle InNode, bool bRemoveUnusedPins, bool bLogChanges)
 		{
+			using namespace Frontend;
+
 			bool bIsNodeDirty = false;
 
 			IMetasoundEditorModule& EditorModule = FModuleManager::GetModuleChecked<IMetasoundEditorModule>("MetaSoundEditor");
 
-			TArray<Frontend::FConstInputHandle> InputHandles;
-			TArray<Frontend::FConstOutputHandle> OutputHandles;
-			auto GetUserModifiableHandles = [InNode](TArray<Frontend::FConstInputHandle>& InHandles, TArray<Frontend::FConstOutputHandle>& OutHandles)
+			TArray<FConstInputHandle> InputHandles;
+			TArray<FConstOutputHandle> OutputHandles;
+			auto GetUserModifiableHandles = [InNode](TArray<FConstInputHandle>& InHandles, TArray<FConstOutputHandle>& OutHandles)
 			{
 				InHandles = InNode->GetConstInputs();
 				OutHandles = InNode->GetConstOutputs();
 
 				// Remove input and output handles which are not user modifiable
-				InHandles.RemoveAll([](const Frontend::FConstInputHandle& FrontendInput) { return !FrontendInput->IsConnectionUserModifiable(); });
-				OutHandles.RemoveAll([](const Frontend::FConstOutputHandle& FrontendOutput) { return !FrontendOutput->IsConnectionUserModifiable(); });
+				InHandles.RemoveAll([](const FConstInputHandle& FrontendInput) { return !FrontendInput->IsConnectionUserModifiable(); });
+				OutHandles.RemoveAll([](const FConstOutputHandle& FrontendOutput) { return !FrontendOutput->IsConnectionUserModifiable(); });
 			};
 			GetUserModifiableHandles(InputHandles, OutputHandles);
 
@@ -1982,12 +1985,12 @@ namespace Metasound
 			{
 				UEdGraphPin* Pin = EditorPins[i];
 
-				auto IsMatchingInputHandle = [&](const Frontend::FConstInputHandle& InputHandle) -> bool
+				auto IsMatchingInputHandle = [&](const FConstInputHandle& InputHandle) -> bool
 				{
 					return IsMatchingInputHandleAndPin(InputHandle, *Pin);
 				};
 
-				auto IsMatchingOutputHandle = [&](const Frontend::FConstOutputHandle& OutputHandle) -> bool
+				auto IsMatchingOutputHandle = [&](const FConstOutputHandle& OutputHandle) -> bool
 				{
 					return IsMatchingOutputHandleAndPin(OutputHandle, *Pin);
 				};
@@ -2040,7 +2043,7 @@ namespace Metasound
 			if (!InputHandles.IsEmpty())
 			{
 				bIsNodeDirty = true;
-				for (Frontend::FConstInputHandle& InputHandle : InputHandles)
+				for (FConstInputHandle& InputHandle : InputHandles)
 				{
 					if (bLogChanges)
 					{
@@ -2056,7 +2059,7 @@ namespace Metasound
 			if (!OutputHandles.IsEmpty())
 			{
 				bIsNodeDirty = true;
-				for (Frontend::FConstOutputHandle& OutputHandle : OutputHandles)
+				for (FConstOutputHandle& OutputHandle : OutputHandles)
 				{
 					if (bLogChanges)
 					{
@@ -2071,8 +2074,9 @@ namespace Metasound
 
 			// Order pins
 			GetUserModifiableHandles(InputHandles, OutputHandles);
-			InNode->GetInputStyle().SortDefaults(InputHandles);
-			InNode->GetOutputStyle().SortDefaults(OutputHandles);
+
+			InNode->GetInputStyle().SortDefaults(InputHandles, [](const FConstInputHandle& Handle) { return FGraphBuilder::GetDisplayName(*Handle); });
+			InNode->GetOutputStyle().SortDefaults(OutputHandles, [](const FConstOutputHandle& Handle) { return FGraphBuilder::GetDisplayName(*Handle); });
 
 			auto SwapAndDirty = [&](int32 IndexA, int32 IndexB)
 			{
@@ -2092,7 +2096,7 @@ namespace Metasound
 					if (!InputHandles.IsEmpty())
 					{
 						constexpr bool bAllowShrinking = false;
-						Frontend::FConstInputHandle InputHandle = InputHandles.Pop(bAllowShrinking);
+						FConstInputHandle InputHandle = InputHandles.Pop(bAllowShrinking);
 						for (int32 j = i; j >= 0; --j)
 						{
 							if (IsMatchingInputHandleAndPin(InputHandle, *InEditorNode.Pins[j]))
@@ -2108,7 +2112,7 @@ namespace Metasound
 					if (!OutputHandles.IsEmpty())
 					{
 						constexpr bool bAllowShrinking = false;
-						Frontend::FConstOutputHandle OutputHandle = OutputHandles.Pop(bAllowShrinking);
+						FConstOutputHandle OutputHandle = OutputHandles.Pop(bAllowShrinking);
 						for (int32 j = i; j >= 0; --j)
 						{
 							if (IsMatchingOutputHandleAndPin(OutputHandle, *InEditorNode.Pins[j]))
diff --git a/Engine/Plugins/Runtime/Metasound/Source/MetasoundEditor/Private/MetasoundEditorModule.cpp b/Engine/Plugins/Runtime/Metasound/Source/MetasoundEditor/Private/MetasoundEditorModule.cpp
--- a/Engine/Plugins/Runtime/Metasound/Source/MetasoundEditor/Private/MetasoundEditorModule.cpp
+++ b/Engine/Plugins/Runtime/Metasound/Source/MetasoundEditor/Private/MetasoundEditorModule.cpp
@@ -707,6 +707,14 @@ namespace Metasound
 				}
 
 				RegisterExplicitProxyClass(*USoundWave::StaticClass());
+
+				// Required to ensure logic to order nodes for presets exclusive to
+				// editor is propagated to transform instances while editing in editor.
+				Frontend::DocumentTransform::RegisterNodeDisplayNamePredicate([](const Frontend::FNodeHandle& NodeHandle)
+				{
+					constexpr bool bIncludeNamespace = false;
+					return FGraphBuilder::GetDisplayName(*NodeHandle, bIncludeNamespace);
+				});
 			}
 
 			virtual void ShutdownModule() override
diff --git a/Engine/Plugins/Runtime/Metasound/Source/MetasoundFrontend/Private/MetasoundFrontendTransform.cpp b/Engine/Plugins/Runtime/Metasound/Source/MetasoundFrontend/Private/MetasoundFrontendTransform.cpp
--- a/Engine/Plugins/Runtime/Metasound/Source/MetasoundFrontend/Private/MetasoundFrontendTransform.cpp
+++ b/Engine/Plugins/Runtime/Metasound/Source/MetasoundFrontend/Private/MetasoundFrontendTransform.cpp
@@ -16,6 +16,24 @@ namespace Metasound
 {
 	namespace Frontend
 	{
+		namespace DocumentTransform
+		{
+#if WITH_EDITOR
+			FGetNodeDisplayNamePredicate NodeDisplayNamePredicate;
+
+			void RegisterNodeDisplayNamePredicate(FGetNodeDisplayNamePredicate&& InNamePredicate)
+			{
+				NodeDisplayNamePredicate = MoveTemp(InNamePredicate);
+			}
+
+			FGetNodeDisplayNamePredicateRef GetNodeDisplayNamePredicate()
+			{
+				return NodeDisplayNamePredicate;
+			}
+#endif // WITH_EDITOR
+		} // namespace DocumentTransform
+
+
 		FModifyRootGraphInterfaces::FModifyRootGraphInterfaces(const TArray<FMetasoundFrontendInterface>& InInterfacesToRemove, const TArray<FMetasoundFrontendInterface>& InInterfacesToAdd)
 			: InterfacesToRemove(InInterfacesToRemove)
 			, InterfacesToAdd(InInterfacesToAdd)
@@ -693,8 +711,8 @@ namespace Metasound
 			ReferencedClassMetadata.SetType(EMetasoundFrontendClassType::External);
 
 			FNodeHandle ReferencedNodeHandle = PresetGraphHandle->AddNode(ReferencedClassMetadata, PresetNodeID);
-#if WITH_EDITOR
 
+#if WITH_EDITOR
 			// Set node location.
 			FMetasoundFrontendNodeStyle RefNodeStyle;
 
@@ -738,7 +756,8 @@ namespace Metasound
 			// inputs/outputs are added but before setting locations to propagate effectively)
 			FMetasoundFrontendInterfaceStyle Style = ReferencedGraphHandle->GetInputStyle();
 			InPresetGraphHandle->SetInputStyle(Style);
-			Style.SortDefaults(NodeHandles);
+
+			Style.SortDefaults(NodeHandles, DocumentTransform::GetNodeDisplayNamePredicate());
 
 			for (const FNodeHandle& InputNode : NodeHandles)
 			{
@@ -777,11 +796,12 @@ namespace Metasound
 			// inputs/outputs are added but before setting locations to propagate effectively)
 			FMetasoundFrontendInterfaceStyle Style = ReferencedGraphHandle->GetOutputStyle();
 			InPresetGraphHandle->SetOutputStyle(Style);
-			Style.SortDefaults(NodeHandles);
 
+			Style.SortDefaults(NodeHandles, DocumentTransform::GetNodeDisplayNamePredicate());
+
+			// Set input node location
 			for (const FNodeHandle& OutputNode : NodeHandles)
 			{
-				// Set input node location
 				FMetasoundFrontendNodeStyle NodeStyle;
 				NodeStyle.Display.Locations.Add(FGuid(), OutputNodeLocation);
 				OutputNode->SetNodeStyle(NodeStyle);
diff --git a/Engine/Plugins/Runtime/Metasound/Source/MetasoundFrontend/Public/MetasoundFrontendDocument.h b/Engine/Plugins/Runtime/Metasound/Source/MetasoundFrontend/Public/MetasoundFrontendDocument.h
--- a/Engine/Plugins/Runtime/Metasound/Source/MetasoundFrontend/Public/MetasoundFrontendDocument.h
+++ b/Engine/Plugins/Runtime/Metasound/Source/MetasoundFrontend/Public/MetasoundFrontendDocument.h
@@ -12,6 +12,8 @@
 #include "MetasoundNodeInterface.h"
 #include "MetasoundVertex.h"
 #include "Misc/Guid.h"
+#include "Templates/Function.h"
+#include "Templates/Invoke.h"
 #include "Templates/TypeHash.h"
 
 #include "MetasoundFrontendDocument.generated.h"
@@ -781,8 +783,8 @@ struct FMetasoundFrontendInterfaceStyle
 	UPROPERTY()
 	TMap<FName, FText> RequiredMembers;
 
-	template <typename HandleType>
-	void SortDefaults(TArray<HandleType>& OutHandles) const
+	template <typename HandleType, typename NamePredicateType>
+	void SortDefaults(TArray<HandleType>& OutHandles, NamePredicateType InGetDisplayNamePredicate) const
 	{
 		TMap<FGuid, int32> NodeIDToSortIndex;
 		int32 HighestSortOrder = TNumericLimits<int32>::Min();
@@ -802,11 +804,19 @@ struct FMetasoundFrontendInterfaceStyle
 			NodeIDToSortIndex.Add(HandleID, SortIndex);
 		}
 
-		OutHandles.Sort([&NodeIDToSortIndex](const HandleType& HandleA, const HandleType& HandleB)
+		OutHandles.Sort([&NodeIDToSortIndex, &InGetDisplayNamePredicate](const HandleType& HandleA, const HandleType& HandleB) -> bool
 		{
 			const FGuid HandleAID = HandleA->GetID();
 			const FGuid HandleBID = HandleB->GetID();
-			return NodeIDToSortIndex[HandleAID] < NodeIDToSortIndex[HandleBID];
+			const int32 AID = NodeIDToSortIndex[HandleAID];
+			const int32 BID = NodeIDToSortIndex[HandleBID];
+
+			// If IDs are equal, sort alphabetically using provided name predicate
+			if (AID == BID)
+			{
+				return Invoke(InGetDisplayNamePredicate, HandleA).CompareTo(Invoke(InGetDisplayNamePredicate, HandleB)) < 0;
+			}
+			return AID < BID;
 		});
 	}
 #endif // #if WITH_EDITORONLY_DATA
diff --git a/Engine/Plugins/Runtime/Metasound/Source/MetasoundFrontend/Public/MetasoundFrontendTransform.h b/Engine/Plugins/Runtime/Metasound/Source/MetasoundFrontend/Public/MetasoundFrontendTransform.h
--- a/Engine/Plugins/Runtime/Metasound/Source/MetasoundFrontend/Public/MetasoundFrontendTransform.h
+++ b/Engine/Plugins/Runtime/Metasound/Source/MetasoundFrontend/Public/MetasoundFrontendTransform.h
@@ -13,6 +13,18 @@ namespace Metasound
 {
 	namespace Frontend
 	{
+
+		namespace DocumentTransform
+		{
+#if WITH_EDITOR
+			using FGetNodeDisplayNamePredicate = TFunction<FText(const FNodeHandle&)>;
+			using FGetNodeDisplayNamePredicateRef = TFunctionRef<FText(const FNodeHandle&)>;
+
+			METASOUNDFRONTEND_API void RegisterNodeDisplayNamePredicate(FGetNodeDisplayNamePredicate&& InNamePredicate);
+			METASOUNDFRONTEND_API FGetNodeDisplayNamePredicateRef GetNodeDisplayNamePredicate();
+#endif // WITH_EDITOR
+		}
+
 		/** Interface for transforms applied to documents. */
 		class IDocumentTransform
 		{
