- Add ability to audition MetaSounds from ContentBrowser
- Add play icon for MetaSound tab when auditioning.
- Use small icons for toolbar
- Move MS asset colors to SlateStyle

diff --git a/Engine/Plugins/Runtime/Metasound/Source/MetasoundEditor/Private/MetasoundAssetTypeActions.cpp b/Engine/Plugins/Runtime/Metasound/Source/MetasoundEditor/Private/MetasoundAssetTypeActions.cpp
--- a/Engine/Plugins/Runtime/Metasound/Source/MetasoundEditor/Private/MetasoundAssetTypeActions.cpp
+++ b/Engine/Plugins/Runtime/Metasound/Source/MetasoundEditor/Private/MetasoundAssetTypeActions.cpp
@@ -2,11 +2,11 @@
 
 #include "MetasoundAssetTypeActions.h"
 
-#include "ContentBrowserModule.h"
+#include "Components/AudioComponent.h"
 #include "ContentBrowserMenuContexts.h"
+#include "ContentBrowserModule.h"
 #include "Editor/EditorPerProjectUserSettings.h"
 #include "IContentBrowserSingleton.h"
-#include "ObjectEditorUtils.h"
 #include "Metasound.h"
 #include "MetasoundAssetBase.h"
 #include "MetasoundSource.h"
@@ -15,8 +15,12 @@
 #include "MetasoundEditorSettings.h"
 #include "MetasoundFactory.h"
 #include "MetasoundUObjectRegistry.h"
+#include "ObjectEditorUtils.h"
+#include "Styling/ISlateStyle.h"
+#include "Styling/SlateStyleRegistry.h"
 #include "ToolMenus.h"
 #include "ToolMenuSection.h"
+#include "EditorStyleSet.h"
 
 
 #define LOCTEXT_NAMESPACE "MetaSoundEditor"
@@ -106,6 +110,60 @@ namespace Metasound
 					InSection.AddMenuEntry(*PresetEntryName, Label, ToolTip, Icon, UIExecuteAction);
 				}));
 			}
+
+			bool IsPlaying(const FSoftObjectPath& InSourcePath)
+			{
+				check(GEditor);
+				if (const UAudioComponent* PreviewComponent = GEditor->GetPreviewAudioComponent())
+				{
+					if (PreviewComponent->IsPlaying())
+					{
+						if (const USoundBase* Sound = PreviewComponent->Sound)
+						{
+							const FName SoundName = *Sound->GetPathName();
+							return SoundName == InSourcePath.GetAssetPathName();
+						}
+					}
+				}
+
+				return false;
+			}
+
+			void PlaySound(UMetaSoundSource& InSource)
+			{
+				// If editor is open, call into it to play to start all visualization requirements therein
+				// specific to auditioning MetaSounds (ex. priming audio bus used for volume metering, playtime
+				// widget, etc.)
+				TSharedPtr<FEditor> Editor = FGraphBuilder::GetEditorForMetasound(InSource);
+				if (Editor.IsValid())
+				{
+					Editor->Play();
+					return;
+				}
+
+				check(GEditor);
+				FGraphBuilder::RegisterGraphWithFrontend(InSource);
+				GEditor->PlayPreviewSound(&InSource);
+			}
+
+			void StopSound(const UMetaSoundSource& InSource)
+			{
+				// If editor is open, call into it to play to start all visualization requirements therein
+				// specific to auditioning MetaSounds (ex. priming audio bus used for volume metering, playtime
+				// widget, etc.)
+				TSharedPtr<FEditor> Editor = FGraphBuilder::GetEditorForMetasound(InSource);
+				if (Editor.IsValid())
+				{
+					Editor->Stop();
+					return;
+				}
+
+				check(GEditor);
+				if (UAudioComponent* PreviewComponent = GEditor->GetPreviewAudioComponent())
+				{
+					PreviewComponent->Stop();
+				}
+			}
 		} // namespace AssetTypeActionsPrivate
 
 		UClass* FAssetTypeActions_MetaSound::GetSupportedClass() const
@@ -113,6 +171,16 @@ namespace Metasound
 			return UMetaSound::StaticClass();
 		}
 
+		FColor FAssetTypeActions_MetaSound::GetTypeColor() const
+		{
+			if (const ISlateStyle* MetasoundStyle = FSlateStyleRegistry::FindSlateStyle("MetaSoundStyle"))
+			{
+				return MetasoundStyle->GetColor("Metasound.Color").ToFColorSRGB();
+			}
+
+			return FColor::White;
+		}
+
 		void FAssetTypeActions_MetaSound::OpenAssetEditor(const TArray<UObject*>& InObjects, TSharedPtr<IToolkitHost> ToolkitHost)
 		{
 			const EToolkitMode::Type Mode = ToolkitHost.IsValid() ? EToolkitMode::WorldCentric : EToolkitMode::Standalone;
@@ -154,6 +222,16 @@ namespace Metasound
 			return UMetaSoundSource::StaticClass();
 		}
 
+		FColor FAssetTypeActions_MetaSoundSource::GetTypeColor() const
+		{
+			if (const ISlateStyle* MetasoundStyle = FSlateStyleRegistry::FindSlateStyle("MetaSoundStyle"))
+			{
+				return MetasoundStyle->GetColor("MetasoundSource.Color").ToFColorSRGB();
+			}
+
+			return FColor::White;
+		}
+
 		void FAssetTypeActions_MetaSoundSource::OpenAssetEditor(const TArray<UObject*>& InObjects, TSharedPtr<IToolkitHost> ToolkitHost)
 		{
 			const EToolkitMode::Type Mode = ToolkitHost.IsValid() ? EToolkitMode::WorldCentric : EToolkitMode::Standalone;
@@ -196,6 +274,95 @@ namespace Metasound
 
 			return SubMenus;
 		}
+
+		TSharedPtr<SWidget> FAssetTypeActions_MetaSoundSource::GetThumbnailOverlay(const FAssetData& AssetData) const
+		{
+			auto OnGetDisplayBrushLambda = [Path = AssetData.ToSoftObjectPath()]()
+			{
+				using namespace AssetTypeActionsPrivate;
+
+				if (IsPlaying(Path))
+				{
+					return FEditorStyle::GetBrush("MediaAsset.AssetActions.Stop.Large");
+				}
+
+				return FEditorStyle::GetBrush("MediaAsset.AssetActions.Play.Large");
+			};
+
+			auto OnClickedLambda = [Path = AssetData.ToSoftObjectPath()]()
+			{
+				using namespace AssetTypeActionsPrivate;
+				// Load and play sound
+				if (UMetaSoundSource* MetaSoundSource = Cast<UMetaSoundSource>(Path.TryLoad()))
+				{
+					if (IsPlaying(Path))
+					{
+						StopSound(*MetaSoundSource);
+					}
+					else
+					{
+						PlaySound(*MetaSoundSource);
+					}
+				}
+
+				return FReply::Handled();
+			};
+
+			auto OnToolTipTextLambda = [Path = AssetData.ToSoftObjectPath()]()
+			{
+				using namespace AssetTypeActionsPrivate;
+
+				FText Format;
+				if (IsPlaying(Path))
+				{
+					Format = LOCTEXT("PreviewMetaSoundFromIconToolTip", "Stop Previewing {0}");
+				}
+				else
+				{
+					Format = LOCTEXT("PreviewMetaSoundFromIconToolTip", "Preview {0}");
+				}
+
+				FName TypeName = UMetaSoundSource::StaticClass()->GetFName();
+				return FText::Format(Format, FText::FromName(TypeName));
+			};
+
+			TSharedPtr<SBox> Box;
+			SAssignNew(Box, SBox)
+				.HAlign(HAlign_Center)
+				.VAlign(VAlign_Center)
+				.Padding(FMargin(2));
+
+			auto OnGetVisibilityLambda = [Box, Path = AssetData.ToSoftObjectPath()]()
+			{
+				using namespace AssetTypeActionsPrivate;
+
+				if (Box.IsValid() && (Box->IsHovered() || IsPlaying(Path)))
+				{
+					return EVisibility::Visible;
+				}
+
+				return EVisibility::Hidden;
+			};
+
+			TSharedPtr<SButton> Widget;
+			SAssignNew(Widget, SButton)
+				.ButtonStyle(FEditorStyle::Get(), "HoverHintOnly")
+				.ToolTipText_Lambda(OnToolTipTextLambda)
+				.Cursor(EMouseCursor::Default) // The outer widget can specify a DragHand cursor, so overriden here
+				.ForegroundColor(FSlateColor::UseForeground())
+				.IsFocusable(false)
+				.OnClicked_Lambda(OnClickedLambda)
+				.Visibility_Lambda(OnGetVisibilityLambda)
+				[
+					SNew(SImage)
+					.Image_Lambda(OnGetDisplayBrushLambda)
+				];
+
+			Box->SetContent(Widget.ToSharedRef());
+			Box->SetVisibility(EVisibility::Visible);
+
+			return Box;
+		}
 	} // namespace Editor
 } // namespace Metasound
 #undef LOCTEXT_NAMESPACE //MetaSoundEditor
diff --git a/Engine/Plugins/Runtime/Metasound/Source/MetasoundEditor/Private/MetasoundAssetTypeActions.h b/Engine/Plugins/Runtime/Metasound/Source/MetasoundEditor/Private/MetasoundAssetTypeActions.h
--- a/Engine/Plugins/Runtime/Metasound/Source/MetasoundEditor/Private/MetasoundAssetTypeActions.h
+++ b/Engine/Plugins/Runtime/Metasound/Source/MetasoundEditor/Private/MetasoundAssetTypeActions.h
@@ -1,7 +1,6 @@
 // Copyright Epic Games, Inc. All Rights Reserved.
 #pragma once
 
-#include "CoreMinimal.h"
 #include "AssetTypeActions_Base.h"
 
 
@@ -17,7 +16,7 @@ namespace Metasound
 		public:
 			// IAssetTypeActions Implementation
 			virtual FText GetName() const override { return NSLOCTEXT("AssetTypeActions", "AssetTypeActions_MetaSound", "MetaSound"); }
-			virtual FColor GetTypeColor() const override { return FColor(13, 55, 13); }
+			virtual FColor GetTypeColor() const override;
 			virtual UClass* GetSupportedClass() const override;
 			virtual uint32 GetCategories() override { return EAssetTypeCategories::Sounds; }
 			virtual const TArray<FText>& GetSubMenus() const override;
@@ -31,12 +30,14 @@ namespace Metasound
 		public:
 			// IAssetTypeActions Implementation
 			virtual FText GetName() const override { return NSLOCTEXT("AssetTypeActions", "AssetTypeActions_MetaSoundSource", "MetaSound Source"); }
-			virtual FColor GetTypeColor() const override { return FColor(103, 214, 66); }
+			virtual FColor GetTypeColor() const override;
 			virtual UClass* GetSupportedClass() const override;
 			virtual uint32 GetCategories() override { return EAssetTypeCategories::Sounds; }
 			virtual const TArray<FText>& GetSubMenus() const override;
 			virtual void OpenAssetEditor(const TArray<UObject*>& InObjects, TSharedPtr<IToolkitHost> ToolkitHost) override;
 
+			virtual TSharedPtr<SWidget> GetThumbnailOverlay(const FAssetData& AssetData) const override;
+
 			static void RegisterMenuActions();
 		};
 	} // namespace Editor
diff --git a/Engine/Plugins/Runtime/Metasound/Source/MetasoundEditor/Private/MetasoundEditor.cpp b/Engine/Plugins/Runtime/Metasound/Source/MetasoundEditor/Private/MetasoundEditor.cpp
--- a/Engine/Plugins/Runtime/Metasound/Source/MetasoundEditor/Private/MetasoundEditor.cpp
+++ b/Engine/Plugins/Runtime/Metasound/Source/MetasoundEditor/Private/MetasoundEditor.cpp
@@ -670,12 +670,15 @@ namespace Metasound
 		{
 			if (Metasound)
 			{
-				FMetasoundAssetBase* MetasoundAsset = IMetasoundUObjectRegistry::Get().GetObjectAsAssetBase(Metasound);
-				check(MetasoundAsset);
-
-				if (UMetasoundEditorGraph* Graph = Cast<UMetasoundEditorGraph>(MetasoundAsset->GetGraph()))
+				if (const UAudioComponent* PreviewComponent = GEditor->GetPreviewAudioComponent())
 				{
-					return Graph->IsPreviewing();
+					if (PreviewComponent->IsPlaying())
+					{
+						if (const USoundBase* Sound = PreviewComponent->Sound)
+						{
+							return Sound->GetUniqueID() == Metasound->GetUniqueID();
+						}
+					}
 				}
 			}
 
@@ -744,9 +747,18 @@ namespace Metasound
 			FGraphEditorCommands::Register();
 			FEditorCommands::Register();
 
-			constexpr bool bForceRefreshNodes = true;
-			FGraphBuilder::RegisterGraphWithFrontend(*Metasound);
-			FGraphBuilder::SynchronizeGraph(*Metasound, bForceRefreshNodes);
+			// If sound was already playing in the editor (ex. from ContentBrowser),
+			// restart to synchronize visual state of editor (ex. volume meter analysis
+			// via transient AudioBus, PlayTime, etc.). If playing, registration is not
+			// required here as it will be handled in play call below after UI is initialized
+			const bool bRestartSound = IsPlaying();
+			if (!bRestartSound)
+			{
+				constexpr bool bForceRefreshNodes = true;
+				FGraphBuilder::RegisterGraphWithFrontend(*Metasound);
+				FGraphBuilder::SynchronizeGraph(*Metasound, bForceRefreshNodes);
+			}
+
 
 			BindGraphCommands();
 			CreateInternalWidgets();
@@ -809,14 +821,21 @@ namespace Metasound
 					)
 				);
 
-			const bool bCreateDefaultStandaloneMenu = true;
-			const bool bCreateDefaultToolbar = true;
-			FAssetEditorToolkit::InitAssetEditor(Mode, InitToolkitHost, TEXT("MetasoundEditorApp"), StandaloneDefaultLayout, bCreateDefaultStandaloneMenu, bCreateDefaultToolbar, ObjectToEdit, false);
+			constexpr bool bCreateDefaultStandaloneMenu = true;
+			constexpr bool bCreateDefaultToolbar = true;
+			constexpr bool bToolbarFocusable = false;
+			constexpr bool bUseSmallToolbarIcons = true;
+			FAssetEditorToolkit::InitAssetEditor(Mode, InitToolkitHost, TEXT("MetasoundEditorApp"), StandaloneDefaultLayout, bCreateDefaultStandaloneMenu, bCreateDefaultToolbar, ObjectToEdit, bToolbarFocusable, bUseSmallToolbarIcons);
 
 			ExtendToolbar();
 			RegenerateMenusAndToolbars();
 
 			NotifyDocumentVersioned();
+
+			if (bRestartSound)
+			{
+				Play();
+			}
 		}
 
 		UObject* FEditor::GetMetasoundObject() const
@@ -855,9 +874,56 @@ namespace Metasound
 
 		FLinearColor FEditor::GetWorldCentricTabColorScale() const
 		{
+			if (const ISlateStyle* MetasoundStyle = FSlateStyleRegistry::FindSlateStyle("MetaSoundStyle"))
+			{
+				if (UMetaSoundSource* MetaSoundSource = Cast<UMetaSoundSource>(Metasound))
+				{
+					return MetasoundStyle->GetColor("MetasoundSource.Color");
+				}
+
+				if (UMetaSound* MetaSound = Cast<UMetaSound>(Metasound))
+				{
+					return MetasoundStyle->GetColor("Metasound.Color");
+				}
+			}
+
 			return FLinearColor(0.3f, 0.2f, 0.5f, 0.5f);
 		}
 
+		const FSlateBrush* FEditor::GetDefaultTabIcon() const
+		{
+			if (IsPlaying())
+			{
+				if (const ISlateStyle* MetasoundStyle = FSlateStyleRegistry::FindSlateStyle("MetaSoundStyle"))
+				{
+					return MetasoundStyle->GetBrush("MetasoundEditor.Play");
+				}
+			}
+
+			return FAssetEditorToolkit::GetDefaultTabIcon();
+		}
+
+		FLinearColor FEditor::GetDefaultTabColor() const
+		{
+			if (IsPlaying())
+			{
+				if (const ISlateStyle* MetasoundStyle = FSlateStyleRegistry::FindSlateStyle("MetaSoundStyle"))
+				{
+					if (UMetaSoundSource* MetaSoundSource = Cast<UMetaSoundSource>(Metasound))
+					{
+						return MetasoundStyle->GetColor("MetasoundSource.Color");
+					}
+
+					if (UMetaSound* MetaSound = Cast<UMetaSound>(Metasound))
+					{
+						return MetasoundStyle->GetColor("Metasound.Color");
+					}
+				}
+			}
+
+			return FAssetEditorToolkit::GetDefaultTabColor();
+		}
+
 		FName FEditor::GetEditorName() const 
 		{
 			return FEditor::EditorName;
diff --git a/Engine/Plugins/Runtime/Metasound/Source/MetasoundEditor/Private/MetasoundEditor.h b/Engine/Plugins/Runtime/Metasound/Source/MetasoundEditor/Private/MetasoundEditor.h
--- a/Engine/Plugins/Runtime/Metasound/Source/MetasoundEditor/Private/MetasoundEditor.h
+++ b/Engine/Plugins/Runtime/Metasound/Source/MetasoundEditor/Private/MetasoundEditor.h
@@ -110,6 +110,8 @@ namespace Metasound
 			virtual FText GetBaseToolkitName() const override;
 			virtual FString GetWorldCentricTabPrefix() const override;
 			virtual FLinearColor GetWorldCentricTabColorScale() const override;
+			virtual const FSlateBrush* GetDefaultTabIcon() const override;
+			virtual FLinearColor GetDefaultTabColor() const override;
 
 			/** IAssetEditorInstance interface */
 			virtual FName GetEditorName() const override;
diff --git a/Engine/Plugins/Runtime/Metasound/Source/MetasoundEditor/Private/MetasoundEditorModule.cpp b/Engine/Plugins/Runtime/Metasound/Source/MetasoundEditor/Private/MetasoundEditorModule.cpp
--- a/Engine/Plugins/Runtime/Metasound/Source/MetasoundEditor/Private/MetasoundEditorModule.cpp
+++ b/Engine/Plugins/Runtime/Metasound/Source/MetasoundEditor/Private/MetasoundEditorModule.cpp
@@ -98,6 +98,9 @@ namespace Metasound
 
 				// Metasound Editor
 				{
+					Set("Metasound.Color", FColor(13, 55, 13));
+					Set("MetasoundSource.Color", FColor(103, 214, 66));
+
 					// Actions
 					Set("MetasoundEditor.Play", new FSlateImageBrush(RootToContentDir(TEXT("Icons/play_40x.png")), Icon40x40));
 					Set("MetasoundEditor.Play.Small", new FSlateImageBrush(RootToContentDir(TEXT("Icons/play_40x.png")), Icon20x20));
