Use loc text for MetaSound node vertex names (Part 2)
- pass on MetaSound nodes to use METASOUND_PARAM() and related macros for information about node vertices that should be localized
- fix typos, loc key collisions, and update audio unit test call sites

diff --git a/Engine/Plugins/Runtime/Metasound/Source/MetasoundEditor/Private/MetasoundEditorGraph.cpp b/Engine/Plugins/Runtime/Metasound/Source/MetasoundEditor/Private/MetasoundEditorGraph.cpp
--- a/Engine/Plugins/Runtime/Metasound/Source/MetasoundEditor/Private/MetasoundEditorGraph.cpp
+++ b/Engine/Plugins/Runtime/Metasound/Source/MetasoundEditor/Private/MetasoundEditorGraph.cpp
@@ -1131,6 +1131,7 @@ void UMetasoundEditorGraphVariable::ResetToClassDefault()
 {
 	using namespace Metasound;
 	using namespace Metasound::Frontend;
+	using namespace Metasound::VariableNames; 
 
 	FMetasoundFrontendLiteral DefaultLiteral;
 	DefaultLiteral.SetFromLiteral(IDataTypeRegistry::Get().CreateDefaultLiteral(GetDataType()));
@@ -1144,7 +1145,7 @@ void UMetasoundEditorGraphVariable::ResetToClassDefault()
 	FNodeHandle MutatorNode = VariableHandle->FindMutatorNode();
 	if (MutatorNode->IsValid())
 	{
-		FInputHandle InputHandle = MutatorNode->GetInputWithVertexName(VariableNames::GetInputDataName());
+		FInputHandle InputHandle = MutatorNode->GetInputWithVertexName(METASOUND_GET_PARAM_NAME(InputData));
 		if (ensure(InputHandle->IsValid()))
 		{
 			InputHandle->ClearLiteral();
@@ -1162,6 +1163,7 @@ void UMetasoundEditorGraphVariable::UpdateFrontendDefaultLiteral(bool bPostTrans
 	using namespace Metasound;
 	using namespace Metasound::Editor;
 	using namespace Metasound::Frontend;
+	using namespace Metasound::VariableNames;
 
 	UObject* Metasound = nullptr;
 	UMetasoundEditorGraph* MetaSoundGraph = GetOwningGraph();
@@ -1190,7 +1192,7 @@ void UMetasoundEditorGraphVariable::UpdateFrontendDefaultLiteral(bool bPostTrans
 	FNodeHandle MutatorNode = VariableHandle->FindMutatorNode();
 	if (MutatorNode->IsValid())
 	{
-		FInputHandle InputHandle = MutatorNode->GetInputWithVertexName(VariableNames::GetInputDataName());
+		FInputHandle InputHandle = MutatorNode->GetInputWithVertexName(METASOUND_GET_PARAM_NAME(InputData));
 		if (ensure(InputHandle->IsValid()))
 		{
 			GraphPrivate::SetLiteralOrClearIfMatchesDefault(InputHandle, DefaultLiteral);
diff --git a/Engine/Plugins/Runtime/Metasound/Source/MetasoundEditor/Private/MetasoundEditorGraphBuilder.cpp b/Engine/Plugins/Runtime/Metasound/Source/MetasoundEditor/Private/MetasoundEditorGraphBuilder.cpp
--- a/Engine/Plugins/Runtime/Metasound/Source/MetasoundEditor/Private/MetasoundEditorGraphBuilder.cpp
+++ b/Engine/Plugins/Runtime/Metasound/Source/MetasoundEditor/Private/MetasoundEditorGraphBuilder.cpp
@@ -201,6 +201,8 @@ namespace Metasound
 
 		FName FGraphBuilder::GetPinName(const Frontend::IOutputController& InFrontendOutput)
 		{
+			using namespace VariableNames; 
+
 			Frontend::FConstNodeHandle OwningNode = InFrontendOutput.GetOwningNode();
 			EMetasoundFrontendClassType OwningNodeClassType = OwningNode->GetClassMetadata().GetType();
 
@@ -215,7 +217,7 @@ namespace Metasound
 					// inputs and outputs and the editor does not display the pin's name. The
 					// editor instead displays the variable's name in place of the pin name to
 					// maintain a consistent look and behavior to input and output nodes.
-					return VariableNames::GetOutputDataName();
+					return METASOUND_GET_PARAM_NAME(OutputData);
 				}
 				case EMetasoundFrontendClassType::Input:
 				case EMetasoundFrontendClassType::Output:
@@ -232,6 +234,8 @@ namespace Metasound
 
 		FName FGraphBuilder::GetPinName(const Frontend::IInputController& InFrontendInput)
 		{
+			using namespace VariableNames;
+			
 			Frontend::FConstNodeHandle OwningNode = InFrontendInput.GetOwningNode();
 			EMetasoundFrontendClassType OwningNodeClassType = OwningNode->GetClassMetadata().GetType();
 
@@ -246,7 +250,7 @@ namespace Metasound
 					// inputs and outputs and the editor does not display the pin's name. The
 					// editor instead displays the variable's name in place of the pin name to
 					// maintain a consistent look and behavior to input and output nodes.
-					return VariableNames::GetInputDataName();
+					return METASOUND_GET_PARAM_NAME(InputData);
 				}
 
 				case EMetasoundFrontendClassType::Input:
@@ -577,6 +581,7 @@ namespace Metasound
 		Frontend::FInputHandle FGraphBuilder::GetInputHandleFromPin(const UEdGraphPin* InPin)
 		{
 			using namespace Frontend;
+			using namespace VariableNames;
 
 			if (InPin && ensure(InPin->Direction == EGPD_Input))
 			{
@@ -586,7 +591,7 @@ namespace Metasound
 					// purposes instead of the underlying vertex's name. The frontend vertices
 					// of a variable node have consistent names no matter what the 
 					// variable is named.
-					return EdVariableNode->GetNodeHandle()->GetInputWithVertexName(VariableNames::GetInputDataName());
+					return EdVariableNode->GetNodeHandle()->GetInputWithVertexName(METASOUND_GET_PARAM_NAME(InputData));
 				}
 				else if (UMetasoundEditorGraphNode* EdNode = CastChecked<UMetasoundEditorGraphNode>(InPin->GetOwningNode()))
 				{
@@ -605,6 +610,7 @@ namespace Metasound
 		Frontend::FOutputHandle FGraphBuilder::GetOutputHandleFromPin(const UEdGraphPin* InPin)
 		{
 			using namespace Frontend;
+			using namespace VariableNames; 
 
 			if (InPin && ensure(InPin->Direction == EGPD_Output))
 			{
@@ -614,7 +620,7 @@ namespace Metasound
 					// purposes instead of the underlying vertex's name. The frontend vertices
 					// of a variable node have consistent names no matter what the 
 					// variable is named.
-					return EdVariableNode->GetNodeHandle()->GetOutputWithVertexName(VariableNames::GetOutputDataName());
+					return EdVariableNode->GetNodeHandle()->GetOutputWithVertexName(METASOUND_GET_PARAM_NAME(OutputData));
 				}
 				else if (UMetasoundEditorGraphNode* EdNode = CastChecked<UMetasoundEditorGraphNode>(InPin->GetOwningNode()))
 				{
diff --git a/Engine/Plugins/Runtime/Metasound/Source/MetasoundEngine/Private/MetasoundWavePlayerNode.cpp b/Engine/Plugins/Runtime/Metasound/Source/MetasoundEngine/Private/MetasoundWavePlayerNode.cpp
--- a/Engine/Plugins/Runtime/Metasound/Source/MetasoundEngine/Private/MetasoundWavePlayerNode.cpp
+++ b/Engine/Plugins/Runtime/Metasound/Source/MetasoundEngine/Private/MetasoundWavePlayerNode.cpp
@@ -11,6 +11,7 @@
 #include "MetasoundExecutableOperator.h"
 #include "MetasoundLog.h"
 #include "MetasoundNodeRegistrationMacro.h"
+#include "MetasoundParamHelper.h"
 #include "MetasoundPrimitives.h"
 #include "MetasoundTrace.h"
 #include "MetasoundTrigger.h"
@@ -24,47 +25,26 @@ namespace Metasound
 {
 	namespace WavePlayerVertexNames
 	{
-		static const TCHAR* InputTriggerPlayName = TEXT("Play");
-		static const TCHAR* InputTriggerStopName = TEXT("Stop");
-		static const TCHAR* InputWaveAssetName = TEXT("Wave Asset");
-		static const TCHAR* InputStartTimeName = TEXT("Start Time");
-		static const TCHAR* InputPitchShiftName = TEXT("Pitch Shift");
-		static const TCHAR* InputLoopName = TEXT("Loop");
-		static const TCHAR* InputLoopStartName = TEXT("Loop Start");
-		static const TCHAR* InputLoopDurationName = TEXT("Loop Duration");
-
-		static const TCHAR* OutputTriggerOnPlayName = TEXT("On Play");
-		static const TCHAR* OutputTriggerOnDoneName = TEXT("On Finished");
-		static const TCHAR* OutputTriggerOnNearlyDoneName = TEXT("On Nearly Finished");
-		static const TCHAR* OutputTriggerOnLoopedName = TEXT("On Looped");
-		static const TCHAR* OutputTriggerOnCuePointName = TEXT("On Cue Point");
-		static const TCHAR* OutputCuePointIDName = TEXT("Cue Point ID");
-		static const TCHAR* OutputCuePointLabelName = TEXT("Cue Point Label");
-		static const TCHAR* OutputLoopPercentName = TEXT("Loop Percent");
-		static const TCHAR* OutputPlaybackLocationName = TEXT("Playback Location");
-		static const TCHAR* OutputAudioLeftName = TEXT("Out Left");
-		static const TCHAR* OutputAudioRightName = TEXT("Out Right");
-
-		static FText InputTriggerPlayTT = METASOUND_LOCTEXT("PlayTT", "Play the wave player.");
-		static FText InputTriggerStopTT = METASOUND_LOCTEXT("StopTT", "Stop the wave player.");
-		static FText InputWaveAssetTT = METASOUND_LOCTEXT("WaveTT", "The wave asset to be real-time decoded.");
-		static FText InputStartTimeTT = METASOUND_LOCTEXT("StartTimeTT", "Time into the wave asset to start (seek) the wave asset. For real-time decoding, the wave asset must be set to seekable!)");
-		static FText InputPitchShiftTT = METASOUND_LOCTEXT("PitchShiftTT", "The pitch shift to use for the wave asset in semitones.");
-		static FText InputLoopTT = METASOUND_LOCTEXT("LoopTT", "Whether or not to loop between the start and specified end times.");
-		static FText InputLoopStartTT = METASOUND_LOCTEXT("LoopStartTT", "When to start the loop.");
-		static FText InputLoopDurationTT = METASOUND_LOCTEXT("LoopDurationTT", "The duration of the loop when wave player is enabled for looping. A negative value will loop the whole wave asset.");
-
-		static FText OutputTriggerOnPlayTT = METASOUND_LOCTEXT("OnPlayTT", "Triggers when Play is triggered.");
-		static FText OutputTriggerOnDoneTT = METASOUND_LOCTEXT("OnDoneTT", "Triggers when the wave played has finished playing.");
-		static FText OutputTriggerOnNearlyDoneTT = METASOUND_LOCTEXT("OnNearlyDoneTT", "Triggers when the wave played has almost finished playing (the block before it finishes). Allows time for logic to trigger different variations to play seamlessly.");
-		static FText OutputTriggerOnLoopedTT = METASOUND_LOCTEXT("OnLoopedTT", "Triggers when the wave player has looped.");
-		static FText OutputTriggerOnCuePointTT = METASOUND_LOCTEXT("OnCuePointTT", "Triggers when a wave cue point was hit during playback.");
-		static FText OutputCuePointIDTT = METASOUND_LOCTEXT("CuePointIDTT", "The cue point ID that was triggered.");
-		static FText OutputCuePointLabelTT = METASOUND_LOCTEXT("CuePointLabelTT", "The cue point label that was triggered (if there was a label parsed in the imported .wav file).");
-		static FText OutputLoopPercentTT = METASOUND_LOCTEXT("LoopPercentTT", "Returns the current loop percent if looping is enabled.");
-		static FText OutputPlaybackLocationTT = METASOUND_LOCTEXT("PlaybackLocationTT", "Returns the absolute position of the wave playback as a precentage of wave duration.");
-		static FText OutputAudioLeftNameTT = METASOUND_LOCTEXT("AudioLeftTT", "The left channel audio output. Mono wave assets will be upmixed to dual stereo.");
-		static FText OutputAudioRightNameTT = METASOUND_LOCTEXT("AudioRightTT", "The right channel audio output. Mono wave assets will be upmixed to dual stereo.");
+		METASOUND_PARAM(InputTriggerPlay, "Play", "Play the wave player.")
+		METASOUND_PARAM(InputTriggerStop, "Stop", "Stop the wave player.")
+		METASOUND_PARAM(InputWaveAsset, "Wave Asset", "The wave asset to be real-time decoded.")
+		METASOUND_PARAM(InputStartTime, "Start Time", "Time into the wave asset to start (seek) the wave asset.")
+		METASOUND_PARAM(InputPitchShift, "Pitch Shift", "The pitch shift to use for the wave asset in semitones.")
+		METASOUND_PARAM(InputLoop, "Loop", "Whether or not to loop between the start and specified end times.")
+		METASOUND_PARAM(InputLoopStart, "Loop Start", "When to start the loop.")
+		METASOUND_PARAM(InputLoopDuration, "Loop Duration", "The duration of the loop when wave player is enabled for looping. A negative value will loop the whole wave asset.")
+		
+		METASOUND_PARAM(OutputTriggerOnPlay, "On Play", "Triggers when Play is triggered.")
+		METASOUND_PARAM(OutputTriggerOnDone, "On Finished", "Triggers when the wave played has finished playing.")
+		METASOUND_PARAM(OutputTriggerOnNearlyDone, "On Nearly Finished", "Triggers when the wave played has almost finished playing (the block before it finishes). Allows time for logic to trigger different variations to play seamlessly. playing.")
+		METASOUND_PARAM(OutputTriggerOnLooped, "On Looped", "Triggers when the wave player has looped.")
+		METASOUND_PARAM(OutputTriggerOnCuePoint, "On Cue Point", "Triggers when a wave cue point was hit during playback.")
+		METASOUND_PARAM(OutputCuePointID, "Cue Point ID", "The cue point ID that was triggered.")
+		METASOUND_PARAM(OutputCuePointLabel, "Cue Point Label", "The cue point label that was triggered (if there was a label parsed in the imported .wav file).")
+		METASOUND_PARAM(OutputLoopRatio, "Loop Ratio", "Returns the current playback location as a ratio of the loop (0-1) if looping is enabled.")
+		METASOUND_PARAM(OutputPlaybackLocation, "OutputPlaybackLocation", "Returns the absolute position of the wave playback as a ratio of wave duration (0-1).")
+		METASOUND_PARAM(OutputAudioLeft, "Out Left", "The left channel audio output. Mono wave assets will be upmixed to dual stereo.")
+		METASOUND_PARAM(OutputAudioRight, "Out Right", "The right channel audio output. Mono wave assets will be upmixed to dual stereo.")
 	}
 
 	class FWavePlayerNode : public FNode
@@ -390,14 +370,14 @@ namespace Metasound
 			using namespace WavePlayerVertexNames;
 
 			FDataReferenceCollection InputDataReferences;
-			InputDataReferences.AddDataReadReference(InputTriggerPlayName, PlayTrigger);
-			InputDataReferences.AddDataReadReference(InputTriggerStopName, StopTrigger);
-			InputDataReferences.AddDataReadReference(InputWaveAssetName, WaveAsset);
-			InputDataReferences.AddDataReadReference(InputStartTimeName, StartTime);
-			InputDataReferences.AddDataReadReference(InputPitchShiftName, PitchShift);
-			InputDataReferences.AddDataReadReference(InputLoopName, bLoop);
-			InputDataReferences.AddDataReadReference(InputLoopStartName, LoopStartTime);
-			InputDataReferences.AddDataReadReference(InputLoopDurationName, LoopDuration);
+			InputDataReferences.AddDataReadReference(METASOUND_GET_PARAM_NAME(InputTriggerPlay), PlayTrigger);
+			InputDataReferences.AddDataReadReference(METASOUND_GET_PARAM_NAME(InputTriggerStop), StopTrigger);
+			InputDataReferences.AddDataReadReference(METASOUND_GET_PARAM_NAME(InputWaveAsset), WaveAsset);
+			InputDataReferences.AddDataReadReference(METASOUND_GET_PARAM_NAME(InputStartTime), StartTime);
+			InputDataReferences.AddDataReadReference(METASOUND_GET_PARAM_NAME(InputPitchShift), PitchShift);
+			InputDataReferences.AddDataReadReference(METASOUND_GET_PARAM_NAME(InputLoop), bLoop);
+			InputDataReferences.AddDataReadReference(METASOUND_GET_PARAM_NAME(InputLoopStart), LoopStartTime);
+			InputDataReferences.AddDataReadReference(METASOUND_GET_PARAM_NAME(InputLoopDuration), LoopDuration);
 			return InputDataReferences;
 		}
 
@@ -406,17 +386,17 @@ namespace Metasound
 			using namespace WavePlayerVertexNames;
 
 			FDataReferenceCollection OutputDataReferences;
-			OutputDataReferences.AddDataReadReference(OutputTriggerOnPlayName, PlayTrigger);
-			OutputDataReferences.AddDataReadReference(OutputTriggerOnDoneName, TriggerOnDone);
-			OutputDataReferences.AddDataReadReference(OutputTriggerOnNearlyDoneName, TriggerOnNearlyDone);
-			OutputDataReferences.AddDataReadReference(OutputTriggerOnLoopedName, TriggerOnLooped);
-			OutputDataReferences.AddDataReadReference(OutputTriggerOnCuePointName, TriggerOnCuePoint);
-			OutputDataReferences.AddDataReadReference(OutputCuePointIDName, CuePointID);
-			OutputDataReferences.AddDataReadReference(OutputCuePointLabelName, CuePointLabel);
-			OutputDataReferences.AddDataReadReference(OutputLoopPercentName, LoopPercent);
-			OutputDataReferences.AddDataReadReference(OutputPlaybackLocationName, PlaybackLocation);
-			OutputDataReferences.AddDataReadReference(OutputAudioLeftName, AudioBufferL);
-			OutputDataReferences.AddDataReadReference(OutputAudioRightName, AudioBufferR);
+			OutputDataReferences.AddDataReadReference(METASOUND_GET_PARAM_NAME(OutputTriggerOnPlay), PlayTrigger);
+			OutputDataReferences.AddDataReadReference(METASOUND_GET_PARAM_NAME(OutputTriggerOnDone), TriggerOnDone);
+			OutputDataReferences.AddDataReadReference(METASOUND_GET_PARAM_NAME(OutputTriggerOnNearlyDone), TriggerOnNearlyDone);
+			OutputDataReferences.AddDataReadReference(METASOUND_GET_PARAM_NAME(OutputTriggerOnLooped), TriggerOnLooped);
+			OutputDataReferences.AddDataReadReference(METASOUND_GET_PARAM_NAME(OutputTriggerOnCuePoint), TriggerOnCuePoint);
+			OutputDataReferences.AddDataReadReference(METASOUND_GET_PARAM_NAME(OutputCuePointID), CuePointID);
+			OutputDataReferences.AddDataReadReference(METASOUND_GET_PARAM_NAME(OutputCuePointLabel), CuePointLabel);
+			OutputDataReferences.AddDataReadReference(METASOUND_GET_PARAM_NAME(OutputLoopRatio), LoopPercent);
+			OutputDataReferences.AddDataReadReference(METASOUND_GET_PARAM_NAME(OutputPlaybackLocation), PlaybackLocation);
+			OutputDataReferences.AddDataReadReference(METASOUND_GET_PARAM_NAME(OutputAudioLeft), AudioBufferL);
+			OutputDataReferences.AddDataReadReference(METASOUND_GET_PARAM_NAME(OutputAudioRight), AudioBufferR);
 			return OutputDataReferences;
 		}
 
@@ -1016,14 +996,14 @@ namespace Metasound
 		FWavePlayerOpArgs Args =
 		{
 			InParams.OperatorSettings,
-			Inputs.GetDataReadReferenceOrConstruct<FTrigger>(InputTriggerPlayName, InParams.OperatorSettings),
-			Inputs.GetDataReadReferenceOrConstruct<FTrigger>(InputTriggerStopName, InParams.OperatorSettings),
-			Inputs.GetDataReadReferenceOrConstruct<FWaveAsset>(InputWaveAssetName),
-			Inputs.GetDataReadReferenceOrConstruct<FTime>(InputStartTimeName),
-			Inputs.GetDataReadReferenceOrConstruct<float>(InputPitchShiftName),
-			Inputs.GetDataReadReferenceOrConstruct<bool>(InputLoopName),
-			Inputs.GetDataReadReferenceOrConstruct<FTime>(InputLoopStartName),
-			Inputs.GetDataReadReferenceOrConstruct<FTime>(InputLoopDurationName)
+			Inputs.GetDataReadReferenceOrConstruct<FTrigger>(METASOUND_GET_PARAM_NAME(InputTriggerPlay), InParams.OperatorSettings),
+			Inputs.GetDataReadReferenceOrConstruct<FTrigger>(METASOUND_GET_PARAM_NAME(InputTriggerStop), InParams.OperatorSettings),
+			Inputs.GetDataReadReferenceOrConstruct<FWaveAsset>(METASOUND_GET_PARAM_NAME(InputWaveAsset)),
+			Inputs.GetDataReadReferenceOrConstruct<FTime>(METASOUND_GET_PARAM_NAME(InputStartTime)),
+			Inputs.GetDataReadReferenceOrConstruct<float>(METASOUND_GET_PARAM_NAME(InputPitchShift)),
+			Inputs.GetDataReadReferenceOrConstruct<bool>(METASOUND_GET_PARAM_NAME(InputLoop)),
+			Inputs.GetDataReadReferenceOrConstruct<FTime>(METASOUND_GET_PARAM_NAME(InputLoopStart)),
+			Inputs.GetDataReadReferenceOrConstruct<FTime>(METASOUND_GET_PARAM_NAME(InputLoopDuration))
 		};
 
 		return MakeUnique<FWavePlayerOperator>(Args);
@@ -1035,27 +1015,27 @@ namespace Metasound
 
 		return FVertexInterface(
 			FInputVertexInterface(
-				TInputDataVertexModel<FTrigger>(InputTriggerPlayName, InputTriggerPlayTT),
-				TInputDataVertexModel<FTrigger>(InputTriggerStopName, InputTriggerStopTT),
-				TInputDataVertexModel<FWaveAsset>(InputWaveAssetName, InputWaveAssetTT),
-				TInputDataVertexModel<FTime>(InputStartTimeName, InputStartTimeTT, 0.0f),
-				TInputDataVertexModel<float>(InputPitchShiftName, InputPitchShiftTT, 0.0f),
-				TInputDataVertexModel<bool>(InputLoopName, InputLoopTT, false),
-				TInputDataVertexModel<FTime>(InputLoopStartName, InputLoopStartTT, 0.0f),
-				TInputDataVertexModel<FTime>(InputLoopDurationName, InputLoopDurationTT, -1.0f)
+				TInputDataVertexModel<FTrigger>(METASOUND_GET_PARAM_NAME_AND_METADATA(InputTriggerPlay)),
+				TInputDataVertexModel<FTrigger>(METASOUND_GET_PARAM_NAME_AND_METADATA(InputTriggerStop)),
+				TInputDataVertexModel<FWaveAsset>(METASOUND_GET_PARAM_NAME_AND_METADATA(InputWaveAsset)),
+				TInputDataVertexModel<FTime>(METASOUND_GET_PARAM_NAME_AND_METADATA(InputStartTime), 0.0f),
+				TInputDataVertexModel<float>(METASOUND_GET_PARAM_NAME_AND_METADATA(InputPitchShift), 0.0f),
+				TInputDataVertexModel<bool>(METASOUND_GET_PARAM_NAME_AND_METADATA(InputLoop), false),
+				TInputDataVertexModel<FTime>(METASOUND_GET_PARAM_NAME_AND_METADATA(InputLoopStart), 0.0f),
+				TInputDataVertexModel<FTime>(METASOUND_GET_PARAM_NAME_AND_METADATA(InputLoopDuration), -1.0f)
 				),
 			FOutputVertexInterface(
-				TOutputDataVertexModel<FTrigger>(OutputTriggerOnPlayName, OutputTriggerOnPlayTT),
-				TOutputDataVertexModel<FTrigger>(OutputTriggerOnDoneName, OutputTriggerOnDoneTT),
-				TOutputDataVertexModel<FTrigger>(OutputTriggerOnNearlyDoneName, OutputTriggerOnNearlyDoneTT),
-				TOutputDataVertexModel<FTrigger>(OutputTriggerOnLoopedName, OutputTriggerOnLoopedTT),
-				TOutputDataVertexModel<FTrigger>(OutputTriggerOnCuePointName, OutputTriggerOnCuePointTT),
-				TOutputDataVertexModel<int32>(OutputCuePointIDName, OutputCuePointIDTT),
-				TOutputDataVertexModel<FString>(OutputCuePointLabelName, OutputCuePointLabelTT),
-				TOutputDataVertexModel<float>(OutputLoopPercentName, OutputLoopPercentTT),
-				TOutputDataVertexModel<float>(OutputPlaybackLocationName, OutputPlaybackLocationTT),
-				TOutputDataVertexModel<FAudioBuffer>(OutputAudioLeftName, OutputAudioLeftNameTT),
-				TOutputDataVertexModel<FAudioBuffer>(OutputAudioRightName, OutputAudioRightNameTT)
+				TOutputDataVertexModel<FTrigger>(METASOUND_GET_PARAM_NAME_AND_METADATA(OutputTriggerOnPlay)),
+				TOutputDataVertexModel<FTrigger>(METASOUND_GET_PARAM_NAME_AND_METADATA(OutputTriggerOnDone)),
+				TOutputDataVertexModel<FTrigger>(METASOUND_GET_PARAM_NAME_AND_METADATA(OutputTriggerOnNearlyDone)),
+				TOutputDataVertexModel<FTrigger>(METASOUND_GET_PARAM_NAME_AND_METADATA(OutputTriggerOnLooped)),
+				TOutputDataVertexModel<FTrigger>(METASOUND_GET_PARAM_NAME_AND_METADATA(OutputTriggerOnCuePoint)),
+				TOutputDataVertexModel<int32>(METASOUND_GET_PARAM_NAME_AND_METADATA(OutputCuePointID)),
+				TOutputDataVertexModel<FString>(METASOUND_GET_PARAM_NAME_AND_METADATA(OutputCuePointLabel)),
+				TOutputDataVertexModel<float>(METASOUND_GET_PARAM_NAME_AND_METADATA(OutputLoopRatio)),
+				TOutputDataVertexModel<float>(METASOUND_GET_PARAM_NAME_AND_METADATA(OutputPlaybackLocation)),
+				TOutputDataVertexModel<FAudioBuffer>(METASOUND_GET_PARAM_NAME_AND_METADATA(OutputAudioLeft)),
+				TOutputDataVertexModel<FAudioBuffer>(METASOUND_GET_PARAM_NAME_AND_METADATA(OutputAudioRight))
 			)
 		);
 	}
diff --git a/Engine/Plugins/Runtime/Metasound/Source/MetasoundFrontend/Private/MetasoundArrayNodes.cpp b/Engine/Plugins/Runtime/Metasound/Source/MetasoundFrontend/Private/MetasoundArrayNodes.cpp
--- a/Engine/Plugins/Runtime/Metasound/Source/MetasoundFrontend/Private/MetasoundArrayNodes.cpp
+++ b/Engine/Plugins/Runtime/Metasound/Source/MetasoundFrontend/Private/MetasoundArrayNodes.cpp
@@ -28,77 +28,6 @@ namespace Metasound
 			return Metadata;
 		}
 	}
-
-	namespace ArrayNodeVertexNames
-	{
-		/* Input Vertx Names */
-		const FVertexName& GetInputArrayName()
-		{
-			static const FVertexName Name = TEXT("Array");
-			return Name;
-		}
-
-		const FVertexName& GetInputLeftArrayName()
-		{
-			static const FVertexName Name = TEXT("Left Array");
-			return Name;
-		}
-
-		const FVertexName& GetInputRightArrayName()
-		{
-			static const FVertexName Name = TEXT("Right Array");
-			return Name;
-		}
-
-		const FVertexName& GetInputTriggerName()
-		{
-			static const FVertexName Name = TEXT("Trigger");
-			return Name;
-		}
-
-		const FVertexName& GetInputIndexName()
-		{
-			static const FVertexName Name = TEXT("Index");
-			return Name;
-		}
-
-		const FVertexName& GetInputStartIndexName()
-		{
-			static const FVertexName Name = TEXT("Start Index");
-			return Name;
-		}
-
-		const FVertexName& GetInputEndIndexName()
-		{
-			static const FVertexName Name = TEXT("End Index");
-			return Name;
-		}
-
-		const FVertexName& GetInputValueName()
-		{
-			static const FVertexName Name = TEXT("Value");
-			return Name;
-		}
-
-		/* Output Vertex Names */
-		const FVertexName& GetOutputNumName()
-		{
-			static const FVertexName Name = TEXT("Num");
-			return Name;
-		}
-
-		const FVertexName& GetOutputValueName()
-		{
-			static const FVertexName Name = TEXT("Element");
-			return Name;
-		}
-
-		const FVertexName& GetOutputArrayName()
-		{
-			static const FVertexName Name = TEXT("Array");
-			return Name;
-		}
-	}
 }
 
 #undef LOCTEXT_NAMESPACE
diff --git a/Engine/Plugins/Runtime/Metasound/Source/MetasoundFrontend/Private/MetasoundArrayShuffleNode.cpp b/Engine/Plugins/Runtime/Metasound/Source/MetasoundFrontend/Private/MetasoundArrayShuffleNode.cpp
--- a/Engine/Plugins/Runtime/Metasound/Source/MetasoundFrontend/Private/MetasoundArrayShuffleNode.cpp
+++ b/Engine/Plugins/Runtime/Metasound/Source/MetasoundFrontend/Private/MetasoundArrayShuffleNode.cpp
@@ -7,75 +7,6 @@
 
 namespace Metasound
 {
-	namespace ArrayNodeShuffleVertexNames
-	{
-		const FVertexName& GetInputTriggerNextName()
-		{
-			static const FVertexName Name = TEXT("Next");
-			return Name;
-		}
-
-		const FVertexName& GetInputTriggerShuffleName()
-		{
-			static const FVertexName Name = TEXT("Shuffle");
-			return Name;
-		}
-
-		const FVertexName& GetInputTriggerResetName()
-		{
-			static const FVertexName Name = TEXT("Reset Seed");
-			return Name;
-		}
-
-		const FVertexName& GetInputShuffleArrayName()
-		{
-			static const FVertexName Name = TEXT("In Array");
-			return Name;
-		}
-
-		const FVertexName& GetInputSeedName()
-		{
-			static const FVertexName Name = TEXT("Seed");
-			return Name;
-		}
-
-		const FVertexName& GetInputAutoShuffleName()
-		{
-			static const FVertexName Name = TEXT("Auto Shuffle");
-			return Name;
-		}
-
-		const FVertexName& GetInputEnableSharedStateName()
-		{
-			static const FVertexName Name = TEXT("Enable Shared State");
-			return Name;
-		}
-
-		const FVertexName& GetOutputTriggerOnNextName()
-		{
-			static const FVertexName Name = TEXT("On Next");
-			return Name;
-		}
-
-		const FVertexName& GetOutputTriggerOnShuffleName()
-		{
-			static const FVertexName Name = TEXT("On Shuffle");
-			return Name;
-		}
-
-		const FVertexName& GetOutputTriggerOnResetName()
-		{
-			static const FVertexName Name = TEXT("On Reset Seed");
-			return Name;
-		}
-
-		const FVertexName& GetOutputValueName()
-		{
-			static const FVertexName Name = TEXT("Value");
-			return Name;
-		}
-	}
-
 	/**
 	* FArrayIndexShuffler
 	*/
diff --git a/Engine/Plugins/Runtime/Metasound/Source/MetasoundFrontend/Private/MetasoundFrontendGraph.cpp b/Engine/Plugins/Runtime/Metasound/Source/MetasoundFrontend/Private/MetasoundFrontendGraph.cpp
--- a/Engine/Plugins/Runtime/Metasound/Source/MetasoundFrontend/Private/MetasoundFrontendGraph.cpp
+++ b/Engine/Plugins/Runtime/Metasound/Source/MetasoundFrontend/Private/MetasoundFrontendGraph.cpp
@@ -589,6 +589,7 @@ namespace Metasound
 	bool FFrontendGraphBuilder::AddDefaultInputLiterals(FBuildGraphContext& InGraphContext)
 	{
 		using namespace Metasound::Frontend;
+		using namespace LiteralNodeNames; 
 
 		using FIterator = FDefaultInputByIDMap::TIterator;
 
@@ -610,7 +611,7 @@ namespace Metasound
 				UE_LOG(LogMetaSound, Error, TEXT("MetaSound '%s': Failed to find node in graph [NodeID:%s]"), *InGraphContext.BuildContext.DebugAssetName, *LiteralNodeID.ToString());
 				return false;
 			}
-			const FVertexName& FromVertexName = LiteralNodeNames::GetOutputDataName();
+			const FVertexName& FromVertexName = METASOUND_GET_PARAM_NAME(OutputValue);
 
 			const INode* ToNode = InGraphContext.Graph->FindNode(LiteralData.DestinationNodeID);
 			if (nullptr == ToNode)
diff --git a/Engine/Plugins/Runtime/Metasound/Source/MetasoundFrontend/Private/MetasoundFrontendGraphController.cpp b/Engine/Plugins/Runtime/Metasound/Source/MetasoundFrontend/Private/MetasoundFrontendGraphController.cpp
--- a/Engine/Plugins/Runtime/Metasound/Source/MetasoundFrontend/Private/MetasoundFrontendGraphController.cpp
+++ b/Engine/Plugins/Runtime/Metasound/Source/MetasoundFrontend/Private/MetasoundFrontendGraphController.cpp
@@ -406,6 +406,8 @@ namespace Metasound
 
 		FNodeHandle FGraphController::FindOrAddVariableMutatorNode(const FGuid& InVariableID)
 		{
+			using namespace VariableNames; 
+
 			if (FMetasoundFrontendVariable* Variable = FindFrontendVariable(InVariableID))
 			{
 				FNodeHandle SetNode = GetNodeWithID(Variable->MutatorNodeID);
@@ -418,7 +420,7 @@ namespace Metasound
 						if (SetNode->IsValid())
 						{
 							// Initialize set default literal value to that of the variable
-							FInputHandle InputHandle = SetNode->GetInputWithVertexName(VariableNames::GetInputDataName());
+							FInputHandle InputHandle = SetNode->GetInputWithVertexName(METASOUND_GET_PARAM_NAME(InputData));
 							if (ensure(InputHandle->IsValid()))
 							{
 								InputHandle->SetLiteral(Variable->Literal);
@@ -436,8 +438,8 @@ namespace Metasound
 
 							if (ensure(SourceVariableNode->IsValid()))
 							{
-								FInputHandle SetNodeInput = SetNode->GetInputWithVertexName(VariableNames::GetInputVariableName());
-								FOutputHandle SourceVariableNodeOutput = SourceVariableNode->GetOutputWithVertexName(VariableNames::GetOutputVariableName());
+								FInputHandle SetNodeInput = SetNode->GetInputWithVertexName(METASOUND_GET_PARAM_NAME(InputVariable));
+								FOutputHandle SourceVariableNodeOutput = SourceVariableNode->GetOutputWithVertexName(METASOUND_GET_PARAM_NAME(OutputVariable));
 
 								ensure(SetNodeInput->Connect(*SourceVariableNodeOutput));
 							}
@@ -448,8 +450,8 @@ namespace Metasound
 								FNodeHandle HeadGetNode = GetNodeWithID(Variable->AccessorNodeIDs[0]);
 								if (HeadGetNode->IsValid())
 								{
-									FOutputHandle SetNodeOutput = SetNode->GetOutputWithVertexName(VariableNames::GetOutputVariableName());
-									FInputHandle GetNodeInput = HeadGetNode->GetInputWithVertexName(VariableNames::GetInputVariableName());
+									FOutputHandle SetNodeOutput = SetNode->GetOutputWithVertexName(METASOUND_GET_PARAM_NAME(OutputVariable));
+									FInputHandle GetNodeInput = HeadGetNode->GetInputWithVertexName(METASOUND_GET_PARAM_NAME(InputVariable));
 
 									ensure(SetNodeOutput->Connect(*GetNodeInput));
 								}
@@ -469,6 +471,8 @@ namespace Metasound
 
 		FNodeHandle FGraphController::AddVariableAccessorNode(const FGuid& InVariableID)
 		{
+			using namespace VariableNames;
+			
 			if (FMetasoundFrontendVariable* Variable = FindFrontendVariable(InVariableID))
 			{
 				FMetasoundFrontendClass NodeClass;
@@ -479,7 +483,7 @@ namespace Metasound
 					if (NewNode->IsValid())
 					{
 						// Connect new node.
-						FInputHandle NewInput = NewNode->GetInputWithVertexName(VariableNames::GetInputVariableName());
+						FInputHandle NewInput = NewNode->GetInputWithVertexName(METASOUND_GET_PARAM_NAME(InputVariable));
 						FNodeHandle TailNode = FindTailNodeInVariableStack(InVariableID);
 						
 						if (!TailNode->IsValid())
@@ -491,7 +495,7 @@ namespace Metasound
 						if (ensure(TailNode->IsValid()))
 						{
 							// connect new node to the last "get" node.
-							FOutputHandle TailNodeOutput = TailNode->GetOutputWithVertexName(VariableNames::GetOutputVariableName());
+							FOutputHandle TailNodeOutput = TailNode->GetOutputWithVertexName(METASOUND_GET_PARAM_NAME(OutputVariable));
 							check(!TailNodeOutput->IsConnected());
 							const bool bSuccess = TailNodeOutput->Connect(*NewInput);
 							check(bSuccess);
@@ -516,6 +520,8 @@ namespace Metasound
 
 		FNodeHandle FGraphController::AddVariableDeferredAccessorNode(const FGuid& InVariableID)
 		{
+			using namespace VariableNames;
+
 			if (FMetasoundFrontendVariable* Variable = FindFrontendVariable(InVariableID))
 			{
 				FMetasoundFrontendClass NodeClass;
@@ -526,20 +532,20 @@ namespace Metasound
 					if (NewNode->IsValid())
 					{
 						// Connect new node.
-						FOutputHandle NewNodeOutput = NewNode->GetOutputWithVertexName(VariableNames::GetOutputVariableName());
+						FOutputHandle NewNodeOutput = NewNode->GetOutputWithVertexName(METASOUND_GET_PARAM_NAME(OutputVariable));
 						FNodeHandle HeadNode = FindHeadNodeInVariableStack(InVariableID);
 						if (HeadNode->IsValid())
 						{
-							FInputHandle HeadNodeInput = HeadNode->GetInputWithVertexName(VariableNames::GetInputVariableName());
+							FInputHandle HeadNodeInput = HeadNode->GetInputWithVertexName(METASOUND_GET_PARAM_NAME(InputVariable));
 							const bool bSuccess = HeadNodeInput->Connect(*NewNodeOutput);
 							check(bSuccess);
 						}
 
-						FInputHandle NewNodeInput = NewNode->GetInputWithVertexName(VariableNames::GetInputVariableName());
+						FInputHandle NewNodeInput = NewNode->GetInputWithVertexName(METASOUND_GET_PARAM_NAME(InputVariable));
 						FNodeHandle VariableNode = GetNodeWithID(Variable->VariableNodeID);
 						if (ensure(VariableNode->IsValid()))
 						{
-							FOutputHandle VariableNodeOutput = VariableNode->GetOutputWithVertexName(VariableNames::GetOutputVariableName());
+							FOutputHandle VariableNodeOutput = VariableNode->GetOutputWithVertexName(METASOUND_GET_PARAM_NAME(OutputVariable));
 							const bool bSuccess = VariableNodeOutput->Connect(*NewNodeInput);
 							check(bSuccess);
 						}
@@ -700,12 +706,14 @@ namespace Metasound
 			// removes a node while maintaining the daisy-chain.
 			check(InNode.IsValid());
 
-			FInputHandle InputToSpliceOut = InNode.GetInputWithVertexName(VariableNames::GetInputVariableName());
+			using namespace VariableNames;
+
+			FInputHandle InputToSpliceOut = InNode.GetInputWithVertexName(METASOUND_GET_PARAM_NAME(InputVariable));
 			FOutputHandle OutputToReroute = InputToSpliceOut->GetConnectedOutput();
 
 			if (OutputToReroute->IsValid())
 			{
-				FOutputHandle OutputToSpliceOut = InNode.GetOutputWithVertexName(VariableNames::GetOutputVariableName());
+				FOutputHandle OutputToSpliceOut = InNode.GetOutputWithVertexName(METASOUND_GET_PARAM_NAME(OutputVariable));
 				for (const FInputHandle& InputToReroute : OutputToSpliceOut->GetConnectedInputs())
 				{
 					ensure(InputToReroute->Connect(*OutputToReroute));
diff --git a/Engine/Plugins/Runtime/Metasound/Source/MetasoundFrontend/Private/MetasoundFrontendInjectReceiveNodes.cpp b/Engine/Plugins/Runtime/Metasound/Source/MetasoundFrontend/Private/MetasoundFrontendInjectReceiveNodes.cpp
--- a/Engine/Plugins/Runtime/Metasound/Source/MetasoundFrontend/Private/MetasoundFrontendInjectReceiveNodes.cpp
+++ b/Engine/Plugins/Runtime/Metasound/Source/MetasoundFrontend/Private/MetasoundFrontendInjectReceiveNodes.cpp
@@ -5,11 +5,16 @@
 #include "MetasoundFacade.h"
 #include "MetasoundFrontendGraph.h"
 #include "MetasoundFrontendRegistries.h"
+#include "MetasoundParamHelper.h"
 #include "MetasoundReceiveNode.h"
 #include "MetasoundVertex.h"
 
+#define LOCTEXT_NAMESPACE "MetasoundFrontend"
+
 namespace Metasound
 {
+	METASOUND_PARAM(OutputAddress, "Address", "Address")
+
 	namespace Frontend 
 	{
 		namespace MetasoundFrontendInjectReceiveNodesPrivate
@@ -18,7 +23,6 @@ namespace Metasound
 			class FAddressOperator : public FNoOpOperator
 			{
 			public:
-				static const TCHAR* GetAddressVertexKey();
 				static TUniquePtr<IOperator> CreateOperator(const FCreateOperatorParams&, TArray<TUniquePtr<IOperatorBuildError>>&);
 				static const FNodeClassMetadata& GetNodeInfo();
 
@@ -32,7 +36,7 @@ namespace Metasound
 				virtual FDataReferenceCollection GetOutputs() const override
 				{
 					FDataReferenceCollection Outputs;
-					Outputs.AddDataReadReference(GetAddressVertexKey(), Address);
+					Outputs.AddDataReadReference(METASOUND_GET_PARAM_NAME(OutputAddress), Address);
 					return Outputs;
 				}
 
@@ -63,8 +67,6 @@ namespace Metasound
 				FReceiveNodeAddressFunction AddressFunction;
 			};
 
-			const TCHAR* FAddressOperator::GetAddressVertexKey() { return TEXT("Address"); }
-
 			TUniquePtr<IOperator> FAddressOperator::CreateOperator(const FCreateOperatorParams& InParams, TArray<TUniquePtr<IOperatorBuildError>>& OutBuildErrors)
 			{
 				const FAddressNode& AddressNode = static_cast<const FAddressNode&>(InParams.Node);
@@ -81,7 +83,7 @@ namespace Metasound
 						FInputVertexInterface{},
 						FOutputVertexInterface
 						{
-							TOutputDataVertexModel<FSendAddress>(FAddressOperator::GetAddressVertexKey(), FText::GetEmpty())
+							TOutputDataVertexModel<FSendAddress>(METASOUND_GET_PARAM_NAME_AND_METADATA(OutputAddress))
 						}
 					};
 
@@ -126,6 +128,8 @@ namespace Metasound
 		bool InjectReceiveNode(FFrontendGraph& InGraph, const FReceiveNodeAddressFunction& InAddressFunction, const FInputDataDestination& InputDestination)
 		{
 			using namespace MetasoundFrontendInjectReceiveNodesPrivate;
+			using namespace ReceiveNodeInfo; 
+
 			// should never contain null nodes for input destination.
 			check(InputDestination.Node != nullptr);
 
@@ -156,7 +160,7 @@ namespace Metasound
 			InGraph.AddNode(ReceiveNodeID, ReceiveNode);
 			InGraph.AddNode(AddressNodeID, AddressNode);
 
-			bool bDataEdgeAdded = InGraph.AddDataEdge(*AddressNode, FAddressOperator::GetAddressVertexKey(), *ReceiveNode, ReceiveNodeInfo::GetAddressInputName());
+			bool bDataEdgeAdded = InGraph.AddDataEdge(*AddressNode, METASOUND_GET_PARAM_NAME(OutputAddress), *ReceiveNode, METASOUND_GET_PARAM_NAME(AddressInput));
 			ensureAlways(bDataEdgeAdded);
 
 			auto IsEdgeConnectedToCurrentInput = [&InputDestination](const FDataEdge& InEdge)
@@ -170,7 +174,7 @@ namespace Metasound
 			InGraph.RemoveDataEdgeByPredicate(IsEdgeConnectedToCurrentInput);
 
 			// Connect previous connections to receive node output.
-			FOutputDataSource ReceiveOutputSource { *ReceiveNode, ReceiveNode->GetVertexInterface().GetOutputVertex(ReceiveNodeInfo::GetOutputName()) };
+			FOutputDataSource ReceiveOutputSource { *ReceiveNode, ReceiveNode->GetVertexInterface().GetOutputVertex(METASOUND_GET_PARAM_NAME(Output)) };
 			for (const FDataEdge& Edge : EdgesFromInput)
 			{
 				FDataEdge NewEdge { ReceiveOutputSource, Edge.To };
@@ -178,7 +182,7 @@ namespace Metasound
 			}
 
 			// Connect input node to receive node.
-			bDataEdgeAdded = InGraph.AddDataEdge(*InputDestination.Node, VertexKey, *ReceiveNode, ReceiveNodeInfo::GetDefaultDataInputName());
+			bDataEdgeAdded = InGraph.AddDataEdge(*InputDestination.Node, VertexKey, *ReceiveNode, METASOUND_GET_PARAM_NAME(DefaultDataInput));
 			ensureAlways(bDataEdgeAdded);
 
 			return true;
@@ -214,3 +218,4 @@ namespace Metasound
 	}
 }
 
+#undef LOCTEXT_NAMESPACE
diff --git a/Engine/Plugins/Runtime/Metasound/Source/MetasoundFrontend/Private/MetasoundReceiveNode.cpp b/Engine/Plugins/Runtime/Metasound/Source/MetasoundFrontend/Private/MetasoundReceiveNode.cpp
--- a/Engine/Plugins/Runtime/Metasound/Source/MetasoundFrontend/Private/MetasoundReceiveNode.cpp
+++ b/Engine/Plugins/Runtime/Metasound/Source/MetasoundFrontend/Private/MetasoundReceiveNode.cpp
@@ -6,24 +6,6 @@ namespace Metasound
 {
 	namespace ReceiveNodeInfo
 	{
-		const FVertexName& GetAddressInputName()
-		{
-			static const FVertexName InputName(TEXT("Address"));
-			return InputName;
-		}
-
-		const FVertexName& GetDefaultDataInputName()
-		{
-			static const FVertexName DefaultDataName(TEXT("Default"));
-			return DefaultDataName;
-		}
-
-		const FVertexName& GetOutputName()
-		{
-			static const FVertexName OutputName(TEXT("Out"));
-			return OutputName;
-		}
-
 		FNodeClassName GetClassNameForDataType(const FName& InDataTypeName)
 		{
 			return FNodeClassName { "Receive", InDataTypeName, FName() };
diff --git a/Engine/Plugins/Runtime/Metasound/Source/MetasoundFrontend/Public/MetasoundArrayNodes.h b/Engine/Plugins/Runtime/Metasound/Source/MetasoundFrontend/Public/MetasoundArrayNodes.h
--- a/Engine/Plugins/Runtime/Metasound/Source/MetasoundFrontend/Public/MetasoundArrayNodes.h
+++ b/Engine/Plugins/Runtime/Metasound/Source/MetasoundFrontend/Public/MetasoundArrayNodes.h
@@ -11,6 +11,7 @@
 #include "MetasoundNodeInterface.h"
 #include "MetasoundNodeRegistrationMacro.h"
 #include "MetasoundOperatorInterface.h"
+#include "MetasoundParamHelper.h"
 #include "MetasoundPrimitives.h"
 #include "MetasoundTrigger.h"
 #include "MetasoundVertex.h"
@@ -44,20 +45,21 @@ namespace Metasound
 
 	namespace ArrayNodeVertexNames
 	{
-		/* Input Vertex Names */
-		METASOUNDFRONTEND_API const FVertexName& GetInputArrayName();
-		METASOUNDFRONTEND_API const FVertexName& GetInputLeftArrayName();
-		METASOUNDFRONTEND_API const FVertexName& GetInputRightArrayName();
-		METASOUNDFRONTEND_API const FVertexName& GetInputTriggerName();
-		METASOUNDFRONTEND_API const FVertexName& GetInputStartIndexName();
-		METASOUNDFRONTEND_API const FVertexName& GetInputEndIndexName();
-		METASOUNDFRONTEND_API const FVertexName& GetInputIndexName();
-		METASOUNDFRONTEND_API const FVertexName& GetInputValueName();
-
-		/* Output Vertex Names */
-		METASOUNDFRONTEND_API const FVertexName& GetOutputNumName();
-		METASOUNDFRONTEND_API const FVertexName& GetOutputValueName();
-		METASOUNDFRONTEND_API const FVertexName& GetOutputArrayName();
+		METASOUND_PARAM(InputArray, "Array", "Input Array.")
+		METASOUND_PARAM(InputLeftArray, "Left Array", "Input Left Array.")
+		METASOUND_PARAM(InputRightArray, "Right Array", "Input Right Array.")
+		METASOUND_PARAM(InputTriggerGet, "Trigger", "Trigger to get value.")
+		METASOUND_PARAM(InputTriggerSet, "Trigger", "Trigger to set value.")
+		METASOUND_PARAM(InputIndex, "Index", "Index in Array.")
+		METASOUND_PARAM(InputStartIndex, "Start Index", "First index to include.")
+		METASOUND_PARAM(InputEndIndex, "End Index", "Last index to include.")
+		METASOUND_PARAM(InputValue, "Value", "Value to set.")
+
+		METASOUND_PARAM(OutputNum, "Num", "Number of elements in the array.")
+		METASOUND_PARAM(OutputValue, "Value", "Value of element at array index.")
+		METASOUND_PARAM(OutputArrayConcat, "Array", "Array after concatenation.")
+		METASOUND_PARAM(OutputArraySet, "Array", "Array after setting.")
+		METASOUND_PARAM(OutputArraySubset, "Array", "Subset of input array.")
 	};
 
 	/** TArrayNumOperator gets the number of elements in an Array. The operator
@@ -76,10 +78,10 @@ namespace Metasound
 
 			static const FVertexInterface DefaultInterface(
 				FInputVertexInterface(
-					TInputDataVertexModel<ArrayType>(GetInputArrayName(), METASOUND_LOCTEXT("ArrayOpArrayNumInput", "Array to inspect."))
+					TInputDataVertexModel<ArrayType>(METASOUND_GET_PARAM_NAME_AND_METADATA(InputArray))
 				),
 				FOutputVertexInterface(
-					TOutputDataVertexModel<int32>(GetOutputNumName(), METASOUND_LOCTEXT("ArrayOpArrayNumOutput", "Number of elements in the array."))
+					TOutputDataVertexModel<int32>(METASOUND_GET_PARAM_NAME_AND_METADATA(OutputNum))
 				)
 			);
 
@@ -112,7 +114,7 @@ namespace Metasound
 			const FInputVertexInterface& Inputs = InParams.Node.GetVertexInterface().GetInputInterface();
 
 			// Get the input array or construct an empty one. 
-			FArrayDataReadReference Array = InParams.InputDataReferences.GetDataReadReferenceOrConstructWithVertexDefault<ArrayType>(Inputs, GetInputArrayName(), InParams.OperatorSettings);
+			FArrayDataReadReference Array = InParams.InputDataReferences.GetDataReadReferenceOrConstructWithVertexDefault<ArrayType>(Inputs, METASOUND_GET_PARAM_NAME(InputArray), InParams.OperatorSettings);
 
 			return MakeUnique<TArrayNumOperator>(Array);
 		}
@@ -133,7 +135,7 @@ namespace Metasound
 
 			FDataReferenceCollection Inputs;
 
-			Inputs.AddDataReadReference(GetInputArrayName(), Array);
+			Inputs.AddDataReadReference(METASOUND_GET_PARAM_NAME(InputArray), Array);
 
 			return Inputs;
 		}
@@ -144,7 +146,7 @@ namespace Metasound
 
 			FDataReferenceCollection Outputs;
 
-			Outputs.AddDataReadReference(GetOutputNumName(), Num);
+			Outputs.AddDataReadReference(METASOUND_GET_PARAM_NAME(OutputNum), Num);
 
 			return Outputs;
 		}
@@ -188,12 +190,12 @@ namespace Metasound
 			using namespace ArrayNodeVertexNames;
 			static const FVertexInterface DefaultInterface(
 				FInputVertexInterface(
-					TInputDataVertexModel<FTrigger>(GetInputTriggerName(), METASOUND_LOCTEXT("ArrayOpArrayGetTrigger", "Trigger to get value.")),
-					TInputDataVertexModel<ArrayType>(GetInputArrayName(), METASOUND_LOCTEXT("ArrayOpArrayGetInput", "Input Array.")),
-					TInputDataVertexModel<int32>(GetInputIndexName(), METASOUND_LOCTEXT("ArrayOpArrayGetIndex", "Index in Array."))
+					TInputDataVertexModel<FTrigger>(METASOUND_GET_PARAM_NAME_AND_METADATA(InputTriggerGet)),
+					TInputDataVertexModel<ArrayType>(METASOUND_GET_PARAM_NAME_AND_METADATA(InputArray)),
+					TInputDataVertexModel<int32>(METASOUND_GET_PARAM_NAME_AND_METADATA(InputIndex))
 				),
 				FOutputVertexInterface(
-					TOutputDataVertexModel<ElementType>(GetOutputValueName(), METASOUND_LOCTEXT("ArrayOpArrayGetOutput", "Value of element at array index."))
+					TOutputDataVertexModel<ElementType>(METASOUND_GET_PARAM_NAME_AND_METADATA(OutputValue))
 				)
 			);
 
@@ -226,13 +228,13 @@ namespace Metasound
 			const FInputVertexInterface& Inputs = InParams.Node.GetVertexInterface().GetInputInterface();
 
 			// Input Trigger
-			TDataReadReference<FTrigger> Trigger = InParams.InputDataReferences.GetDataReadReferenceOrConstructWithVertexDefault<FTrigger>(Inputs, GetInputTriggerName(), InParams.OperatorSettings);
+			TDataReadReference<FTrigger> Trigger = InParams.InputDataReferences.GetDataReadReferenceOrConstructWithVertexDefault<FTrigger>(Inputs, METASOUND_GET_PARAM_NAME(InputTriggerGet), InParams.OperatorSettings);
 			
 			// Input Array
-			FArrayDataReadReference Array = InParams.InputDataReferences.GetDataReadReferenceOrConstructWithVertexDefault<ArrayType>(Inputs, GetInputArrayName(), InParams.OperatorSettings);
+			FArrayDataReadReference Array = InParams.InputDataReferences.GetDataReadReferenceOrConstructWithVertexDefault<ArrayType>(Inputs, METASOUND_GET_PARAM_NAME(InputArray), InParams.OperatorSettings);
 
 			// Input Index
-			TDataReadReference<int32> Index = InParams.InputDataReferences.GetDataReadReferenceOrConstructWithVertexDefault<int32>(Inputs, GetInputIndexName(), InParams.OperatorSettings);
+			TDataReadReference<int32> Index = InParams.InputDataReferences.GetDataReadReferenceOrConstructWithVertexDefault<int32>(Inputs, METASOUND_GET_PARAM_NAME(InputIndex), InParams.OperatorSettings);
 
 			return MakeUnique<TArrayGetOperator>(InParams.OperatorSettings, Trigger, Array, Index);
 		}
@@ -254,9 +256,9 @@ namespace Metasound
 
 			FDataReferenceCollection Inputs;
 
-			Inputs.AddDataReadReference(GetInputTriggerName(), Trigger);
-			Inputs.AddDataReadReference(GetInputArrayName(), Array);
-			Inputs.AddDataReadReference(GetInputIndexName(), Index);
+			Inputs.AddDataReadReference(METASOUND_GET_PARAM_NAME(InputTriggerGet), Trigger);
+			Inputs.AddDataReadReference(METASOUND_GET_PARAM_NAME(InputArray), Array);
+			Inputs.AddDataReadReference(METASOUND_GET_PARAM_NAME(InputIndex), Index);
 
 			return Inputs;
 		}
@@ -267,7 +269,7 @@ namespace Metasound
 
 			FDataReferenceCollection Outputs;
 
-			Outputs.AddDataReadReference(GetOutputValueName(), Value);
+			Outputs.AddDataReadReference(METASOUND_GET_PARAM_NAME(OutputValue), Value);
 
 			return Outputs;
 		}
@@ -325,13 +327,13 @@ namespace Metasound
 			using namespace ArrayNodeVertexNames;
 			static const FVertexInterface DefaultInterface(
 				FInputVertexInterface(
-					TInputDataVertexModel<FTrigger>(GetInputTriggerName(), METASOUND_LOCTEXT("ArrayOpArraySetTrigger", "Trigger to set value.")),
-					TInputDataVertexModel<ArrayType>(GetInputArrayName(), METASOUND_LOCTEXT("ArrayOpArraySetInput", "Input Array.")),
-					TInputDataVertexModel<int32>(GetInputIndexName(), METASOUND_LOCTEXT("ArrayOpArraySetIndex", "Index in Array.")),
-					TInputDataVertexModel<ElementType>(GetInputValueName(), METASOUND_LOCTEXT("ArrayOpArraySetElement", "Value to set"))
+					TInputDataVertexModel<FTrigger>(METASOUND_GET_PARAM_NAME_AND_METADATA(InputTriggerSet)),
+					TInputDataVertexModel<ArrayType>(METASOUND_GET_PARAM_NAME_AND_METADATA(InputArray)),
+					TInputDataVertexModel<int32>(METASOUND_GET_PARAM_NAME_AND_METADATA(InputIndex)),
+					TInputDataVertexModel<ElementType>(METASOUND_GET_PARAM_NAME_AND_METADATA(InputValue))
 				),
 				FOutputVertexInterface(
-					TOutputDataVertexModel<ArrayType>(GetOutputArrayName(), METASOUND_LOCTEXT("ArrayOpArraySetOutput", "Array after setting."))
+					TOutputDataVertexModel<ArrayType>(METASOUND_GET_PARAM_NAME_AND_METADATA(OutputArraySet))
 				)
 			);
 
@@ -363,14 +365,14 @@ namespace Metasound
 
 			const FInputVertexInterface& Inputs = InParams.Node.GetVertexInterface().GetInputInterface();
 			
-			TDataReadReference<FTrigger> Trigger = InParams.InputDataReferences.GetDataReadReferenceOrConstructWithVertexDefault<FTrigger>(Inputs, GetInputTriggerName(), InParams.OperatorSettings);
+			TDataReadReference<FTrigger> Trigger = InParams.InputDataReferences.GetDataReadReferenceOrConstructWithVertexDefault<FTrigger>(Inputs, METASOUND_GET_PARAM_NAME(InputTriggerSet), InParams.OperatorSettings);
 
-			FArrayDataReadReference InitArray = InParams.InputDataReferences.GetDataReadReferenceOrConstructWithVertexDefault<ArrayType>(Inputs, GetInputArrayName(), InParams.OperatorSettings);
+			FArrayDataReadReference InitArray = InParams.InputDataReferences.GetDataReadReferenceOrConstructWithVertexDefault<ArrayType>(Inputs, METASOUND_GET_PARAM_NAME(InputArray), InParams.OperatorSettings);
 			FArrayDataWriteReference Array = TDataWriteReferenceFactory<ArrayType>::CreateExplicitArgs(InParams.OperatorSettings, *InitArray);
 
-			TDataReadReference<int32> Index = InParams.InputDataReferences.GetDataReadReferenceOrConstructWithVertexDefault<int32>(Inputs, GetInputIndexName(), InParams.OperatorSettings);
+			TDataReadReference<int32> Index = InParams.InputDataReferences.GetDataReadReferenceOrConstructWithVertexDefault<int32>(Inputs, METASOUND_GET_PARAM_NAME(InputIndex), InParams.OperatorSettings);
 
-			TDataReadReference<ElementType> Value = InParams.InputDataReferences.GetDataReadReferenceOrConstructWithVertexDefault<ElementType>(Inputs, GetInputValueName(), InParams.OperatorSettings);
+			TDataReadReference<ElementType> Value = InParams.InputDataReferences.GetDataReadReferenceOrConstructWithVertexDefault<ElementType>(Inputs, METASOUND_GET_PARAM_NAME(InputValue), InParams.OperatorSettings);
 
 			return MakeUnique<TArraySetOperator>(InParams.OperatorSettings, Trigger, InitArray, Array, Index, Value);
 		}
@@ -393,10 +395,10 @@ namespace Metasound
 			using namespace ArrayNodeVertexNames;
 			FDataReferenceCollection Inputs;
 
-			Inputs.AddDataReadReference(GetInputTriggerName(), Trigger);
-			Inputs.AddDataReadReference(GetInputArrayName(), InitArray);
-			Inputs.AddDataReadReference(GetInputIndexName(), Index);
-			Inputs.AddDataReadReference(GetInputValueName(), Value);
+			Inputs.AddDataReadReference(METASOUND_GET_PARAM_NAME(InputTriggerSet), Trigger);
+			Inputs.AddDataReadReference(METASOUND_GET_PARAM_NAME(InputArray), InitArray);
+			Inputs.AddDataReadReference(METASOUND_GET_PARAM_NAME(InputIndex), Index);
+			Inputs.AddDataReadReference(METASOUND_GET_PARAM_NAME(InputValue), Value);
 
 			return Inputs;
 		}
@@ -406,7 +408,7 @@ namespace Metasound
 			using namespace ArrayNodeVertexNames;
 			FDataReferenceCollection Outputs;
 
-			Outputs.AddDataReadReference(GetOutputArrayName(), Array);
+			Outputs.AddDataReadReference(METASOUND_GET_PARAM_NAME(OutputArraySet), Array);
 
 			return Outputs;
 		}
@@ -466,12 +468,12 @@ namespace Metasound
 
 			static const FVertexInterface DefaultInterface(
 				FInputVertexInterface(
-					TInputDataVertexModel<FTrigger>(GetInputTriggerName(), METASOUND_LOCTEXT("ArrayOpArrayConcatTrigger", "Trigger to set value.")),
-					TInputDataVertexModel<ArrayType>(GetInputLeftArrayName(), METASOUND_LOCTEXT("ArrayOpArrayConcatInputLeft", "Input Left Array.")),
-					TInputDataVertexModel<ArrayType>(GetInputRightArrayName(), METASOUND_LOCTEXT("ArrayOpArrayConcatInputRight", "Input Left Array."))
+					TInputDataVertexModel<FTrigger>(METASOUND_GET_PARAM_NAME_AND_METADATA(InputTriggerGet)),
+					TInputDataVertexModel<ArrayType>(METASOUND_GET_PARAM_NAME_AND_METADATA(InputLeftArray)),
+					TInputDataVertexModel<ArrayType>(METASOUND_GET_PARAM_NAME_AND_METADATA(InputRightArray))
 				),
 				FOutputVertexInterface(
-					TOutputDataVertexModel<ArrayType>(GetOutputArrayName(), METASOUND_LOCTEXT("ArrayOpArrayConcatOutput", "Array after concatenation."))
+					TOutputDataVertexModel<ArrayType>(METASOUND_GET_PARAM_NAME_AND_METADATA(OutputArrayConcat))
 				)
 			);
 
@@ -503,10 +505,10 @@ namespace Metasound
 
 			const FInputVertexInterface& Inputs = InParams.Node.GetVertexInterface().GetInputInterface();
 			
-			TDataReadReference<FTrigger> Trigger = InParams.InputDataReferences.GetDataReadReferenceOrConstructWithVertexDefault<FTrigger>(Inputs, GetInputTriggerName(), InParams.OperatorSettings);
+			TDataReadReference<FTrigger> Trigger = InParams.InputDataReferences.GetDataReadReferenceOrConstructWithVertexDefault<FTrigger>(Inputs, METASOUND_GET_PARAM_NAME(InputTriggerGet), InParams.OperatorSettings);
 
-			FArrayDataReadReference LeftArray = InParams.InputDataReferences.GetDataReadReferenceOrConstructWithVertexDefault<ArrayType>(Inputs, GetInputLeftArrayName(), InParams.OperatorSettings);
-			FArrayDataReadReference RightArray = InParams.InputDataReferences.GetDataReadReferenceOrConstructWithVertexDefault<ArrayType>(Inputs, GetInputRightArrayName(), InParams.OperatorSettings);
+			FArrayDataReadReference LeftArray = InParams.InputDataReferences.GetDataReadReferenceOrConstructWithVertexDefault<ArrayType>(Inputs, METASOUND_GET_PARAM_NAME(InputLeftArray), InParams.OperatorSettings);
+			FArrayDataReadReference RightArray = InParams.InputDataReferences.GetDataReadReferenceOrConstructWithVertexDefault<ArrayType>(Inputs, METASOUND_GET_PARAM_NAME(InputRightArray), InParams.OperatorSettings);
 
 			FArrayDataWriteReference OutArray = TDataWriteReferenceFactory<ArrayType>::CreateExplicitArgs(InParams.OperatorSettings);
 
@@ -529,9 +531,9 @@ namespace Metasound
 			using namespace ArrayNodeVertexNames;
 			FDataReferenceCollection Inputs;
 
-			Inputs.AddDataReadReference(GetInputTriggerName(), Trigger);
-			Inputs.AddDataReadReference(GetInputLeftArrayName(), LeftArray);
-			Inputs.AddDataReadReference(GetInputRightArrayName(), RightArray);
+			Inputs.AddDataReadReference(METASOUND_GET_PARAM_NAME(InputTriggerGet), Trigger);
+			Inputs.AddDataReadReference(METASOUND_GET_PARAM_NAME(InputLeftArray), LeftArray);
+			Inputs.AddDataReadReference(METASOUND_GET_PARAM_NAME(InputRightArray), RightArray);
 
 			return Inputs;
 		}
@@ -541,7 +543,7 @@ namespace Metasound
 			using namespace ArrayNodeVertexNames;
 			FDataReferenceCollection Outputs;
 
-			Outputs.AddDataReadReference(GetOutputArrayName(), OutArray);
+			Outputs.AddDataReadReference(METASOUND_GET_PARAM_NAME(OutputArrayConcat), OutArray);
 
 			return Outputs;
 		}
@@ -590,14 +592,14 @@ namespace Metasound
 
 			static const FVertexInterface DefaultInterface(
 				FInputVertexInterface(
-					TInputDataVertexModel<FTrigger>(GetInputTriggerName(), METASOUND_LOCTEXT("ArrayOpArraySubsetTrigger", "Trigger to set value.")),
-					TInputDataVertexModel<ArrayType>(GetInputArrayName(), METASOUND_LOCTEXT("ArrayOpArraySubsetInputLeft", "Input Array.")),
-					TInputDataVertexModel<int32>(GetInputStartIndexName(), METASOUND_LOCTEXT("ArrayOpArraySubsetStartIndex", "First index to include.")),
-					TInputDataVertexModel<int32>(GetInputEndIndexName(), METASOUND_LOCTEXT("ArrayOpArraySubsetEndIndex", "Last index to include."))
+					TInputDataVertexModel<FTrigger>(METASOUND_GET_PARAM_NAME_AND_METADATA(InputTriggerGet)),
+					TInputDataVertexModel<ArrayType>(METASOUND_GET_PARAM_NAME_AND_METADATA(InputArray)),
+					TInputDataVertexModel<int32>(METASOUND_GET_PARAM_NAME_AND_METADATA(InputStartIndex)),
+					TInputDataVertexModel<int32>(METASOUND_GET_PARAM_NAME_AND_METADATA(InputEndIndex))
 
 				),
 				FOutputVertexInterface(
-					TOutputDataVertexModel<ArrayType>(GetOutputArrayName(), METASOUND_LOCTEXT("ArrayOpArraySubsetOutput", "Subset of input array."))
+					TOutputDataVertexModel<ArrayType>(METASOUND_GET_PARAM_NAME_AND_METADATA(OutputArraySubset))
 				)
 			);
 
@@ -629,12 +631,12 @@ namespace Metasound
 
 			const FInputVertexInterface& Inputs = InParams.Node.GetVertexInterface().GetInputInterface();
 			
-			TDataReadReference<FTrigger> Trigger = InParams.InputDataReferences.GetDataReadReferenceOrConstructWithVertexDefault<FTrigger>(Inputs, GetInputTriggerName(), InParams.OperatorSettings);
+			TDataReadReference<FTrigger> Trigger = InParams.InputDataReferences.GetDataReadReferenceOrConstructWithVertexDefault<FTrigger>(Inputs, METASOUND_GET_PARAM_NAME(InputTriggerGet), InParams.OperatorSettings);
 
-			FArrayDataReadReference InArray = InParams.InputDataReferences.GetDataReadReferenceOrConstructWithVertexDefault<ArrayType>(Inputs, GetInputArrayName(), InParams.OperatorSettings);
+			FArrayDataReadReference InArray = InParams.InputDataReferences.GetDataReadReferenceOrConstructWithVertexDefault<ArrayType>(Inputs, METASOUND_GET_PARAM_NAME(InputArray), InParams.OperatorSettings);
 
-			TDataReadReference<int32> StartIndex = InParams.InputDataReferences.GetDataReadReferenceOrConstructWithVertexDefault<int32>(Inputs, GetInputStartIndexName(), InParams.OperatorSettings);
-			TDataReadReference<int32> EndIndex = InParams.InputDataReferences.GetDataReadReferenceOrConstructWithVertexDefault<int32>(Inputs, GetInputEndIndexName(), InParams.OperatorSettings);
+			TDataReadReference<int32> StartIndex = InParams.InputDataReferences.GetDataReadReferenceOrConstructWithVertexDefault<int32>(Inputs, METASOUND_GET_PARAM_NAME(InputStartIndex), InParams.OperatorSettings);
+			TDataReadReference<int32> EndIndex = InParams.InputDataReferences.GetDataReadReferenceOrConstructWithVertexDefault<int32>(Inputs, METASOUND_GET_PARAM_NAME(InputEndIndex), InParams.OperatorSettings);
 
 			FArrayDataWriteReference OutArray = TDataWriteReferenceFactory<ArrayType>::CreateExplicitArgs(InParams.OperatorSettings);
 
@@ -659,10 +661,10 @@ namespace Metasound
 
 			FDataReferenceCollection Inputs;
 
-			Inputs.AddDataReadReference(GetInputTriggerName(), Trigger);
-			Inputs.AddDataReadReference(GetInputArrayName(), InputArray);
-			Inputs.AddDataReadReference(GetInputStartIndexName(), StartIndex);
-			Inputs.AddDataReadReference(GetInputEndIndexName(), EndIndex);
+			Inputs.AddDataReadReference(METASOUND_GET_PARAM_NAME(InputTriggerGet), Trigger);
+			Inputs.AddDataReadReference(METASOUND_GET_PARAM_NAME(InputArray), InputArray);
+			Inputs.AddDataReadReference(METASOUND_GET_PARAM_NAME(InputStartIndex), StartIndex);
+			Inputs.AddDataReadReference(METASOUND_GET_PARAM_NAME(InputEndIndex), EndIndex);
 
 			return Inputs;
 		}
@@ -673,7 +675,7 @@ namespace Metasound
 
 			FDataReferenceCollection Outputs;
 
-			Outputs.AddDataReadReference(GetOutputArrayName(), OutputArray);
+			Outputs.AddDataReadReference(METASOUND_GET_PARAM_NAME(OutputArraySubset), OutputArray);
 
 			return Outputs;
 		}
diff --git a/Engine/Plugins/Runtime/Metasound/Source/MetasoundFrontend/Public/MetasoundArrayRandomNode.h b/Engine/Plugins/Runtime/Metasound/Source/MetasoundFrontend/Public/MetasoundArrayRandomNode.h
--- a/Engine/Plugins/Runtime/Metasound/Source/MetasoundFrontend/Public/MetasoundArrayRandomNode.h
+++ b/Engine/Plugins/Runtime/Metasound/Source/MetasoundFrontend/Public/MetasoundArrayRandomNode.h
@@ -8,6 +8,7 @@
 #include "MetasoundFacade.h"
 #include "MetasoundSourceInterface.h"
 #include "MetasoundNodeInterface.h"
+#include "MetasoundParamHelper.h"
 #include "MetasoundTrigger.h"
 
 #define LOCTEXT_NAMESPACE "MetasoundFrontend"
@@ -16,27 +17,16 @@ namespace Metasound
 {
 	namespace ArrayNodeRandomGetVertexNames
 	{
-		static const FName InputTriggerNextName = "Next";
-		static const FName InputTriggerResetName = "Reset";
-		static const FName InputRandomArrayName = "In Array";
-		static const FName InputWeightsName = "Weights";
-		static const FName InputSeedName = "Seed";
-		static const FName InputNoRepeatOrderName = "No Repeats";
-		static const FName InputEnableSharedStateName = "Enable Shared State";
-		static const FName OutputTriggerOnNextName = "On Next";
-		static const FName OutputTriggerOnResetName = "On Reset";
-		static const FName OutputValueName = "Value";
-
-		static const FText InputTriggerNextTooltip = METASOUND_LOCTEXT("RandomArrayGetNode_TriggerNextTooltip", "Trigger to get the next value in the randomized array.");
-		static const FText InputTriggerResetTooltip = METASOUND_LOCTEXT("RandomArrayGetNode_TriggerResetTooltip", "Trigger to reset the seed for the randomized array.");
-		static const FText InputRandomArrayTooltip = METASOUND_LOCTEXT("RandomArrayGetNode_RandomArrayTooltip", "Input array to randomized.");
-		static const FText InputWeightsTooltip = METASOUND_LOCTEXT("RandomArrayGetNode_WeightsTooltip", "Input array of weights to use for random selection. Will repeat if this array is shorter than the input array to select from.");
-		static const FText InputSeedTooltip = METASOUND_LOCTEXT("RandomArrayGetNode_SeedTooltip", "Seed to use for the random shuffle.");
-		static const FText InputNoRepeatOrderTooltip = METASOUND_LOCTEXT("RandomArrayGetNode_NoRepeatOrderTooltip", "The number of elements to track to avoid repeating in a row.");
-		static const FText InputEnableSharedStateTooltip = METASOUND_LOCTEXT("RandomArrayGetNode_EnableSharedStateTooltip", "Set to enabled to share state across instances of this MetaSound.");
-		static const FText OutputTriggerOnNextTooltip = METASOUND_LOCTEXT("RandomArrayGetNode_TriggerOnNextTooltip", "Triggers when the \"Next\" input is triggered.");
-		static const FText OutputTriggerOnResetTooltip = METASOUND_LOCTEXT("RandomArrayGetNode_TriggerOnResetTooltip", "Triggers when the \"Shuffle\" input is triggered or if the array is auto-shuffled.");
-		static const FText OutputValueTooltip = METASOUND_LOCTEXT("RandomArrayGetNode_ValueTooltip", "Value of the current shuffled element.");
+		METASOUND_PARAM(InputTriggerNextValue, "Next", "Trigger to get the next value in the randomized array.")
+		METASOUND_PARAM(InputTriggerResetSeed, "Reset", "Trigger to reset the seed for the randomized array.")
+		METASOUND_PARAM(InputRandomArray, "In Array", "Input array to randomized.")
+		METASOUND_PARAM(InputWeights, "Weights", "Input array of weights to use for random selection. Will repeat if this array is shorter than the input array to select from.")
+		METASOUND_PARAM(InputSeed, "Seed", "Seed to use for the random shuffle.")
+		METASOUND_PARAM(InputNoRepeatOrder, "No Repeats", "The number of elements to track to avoid repeating in a row.")
+		METASOUND_PARAM(InputEnableSharedState, "Enable Shared State", "Set to enabled to share state across instances of this MetaSound.")
+		METASOUND_PARAM(OutputTriggerOnNext, "On Next", "Triggers when the \"Next\" input is triggered.")
+		METASOUND_PARAM(OutputTriggerOnReset, "On Reset", "Triggers when the \"Shuffle\" input is triggered or if the array is auto-shuffled.")
+		METASOUND_PARAM(ShuffleOutputValue, "Value", "Value of the current shuffled element.")
 	}
 
 	class METASOUNDFRONTEND_API FArrayRandomGet
@@ -119,18 +109,18 @@ namespace Metasound
 
 			static const FVertexInterface DefaultInterface(
 				FInputVertexInterface(
-					TInputDataVertexModel<FTrigger>(InputTriggerNextName, InputTriggerNextTooltip),
-					TInputDataVertexModel<FTrigger>(InputTriggerResetName, InputTriggerResetTooltip),
-					TInputDataVertexModel<ArrayType>(InputRandomArrayName, InputRandomArrayTooltip),
-					TInputDataVertexModel<WeightsArrayType>(InputWeightsName, InputWeightsTooltip),
-					TInputDataVertexModel<int32>(InputSeedName, InputSeedTooltip, -1),
-					TInputDataVertexModel<int32>(InputNoRepeatOrderName, InputNoRepeatOrderTooltip, 1),
-					TInputDataVertexModel<bool>(InputEnableSharedStateName, InputEnableSharedStateTooltip, false)
+					TInputDataVertexModel<FTrigger>(METASOUND_GET_PARAM_NAME_AND_METADATA(InputTriggerNextValue)),
+					TInputDataVertexModel<FTrigger>(METASOUND_GET_PARAM_NAME_AND_METADATA(InputTriggerResetSeed)),
+					TInputDataVertexModel<ArrayType>(METASOUND_GET_PARAM_NAME_AND_METADATA(InputRandomArray)),
+					TInputDataVertexModel<WeightsArrayType>(METASOUND_GET_PARAM_NAME_AND_METADATA(InputWeights)),
+					TInputDataVertexModel<int32>(METASOUND_GET_PARAM_NAME_AND_METADATA(InputSeed), -1),
+					TInputDataVertexModel<int32>(METASOUND_GET_PARAM_NAME_AND_METADATA(InputNoRepeatOrder), 1),
+					TInputDataVertexModel<bool>(METASOUND_GET_PARAM_NAME_AND_METADATA(InputEnableSharedState), false)
 				),
 				FOutputVertexInterface(
-					TOutputDataVertexModel<FTrigger>(OutputTriggerOnNextName, OutputTriggerOnNextTooltip),
-					TOutputDataVertexModel<FTrigger>(OutputTriggerOnResetName, OutputTriggerOnResetTooltip),
-					TOutputDataVertexModel<ElementType>(OutputValueName, OutputValueTooltip)
+					TOutputDataVertexModel<FTrigger>(METASOUND_GET_PARAM_NAME_AND_METADATA(OutputTriggerOnNext)),
+					TOutputDataVertexModel<FTrigger>(METASOUND_GET_PARAM_NAME_AND_METADATA(OutputTriggerOnReset)),
+					TOutputDataVertexModel<ElementType>(METASOUND_GET_PARAM_NAME_AND_METADATA(ShuffleOutputValue))
 				)
 			);
 
@@ -162,13 +152,13 @@ namespace Metasound
 
 			const FInputVertexInterface& Inputs = InParams.Node.GetVertexInterface().GetInputInterface();
 
-			FTriggerReadRef InTriggerNext = InParams.InputDataReferences.GetDataReadReferenceOrConstructWithVertexDefault<FTrigger>(Inputs, InputTriggerNextName, InParams.OperatorSettings);
-			FTriggerReadRef InTriggerReset = InParams.InputDataReferences.GetDataReadReferenceOrConstructWithVertexDefault<FTrigger>(Inputs, InputTriggerResetName, InParams.OperatorSettings);
-			FArrayDataReadReference InInputArray = InParams.InputDataReferences.GetDataReadReferenceOrConstructWithVertexDefault<ArrayType>(Inputs, InputRandomArrayName, InParams.OperatorSettings);
-			FArrayWeightReadReference InInputWeightsArray = InParams.InputDataReferences.GetDataReadReferenceOrConstructWithVertexDefault<WeightsArrayType>(Inputs, InputWeightsName, InParams.OperatorSettings);
-			FInt32ReadRef InSeedValue = InParams.InputDataReferences.GetDataReadReferenceOrConstructWithVertexDefault<int32>(Inputs, InputSeedName, InParams.OperatorSettings);
-			FInt32ReadRef InNoRepeatOrder = InParams.InputDataReferences.GetDataReadReferenceOrConstructWithVertexDefault<int32>(Inputs, InputNoRepeatOrderName, InParams.OperatorSettings);
-			FBoolReadRef bInEnableSharedState = InParams.InputDataReferences.GetDataReadReferenceOrConstructWithVertexDefault<bool>(Inputs, InputEnableSharedStateName, InParams.OperatorSettings);
+			FTriggerReadRef InTriggerNext = InParams.InputDataReferences.GetDataReadReferenceOrConstructWithVertexDefault<FTrigger>(Inputs, METASOUND_GET_PARAM_NAME(InputTriggerNextValue), InParams.OperatorSettings);
+			FTriggerReadRef InTriggerReset = InParams.InputDataReferences.GetDataReadReferenceOrConstructWithVertexDefault<FTrigger>(Inputs, METASOUND_GET_PARAM_NAME(InputTriggerResetSeed), InParams.OperatorSettings);
+			FArrayDataReadReference InInputArray = InParams.InputDataReferences.GetDataReadReferenceOrConstructWithVertexDefault<ArrayType>(Inputs, METASOUND_GET_PARAM_NAME(InputRandomArray), InParams.OperatorSettings);
+			FArrayWeightReadReference InInputWeightsArray = InParams.InputDataReferences.GetDataReadReferenceOrConstructWithVertexDefault<WeightsArrayType>(Inputs, METASOUND_GET_PARAM_NAME(InputWeights), InParams.OperatorSettings);
+			FInt32ReadRef InSeedValue = InParams.InputDataReferences.GetDataReadReferenceOrConstructWithVertexDefault<int32>(Inputs, METASOUND_GET_PARAM_NAME(InputSeed), InParams.OperatorSettings);
+			FInt32ReadRef InNoRepeatOrder = InParams.InputDataReferences.GetDataReadReferenceOrConstructWithVertexDefault<int32>(Inputs, METASOUND_GET_PARAM_NAME(InputNoRepeatOrder), InParams.OperatorSettings);
+			FBoolReadRef bInEnableSharedState = InParams.InputDataReferences.GetDataReadReferenceOrConstructWithVertexDefault<bool>(Inputs, METASOUND_GET_PARAM_NAME(InputEnableSharedState), InParams.OperatorSettings);
 
 			return MakeUnique<TArrayRandomGetOperator<ArrayType>>(InParams, InTriggerNext, InTriggerReset, InInputArray, InInputWeightsArray, InSeedValue, InNoRepeatOrder, bInEnableSharedState);
 		}
@@ -251,13 +241,13 @@ namespace Metasound
 			using namespace ArrayNodeRandomGetVertexNames;
 
 			FDataReferenceCollection Inputs;
-			Inputs.AddDataReadReference(InputTriggerNextName, TriggerNext);
-			Inputs.AddDataReadReference(InputTriggerResetName, TriggerReset);
-			Inputs.AddDataReadReference(InputRandomArrayName, InputArray);
-			Inputs.AddDataReadReference(InputWeightsName, InputWeightsArray);
-			Inputs.AddDataReadReference(InputSeedName, SeedValue);
-			Inputs.AddDataReadReference(InputNoRepeatOrderName, NoRepeatOrder);
-			Inputs.AddDataReadReference(InputEnableSharedStateName, bEnableSharedState);
+			Inputs.AddDataReadReference(METASOUND_GET_PARAM_NAME(InputTriggerNextValue), TriggerNext);
+			Inputs.AddDataReadReference(METASOUND_GET_PARAM_NAME(InputTriggerResetSeed), TriggerReset);
+			Inputs.AddDataReadReference(METASOUND_GET_PARAM_NAME(InputRandomArray), InputArray);
+			Inputs.AddDataReadReference(METASOUND_GET_PARAM_NAME(InputWeights), InputWeightsArray);
+			Inputs.AddDataReadReference(METASOUND_GET_PARAM_NAME(InputSeed), SeedValue);
+			Inputs.AddDataReadReference(METASOUND_GET_PARAM_NAME(InputNoRepeatOrder), NoRepeatOrder);
+			Inputs.AddDataReadReference(METASOUND_GET_PARAM_NAME(InputEnableSharedState), bEnableSharedState);
 
 			return Inputs;
 		}
@@ -267,9 +257,9 @@ namespace Metasound
 			using namespace ArrayNodeRandomGetVertexNames;
 
 			FDataReferenceCollection Outputs;
-			Outputs.AddDataReadReference(OutputTriggerOnNextName, TriggerOnNext);
-			Outputs.AddDataReadReference(OutputTriggerOnResetName, TriggerOnReset);
-			Outputs.AddDataReadReference(OutputValueName, OutValue);
+			Outputs.AddDataReadReference(METASOUND_GET_PARAM_NAME(OutputTriggerOnNext), TriggerOnNext);
+			Outputs.AddDataReadReference(METASOUND_GET_PARAM_NAME(OutputTriggerOnReset), TriggerOnReset);
+			Outputs.AddDataReadReference(METASOUND_GET_PARAM_NAME(ShuffleOutputValue), OutValue);
 
 			return Outputs;
 		}
diff --git a/Engine/Plugins/Runtime/Metasound/Source/MetasoundFrontend/Public/MetasoundArrayShuffleNode.h b/Engine/Plugins/Runtime/Metasound/Source/MetasoundFrontend/Public/MetasoundArrayShuffleNode.h
--- a/Engine/Plugins/Runtime/Metasound/Source/MetasoundFrontend/Public/MetasoundArrayShuffleNode.h
+++ b/Engine/Plugins/Runtime/Metasound/Source/MetasoundFrontend/Public/MetasoundArrayShuffleNode.h
@@ -12,6 +12,7 @@
 #include "MetasoundLog.h"
 #include "MetasoundNodeInterface.h"
 #include "MetasoundOperatorInterface.h"
+#include "MetasoundParamHelper.h"
 #include "MetasoundPrimitives.h"
 #include "MetasoundSourceInterface.h"
 #include "MetasoundTrigger.h"
@@ -28,19 +29,18 @@ namespace Metasound
 	/** Shuffle Node Vertex Names */
 	namespace ArrayNodeShuffleVertexNames
 	{
-		/** Input Vertex Names */
-		METASOUNDFRONTEND_API const FVertexName& GetInputTriggerNextName();
-		METASOUNDFRONTEND_API const FVertexName& GetInputTriggerShuffleName();
-		METASOUNDFRONTEND_API const FVertexName& GetInputTriggerResetName();
-		METASOUNDFRONTEND_API const FVertexName& GetInputShuffleArrayName();
-		METASOUNDFRONTEND_API const FVertexName& GetInputSeedName();
-		METASOUNDFRONTEND_API const FVertexName& GetInputAutoShuffleName();
-		METASOUNDFRONTEND_API const FVertexName& GetInputEnableSharedStateName();
-
-		METASOUNDFRONTEND_API const FVertexName& GetOutputTriggerOnNextName();
-		METASOUNDFRONTEND_API const FVertexName& GetOutputTriggerOnShuffleName();
-		METASOUNDFRONTEND_API const FVertexName& GetOutputTriggerOnResetName();
-		METASOUNDFRONTEND_API const FVertexName& GetOutputValueName();
+		METASOUND_PARAM(InputTriggerNext, "Next", "Trigger to get the next value in the shuffled array.")
+		METASOUND_PARAM(InputTriggerShuffle, "Shuffle", "Trigger to shuffle the array manually.")
+		METASOUND_PARAM(InputTriggerReset, "Reset Seed", "Trigger to reset the random seed stream of the shuffle node.")
+		METASOUND_PARAM(InputShuffleArray, "In Array", "Input Array.")
+		METASOUND_PARAM(InputShuffleSeed, "Seed", "Seed to use for the the random shuffle.")
+		METASOUND_PARAM(InputAutoShuffle, "Auto Shuffle", "Set to true to automatically shuffle when the array has been read.")
+		METASOUND_PARAM(InputShuffleEnableSharedState, "Enable Shared State", "Set to enabled shared state across instances of this metasound.")
+
+		METASOUND_PARAM(OutputTriggerOnNext, "On Next", "Triggers when the \"Next\" input is triggered.")
+		METASOUND_PARAM(OutputTriggerOnShuffle, "On Shuffle", "Triggers when the \"Shuffle\" input is triggered or if the array is auto-shuffled.")
+		METASOUND_PARAM(OutputTriggerOnResetSeed, "On Reset Seed", "Triggers when the \"Reset Seed\" input is triggered.")
+		METASOUND_PARAM(OutputShuffledValue, "Value", "Value of the current shuffled element.")
 	}
 
 	class METASOUNDFRONTEND_API FArrayIndexShuffler
@@ -147,19 +147,19 @@ namespace Metasound
 
 			static const FVertexInterface DefaultInterface(
 				FInputVertexInterface(
-					TInputDataVertexModel<FTrigger>(GetInputTriggerNextName(), METASOUND_LOCTEXT("ShuffleOpInputTriggerNextTT", "Trigger to get the next value in the shuffled array.")),
-					TInputDataVertexModel<FTrigger>(GetInputTriggerShuffleName(), METASOUND_LOCTEXT("ShuffleOpInputTriggerShuffleTT", "Trigger to shuffle the array manually.")),
-					TInputDataVertexModel<FTrigger>(GetInputTriggerResetName(), METASOUND_LOCTEXT("ShuffleOpInputTriggerResetTT", "Trigger to reset the random seed stream of the shuffle node.")),
-					TInputDataVertexModel<ArrayType>(GetInputShuffleArrayName(), METASOUND_LOCTEXT("ShuffleOpInputShuffleArrayTT", "Input Array.")),
-					TInputDataVertexModel<int32>(GetInputSeedName(), METASOUND_LOCTEXT("ShuffleOpInputSeedTT", "Seed to use for the the random shuffle."), -1),
-					TInputDataVertexModel<bool>(GetInputAutoShuffleName(), METASOUND_LOCTEXT("ShuffleOpInputAutoShuffleTT", "Set to true to automatically shuffle when the array has been read."), true),
-					TInputDataVertexModel<bool>(GetInputEnableSharedStateName(), METASOUND_LOCTEXT("ShuffleOpInputEnableSharedStatTT", "Set to enabled shared state across instances of this metasound."), false)
+					TInputDataVertexModel<FTrigger>(METASOUND_GET_PARAM_NAME_AND_METADATA(InputTriggerNext)),
+					TInputDataVertexModel<FTrigger>(METASOUND_GET_PARAM_NAME_AND_METADATA(InputTriggerShuffle)),
+					TInputDataVertexModel<FTrigger>(METASOUND_GET_PARAM_NAME_AND_METADATA(InputTriggerReset)),
+					TInputDataVertexModel<ArrayType>(METASOUND_GET_PARAM_NAME_AND_METADATA(InputShuffleArray)),
+					TInputDataVertexModel<int32>(METASOUND_GET_PARAM_NAME_AND_METADATA(InputShuffleSeed), -1),
+					TInputDataVertexModel<bool>(METASOUND_GET_PARAM_NAME_AND_METADATA(InputAutoShuffle), true),
+					TInputDataVertexModel<bool>(METASOUND_GET_PARAM_NAME_AND_METADATA(InputShuffleEnableSharedState), false)
 					),
 				FOutputVertexInterface(
-					TOutputDataVertexModel<FTrigger>(GetOutputTriggerOnNextName(), METASOUND_LOCTEXT("ShuffleOpOutputTriggerOnNextNameTT", "Triggers when the \"Next\" input is triggered.")),
-					TOutputDataVertexModel<FTrigger>(GetOutputTriggerOnShuffleName(), METASOUND_LOCTEXT("ShuffleOpOutputTriggerOnShuffleNameTT", "Triggers when the \"Shuffle\" input is triggered or if the array is auto-shuffled.")),
-					TOutputDataVertexModel<FTrigger>(GetOutputTriggerOnResetName(), METASOUND_LOCTEXT("ShuffleOpOutputTriggerOnResetNameTT", "Triggers when the \"Reset Seed\" input is triggered.")),
-					TOutputDataVertexModel<ElementType>(GetOutputValueName(), METASOUND_LOCTEXT("ShuffleOpOutputValueTT", "Value of the current shuffled element."))
+					TOutputDataVertexModel<FTrigger>(METASOUND_GET_PARAM_NAME_AND_METADATA(OutputTriggerOnNext)),
+					TOutputDataVertexModel<FTrigger>(METASOUND_GET_PARAM_NAME_AND_METADATA(OutputTriggerOnShuffle)),
+					TOutputDataVertexModel<FTrigger>(METASOUND_GET_PARAM_NAME_AND_METADATA(OutputTriggerOnResetSeed)),
+					TOutputDataVertexModel<ElementType>(METASOUND_GET_PARAM_NAME_AND_METADATA(OutputShuffledValue))
 				)
 			);
 
@@ -191,13 +191,13 @@ namespace Metasound
 
 			const FInputVertexInterface& Inputs = InParams.Node.GetVertexInterface().GetInputInterface();
 
-			TDataReadReference<FTrigger> InTriggerNext = InParams.InputDataReferences.GetDataReadReferenceOrConstructWithVertexDefault<FTrigger>(Inputs, GetInputTriggerNextName(), InParams.OperatorSettings);
-			TDataReadReference<FTrigger> InTriggerShuffle = InParams.InputDataReferences.GetDataReadReferenceOrConstructWithVertexDefault<FTrigger>(Inputs, GetInputTriggerShuffleName(), InParams.OperatorSettings);
-			TDataReadReference<FTrigger> InTriggerReset = InParams.InputDataReferences.GetDataReadReferenceOrConstructWithVertexDefault<FTrigger>(Inputs, GetInputTriggerResetName(), InParams.OperatorSettings);
-			FArrayDataReadReference InInputArray = InParams.InputDataReferences.GetDataReadReferenceOrConstructWithVertexDefault<ArrayType>(Inputs, GetInputShuffleArrayName(), InParams.OperatorSettings);
-			TDataReadReference<int32> InSeedValue = InParams.InputDataReferences.GetDataReadReferenceOrConstructWithVertexDefault<int32>(Inputs, GetInputSeedName(), InParams.OperatorSettings);
-			TDataReadReference<bool> bInAutoShuffle = InParams.InputDataReferences.GetDataReadReferenceOrConstructWithVertexDefault<bool>(Inputs, GetInputAutoShuffleName(), InParams.OperatorSettings);
-			TDataReadReference<bool> bInEnableSharedState = InParams.InputDataReferences.GetDataReadReferenceOrConstructWithVertexDefault<bool>(Inputs, GetInputEnableSharedStateName(), InParams.OperatorSettings);
+			TDataReadReference<FTrigger> InTriggerNext = InParams.InputDataReferences.GetDataReadReferenceOrConstructWithVertexDefault<FTrigger>(Inputs, METASOUND_GET_PARAM_NAME(InputTriggerNext), InParams.OperatorSettings);
+			TDataReadReference<FTrigger> InTriggerShuffle = InParams.InputDataReferences.GetDataReadReferenceOrConstructWithVertexDefault<FTrigger>(Inputs, METASOUND_GET_PARAM_NAME(InputTriggerShuffle), InParams.OperatorSettings);
+			TDataReadReference<FTrigger> InTriggerReset = InParams.InputDataReferences.GetDataReadReferenceOrConstructWithVertexDefault<FTrigger>(Inputs, METASOUND_GET_PARAM_NAME(InputTriggerReset), InParams.OperatorSettings);
+			FArrayDataReadReference InInputArray = InParams.InputDataReferences.GetDataReadReferenceOrConstructWithVertexDefault<ArrayType>(Inputs, METASOUND_GET_PARAM_NAME(InputShuffleArray), InParams.OperatorSettings);
+			TDataReadReference<int32> InSeedValue = InParams.InputDataReferences.GetDataReadReferenceOrConstructWithVertexDefault<int32>(Inputs, METASOUND_GET_PARAM_NAME(InputShuffleSeed), InParams.OperatorSettings);
+			TDataReadReference<bool> bInAutoShuffle = InParams.InputDataReferences.GetDataReadReferenceOrConstructWithVertexDefault<bool>(Inputs, METASOUND_GET_PARAM_NAME(InputAutoShuffle), InParams.OperatorSettings);
+			TDataReadReference<bool> bInEnableSharedState = InParams.InputDataReferences.GetDataReadReferenceOrConstructWithVertexDefault<bool>(Inputs, METASOUND_GET_PARAM_NAME(InputShuffleEnableSharedState), InParams.OperatorSettings);
 
 			return MakeUnique<TArrayShuffleOperator>(InParams, InTriggerNext, InTriggerShuffle, InTriggerReset, InInputArray, InSeedValue, bInAutoShuffle, bInEnableSharedState);
 		}
@@ -261,13 +261,13 @@ namespace Metasound
 			using namespace ArrayNodeShuffleVertexNames;
 
 			FDataReferenceCollection Inputs;
-			Inputs.AddDataReadReference(GetInputTriggerNextName(), TriggerNext);
-			Inputs.AddDataReadReference(GetInputTriggerShuffleName(), TriggerShuffle);
-			Inputs.AddDataReadReference(GetInputTriggerResetName(), TriggerReset);
-			Inputs.AddDataReadReference(GetInputShuffleArrayName(), InputArray);
-			Inputs.AddDataReadReference(GetInputSeedName(), SeedValue);
-			Inputs.AddDataReadReference(GetInputAutoShuffleName(), bAutoShuffle);
-			Inputs.AddDataReadReference(GetInputEnableSharedStateName(), bEnableSharedState);
+			Inputs.AddDataReadReference(METASOUND_GET_PARAM_NAME(InputTriggerNext), TriggerNext);
+			Inputs.AddDataReadReference(METASOUND_GET_PARAM_NAME(InputTriggerShuffle), TriggerShuffle);
+			Inputs.AddDataReadReference(METASOUND_GET_PARAM_NAME(InputTriggerReset), TriggerReset);
+			Inputs.AddDataReadReference(METASOUND_GET_PARAM_NAME(InputShuffleArray), InputArray);
+			Inputs.AddDataReadReference(METASOUND_GET_PARAM_NAME(InputShuffleSeed), SeedValue);
+			Inputs.AddDataReadReference(METASOUND_GET_PARAM_NAME(InputAutoShuffle), bAutoShuffle);
+			Inputs.AddDataReadReference(METASOUND_GET_PARAM_NAME(InputShuffleEnableSharedState), bEnableSharedState);
 
 			return Inputs;
 		}
@@ -277,10 +277,10 @@ namespace Metasound
 			using namespace ArrayNodeShuffleVertexNames;
 
 			FDataReferenceCollection Outputs;
-			Outputs.AddDataReadReference(GetOutputTriggerOnNextName(), TriggerOnNext);
-			Outputs.AddDataReadReference(GetOutputTriggerOnShuffleName(), TriggerOnShuffle);
-			Outputs.AddDataReadReference(GetOutputTriggerOnResetName(), TriggerOnReset);
-			Outputs.AddDataReadReference(GetOutputValueName(), OutValue);
+			Outputs.AddDataReadReference(METASOUND_GET_PARAM_NAME(OutputTriggerOnNext), TriggerOnNext);
+			Outputs.AddDataReadReference(METASOUND_GET_PARAM_NAME(OutputTriggerOnShuffle), TriggerOnShuffle);
+			Outputs.AddDataReadReference(METASOUND_GET_PARAM_NAME(OutputTriggerOnResetSeed), TriggerOnReset);
+			Outputs.AddDataReadReference(METASOUND_GET_PARAM_NAME(OutputShuffledValue), OutValue);
 
 			return Outputs;
 		}
diff --git a/Engine/Plugins/Runtime/Metasound/Source/MetasoundFrontend/Public/MetasoundReceiveNode.h b/Engine/Plugins/Runtime/Metasound/Source/MetasoundFrontend/Public/MetasoundReceiveNode.h
--- a/Engine/Plugins/Runtime/Metasound/Source/MetasoundFrontend/Public/MetasoundReceiveNode.h
+++ b/Engine/Plugins/Runtime/Metasound/Source/MetasoundFrontend/Public/MetasoundReceiveNode.h
@@ -9,6 +9,7 @@
 #include "MetasoundDataFactory.h"
 #include "MetasoundDataReference.h"
 #include "MetasoundExecutableOperator.h"
+#include "MetasoundParamHelper.h"
 #include "MetasoundRouter.h"
 
 #include <type_traits>
@@ -20,9 +21,11 @@ namespace Metasound
 {
 	namespace ReceiveNodeInfo
 	{
-		METASOUNDFRONTEND_API const FVertexName& GetAddressInputName();
-		METASOUNDFRONTEND_API const FVertexName& GetDefaultDataInputName();
-		METASOUNDFRONTEND_API const FVertexName& GetOutputName();
+		METASOUND_PARAM(AddressInput, "Address", "Address")
+		METASOUND_PARAM(DefaultDataInput, "Default", "Default")
+		METASOUND_PARAM(Output, "Out", "Out")
+
+
 		METASOUNDFRONTEND_API FNodeClassName GetClassNameForDataType(const FName& InDataTypeName);
 		METASOUNDFRONTEND_API int32 GetCurrentMajorVersion();
 		METASOUNDFRONTEND_API int32 GetCurrentMinorVersion();
@@ -34,13 +37,29 @@ namespace Metasound
 	public:
 		static FVertexInterface DeclareVertexInterface()
 		{
+			using namespace ReceiveNodeInfo;
+			static const FDataVertexMetadata AddressInputMetadata
+			{
+				  FText::GetEmpty() // description
+				, METASOUND_GET_PARAM_DISPLAYNAME(AddressInput) // display name
+			};
+			static const FDataVertexMetadata DefaultDataInputMetadata
+			{
+				  FText::GetEmpty() // description
+				, METASOUND_GET_PARAM_DISPLAYNAME(DefaultDataInput) // display name
+			};
+			static const FDataVertexMetadata OutputMetadata
+			{
+				  FText::GetEmpty() // description
+				, METASOUND_GET_PARAM_DISPLAYNAME(Output) // display name
+			};
 			return FVertexInterface(
 				FInputVertexInterface(
-					TInputDataVertexModel<FSendAddress>(ReceiveNodeInfo::GetAddressInputName(), FText::GetEmpty()),
-					TInputDataVertexModel<TDataType>(ReceiveNodeInfo::GetDefaultDataInputName(), FText::GetEmpty())
+					TInputDataVertexModel<FSendAddress>(METASOUND_GET_PARAM_NAME(AddressInput), AddressInputMetadata),
+					TInputDataVertexModel<TDataType>(METASOUND_GET_PARAM_NAME(DefaultDataInput), DefaultDataInputMetadata)
 				),
 				FOutputVertexInterface(
-					TOutputDataVertexModel<TDataType>(ReceiveNodeInfo::GetOutputName(), FText::GetEmpty())
+					TOutputDataVertexModel<TDataType>(METASOUND_GET_PARAM_NAME(Output), OutputMetadata)
 				)
 			);
 		}
@@ -99,18 +118,22 @@ namespace Metasound
 
 				virtual FDataReferenceCollection GetInputs() const override
 				{
+					using namespace ReceiveNodeInfo; 
+
 					FDataReferenceCollection Inputs;
 
-					Inputs.AddDataReadReference<TDataType>(ReceiveNodeInfo::GetDefaultDataInputName(), DefaultData);
-					Inputs.AddDataReadReference<FSendAddress>(ReceiveNodeInfo::GetAddressInputName(), SendAddress);
+					Inputs.AddDataReadReference<TDataType>(METASOUND_GET_PARAM_NAME(DefaultDataInput), DefaultData);
+					Inputs.AddDataReadReference<FSendAddress>(METASOUND_GET_PARAM_NAME(AddressInput), SendAddress);
 
 					return Inputs;
 				}
 
 				virtual FDataReferenceCollection GetOutputs() const override
 				{
+					using namespace ReceiveNodeInfo;
+
 					FDataReferenceCollection Outputs;
-					Outputs.AddDataReadReference<TDataType>(ReceiveNodeInfo::GetOutputName(), TDataReadReference<TDataType>(OutputData));
+					Outputs.AddDataReadReference<TDataType>(METASOUND_GET_PARAM_NAME(Output), TDataReadReference<TDataType>(OutputData));
 					return Outputs;
 				}
 
@@ -191,17 +214,19 @@ namespace Metasound
 
 				virtual TUniquePtr<IOperator> CreateOperator(const FCreateOperatorParams& InParams, FBuildErrorArray& OutErrors) override
 				{
+					using namespace ReceiveNodeInfo; 
+
 					TDataReadReference<TDataType> DefaultReadRef = TDataReadReferenceFactory<TDataType>::CreateAny(InParams.OperatorSettings);
 
-					if (InParams.InputDataReferences.ContainsDataReadReference<TDataType>(ReceiveNodeInfo::GetDefaultDataInputName()))
+					if (InParams.InputDataReferences.ContainsDataReadReference<TDataType>(METASOUND_GET_PARAM_NAME(DefaultDataInput)))
 					{
-						DefaultReadRef = InParams.InputDataReferences.GetDataReadReference<TDataType>(ReceiveNodeInfo::GetDefaultDataInputName());
+						DefaultReadRef = InParams.InputDataReferences.GetDataReadReference<TDataType>(METASOUND_GET_PARAM_NAME(DefaultDataInput));
 					}
 
 					return MakeUnique<TReceiverOperator>(
 						DefaultReadRef,
 						TDataWriteReferenceFactory<TDataType>::CreateAny(InParams.OperatorSettings, *DefaultReadRef),
-						InParams.InputDataReferences.GetDataReadReferenceOrConstruct<FSendAddress>(ReceiveNodeInfo::GetAddressInputName()),
+						InParams.InputDataReferences.GetDataReadReferenceOrConstruct<FSendAddress>(METASOUND_GET_PARAM_NAME(AddressInput)),
 						InParams.OperatorSettings
 						);
 				}
diff --git a/Engine/Plugins/Runtime/Metasound/Source/MetasoundFrontend/Public/MetasoundSendNode.h b/Engine/Plugins/Runtime/Metasound/Source/MetasoundFrontend/Public/MetasoundSendNode.h
--- a/Engine/Plugins/Runtime/Metasound/Source/MetasoundFrontend/Public/MetasoundSendNode.h
+++ b/Engine/Plugins/Runtime/Metasound/Source/MetasoundFrontend/Public/MetasoundSendNode.h
@@ -7,6 +7,7 @@
 #include "MetasoundOperatorInterface.h"
 #include "MetasoundDataReference.h"
 #include "MetasoundExecutableOperator.h"
+#include "MetasoundParamHelper.h"
 #include "MetasoundRouter.h"
 #include "MetasoundVertex.h"
 
@@ -17,16 +18,15 @@
 
 namespace Metasound
 {
+	namespace SendVertexNames
+	{
+		METASOUND_PARAM(AddressInput, "Address", "Address")
+	}
+
 	template<typename TDataType>
 	class TSendNode : public FNode
 	{
 	public:
-		static const FVertexName& GetAddressInputName()
-		{
-			static const FVertexName InputName = TEXT("Address");
-			return InputName;
-		}
-
 		static const FVertexName& GetSendInputName()
 		{
 			static const FVertexName& SendInput = GetMetasoundDataTypeName<TDataType>();
@@ -35,9 +35,16 @@ namespace Metasound
 
 		static FVertexInterface DeclareVertexInterface()
 		{
+			using namespace SendVertexNames; 
+			static const FDataVertexMetadata AddressInputMetadata
+			{
+				  FText::GetEmpty() // description
+				, METASOUND_GET_PARAM_DISPLAYNAME(AddressInput) // display name
+			};
+
 			return FVertexInterface(
 				FInputVertexInterface(
-					TInputDataVertexModel<FSendAddress>(GetAddressInputName(), FText::GetEmpty()),
+					TInputDataVertexModel<FSendAddress>(METASOUND_GET_PARAM_NAME(AddressInput), AddressInputMetadata),
 					TInputDataVertexModel<TDataType>(GetSendInputName(), FText::GetEmpty())
 				),
 				FOutputVertexInterface(
@@ -99,8 +106,10 @@ namespace Metasound
 
 				virtual FDataReferenceCollection GetInputs() const override
 				{
+					using namespace SendVertexNames; 
+
 					FDataReferenceCollection Inputs;
-					Inputs.AddDataReadReference<FSendAddress>(GetAddressInputName(), SendAddress);
+					Inputs.AddDataReadReference<FSendAddress>(METASOUND_GET_PARAM_NAME(AddressInput), SendAddress);
 					Inputs.AddDataReadReference<TDataType>(GetSendInputName(), TDataReadReference<TDataType>(InputData));
 					return Inputs;
 				}
@@ -162,10 +171,12 @@ namespace Metasound
 
 				virtual TUniquePtr<IOperator> CreateOperator(const FCreateOperatorParams& InParams, FBuildErrorArray& OutErrors) override
 				{
+					using namespace SendVertexNames; 
+
 					if (InParams.InputDataReferences.ContainsDataReadReference<TDataType>(GetSendInputName()))
 					{
 						return MakeUnique<TSendOperator>(InParams.InputDataReferences.GetDataReadReference<TDataType>(GetSendInputName()),
-							InParams.InputDataReferences.GetDataReadReferenceOrConstruct<FSendAddress>(GetAddressInputName()),
+							InParams.InputDataReferences.GetDataReadReferenceOrConstruct<FSendAddress>(METASOUND_GET_PARAM_NAME(AddressInput)),
 							InParams.OperatorSettings
 						);
 					}
diff --git a/Engine/Plugins/Runtime/Metasound/Source/MetasoundGraphCore/Private/MetasoundLiteralNode.cpp b/Engine/Plugins/Runtime/Metasound/Source/MetasoundGraphCore/Private/MetasoundLiteralNode.cpp
--- a/Engine/Plugins/Runtime/Metasound/Source/MetasoundGraphCore/Private/MetasoundLiteralNode.cpp
+++ b/Engine/Plugins/Runtime/Metasound/Source/MetasoundGraphCore/Private/MetasoundLiteralNode.cpp
@@ -5,11 +5,4 @@
 
 namespace Metasound
 {
-	namespace LiteralNodeNames
-	{
-		const TCHAR* GetOutputDataName()
-		{
-			return TEXT("Value");
-		}
-	}
 }
diff --git a/Engine/Plugins/Runtime/Metasound/Source/MetasoundGraphCore/Private/MetasoundVariableNodes.cpp b/Engine/Plugins/Runtime/Metasound/Source/MetasoundGraphCore/Private/MetasoundVariableNodes.cpp
--- a/Engine/Plugins/Runtime/Metasound/Source/MetasoundGraphCore/Private/MetasoundVariableNodes.cpp
+++ b/Engine/Plugins/Runtime/Metasound/Source/MetasoundGraphCore/Private/MetasoundVariableNodes.cpp
@@ -9,31 +9,6 @@ namespace Metasound
 {
 	namespace VariableNames
 	{
-		const FVertexName& GetInputDataName()
-		{
-			static const FVertexName Name = TEXT("Value");
-			return Name;
-		}
-
-		const FVertexName& GetOutputDataName()
-		{
-			static const FVertexName Name = TEXT("Value");
-			return Name;
-		}
-
-		const FVertexName& GetInputVariableName()
-		{
-			static const FVertexName Name = TEXT("Variable");
-			return Name;
-		}
-
-		const FVertexName& GetOutputVariableName()
-		{
-			static const FVertexName Name = TEXT("Variable");
-			return Name;
-		}
-
-
 		FNodeClassName GetVariableNodeClassName(const FName& InDataTypeName)
 		{
 			static const FName NodeNamespace = "InitVariable";
diff --git a/Engine/Plugins/Runtime/Metasound/Source/MetasoundGraphCore/Public/MetasoundLiteralNode.h b/Engine/Plugins/Runtime/Metasound/Source/MetasoundGraphCore/Public/MetasoundLiteralNode.h
--- a/Engine/Plugins/Runtime/Metasound/Source/MetasoundGraphCore/Public/MetasoundLiteralNode.h
+++ b/Engine/Plugins/Runtime/Metasound/Source/MetasoundGraphCore/Public/MetasoundLiteralNode.h
@@ -8,6 +8,7 @@
 #include "MetasoundNode.h"
 #include "MetasoundNodeInterface.h"
 #include "MetasoundOperatorInterface.h"
+#include "MetasoundParamHelper.h"
 #include "MetasoundVertex.h"
 
 #define LOCTEXT_NAMESPACE "MetasoundGraphCore"
@@ -17,7 +18,7 @@ namespace Metasound
 {
 	namespace LiteralNodeNames
 	{
-		METASOUNDGRAPHCORE_API const TCHAR* GetOutputDataName();
+		METASOUND_PARAM(OutputValue, "Value", "Value")
 	}
 
 	template<typename DataType>
@@ -44,8 +45,10 @@ namespace Metasound
 
 		virtual FDataReferenceCollection GetOutputs() const override
 		{
+			using namespace LiteralNodeNames;
+
 			FDataReferenceCollection Outputs;
-			Outputs.AddDataReadReference<DataType>(LiteralNodeNames::GetOutputDataName(), OutputValue);
+			Outputs.AddDataReadReference<DataType>(METASOUND_GET_PARAM_NAME(OutputValue), OutputValue);
 			return Outputs;
 		}
 
@@ -101,10 +104,16 @@ namespace Metasound
 
 		static FVertexInterface DeclareVertexInterface()
 		{
+			using namespace LiteralNodeNames; 
+			static const FDataVertexMetadata OutputMetadata
+			{
+				  FText::GetEmpty() // description
+				, METASOUND_GET_PARAM_DISPLAYNAME(OutputValue) // display name
+			};
 			return FVertexInterface(
 				FInputVertexInterface(),
 				FOutputVertexInterface(
-					TOutputDataVertexModel<DataType>(LiteralNodeNames::GetOutputDataName(), FText::GetEmpty())
+					TOutputDataVertexModel<DataType>(METASOUND_GET_PARAM_NAME(OutputValue), OutputMetadata)
 				)
 			);
 		}
diff --git a/Engine/Plugins/Runtime/Metasound/Source/MetasoundGraphCore/Public/MetasoundVariableNodes.h b/Engine/Plugins/Runtime/Metasound/Source/MetasoundGraphCore/Public/MetasoundVariableNodes.h
--- a/Engine/Plugins/Runtime/Metasound/Source/MetasoundGraphCore/Public/MetasoundVariableNodes.h
+++ b/Engine/Plugins/Runtime/Metasound/Source/MetasoundGraphCore/Public/MetasoundVariableNodes.h
@@ -8,6 +8,7 @@
 #include "MetasoundLog.h"
 #include "MetasoundNode.h"
 #include "MetasoundNodeInterface.h"
+#include "MetasoundParamHelper.h"
 #include "MetasoundVariable.h"
 #include "MetasoundVertex.h"
 
@@ -18,17 +19,34 @@ namespace Metasound
 {
 	namespace VariableNames
 	{
-		/** Input vertex name for data */
-		METASOUNDGRAPHCORE_API const FVertexName& GetInputDataName();
+		METASOUND_PARAM(InputData, "Value", "Value")
+		METASOUND_PARAM(OutputData, "Value", "Value")
+		METASOUND_PARAM(InputVariable, "Variable", "Variable")
+		METASOUND_PARAM(OutputVariable, "Variable", "Variable")
 
-		/** Output vertex name for data */
-		METASOUNDGRAPHCORE_API const FVertexName& GetOutputDataName();
+		static const FDataVertexMetadata InputDataMetadata
+		{
+			  FText::GetEmpty() // description
+			, METASOUND_GET_PARAM_DISPLAYNAME(InputData) // display name
+		};
+
+		static const FDataVertexMetadata OutputDataMetadata
+		{
+			  FText::GetEmpty() // description
+			, METASOUND_GET_PARAM_DISPLAYNAME(OutputData) // display name
+		};
 
-		/** Input vertex name for variables */
-		METASOUNDGRAPHCORE_API const FVertexName& GetInputVariableName();
-		
-		/** Output vertex name for variables */
-		METASOUNDGRAPHCORE_API const FVertexName& GetOutputVariableName();
+		static const FDataVertexMetadata InputVariableMetadata
+		{
+			  FText::GetEmpty() // description
+			, METASOUND_GET_PARAM_DISPLAYNAME(InputVariable) // display name
+		};
+
+		static const FDataVertexMetadata OutputVariableMetadata
+		{
+			  FText::GetEmpty() // description
+			, METASOUND_GET_PARAM_DISPLAYNAME(OutputVariable) // display name
+		};
 
 		/** Class name for variable node. */
 		METASOUNDGRAPHCORE_API FNodeClassName GetVariableNodeClassName(const FName& InDataTypeName);
@@ -104,8 +122,10 @@ namespace Metasound
 
 			virtual FDataReferenceCollection GetOutputs() const override
 			{
+				using namespace VariableNames; 
+
 				FDataReferenceCollection Collection;
-				Collection.AddDataWriteReference(VariableNames::GetOutputVariableName(), Variable);
+				Collection.AddDataWriteReference(METASOUND_GET_PARAM_NAME(OutputVariable), Variable);
 				return Collection;
 			}
 
@@ -165,11 +185,13 @@ namespace Metasound
 
 		static FVertexInterface DeclareVertexInterface()
 		{
+			using namespace VariableNames; 
+
 			return FVertexInterface(
 				FInputVertexInterface(
 				),
 				FOutputVertexInterface(
-					TOutputDataVertexModel<FVariable>(VariableNames::GetOutputVariableName(), FText::GetEmpty())
+					TOutputDataVertexModel<FVariable>(METASOUND_GET_PARAM_NAME(OutputVariable), OutputVariableMetadata)
 				)
 			);
 		}
@@ -251,8 +273,10 @@ namespace Metasound
 
 			virtual FDataReferenceCollection GetOutputs() const override
 			{
+				using namespace VariableNames; 
+
 				FDataReferenceCollection Collection;
-				Collection.AddDataReadReference<FVariable>(VariableNames::GetOutputVariableName(), Variable);
+				Collection.AddDataReadReference<FVariable>(METASOUND_GET_PARAM_NAME(OutputVariable), Variable);
 				return Collection;
 			}
 
@@ -278,23 +302,23 @@ namespace Metasound
 
 		const FDataReferenceCollection& Inputs = InParams.InputDataReferences;
 
-		if (Inputs.ContainsDataWriteReference<FVariable>(GetInputVariableName()))
+		if (Inputs.ContainsDataWriteReference<FVariable>(METASOUND_GET_PARAM_NAME(InputVariable)))
 		{
 			// If a variable is provided, set the reference to read.
-			TDataWriteReference<FVariable> Variable = Inputs.GetDataWriteReference<FVariable>(GetInputVariableName());
-			if (Inputs.ContainsDataReadReference<DataType>(GetInputDataName()))
+			TDataWriteReference<FVariable> Variable = Inputs.GetDataWriteReference<FVariable>(METASOUND_GET_PARAM_NAME(InputVariable));
+			if (Inputs.ContainsDataReadReference<DataType>(METASOUND_GET_PARAM_NAME(InputData)))
 			{
 				// Update the input variable with the data reference to copy from.
-				TDataReadReference InputData = Inputs.GetDataReadReference<DataType>(GetInputDataName());
+				TDataReadReference InputData = Inputs.GetDataReadReference<DataType>(METASOUND_GET_PARAM_NAME(InputData));
 				Variable->SetDataReference(InputData);
 			}
 
 			return MakeUnique<FOperator>(Variable);
 		}
-		else if (Inputs.ContainsDataReadReference<DataType>(GetInputDataName()))
+		else if (Inputs.ContainsDataReadReference<DataType>(METASOUND_GET_PARAM_NAME(InputData)))
 		{
 			// If no input variable is provided create Variable with input variable
-			TDataReadReference<DataType> InputData = Inputs.GetDataReadReference<DataType>(GetInputDataName());
+			TDataReadReference<DataType> InputData = Inputs.GetDataReadReference<DataType>(METASOUND_GET_PARAM_NAME(InputData));
 			TDataWriteReference<DataType> InitData = TDataWriteReference<DataType>::CreateNew(*InputData);
 			TDataWriteReference<FVariable> Variable = TDataWriteReference<FVariable>::CreateNew(InitData);
 			Variable->SetDataReference(InputData);
@@ -310,13 +334,15 @@ namespace Metasound
 	template<typename DataType>
 	FVertexInterface TVariableMutatorNode<DataType>::FOperator::DeclareVertexInterface()
 	{
+		using namespace VariableNames;
+
 		return FVertexInterface(
 			FInputVertexInterface(
-				TInputDataVertexModel<DataType>(VariableNames::GetInputDataName(), FText::GetEmpty()), 
-				TInputDataVertexModel<FVariable>(VariableNames::GetInputVariableName(), FText::GetEmpty())
+				TInputDataVertexModel<DataType>(METASOUND_GET_PARAM_NAME(InputData), InputDataMetadata),
+				TInputDataVertexModel<FVariable>(METASOUND_GET_PARAM_NAME(InputVariable), InputVariableMetadata)
 			),
 			FOutputVertexInterface(
-				TOutputDataVertexModel<FVariable>(VariableNames::GetOutputVariableName(), FText::GetEmpty())
+				TOutputDataVertexModel<FVariable>(METASOUND_GET_PARAM_NAME(OutputVariable), OutputVariableMetadata)
 			)
 		);
 	}
@@ -375,9 +401,11 @@ namespace Metasound
 
 			virtual FDataReferenceCollection GetOutputs() const override
 			{
+				using namespace VariableNames; 
+
 				FDataReferenceCollection Collection;
-				Collection.AddDataWriteReference<FVariable>(VariableNames::GetOutputVariableName(), DelayedVariable);
-				Collection.AddDataReadReference<DataType>(VariableNames::GetOutputDataName(), DelayedVariable->GetDelayedDataReference());
+				Collection.AddDataWriteReference<FVariable>(METASOUND_GET_PARAM_NAME(OutputVariable), DelayedVariable);
+				Collection.AddDataReadReference<DataType>(METASOUND_GET_PARAM_NAME(OutputData), DelayedVariable->GetDelayedDataReference());
 				return Collection;
 			}
 
@@ -405,9 +433,9 @@ namespace Metasound
 
 		const FDataReferenceCollection& Inputs = InParams.InputDataReferences;
 
-		if (ensure(Inputs.ContainsDataWriteReference<FVariable>(GetInputVariableName())))
+		if (ensure(Inputs.ContainsDataWriteReference<FVariable>(METASOUND_GET_PARAM_NAME(InputVariable))))
 		{
-			TDataWriteReference<FVariable> DelayedVariable = Inputs.GetDataWriteReference<FVariable>(GetInputVariableName());
+			TDataWriteReference<FVariable> DelayedVariable = Inputs.GetDataWriteReference<FVariable>(METASOUND_GET_PARAM_NAME(InputVariable));
 
 			return MakeUnique<FOperator>(DelayedVariable);
 		}
@@ -422,13 +450,15 @@ namespace Metasound
 	template<typename DataType>
 	FVertexInterface TVariableDeferredAccessorNode<DataType>::FOperator::DeclareVertexInterface()
 	{
+		using namespace VariableNames; 
+
 		return FVertexInterface(
 			FInputVertexInterface(
-				TInputDataVertexModel<FVariable>(VariableNames::GetInputVariableName(), FText::GetEmpty())
+				TInputDataVertexModel<FVariable>(METASOUND_GET_PARAM_NAME(InputVariable), InputVariableMetadata)
 			),
 			FOutputVertexInterface(
-				TOutputDataVertexModel<FVariable>(VariableNames::GetOutputVariableName(), FText::GetEmpty()),
-				TOutputDataVertexModel<DataType>(VariableNames::GetOutputDataName(), FText::GetEmpty())
+				TOutputDataVertexModel<FVariable>(METASOUND_GET_PARAM_NAME(OutputVariable), OutputVariableMetadata),
+				TOutputDataVertexModel<DataType>(METASOUND_GET_PARAM_NAME(OutputData), OutputDataMetadata)
 			)
 		);
 	}
@@ -482,9 +512,11 @@ namespace Metasound
 
 			virtual FDataReferenceCollection GetOutputs() const override
 			{
+				using namespace VariableNames; 
+
 				FDataReferenceCollection Collection;
-				Collection.AddDataReadReference<DataType>(VariableNames::GetOutputDataName(), Variable->GetDataReference());
-				Collection.AddDataReadReference<FVariable>(VariableNames::GetOutputVariableName(), Variable);
+				Collection.AddDataReadReference<DataType>(METASOUND_GET_PARAM_NAME(OutputData), Variable->GetDataReference());
+				Collection.AddDataReadReference<FVariable>(METASOUND_GET_PARAM_NAME(OutputVariable), Variable);
 				return Collection;
 			}
 
@@ -511,14 +543,14 @@ namespace Metasound
 		const FDataReferenceCollection& Inputs = InParams.InputDataReferences;
 
 		// Update delayed variable.
-		if (Inputs.ContainsDataReadReference<FVariable>(GetInputVariableName()))
+		if (Inputs.ContainsDataReadReference<FVariable>(METASOUND_GET_PARAM_NAME(InputVariable)))
 		{
-			TDataReadReference<FVariable> Variable = Inputs.GetDataReadReference<FVariable>(GetInputVariableName());
+			TDataReadReference<FVariable> Variable = Inputs.GetDataReadReference<FVariable>(METASOUND_GET_PARAM_NAME(InputVariable));
 			return MakeUnique<FOperator>(Variable);
 		}
-		else if (Inputs.ContainsDataWriteReference<FVariable>(GetInputVariableName()))
+		else if (Inputs.ContainsDataWriteReference<FVariable>(METASOUND_GET_PARAM_NAME(InputVariable)))
 		{
-			TDataReadReference<FVariable> Variable = Inputs.GetDataWriteReference<FVariable>(GetInputVariableName());
+			TDataReadReference<FVariable> Variable = Inputs.GetDataWriteReference<FVariable>(METASOUND_GET_PARAM_NAME(InputVariable));
 			return MakeUnique<FOperator>(Variable);
 		}
 		else 
@@ -532,13 +564,15 @@ namespace Metasound
 	template<typename DataType>
 	FVertexInterface TVariableAccessorNode<DataType>::FOperator::DeclareVertexInterface()
 	{
+		using namespace VariableNames; 
+
 		return FVertexInterface(
 			FInputVertexInterface(
-				TInputDataVertexModel<FVariable>(VariableNames::GetInputVariableName(), FText::GetEmpty())
+				TInputDataVertexModel<FVariable>(METASOUND_GET_PARAM_NAME(InputVariable), InputDataMetadata)
 			),
 			FOutputVertexInterface(
-				TOutputDataVertexModel<DataType>(VariableNames::GetOutputDataName(), FText::GetEmpty()),
-				TOutputDataVertexModel<FVariable>(VariableNames::GetOutputVariableName(), FText::GetEmpty())
+				TOutputDataVertexModel<DataType>(METASOUND_GET_PARAM_NAME(OutputData), OutputDataMetadata),
+				TOutputDataVertexModel<FVariable>(METASOUND_GET_PARAM_NAME(OutputVariable), OutputVariableMetadata)
 			)
 		);
 	}
diff --git a/Engine/Plugins/Runtime/Metasound/Source/MetasoundStandardNodes/Private/MetasoundAudioFormats.cpp b/Engine/Plugins/Runtime/Metasound/Source/MetasoundStandardNodes/Private/MetasoundAudioFormats.cpp
--- a/Engine/Plugins/Runtime/Metasound/Source/MetasoundStandardNodes/Private/MetasoundAudioFormats.cpp
+++ b/Engine/Plugins/Runtime/Metasound/Source/MetasoundStandardNodes/Private/MetasoundAudioFormats.cpp
@@ -88,8 +88,8 @@ namespace Metasound
 	// Special vertex keys for stereo input/output nodes.
 	namespace StereoAudioFormatVertexKeys
 	{
-		static const FVertexName LeftChannelVertexKey = "Left";
-		static const FVertexName RightChannelVertexKey = "Right";
+		METASOUND_PARAM(LeftChannelVertex, "Left", "Left channel audio output.")
+		METASOUND_PARAM(RightChannelVertex, "Right", "Right channel audio output.")
 	}
 
 	// Specialization of TOutputNode<FStereoAudio> to support direct connection
@@ -113,10 +113,12 @@ namespace Metasound
 
 			virtual FDataReferenceCollection GetInputs() const override
 			{
+				using namespace StereoAudioFormatVertexKeys;
+
 				FDataReferenceCollection Inputs;
 
-				Inputs.AddDataReadReference(StereoAudioFormatVertexKeys::LeftChannelVertexKey, Left);
-				Inputs.AddDataReadReference(StereoAudioFormatVertexKeys::RightChannelVertexKey, Right);
+				Inputs.AddDataReadReference(METASOUND_GET_PARAM_NAME(LeftChannelVertex), Left);
+				Inputs.AddDataReadReference(METASOUND_GET_PARAM_NAME(RightChannelVertex), Right);
 
 				return Inputs;
 			}
@@ -153,9 +155,11 @@ namespace Metasound
 
 			TUniquePtr<IOperator> CreateOperator(const FCreateOperatorParams& InParams, FBuildErrorArray& OutErrors) override
 			{
+				using namespace StereoAudioFormatVertexKeys;
+
 				// Construct stereo from left and right audio buffers. 
-				TDataReadReference<FAudioBuffer> Left = InParams.InputDataReferences.GetDataReadReferenceOrConstruct<FAudioBuffer>(StereoAudioFormatVertexKeys::LeftChannelVertexKey, InParams.OperatorSettings);
-				TDataReadReference<FAudioBuffer> Right = InParams.InputDataReferences.GetDataReadReferenceOrConstruct<FAudioBuffer>(StereoAudioFormatVertexKeys::RightChannelVertexKey, InParams.OperatorSettings);
+				TDataReadReference<FAudioBuffer> Left = InParams.InputDataReferences.GetDataReadReferenceOrConstruct<FAudioBuffer>(METASOUND_GET_PARAM_NAME(LeftChannelVertex), InParams.OperatorSettings);
+				TDataReadReference<FAudioBuffer> Right = InParams.InputDataReferences.GetDataReadReferenceOrConstruct<FAudioBuffer>(METASOUND_GET_PARAM_NAME(RightChannelVertex), InParams.OperatorSettings);
 
 				TDataReadReference<FStereoAudioFormat> Stereo = TDataReadReferenceFactory<FStereoAudioFormat>::CreateExplicitArgs(InParams.OperatorSettings, WriteCast(Left), WriteCast(Right));
 
@@ -168,10 +172,12 @@ namespace Metasound
 
 		static FVertexInterface GetVertexInterface(const FVertexName& InVertexName)
 		{
+			using namespace StereoAudioFormatVertexKeys; 
+
 			return FVertexInterface(
 				FInputVertexInterface(
-					TInputDataVertexModel<FAudioBuffer>(StereoAudioFormatVertexKeys::LeftChannelVertexKey, METASOUND_LOCTEXT("Metasound_LeftStereoOutputVertexDescription", "Left channel audio output.")),
-					TInputDataVertexModel<FAudioBuffer>(StereoAudioFormatVertexKeys::RightChannelVertexKey, METASOUND_LOCTEXT("Metasound_RightStereoOutputVertexDescription", "Right channel audio output."))
+					TInputDataVertexModel<FAudioBuffer>(METASOUND_GET_PARAM_NAME_AND_METADATA(LeftChannelVertex)),
+					TInputDataVertexModel<FAudioBuffer>(METASOUND_GET_PARAM_NAME_AND_METADATA(RightChannelVertex))
 				),
 				FOutputVertexInterface(
 					TOutputDataVertexModel<FStereoAudioFormat>(InVertexName, METASOUND_LOCTEXT("Metasound_StereoOutputVertexDescription", "Stereo Output."))
@@ -263,9 +269,10 @@ namespace Metasound
 			virtual FDataReferenceCollection GetOutputs() const override
 			{
 				FDataReferenceCollection Outputs;
+				using namespace StereoAudioFormatVertexKeys;
 
-				Outputs.AddDataReadReference(StereoAudioFormatVertexKeys::LeftChannelVertexKey, Left);
-				Outputs.AddDataReadReference(StereoAudioFormatVertexKeys::RightChannelVertexKey, Right);
+				Outputs.AddDataReadReference(METASOUND_GET_PARAM_NAME(LeftChannelVertex), Left);
+				Outputs.AddDataReadReference(METASOUND_GET_PARAM_NAME(RightChannelVertex), Right);
 
 				return Outputs;
 			}
@@ -307,13 +314,15 @@ namespace Metasound
 
 		static FVertexInterface GetVertexInterface(const FVertexName& InVertexName)
 		{
+			using namespace StereoAudioFormatVertexKeys;
+			
 			return FVertexInterface(
 				FInputVertexInterface(
 					TInputDataVertexModel<FStereoAudioFormat>(InVertexName, METASOUND_LOCTEXT("Metasound_StereoInputVertexDescription", "Stereo Input."))
 				),
 				FOutputVertexInterface(
-					TOutputDataVertexModel<FAudioBuffer>(StereoAudioFormatVertexKeys::LeftChannelVertexKey, METASOUND_LOCTEXT("Metasound_LeftStereoInputVertexDescription", "Left channel audio output.")),
-					TOutputDataVertexModel<FAudioBuffer>(StereoAudioFormatVertexKeys::RightChannelVertexKey, METASOUND_LOCTEXT("Metasound_RightStereoInputVertexDescription", "Right channel audio output."))
+					TOutputDataVertexModel<FAudioBuffer>(METASOUND_GET_PARAM_NAME_AND_METADATA(LeftChannelVertex)),
+					TOutputDataVertexModel<FAudioBuffer>(METASOUND_GET_PARAM_NAME_AND_METADATA(RightChannelVertex))
 				)
 			);
 		}
diff --git a/Engine/Plugins/Runtime/Metasound/Source/MetasoundStandardNodes/Private/MetasoundBPMToSecondsNode.cpp b/Engine/Plugins/Runtime/Metasound/Source/MetasoundStandardNodes/Private/MetasoundBPMToSecondsNode.cpp
--- a/Engine/Plugins/Runtime/Metasound/Source/MetasoundStandardNodes/Private/MetasoundBPMToSecondsNode.cpp
+++ b/Engine/Plugins/Runtime/Metasound/Source/MetasoundStandardNodes/Private/MetasoundBPMToSecondsNode.cpp
@@ -5,6 +5,7 @@
 #include "Internationalization/Text.h"
 #include "MetasoundExecutableOperator.h"
 #include "MetasoundNodeRegistrationMacro.h"
+#include "MetasoundParamHelper.h"
 #include "MetasoundPrimitives.h"
 #include "MetasoundStandardNodesNames.h"
 #include "MetasoundTime.h"
@@ -15,29 +16,10 @@ namespace Metasound
 {
 	namespace BPMToSecondsVertexNames
 	{
-		const FVertexName& GetInputBPMTempoName()
-		{
-			static FVertexName Name = TEXT("BPM");
-			return Name;
-		}
-
-		const FVertexName& GetInputBeatMultiplierName()
-		{
-			static FVertexName Name = TEXT("Beat Multiplier");
-			return Name;
-		}
-
-		const FVertexName& GetInputDivisionsOfWholeNoteName()
-		{
-			static FVertexName Name = TEXT("Divisions of Whole Note");
-			return Name;
-		}
-
-		const FVertexName& GetOutputTimeSecondsName()
-		{
-			static FVertexName Name = TEXT("Seconds");
-			return Name;
-		}
+		METASOUND_PARAM(InputBPM, "BPM", "Beats Per Minute.")
+		METASOUND_PARAM(InputBeatMultiplier, "Beat Multiplier", "The multiplier of the BPM.")
+		METASOUND_PARAM(InputDivOfWholeNote, "Divisions of Whole Note", "Divisions of a whole note.")
+		METASOUND_PARAM(OutputTimeSeconds, "Seconds", "The output time in seconds.")
 	}
 
 	class FBPMToSecondsOperator : public TExecutableOperator<FBPMToSecondsOperator>
@@ -84,9 +66,9 @@ namespace Metasound
 		using namespace BPMToSecondsVertexNames;
 
 		FDataReferenceCollection InputDataReferences;
-		InputDataReferences.AddDataReadReference(GetInputBPMTempoName(), BPM);
-		InputDataReferences.AddDataReadReference(GetInputBeatMultiplierName(), BeatMultiplier);
-		InputDataReferences.AddDataReadReference(GetInputDivisionsOfWholeNoteName(), DivOfWholeNote);
+		InputDataReferences.AddDataReadReference(METASOUND_GET_PARAM_NAME(InputBPM), BPM);
+		InputDataReferences.AddDataReadReference(METASOUND_GET_PARAM_NAME(InputBeatMultiplier), BeatMultiplier);
+		InputDataReferences.AddDataReadReference(METASOUND_GET_PARAM_NAME(InputDivOfWholeNote), DivOfWholeNote);
 
 		return InputDataReferences;
 	}
@@ -96,7 +78,7 @@ namespace Metasound
 		using namespace BPMToSecondsVertexNames;
 
 		FDataReferenceCollection OutputDataReferences;
-		OutputDataReferences.AddDataReadReference(GetOutputTimeSecondsName(), TimeSeconds);
+		OutputDataReferences.AddDataReadReference(METASOUND_GET_PARAM_NAME(OutputTimeSeconds), TimeSeconds);
 		return OutputDataReferences;
 	}
 
@@ -134,12 +116,12 @@ namespace Metasound
 
 		static const FVertexInterface Interface(
 			FInputVertexInterface(
-				TInputDataVertexModel<float>(GetInputBPMTempoName(), METASOUND_LOCTEXT("BPMToSecondsNode_BPMTT", "Beats Per Minute."), 90.0f),
-				TInputDataVertexModel<float>(GetInputBeatMultiplierName(), METASOUND_LOCTEXT("BPMToSecondsNode_BeatMultiplierTT", "The multiplier of the BPM."), 1.0f),
-				TInputDataVertexModel<float>(GetInputDivisionsOfWholeNoteName(), METASOUND_LOCTEXT("BPMToSecondsNode_DivOfWholeNoteTT", "Divisions of a whole note."), 4.0f)
+				TInputDataVertexModel<float>(METASOUND_GET_PARAM_NAME_AND_METADATA(InputBPM), 90.0f),
+				TInputDataVertexModel<float>(METASOUND_GET_PARAM_NAME_AND_METADATA(InputBeatMultiplier), 1.0f),
+				TInputDataVertexModel<float>(METASOUND_GET_PARAM_NAME_AND_METADATA(InputDivOfWholeNote), 4.0f)
 			),
 			FOutputVertexInterface(
-				TOutputDataVertexModel<FTime>(GetOutputTimeSecondsName(), METASOUND_LOCTEXT("BPMToSecondsNode_OutputTimeTT", "The output time in seconds."))
+				TOutputDataVertexModel<FTime>(METASOUND_GET_PARAM_NAME_AND_METADATA(OutputTimeSeconds))
 			)
 		);
 
@@ -175,9 +157,9 @@ namespace Metasound
 		const FDataReferenceCollection& InputCollection = InParams.InputDataReferences;
 		const FInputVertexInterface& InputInterface = GetVertexInterface().GetInputInterface();
 
-		FFloatReadRef InBPM = InputCollection.GetDataReadReferenceOrConstructWithVertexDefault<float>(InputInterface, GetInputBPMTempoName(), InParams.OperatorSettings);
-		FFloatReadRef InBeatMultiplier = InputCollection.GetDataReadReferenceOrConstructWithVertexDefault<float>(InputInterface, GetInputBeatMultiplierName(), InParams.OperatorSettings);
-		FFloatReadRef InDivOfWholeNote = InputCollection.GetDataReadReferenceOrConstructWithVertexDefault<float>(InputInterface, GetInputDivisionsOfWholeNoteName(), InParams.OperatorSettings);
+		FFloatReadRef InBPM = InputCollection.GetDataReadReferenceOrConstructWithVertexDefault<float>(InputInterface, METASOUND_GET_PARAM_NAME(InputBPM), InParams.OperatorSettings);
+		FFloatReadRef InBeatMultiplier = InputCollection.GetDataReadReferenceOrConstructWithVertexDefault<float>(InputInterface, METASOUND_GET_PARAM_NAME(InputBeatMultiplier), InParams.OperatorSettings);
+		FFloatReadRef InDivOfWholeNote = InputCollection.GetDataReadReferenceOrConstructWithVertexDefault<float>(InputInterface, METASOUND_GET_PARAM_NAME(InputDivOfWholeNote), InParams.OperatorSettings);
 
 		return MakeUnique<FBPMToSecondsOperator>(InParams.OperatorSettings, InBPM, InBeatMultiplier, InDivOfWholeNote);
 	}
diff --git a/Engine/Plugins/Runtime/Metasound/Source/MetasoundStandardNodes/Private/MetasoundBandsplitterNode.cpp b/Engine/Plugins/Runtime/Metasound/Source/MetasoundStandardNodes/Private/MetasoundBandsplitterNode.cpp
--- a/Engine/Plugins/Runtime/Metasound/Source/MetasoundStandardNodes/Private/MetasoundBandsplitterNode.cpp
+++ b/Engine/Plugins/Runtime/Metasound/Source/MetasoundStandardNodes/Private/MetasoundBandsplitterNode.cpp
@@ -173,16 +173,28 @@ namespace Metasound
 				// audio channels
 				for (uint32 ChannelIndex = 0; ChannelIndex < NumChannels; ChannelIndex++)
 				{
-					InputInterface.Add(TInputDataVertexModel<FAudioBuffer>(GetAudioInputName(ChannelIndex), GetAudioInputDescription(ChannelIndex)));
+					const FDataVertexMetadata AudioInputMetadata
+					{
+						GetAudioInputDescription(ChannelIndex),
+						GetAudioInputDisplayName(ChannelIndex)
+					};
+
+					InputInterface.Add(TInputDataVertexModel<FAudioBuffer>(GetAudioInputName(ChannelIndex), AudioInputMetadata));
 				}
 
-				InputInterface.Add(TInputDataVertexModel<FEnumBandSplitterFilterOrder>(METASOUND_GET_PARAM_NAME_AND_TT(InputFilterOrder)));
-				InputInterface.Add(TInputDataVertexModel<bool>(METASOUND_GET_PARAM_NAME_AND_TT(InputPhaseCompensate), true));
+				InputInterface.Add(TInputDataVertexModel<FEnumBandSplitterFilterOrder>(METASOUND_GET_PARAM_NAME_AND_METADATA(InputFilterOrder)));
+				InputInterface.Add(TInputDataVertexModel<bool>(METASOUND_GET_PARAM_NAME_AND_METADATA(InputPhaseCompensate), true));
 
 				// Crossover frequencies
 				for (uint32 InputIndex = 0; InputIndex < NumBands - 1; InputIndex++)
 				{
-					InputInterface.Add(TInputDataVertexModel<float>(GetCrossoverInputName(InputIndex), GetCrossoverInputDescription(InputIndex), (1 + InputIndex) * 500.0f));
+					const FDataVertexMetadata CrossoverInputMetadata
+					{
+						GetCrossoverInputDescription(InputIndex),
+						GetCrossoverInputDisplayName(InputIndex)
+					};
+
+					InputInterface.Add(TInputDataVertexModel<float>(GetCrossoverInputName(InputIndex), CrossoverInputMetadata, (1 + InputIndex) * 500.0f));
 				}
 
 				// outputs
@@ -191,7 +203,13 @@ namespace Metasound
 				{
 					for (uint32 ChannelIndex = 0; ChannelIndex < NumChannels; ChannelIndex++)
 					{
-						OutputInterface.Add(TOutputDataVertexModel<FAudioBuffer>(GetAudioOutputName(OutputIndex, ChannelIndex), GetAudioOutputDescription(OutputIndex, ChannelIndex)));
+						const FDataVertexMetadata AudioOutputMetadata
+						{
+							GetAudioOutputDescription(OutputIndex, ChannelIndex),
+							GetAudioOutputDisplayName(OutputIndex, ChannelIndex)
+						};
+
+						OutputInterface.Add(TOutputDataVertexModel<FAudioBuffer>(GetAudioOutputName(OutputIndex, ChannelIndex), AudioOutputMetadata));
 					}
 				}
 
@@ -339,11 +357,37 @@ namespace Metasound
 			return METASOUND_LOCTEXT_FORMAT("BandSplitterAudioInputDescription", "Audio Channel: {0}", ChannelIndex);
 		}
 
+		static const FText GetAudioInputDisplayName(uint32 ChannelIndex)
+		{
+			if (NumChannels == 1)
+			{
+				return METASOUND_LOCTEXT("BandsplitterAudioInputIn1", "In");
+			}
+			else if (NumChannels == 2)
+			{
+				if (ChannelIndex == 0)
+				{
+					return METASOUND_LOCTEXT("BandsplitterAudioInputIn2L", "In L");
+				}
+				else
+				{
+					return METASOUND_LOCTEXT("BandsplitterAudioInputIn2R", "In R");
+				}
+			}
+
+			return METASOUND_LOCTEXT_FORMAT("BandsplitterAudioInputIn", "In {0}", ChannelIndex);
+		}
+
 		static const FVertexName GetCrossoverInputName(uint32 InputIndex)
 		{
 			return *FString::Printf(TEXT("Crossover %i"), InputIndex);
 		}
 
+		static const FText GetCrossoverInputDisplayName(uint32 InputIndex)
+		{
+			return METASOUND_LOCTEXT_FORMAT("BandsplitterCrossoverInputDisplayName", "Crossover {0}", InputIndex);
+		}
+
 		static const FText GetCrossoverInputDescription(uint32 InputIndex)
 		{
 			return METASOUND_LOCTEXT_FORMAT("BandSplitterCrossoverInputDescription", "Crossover Input #: {0}", InputIndex);
@@ -368,6 +412,27 @@ namespace Metasound
 			return METASOUND_LOCTEXT_FORMAT("BandSplitterAudioOutputDescription", "Channel {0} Output for Band {1}", ChannelIndex, OutputIndex);
 		}
 
+		static const FText GetAudioOutputDisplayName(uint32 OutputIndex, uint32 ChannelIndex)
+		{
+			if (NumChannels == 1)
+			{
+				return METASOUND_LOCTEXT_FORMAT("BandsplitterAudioOutputOut1", "Band {0} Out", OutputIndex);
+			}
+			else if (NumChannels == 2)
+			{
+				if (ChannelIndex == 0)
+				{
+					return METASOUND_LOCTEXT_FORMAT("BandsplitterAudioOutputOut2L", "Band {0} L", OutputIndex);
+				}
+				else
+				{
+					return METASOUND_LOCTEXT_FORMAT("BandsplitterAudioOutputOut2R", "Band {0} R", OutputIndex);
+				}
+			}
+
+			return METASOUND_LOCTEXT_FORMAT("BandsplitterAudioOutputOut", "Band {0} Out {1}", OutputIndex, ChannelIndex);
+		}
+
 		FORCEINLINE Audio::EFilterOrder GetFilterOrder(EBandSplitterFilterOrder InFilterOrder) const
 		{
 			using namespace Audio;
diff --git a/Engine/Plugins/Runtime/Metasound/Source/MetasoundStandardNodes/Private/MetasoundCrossfadeNode.cpp b/Engine/Plugins/Runtime/Metasound/Source/MetasoundStandardNodes/Private/MetasoundCrossfadeNode.cpp
--- a/Engine/Plugins/Runtime/Metasound/Source/MetasoundStandardNodes/Private/MetasoundCrossfadeNode.cpp
+++ b/Engine/Plugins/Runtime/Metasound/Source/MetasoundStandardNodes/Private/MetasoundCrossfadeNode.cpp
@@ -25,17 +25,8 @@ namespace Metasound
 {
 	namespace CrossfadeVertexNames
 	{
-		const FVertexName GetInputCrossfadeValueName()
-		{
-			static const FVertexName CrossfadeName = TEXT("Crossfade Value");
-			return CrossfadeName;
-		}
-
-		const FText GetInputCrossfadeValueDescription()
-		{
-			static const FText Desc = METASOUND_LOCTEXT("InputCrossfadeValueDesc", "Crossfade value to crossfade across inputs. Output will be the float value between adjacent whole number values.");
-			return Desc;
-		}
+		METASOUND_PARAM(InputCrossfadeValue, "Crossfade Value", "Crossfade value to crossfade across inputs. Output will be the float value between adjacent whole number values.")
+		METASOUND_PARAM(OutputTrigger, "Out", "Triggered when all input triggers have been triggered. Call Reset to reset the state or use \"Auto Reset\"")
 
 		const FVertexName GetInputName(uint32 InIndex)
 		{
@@ -47,16 +38,9 @@ namespace Metasound
 			return METASOUND_LOCTEXT_FORMAT("CrossfadeInputDesc", "Cross fade {0} input.", InIndex);
 		}
 
-		const FVertexName& GetOutputName()
+		const FText GetInputDisplayName(uint32 InIndex)
 		{
-			static const FVertexName Name = TEXT("Out");
-			return Name;
-		}
-
-		const FText& GetOutputDescription()
-		{
-			static const FText Desc = METASOUND_LOCTEXT("TriggerAccumulateOutputTriggerDesc", "Triggered when all input triggers have been triggered. Call Reset to reset the state or use \"Auto Reset\"");
-			return Desc;
+			return METASOUND_LOCTEXT_FORMAT("CrossfadeInputDisplayName", "In {0}", InIndex);
 		}
 	}
 
@@ -162,15 +146,21 @@ namespace Metasound
 			{
 				FInputVertexInterface InputInterface;
 
-				InputInterface.Add(TInputDataVertexModel<float>(GetInputCrossfadeValueName(), GetInputCrossfadeValueDescription()));
+				InputInterface.Add(TInputDataVertexModel<float>(METASOUND_GET_PARAM_NAME_AND_METADATA(InputCrossfadeValue)));
 
 				for (uint32 i = 0; i < NumInputs; ++i)
 				{
-					InputInterface.Add(TInputDataVertexModel<ValueType>(GetInputName(i), GetInputDescription(i)));
+					const FDataVertexMetadata InputMetadata
+					{
+						GetInputDescription(i),
+						GetInputDisplayName(i)
+					};
+
+					InputInterface.Add(TInputDataVertexModel<ValueType>(GetInputName(i), InputMetadata));
 				}
 
 				FOutputVertexInterface OutputInterface;
-				OutputInterface.Add(TOutputDataVertexModel<ValueType>(GetOutputName(), GetOutputDescription()));
+				OutputInterface.Add(TOutputDataVertexModel<ValueType>(METASOUND_GET_PARAM_NAME_AND_METADATA(OutputTrigger)));
 
 				return FVertexInterface(InputInterface, OutputInterface);
 			};
@@ -217,7 +207,7 @@ namespace Metasound
 			const FInputVertexInterface& InputInterface = InParams.Node.GetVertexInterface().GetInputInterface();
 			const FDataReferenceCollection& InputCollection = InParams.InputDataReferences;
 
-			FFloatReadRef CrossfadeValue = InputCollection.GetDataReadReferenceOrConstructWithVertexDefault<float>(InputInterface, GetInputCrossfadeValueName(), InParams.OperatorSettings);
+			FFloatReadRef CrossfadeValue = InputCollection.GetDataReadReferenceOrConstructWithVertexDefault<float>(InputInterface, METASOUND_GET_PARAM_NAME(InputCrossfadeValue), InParams.OperatorSettings);
 
 			TArray<TDataReadReference<ValueType>> InputValues;
 			for (uint32 i = 0; i < NumInputs; ++i)
@@ -244,7 +234,7 @@ namespace Metasound
 		{
 			using namespace CrossfadeVertexNames;
 			FDataReferenceCollection Inputs;
-			Inputs.AddDataReadReference(GetInputCrossfadeValueName(), CrossfadeValue);
+			Inputs.AddDataReadReference(METASOUND_GET_PARAM_NAME(InputCrossfadeValue), CrossfadeValue);
 			for (uint32 i = 0; i < NumInputs; ++i)
 			{
 				Inputs.AddDataReadReference(GetInputName(i), InputValues[i]);
@@ -257,7 +247,7 @@ namespace Metasound
 			using namespace CrossfadeVertexNames;
 
 			FDataReferenceCollection Outputs;
-			Outputs.AddDataReadReference(GetOutputName(), OutputValue);
+			Outputs.AddDataReadReference(METASOUND_GET_PARAM_NAME(OutputTrigger), OutputValue);
 
 			return Outputs;
 		}
diff --git a/Engine/Plugins/Runtime/Metasound/Source/MetasoundStandardNodes/Private/MetasoundDebugLogNode.cpp b/Engine/Plugins/Runtime/Metasound/Source/MetasoundStandardNodes/Private/MetasoundDebugLogNode.cpp
--- a/Engine/Plugins/Runtime/Metasound/Source/MetasoundStandardNodes/Private/MetasoundDebugLogNode.cpp
+++ b/Engine/Plugins/Runtime/Metasound/Source/MetasoundStandardNodes/Private/MetasoundDebugLogNode.cpp
@@ -1,12 +1,14 @@
 // Copyright Epic Games, Inc. All Rights Reserved.
 
+#include "MetasoundFacade.h"
 #include "MetasoundNodeRegistrationMacro.h"
+#include "MetasoundParamHelper.h"
+#include "MetasoundPrimitives.h"
 #include "MetasoundStandardNodesCategories.h"
-#include "MetasoundFacade.h"
 #include "MetasoundSourceInterface.h"
 #include "MetasoundStandardNodesNames.h"
 #include "MetasoundTrigger.h"
-#include "MetasoundPrimitives.h"
+
 
 #define LOCTEXT_NAMESPACE "MetasoundStandardNodes_DebugLogNode"
 
@@ -39,29 +41,9 @@ namespace Metasound
 	//Getters for the name of each parameter used in Print Log
 	namespace PrintLogVertexNames
 	{
-		const FVertexName& GetInputTriggerName()
-		{
-			static const FVertexName Name = TEXT("Trigger");
-			return Name;
-		}
-
-		const FVertexName& GetLabelPrintLogName()
-		{
-			static const FVertexName Name = TEXT("Label");
-			return Name;
-		}
-
-		const FVertexName& GetToLogPrintLogName()
-		{
-			static const FVertexName Name = TEXT("Value To Log");
-			return Name;
-		}
-
-		const FVertexName& GetOutputPrintLogName()
-		{
-			static const FVertexName Name = TEXT("Was Successful");
-			return Name;
-		}
+		METASOUND_PARAM(InputTrigger, "Trigger", "Trigger to write the set value to the log.")
+		METASOUND_PARAM(InputLabel, "Label", "The label to attach to the value that will be logged.")
+		METASOUND_PARAM(InputValueToLog, "Value To Log", "The value to record to the log when triggered.")
 	}
 
 
@@ -73,12 +55,13 @@ namespace Metasound
 
 			static const FVertexInterface& GetDefaultInterface()
 			{
+				using namespace PrintLogVertexNames;
 				static const FVertexInterface DefaultInterface(
 					FInputVertexInterface(
 
-						TInputDataVertexModel<FTrigger>(PrintLogVertexNames::GetInputTriggerName(), METASOUND_LOCTEXT("PrintLogTrigger", "Trigger to write the set value to the log.")),
-						TInputDataVertexModel<FString>(PrintLogVertexNames::GetLabelPrintLogName(), METASOUND_LOCTEXT("PrintLogLabel", "The label to attach to the value that will be logged")),
-						TInputDataVertexModel<PrintLogType>(PrintLogVertexNames::GetToLogPrintLogName(), METASOUND_LOCTEXT("PrintLogValueToLog", "The value to record to the log when triggered"))
+						TInputDataVertexModel<FTrigger>(METASOUND_GET_PARAM_NAME_AND_METADATA(InputTrigger)),
+						TInputDataVertexModel<FString>(METASOUND_GET_PARAM_NAME_AND_METADATA(InputLabel)),
+						TInputDataVertexModel<PrintLogType>(METASOUND_GET_PARAM_NAME_AND_METADATA(InputValueToLog))
 					),
 					FOutputVertexInterface(
 					)
@@ -106,13 +89,15 @@ namespace Metasound
 
 			static TUniquePtr<IOperator> CreateOperator(const FCreateOperatorParams& InParams, TArray<TUniquePtr<IOperatorBuildError>>& OutErrors)
 			{
+				using namespace PrintLogVertexNames;
+
 				const FInputVertexInterface& InputInterface = InParams.Node.GetVertexInterface().GetInputInterface();
 				const FDataReferenceCollection& InputCollection = InParams.InputDataReferences;
 
-				FTriggerReadRef Trigger = InputCollection.GetDataReadReferenceOrConstruct<FTrigger>(PrintLogVertexNames::GetInputTriggerName(), InParams.OperatorSettings);
+				FTriggerReadRef Trigger = InputCollection.GetDataReadReferenceOrConstruct<FTrigger>(METASOUND_GET_PARAM_NAME(InputTrigger), InParams.OperatorSettings);
 
-				TDataReadReference<FString> Label = InputCollection.GetDataReadReferenceOrConstructWithVertexDefault<FString>(InputInterface, PrintLogVertexNames::GetLabelPrintLogName(), InParams.OperatorSettings);
-				TDataReadReference<PrintLogType> ValueToLog = InputCollection.GetDataReadReferenceOrConstructWithVertexDefault<PrintLogType>(InputInterface, PrintLogVertexNames::GetToLogPrintLogName(), InParams.OperatorSettings);
+				TDataReadReference<FString> Label = InputCollection.GetDataReadReferenceOrConstructWithVertexDefault<FString>(InputInterface, METASOUND_GET_PARAM_NAME(InputLabel), InParams.OperatorSettings);
+				TDataReadReference<PrintLogType> ValueToLog = InputCollection.GetDataReadReferenceOrConstructWithVertexDefault<PrintLogType>(InputInterface, METASOUND_GET_PARAM_NAME(InputValueToLog), InParams.OperatorSettings);
 
 				FString GraphNameFull = InParams.Environment.GetValue<FString>(Frontend::SourceInterface::Environment::GraphName);
 				TArray<FString> ParsedString;
@@ -137,11 +122,12 @@ namespace Metasound
 
 			virtual FDataReferenceCollection GetInputs() const override
 			{
+				using namespace PrintLogVertexNames;
 				FDataReferenceCollection Inputs;
 
-				Inputs.AddDataReadReference(PrintLogVertexNames::GetInputTriggerName(), Trigger);
-				Inputs.AddDataReadReference(PrintLogVertexNames::GetLabelPrintLogName(), Label);
-				Inputs.AddDataReadReference(PrintLogVertexNames::GetToLogPrintLogName(), ValueToLog);
+				Inputs.AddDataReadReference(METASOUND_GET_PARAM_NAME(InputTrigger), Trigger);
+				Inputs.AddDataReadReference(METASOUND_GET_PARAM_NAME(InputLabel), Label);
+				Inputs.AddDataReadReference(METASOUND_GET_PARAM_NAME(InputValueToLog), ValueToLog);
 
 				return Inputs;
 			}
diff --git a/Engine/Plugins/Runtime/Metasound/Source/MetasoundStandardNodes/Private/MetasoundDelayNode.cpp b/Engine/Plugins/Runtime/Metasound/Source/MetasoundStandardNodes/Private/MetasoundDelayNode.cpp
--- a/Engine/Plugins/Runtime/Metasound/Source/MetasoundStandardNodes/Private/MetasoundDelayNode.cpp
+++ b/Engine/Plugins/Runtime/Metasound/Source/MetasoundStandardNodes/Private/MetasoundDelayNode.cpp
@@ -3,6 +3,7 @@
 #include "Internationalization/Text.h"
 #include "MetasoundExecutableOperator.h"
 #include "MetasoundNodeRegistrationMacro.h"
+#include "MetasoundParamHelper.h"
 #include "MetasoundPrimitives.h"
 #include "MetasoundStandardNodesNames.h"
 #include "MetasoundTrigger.h"
@@ -18,12 +19,12 @@ namespace Metasound
 {
 	namespace Delay
 	{
-		static const FName InParamNameAudioInput = "In";
-		static const FName InParamNameDelayTime = "Delay Time";
-		static const FName InParamNameDryLevel = "Dry Level";
-		static const FName InParamNameWetLevel = "Wet Level";
-		static const FName InParamNameFeedbackAmount = "Feedback";
-		static const FName OutParamNameAudio = "Out";
+		METASOUND_PARAM(InParamAudioInput, "In", "Audio input.")
+		METASOUND_PARAM(InParamDelayTime, "Delay Time", "The amount of time to delay the audio, in seconds.")
+		METASOUND_PARAM(InParamDryLevel, "Dry Level", "The dry level of the delay.")
+		METASOUND_PARAM(InParamWetLevel, "Wet Level", "The wet level of the delay.")
+		METASOUND_PARAM(InParamFeedbackAmount, "Feedback", "Feedback amount.")
+		METASOUND_PARAM(OutParamAudio, "Out", "Audio output.")
 
 		// TODO: make this a static vertex
 		static float MaxDelaySeconds = 5.0f;
@@ -101,20 +102,24 @@ namespace Metasound
 
 	FDataReferenceCollection FDelayOperator::GetInputs() const
 	{
+		using namespace Delay;
+
 		FDataReferenceCollection InputDataReferences;
-		InputDataReferences.AddDataReadReference(Delay::InParamNameAudioInput, FAudioBufferReadRef(AudioInput));
-		InputDataReferences.AddDataReadReference(Delay::InParamNameDelayTime, FTimeReadRef(DelayTime));
-		InputDataReferences.AddDataReadReference(Delay::InParamNameDryLevel, FFloatReadRef(DryLevel));
-		InputDataReferences.AddDataReadReference(Delay::InParamNameWetLevel, FFloatReadRef(WetLevel));
-		InputDataReferences.AddDataReadReference(Delay::InParamNameFeedbackAmount, FFloatReadRef(Feedback));
+		InputDataReferences.AddDataReadReference(METASOUND_GET_PARAM_NAME(InParamAudioInput), FAudioBufferReadRef(AudioInput));
+		InputDataReferences.AddDataReadReference(METASOUND_GET_PARAM_NAME(InParamDelayTime), FTimeReadRef(DelayTime));
+		InputDataReferences.AddDataReadReference(METASOUND_GET_PARAM_NAME(InParamDryLevel), FFloatReadRef(DryLevel));
+		InputDataReferences.AddDataReadReference(METASOUND_GET_PARAM_NAME(InParamWetLevel), FFloatReadRef(WetLevel));
+		InputDataReferences.AddDataReadReference(METASOUND_GET_PARAM_NAME(InParamFeedbackAmount), FFloatReadRef(Feedback));
 
 		return InputDataReferences;
 	}
 
 	FDataReferenceCollection FDelayOperator::GetOutputs() const
 	{
+		using namespace Delay;
+
 		FDataReferenceCollection OutputDataReferences;
-		OutputDataReferences.AddDataReadReference(Delay::OutParamNameAudio, FAudioBufferReadRef(AudioOutput));
+		OutputDataReferences.AddDataReadReference(METASOUND_GET_PARAM_NAME(OutParamAudio), FAudioBufferReadRef(AudioOutput));
 		return OutputDataReferences;
 	}
 
@@ -168,16 +173,18 @@ namespace Metasound
 
 	const FVertexInterface& FDelayOperator::GetVertexInterface()
 	{
+		using namespace Delay;
+
 		static const FVertexInterface Interface(
 			FInputVertexInterface(
-				TInputDataVertexModel<FAudioBuffer>(Delay::InParamNameAudioInput, METASOUND_LOCTEXT("DelayNode_AudioInputTooltip", "Audio input.")),
-				TInputDataVertexModel<FTime>(Delay::InParamNameDelayTime, METASOUND_LOCTEXT("DelayNode_DelayTimeTooltip", "The amount of time to delay the audio, in seconds."), 1.0f),
-				TInputDataVertexModel<float>(Delay::InParamNameDryLevel, METASOUND_LOCTEXT("DelayNode_DryLevelTooltip", "The dry level of the delay."), 0.0f),
-				TInputDataVertexModel<float>(Delay::InParamNameWetLevel, METASOUND_LOCTEXT("DelayNode_WetlevelTooltip", "The wet level of the delay."), 1.0f),
-				TInputDataVertexModel<float>(Delay::InParamNameFeedbackAmount, METASOUND_LOCTEXT("DelayNode_FeedbackTooltip", "Feedback amount."), 0.0f)
+				TInputDataVertexModel<FAudioBuffer>(METASOUND_GET_PARAM_NAME_AND_METADATA(InParamAudioInput)),
+				TInputDataVertexModel<FTime>(METASOUND_GET_PARAM_NAME_AND_METADATA(InParamDelayTime), 1.0f),
+				TInputDataVertexModel<float>(METASOUND_GET_PARAM_NAME_AND_METADATA(InParamDryLevel), 0.0f),
+				TInputDataVertexModel<float>(METASOUND_GET_PARAM_NAME_AND_METADATA(InParamWetLevel), 1.0f),
+				TInputDataVertexModel<float>(METASOUND_GET_PARAM_NAME_AND_METADATA(InParamFeedbackAmount), 0.0f)
 			),
 			FOutputVertexInterface(
-				TOutputDataVertexModel<FAudioBuffer>(Delay::OutParamNameAudio, METASOUND_LOCTEXT("DelayNode_DelayOutputTooltip", "Audio output."))
+				TOutputDataVertexModel<FAudioBuffer>(METASOUND_GET_PARAM_NAME_AND_METADATA(OutParamAudio))
 			)
 		);
 
@@ -208,14 +215,16 @@ namespace Metasound
 
 	TUniquePtr<IOperator> FDelayOperator::CreateOperator(const FCreateOperatorParams& InParams, FBuildErrorArray& OutErrors)
 	{
+		using namespace Delay; 
+
 		const FDataReferenceCollection& InputCollection = InParams.InputDataReferences;
 		const FInputVertexInterface& InputInterface = GetVertexInterface().GetInputInterface();
 
-		FAudioBufferReadRef AudioIn = InputCollection.GetDataReadReferenceOrConstruct<FAudioBuffer>(Delay::InParamNameAudioInput, InParams.OperatorSettings);
-		FTimeReadRef DelayTime = InputCollection.GetDataReadReferenceOrConstructWithVertexDefault<FTime>(InputInterface, Delay::InParamNameDelayTime, InParams.OperatorSettings);
-		FFloatReadRef DryLevel = InputCollection.GetDataReadReferenceOrConstructWithVertexDefault<float>(InputInterface, Delay::InParamNameDryLevel, InParams.OperatorSettings);
-		FFloatReadRef WetLevel = InputCollection.GetDataReadReferenceOrConstructWithVertexDefault<float>(InputInterface, Delay::InParamNameWetLevel, InParams.OperatorSettings);
-		FFloatReadRef Feedback = InputCollection.GetDataReadReferenceOrConstructWithVertexDefault<float>(InputInterface, Delay::InParamNameFeedbackAmount, InParams.OperatorSettings);
+		FAudioBufferReadRef AudioIn = InputCollection.GetDataReadReferenceOrConstruct<FAudioBuffer>(METASOUND_GET_PARAM_NAME(InParamAudioInput), InParams.OperatorSettings);
+		FTimeReadRef DelayTime = InputCollection.GetDataReadReferenceOrConstructWithVertexDefault<FTime>(InputInterface, METASOUND_GET_PARAM_NAME(InParamDelayTime), InParams.OperatorSettings);
+		FFloatReadRef DryLevel = InputCollection.GetDataReadReferenceOrConstructWithVertexDefault<float>(InputInterface, METASOUND_GET_PARAM_NAME(InParamDryLevel), InParams.OperatorSettings);
+		FFloatReadRef WetLevel = InputCollection.GetDataReadReferenceOrConstructWithVertexDefault<float>(InputInterface, METASOUND_GET_PARAM_NAME(InParamWetLevel), InParams.OperatorSettings);
+		FFloatReadRef Feedback = InputCollection.GetDataReadReferenceOrConstructWithVertexDefault<float>(InputInterface, METASOUND_GET_PARAM_NAME(InParamFeedbackAmount), InParams.OperatorSettings);
 
 		return MakeUnique<FDelayOperator>(InParams.OperatorSettings, AudioIn, DelayTime, DryLevel, WetLevel, Feedback);
 	}
diff --git a/Engine/Plugins/Runtime/Metasound/Source/MetasoundStandardNodes/Private/MetasoundEnvelopeFollowerNode.cpp b/Engine/Plugins/Runtime/Metasound/Source/MetasoundStandardNodes/Private/MetasoundEnvelopeFollowerNode.cpp
--- a/Engine/Plugins/Runtime/Metasound/Source/MetasoundStandardNodes/Private/MetasoundEnvelopeFollowerNode.cpp
+++ b/Engine/Plugins/Runtime/Metasound/Source/MetasoundStandardNodes/Private/MetasoundEnvelopeFollowerNode.cpp
@@ -24,15 +24,15 @@
 
 namespace Metasound
 {
-	namespace EnvelopeFollower
+	namespace EnvelopeFollowerVertexNames
 	{
-		static const TCHAR* InParamNameAudioInput = TEXT("In");
-		static const TCHAR* InParamNameAttackTime = TEXT("Attack Time");
-		static const TCHAR* InParamNameReleaseTime = TEXT("Release Time");
-		static const TCHAR* InParamNameFollowMode = TEXT("Peak Mode");
-		static const TCHAR* OutParamNameEnvelope = TEXT("Envelope");
+		METASOUND_PARAM(InParamAudioInput, "In", "Audio input.")
+		METASOUND_PARAM(InParamAttackTime, "Attack Time", "The attack time of the envelope follower.")
+		METASOUND_PARAM(InParamReleaseTime,"Release Time" , "The release time of the envelope follower.")
+		METASOUND_PARAM(InParamFollowMode, "Peak Mode", "The following-method of the envelope follower.")
+		METASOUND_PARAM(OutParamEnvelope, "Envelope", "The output envelope value of the audio signal.")
+		METASOUND_PARAM(OutputAudioEnvelope, "Audio Envelope", "The output envelope value of the audio signal (audio rate).");
 	}
-	METASOUND_PARAM(OutputAudioEnvelope, "Audio Envelope", "The output envelope value of the audio signal (audio rate).");
 
 	class FEnvelopeFollowerOperator : public TExecutableOperator<FEnvelopeFollowerOperator>
 	{
@@ -104,19 +104,23 @@ namespace Metasound
 
 	FDataReferenceCollection FEnvelopeFollowerOperator::GetInputs() const
 	{
+		using namespace EnvelopeFollowerVertexNames;
+
 		FDataReferenceCollection InputDataReferences;
-		InputDataReferences.AddDataReadReference(EnvelopeFollower::InParamNameAudioInput, AudioInput);
-		InputDataReferences.AddDataReadReference(EnvelopeFollower::InParamNameAttackTime, AttackTimeInput);
-		InputDataReferences.AddDataReadReference(EnvelopeFollower::InParamNameReleaseTime, ReleaseTimeInput);
-		InputDataReferences.AddDataReadReference(EnvelopeFollower::InParamNameFollowMode, FollowModeInput);
+		InputDataReferences.AddDataReadReference(METASOUND_GET_PARAM_NAME(InParamAudioInput), AudioInput);
+		InputDataReferences.AddDataReadReference(METASOUND_GET_PARAM_NAME(InParamAttackTime), AttackTimeInput);
+		InputDataReferences.AddDataReadReference(METASOUND_GET_PARAM_NAME(InParamReleaseTime), ReleaseTimeInput);
+		InputDataReferences.AddDataReadReference(METASOUND_GET_PARAM_NAME(InParamFollowMode), FollowModeInput);
 
 		return InputDataReferences;
 	}
 
 	FDataReferenceCollection FEnvelopeFollowerOperator::GetOutputs() const
 	{
+		using namespace EnvelopeFollowerVertexNames;
+
 		FDataReferenceCollection OutputDataReferences;
-		OutputDataReferences.AddDataReadReference(EnvelopeFollower::OutParamNameEnvelope, EnvelopeFloatOutput);
+		OutputDataReferences.AddDataReadReference(METASOUND_GET_PARAM_NAME(OutParamEnvelope), EnvelopeFloatOutput);
 		OutputDataReferences.AddDataReadReference(METASOUND_GET_PARAM_NAME(OutputAudioEnvelope), EnvelopeAudioOutput);
 		return OutputDataReferences;
 	}
@@ -170,16 +174,18 @@ namespace Metasound
 
 	const FVertexInterface& FEnvelopeFollowerOperator::GetVertexInterface()
 	{
+		using namespace EnvelopeFollowerVertexNames;
+
 		static const FVertexInterface Interface(
 			FInputVertexInterface(
-				TInputDataVertexModel<FAudioBuffer>(EnvelopeFollower::InParamNameAudioInput, METASOUND_LOCTEXT("AudioInputToolTT", "Audio input.")),
-				TInputDataVertexModel<FTime>(EnvelopeFollower::InParamNameAttackTime, METASOUND_LOCTEXT("AttackTimeTT", "The attack time of the envelope follower."), 0.01f),
-				TInputDataVertexModel<FTime>(EnvelopeFollower::InParamNameReleaseTime, METASOUND_LOCTEXT("ReleaseTimeTT", "The release time of the envelope follower."), 0.1f),
-				TInputDataVertexModel<FEnumEnvelopePeakMode>(EnvelopeFollower::InParamNameFollowMode, METASOUND_LOCTEXT("FollowModeTT", "The following-method of the envelope follower."))
+				TInputDataVertexModel<FAudioBuffer>(METASOUND_GET_PARAM_NAME_AND_METADATA(InParamAudioInput)),
+				TInputDataVertexModel<FTime>(METASOUND_GET_PARAM_NAME_AND_METADATA(InParamAttackTime), 0.01f),
+				TInputDataVertexModel<FTime>(METASOUND_GET_PARAM_NAME_AND_METADATA(InParamReleaseTime), 0.1f),
+				TInputDataVertexModel<FEnumEnvelopePeakMode>(METASOUND_GET_PARAM_NAME_AND_METADATA(InParamFollowMode))
 			),
 			FOutputVertexInterface(
-				TOutputDataVertexModel<float>(EnvelopeFollower::OutParamNameEnvelope, METASOUND_LOCTEXT("EnvelopeFollowerOutputTT", "The output envelope value of the audio signal.")),
-				TOutputDataVertexModel<FAudioBuffer>(METASOUND_GET_PARAM_NAME_AND_TT(OutputAudioEnvelope))
+				TOutputDataVertexModel<float>(METASOUND_GET_PARAM_NAME_AND_METADATA(OutParamEnvelope)),
+				TOutputDataVertexModel<FAudioBuffer>(METASOUND_GET_PARAM_NAME_AND_METADATA(OutputAudioEnvelope))
 			)
 		);
 
@@ -210,14 +216,16 @@ namespace Metasound
 
 	TUniquePtr<IOperator> FEnvelopeFollowerOperator::CreateOperator(const FCreateOperatorParams& InParams, FBuildErrorArray& OutErrors)
 	{
+		using namespace EnvelopeFollowerVertexNames;
+
 		const FEnvelopeFollowerNode& EnvelopeFollowerNode = static_cast<const FEnvelopeFollowerNode&>(InParams.Node);
 		const FDataReferenceCollection& InputCollection = InParams.InputDataReferences;
 		const FInputVertexInterface& InputInterface = GetVertexInterface().GetInputInterface();
 
-		FAudioBufferReadRef AudioIn = InputCollection.GetDataReadReferenceOrConstruct<FAudioBuffer>(EnvelopeFollower::InParamNameAudioInput, InParams.OperatorSettings);
-		FTimeReadRef AttackTime = InputCollection.GetDataReadReferenceOrConstructWithVertexDefault<FTime>(InputInterface, EnvelopeFollower::InParamNameAttackTime, InParams.OperatorSettings);
-		FTimeReadRef ReleaseTime = InputCollection.GetDataReadReferenceOrConstructWithVertexDefault<FTime>(InputInterface, EnvelopeFollower::InParamNameReleaseTime, InParams.OperatorSettings);
-		FEnvelopePeakModeReadRef EnvelopeModeIn = InputCollection.GetDataReadReferenceOrConstruct<FEnumEnvelopePeakMode>(EnvelopeFollower::InParamNameFollowMode);
+		FAudioBufferReadRef AudioIn = InputCollection.GetDataReadReferenceOrConstruct<FAudioBuffer>(METASOUND_GET_PARAM_NAME(InParamAudioInput), InParams.OperatorSettings);
+		FTimeReadRef AttackTime = InputCollection.GetDataReadReferenceOrConstructWithVertexDefault<FTime>(InputInterface, METASOUND_GET_PARAM_NAME(InParamAttackTime), InParams.OperatorSettings);
+		FTimeReadRef ReleaseTime = InputCollection.GetDataReadReferenceOrConstructWithVertexDefault<FTime>(InputInterface, METASOUND_GET_PARAM_NAME(InParamReleaseTime), InParams.OperatorSettings);
+		FEnvelopePeakModeReadRef EnvelopeModeIn = InputCollection.GetDataReadReferenceOrConstruct<FEnumEnvelopePeakMode>(METASOUND_GET_PARAM_NAME(InParamFollowMode));
 
 
 		return MakeUnique<FEnvelopeFollowerOperator>(InParams.OperatorSettings, AudioIn, AttackTime, ReleaseTime, EnvelopeModeIn);
diff --git a/Engine/Plugins/Runtime/Metasound/Source/MetasoundStandardNodes/Private/MetasoundITDPannerNode.cpp b/Engine/Plugins/Runtime/Metasound/Source/MetasoundStandardNodes/Private/MetasoundITDPannerNode.cpp
--- a/Engine/Plugins/Runtime/Metasound/Source/MetasoundStandardNodes/Private/MetasoundITDPannerNode.cpp
+++ b/Engine/Plugins/Runtime/Metasound/Source/MetasoundStandardNodes/Private/MetasoundITDPannerNode.cpp
@@ -4,6 +4,7 @@
 #include "MetasoundExecutableOperator.h"
 #include "MetasoundNodeRegistrationMacro.h"
 #include "MetasoundDataTypeRegistrationMacro.h"
+#include "MetasoundParamHelper.h"
 #include "MetasoundPrimitives.h"
 #include "MetasoundStandardNodesNames.h"
 #include "MetasoundTrigger.h"
@@ -22,77 +23,12 @@ namespace Metasound
 {
 	namespace ITDPannerVertexNames
 	{
-		const FVertexName& GetInputAudioName()
-		{
-			static FVertexName Name = TEXT("In");
-			return Name;
-		}
-
-		const FText& GetInputAudioDescription()
-		{
-			static FText Desc = METASOUND_LOCTEXT("ITDPannerNodeInDesc", "The input audio to spatialize.");
-			return Desc;
-		}
-
-		const FVertexName& GetInputPanAngleName()
-		{
-			static FVertexName Name = TEXT("Angle");
-			return Name;
-		}
-
-		const FText& GetInputPanAngleDescription()
-		{
-			static FText Desc = METASOUND_LOCTEXT("ITDPannerNodePanAngleDesc", "The sound source angle in degrees. 90 degrees is in front, 0 degrees is to the right, 270 degrees is behind, 180 degrees is to the left.");
-			return Desc;
-		}
-
-		const FVertexName& GetInputDistanceFactorName()
-		{
-			static FVertexName Name = TEXT("Distance Factor");
-			return Name;
-		}
-
-		const FText& GetInputDistanceFactorDescription()
-		{
-			static FText Desc = METASOUND_LOCTEXT("ITDPannerNodeDistanceFactorDesc", "The normalized distance factor (0.0 to 1.0) to use for ILD (Inter-aural level difference) calculations. 0.0 is near, 1.0 is far. The further away something is the less there is a difference in levels (gain) between the ears.");
-			return Desc;
-		}
-
-		const FVertexName& GetInputHeadWidthName()
-		{
-			static FVertexName Name = TEXT("Head Width");
-			return Name;
-		}
-
-		const FText& GetInputHeadWidthDescription()
-		{
-			static FText Desc = METASOUND_LOCTEXT("ITDPannerNodeHeadWidthDesc", "The width of the listener head to use for ITD calculations in centimeters.");
-			return Desc;
-		}
-
-		const FVertexName& GetOutputAudioLeftName()
-		{
-			static FVertexName Name = TEXT("Out Left");
-			return Name;
-		}
-
-		const FText& GetOutputAudioLeftDescription()
-		{
-			static FText Desc = METASOUND_LOCTEXT("ITDPannerNodeOutputLeftDescription", "Left channel audio output.");
-			return Desc;
-		}
-
-		const FVertexName& GetOutputAudioRightName()
-		{
-			static FVertexName Name = TEXT("Out Right");
-			return Name;
-		}
-
-		const FText& GetOutputAudioRightDescription()
-		{
-			static FText Desc = METASOUND_LOCTEXT("ITDPannerNodeOutputRightDescription", "Right channel audio output.");
-			return Desc;
-		}
+		METASOUND_PARAM(InputAudio, "In", "The input audio to spatialize.")
+		METASOUND_PARAM(InputPanAngle, "Angle", "The sound source angle in degrees. 90 degrees is in front, 0 degrees is to the right, 270 degrees is behind, 180 degrees is to the left.")
+		METASOUND_PARAM(InputDistanceFactor, "Distance Factor", "The normalized distance factor (0.0 to 1.0) to use for ILD (Inter-aural level difference) calculations. 0.0 is near, 1.0 is far. The further away something is the less there is a difference in levels (gain) between the ears.")
+		METASOUND_PARAM(InputHeadWidth, "Head Width", "The width of the listener head to use for ITD calculations in centimeters.")
+		METASOUND_PARAM(OutputAudioLeft, "Out Left", "Left channel audio output.")
+		METASOUND_PARAM(OutputAudioRight, "Out Right", "Right channel audio output.")
 	}
 
 	class FITDPannerOperator : public TExecutableOperator<FITDPannerOperator>
@@ -175,10 +111,10 @@ namespace Metasound
 		using namespace ITDPannerVertexNames;
 
 		FDataReferenceCollection InputDataReferences;
-		InputDataReferences.AddDataReadReference(GetInputAudioName(), AudioInput);
-		InputDataReferences.AddDataReadReference(GetInputPanAngleName(), PanningAngle);
-		InputDataReferences.AddDataReadReference(GetInputDistanceFactorName(), DistanceFactor);
-		InputDataReferences.AddDataReadReference(GetInputHeadWidthName(), HeadWidth);
+		InputDataReferences.AddDataReadReference(METASOUND_GET_PARAM_NAME(InputAudio), AudioInput);
+		InputDataReferences.AddDataReadReference(METASOUND_GET_PARAM_NAME(InputPanAngle), PanningAngle);
+		InputDataReferences.AddDataReadReference(METASOUND_GET_PARAM_NAME(InputDistanceFactor), DistanceFactor);
+		InputDataReferences.AddDataReadReference(METASOUND_GET_PARAM_NAME(InputHeadWidth), HeadWidth);
 
 		return InputDataReferences;
 	}
@@ -188,8 +124,8 @@ namespace Metasound
 		using namespace ITDPannerVertexNames;
 
 		FDataReferenceCollection OutputDataReferences;
-		OutputDataReferences.AddDataReadReference(GetOutputAudioLeftName(), AudioLeftOutput);
-		OutputDataReferences.AddDataReadReference(GetOutputAudioRightName(), AudioRightOutput);
+		OutputDataReferences.AddDataReadReference(METASOUND_GET_PARAM_NAME(OutputAudioLeft), AudioLeftOutput);
+		OutputDataReferences.AddDataReadReference(METASOUND_GET_PARAM_NAME(OutputAudioRight), AudioRightOutput);
 
 		return OutputDataReferences;
 	}
@@ -292,14 +228,14 @@ namespace Metasound
 
 		static const FVertexInterface Interface(
 			FInputVertexInterface(
-				TInputDataVertexModel<FAudioBuffer>(GetInputAudioName(), GetInputAudioDescription()),
-				TInputDataVertexModel<float>(GetInputPanAngleName(), GetInputPanAngleDescription(), 90.0f),
-				TInputDataVertexModel<float>(GetInputDistanceFactorName(), GetInputDistanceFactorDescription(), 0.0f),
-				TInputDataVertexModel<float>(GetInputHeadWidthName(), GetInputHeadWidthDescription(), 34.0f)
+				TInputDataVertexModel<FAudioBuffer>(METASOUND_GET_PARAM_NAME_AND_METADATA(InputAudio)),
+				TInputDataVertexModel<float>(METASOUND_GET_PARAM_NAME_AND_METADATA(InputPanAngle), 90.0f),
+				TInputDataVertexModel<float>(METASOUND_GET_PARAM_NAME_AND_METADATA(InputDistanceFactor), 0.0f),
+				TInputDataVertexModel<float>(METASOUND_GET_PARAM_NAME_AND_METADATA(InputHeadWidth), 34.0f)
 			),
 			FOutputVertexInterface(
-				TOutputDataVertexModel<FAudioBuffer>(GetOutputAudioLeftName(), GetOutputAudioLeftDescription()),
-				TOutputDataVertexModel<FAudioBuffer>(GetOutputAudioRightName(), GetOutputAudioRightDescription())
+				TOutputDataVertexModel<FAudioBuffer>(METASOUND_GET_PARAM_NAME_AND_METADATA(OutputAudioLeft)),
+				TOutputDataVertexModel<FAudioBuffer>(METASOUND_GET_PARAM_NAME_AND_METADATA(OutputAudioRight))
 			)
 		);
 
@@ -335,10 +271,10 @@ namespace Metasound
 
 		using namespace ITDPannerVertexNames;
 
-		FAudioBufferReadRef AudioIn = InputCollection.GetDataReadReferenceOrConstruct<FAudioBuffer>(GetInputAudioName(), InParams.OperatorSettings);
-		FFloatReadRef PanningAngle = InputCollection.GetDataReadReferenceOrConstructWithVertexDefault<float>(InputInterface, GetInputPanAngleName(), InParams.OperatorSettings);
-		FFloatReadRef DistanceFactor = InputCollection.GetDataReadReferenceOrConstructWithVertexDefault<float>(InputInterface, GetInputDistanceFactorName(), InParams.OperatorSettings);
-		FFloatReadRef HeadWidth = InputCollection.GetDataReadReferenceOrConstructWithVertexDefault<float>(InputInterface, GetInputHeadWidthName(), InParams.OperatorSettings);
+		FAudioBufferReadRef AudioIn = InputCollection.GetDataReadReferenceOrConstruct<FAudioBuffer>(METASOUND_GET_PARAM_NAME(InputAudio), InParams.OperatorSettings);
+		FFloatReadRef PanningAngle = InputCollection.GetDataReadReferenceOrConstructWithVertexDefault<float>(InputInterface, METASOUND_GET_PARAM_NAME(InputPanAngle), InParams.OperatorSettings);
+		FFloatReadRef DistanceFactor = InputCollection.GetDataReadReferenceOrConstructWithVertexDefault<float>(InputInterface, METASOUND_GET_PARAM_NAME(InputDistanceFactor), InParams.OperatorSettings);
+		FFloatReadRef HeadWidth = InputCollection.GetDataReadReferenceOrConstructWithVertexDefault<float>(InputInterface, METASOUND_GET_PARAM_NAME(InputHeadWidth), InParams.OperatorSettings);
 
 		return MakeUnique<FITDPannerOperator>(InParams.OperatorSettings, AudioIn, PanningAngle, DistanceFactor, HeadWidth);
 	}
diff --git a/Engine/Plugins/Runtime/Metasound/Source/MetasoundStandardNodes/Private/MetasoundInterpToNode.cpp b/Engine/Plugins/Runtime/Metasound/Source/MetasoundStandardNodes/Private/MetasoundInterpToNode.cpp
--- a/Engine/Plugins/Runtime/Metasound/Source/MetasoundStandardNodes/Private/MetasoundInterpToNode.cpp
+++ b/Engine/Plugins/Runtime/Metasound/Source/MetasoundStandardNodes/Private/MetasoundInterpToNode.cpp
@@ -5,6 +5,7 @@
 #include "Internationalization/Text.h"
 #include "MetasoundExecutableOperator.h"
 #include "MetasoundNodeRegistrationMacro.h"
+#include "MetasoundParamHelper.h"
 #include "MetasoundPrimitives.h"
 #include "MetasoundStandardNodesNames.h"
 #include "MetasoundTrigger.h"
@@ -15,11 +16,11 @@
 
 namespace Metasound
 {
-	namespace InterpTo
+	namespace InterpToVertexNames
 	{
-		static const TCHAR* InParamNameTarget = TEXT("Target");
-		static const TCHAR* InParamNameInterpTime = TEXT("Interp Time");
-		static const TCHAR* OutParamNameValue = TEXT("Value");
+		METASOUND_PARAM(InParamTarget, "Target", "Target value.")
+		METASOUND_PARAM(InParamInterpTime, "Interp Time", "The time to interpolate from the current value to the target value.")
+		METASOUND_PARAM(OutParamValue, "Value", "The current value.")
 	}
 
 	/** FInterpToNode
@@ -87,17 +88,21 @@ namespace Metasound
 
 	FDataReferenceCollection FInterpToOperator::GetInputs() const
 	{
+		using namespace InterpToVertexNames;
+
 		FDataReferenceCollection InputDataReferences;
-		InputDataReferences.AddDataReadReference(InterpTo::InParamNameTarget, FFloatReadRef(TargetValue));
-		InputDataReferences.AddDataReadReference(InterpTo::InParamNameInterpTime, FTimeReadRef(InterpTime));
+		InputDataReferences.AddDataReadReference(METASOUND_GET_PARAM_NAME(InParamTarget), FFloatReadRef(TargetValue));
+		InputDataReferences.AddDataReadReference(METASOUND_GET_PARAM_NAME(InParamInterpTime), FTimeReadRef(InterpTime));
 
 		return InputDataReferences;
 	}
 
 	FDataReferenceCollection FInterpToOperator::GetOutputs() const
 	{
+		using namespace InterpToVertexNames;
+
 		FDataReferenceCollection OutputDataReferences;
-		OutputDataReferences.AddDataReadReference(InterpTo::OutParamNameValue, FFloatReadRef(ValueOutput));
+		OutputDataReferences.AddDataReadReference(METASOUND_GET_PARAM_NAME(OutParamValue), FFloatReadRef(ValueOutput));
 		return OutputDataReferences;
 	}
 
@@ -124,13 +129,15 @@ namespace Metasound
 
 	const FVertexInterface& FInterpToOperator::GetVertexInterface()
 	{
+		using namespace InterpToVertexNames;
+
 		static const FVertexInterface Interface(
 			FInputVertexInterface(
-				TInputDataVertexModel<FTime>(InterpTo::InParamNameInterpTime, METASOUND_LOCTEXT("InterpTimeTooltip", "The time to interpolate from the current value to the target value."), 0.1f),
-				TInputDataVertexModel<float>(InterpTo::InParamNameTarget, METASOUND_LOCTEXT("TargetValueTooltip", "Target value."), 1.0f)
+				TInputDataVertexModel<FTime>(METASOUND_GET_PARAM_NAME_AND_METADATA(InParamInterpTime), 0.1f),
+				TInputDataVertexModel<float>(METASOUND_GET_PARAM_NAME_AND_METADATA(InParamTarget), 1.0f)
 			),
 			FOutputVertexInterface(
-				TOutputDataVertexModel<float>(InterpTo::OutParamNameValue, METASOUND_LOCTEXT("ValueTooltip", "The current value."))
+				TOutputDataVertexModel<float>(METASOUND_GET_PARAM_NAME_AND_METADATA(OutParamValue))
 			)
 		);
 
@@ -167,12 +174,14 @@ namespace Metasound
 
 	TUniquePtr<IOperator> FInterpToOperator::CreateOperator(const FCreateOperatorParams& InParams, FBuildErrorArray& OutErrors)
 	{
+		using namespace InterpToVertexNames; 
+
 		const FInterpToNode& InterpToNode = static_cast<const FInterpToNode&>(InParams.Node);
 		const FDataReferenceCollection& InputCollection = InParams.InputDataReferences;
 		const FInputVertexInterface& InputInterface = GetVertexInterface().GetInputInterface();
 
-		FFloatReadRef TargetValue = InputCollection.GetDataReadReferenceOrConstructWithVertexDefault<float>(InputInterface, InterpTo::InParamNameTarget, InParams.OperatorSettings);
-		FTimeReadRef InterpTime = InputCollection.GetDataReadReferenceOrConstructWithVertexDefault<FTime>(InputInterface, InterpTo::InParamNameInterpTime, InParams.OperatorSettings);
+		FFloatReadRef TargetValue = InputCollection.GetDataReadReferenceOrConstructWithVertexDefault<float>(InputInterface, METASOUND_GET_PARAM_NAME(InParamTarget), InParams.OperatorSettings);
+		FTimeReadRef InterpTime = InputCollection.GetDataReadReferenceOrConstructWithVertexDefault<FTime>(InputInterface, METASOUND_GET_PARAM_NAME(InParamInterpTime), InParams.OperatorSettings);
 
 		return MakeUnique<FInterpToOperator>(InParams.OperatorSettings, TargetValue, InterpTime);
 	}
diff --git a/Engine/Plugins/Runtime/Metasound/Source/MetasoundStandardNodes/Private/MetasoundMathNodes.cpp b/Engine/Plugins/Runtime/Metasound/Source/MetasoundStandardNodes/Private/MetasoundMathNodes.cpp
--- a/Engine/Plugins/Runtime/Metasound/Source/MetasoundStandardNodes/Private/MetasoundMathNodes.cpp
+++ b/Engine/Plugins/Runtime/Metasound/Source/MetasoundStandardNodes/Private/MetasoundMathNodes.cpp
@@ -12,6 +12,7 @@
 #include "MetasoundNodeRegistrationMacro.h"
 #include "MetasoundNodeInterface.h"
 #include "MetasoundOperatorInterface.h"
+#include "MetasoundParamHelper.h"
 #include "MetasoundPrimitives.h"
 #include "MetasoundStandardNodesCategories.h"
 #include "MetasoundStandardNodesNames.h"
@@ -52,9 +53,11 @@ namespace Metasound
 {
 	namespace MathOpNames
 	{
-		static const FVertexName PrimaryOperandName = TEXT("PrimaryOperand");
+		METASOUND_PARAM(PrimaryOperand, "Primary Operand", "Initial operand")
+		METASOUND_PARAM(AdditionalOperands, "Additional Operand", "Additional operand(s)")
 
-		static const FVertexName AdditionalOperandsName = TEXT("AdditionalOperands");
+		METASOUND_PARAM(OutMath, "Out", "Math operation result.")
+		METASOUND_PARAM(OutAudio, "Out", "Resulting buffer.")
 
 		static const TArray<FText> AddKeywords = { METASOUND_LOCTEXT("AddMathKeyword", "+") };
 		static const TArray<FText> SubtractKeywords = { METASOUND_LOCTEXT("SubtractMathKeyword", "-") };
@@ -238,20 +241,21 @@ namespace Metasound
 
 			static TUniquePtr<IOperator> CreateOperator(const FCreateOperatorParams& InParams, FBuildErrorArray& OutErrors)
 			{
+				using namespace MathOpNames;
+
 				const TMathOpNode& MathOpNode = static_cast<const TMathOpNode&>(InParams.Node);
 
 				const FInputVertexInterface& InputInterface = MathOpNode.GetVertexInterface().GetInputInterface();
 
 				FLiteral DefaultValue = FLiteral::CreateInvalid();
-				const FVertexName& PrimaryOperandName = MathOpNames::PrimaryOperandName;
-				if (InputInterface.Contains(PrimaryOperandName))
+				if (InputInterface.Contains(METASOUND_GET_PARAM_NAME(PrimaryOperand)))
 				{
-					DefaultValue = InputInterface[PrimaryOperandName].GetDefaultLiteral();
+					DefaultValue = InputInterface[METASOUND_GET_PARAM_NAME(PrimaryOperand)].GetDefaultLiteral();
 				}
-				TDataClassReadRef PrimaryOperand = InParams.InputDataReferences.GetDataReadReferenceOrConstruct<TDataClass>(MathOpNames::PrimaryOperandName, TMathOpClass::GetDefault(InParams.OperatorSettings, DefaultValue));
+				TDataClassReadRef PrimaryOperand = InParams.InputDataReferences.GetDataReadReferenceOrConstruct<TDataClass>(METASOUND_GET_PARAM_NAME(PrimaryOperand), TMathOpClass::GetDefault(InParams.OperatorSettings, DefaultValue));
 
 				// TODO: Support dynamic number of inputs
-				const FVertexName& OpName = MathOpNames::AdditionalOperandsName;
+				const FVertexName& OpName = METASOUND_GET_PARAM_NAME(AdditionalOperands);
 				if (InputInterface.Contains(OpName))
 				{
 					DefaultValue = InputInterface[OpName].GetDefaultLiteral();
@@ -273,16 +277,20 @@ namespace Metasound
 
 			virtual FDataReferenceCollection GetInputs() const override
 			{
+				using namespace MathOpNames; 
+
 				FDataReferenceCollection InputDataReferences;
-				InputDataReferences.AddDataReadReference(MathOpNames::PrimaryOperandName, PrimaryOperandRef);
-				InputDataReferences.AddDataReadReference(MathOpNames::AdditionalOperandsName, AdditionalOperandRefs[0]);
+				InputDataReferences.AddDataReadReference(METASOUND_GET_PARAM_NAME(PrimaryOperand), PrimaryOperandRef);
+				InputDataReferences.AddDataReadReference(METASOUND_GET_PARAM_NAME(AdditionalOperands), AdditionalOperandRefs[0]);
 				return InputDataReferences;
 			}
 
 			virtual FDataReferenceCollection GetOutputs() const override
 			{
+				using namespace MathOpNames;
+				
 				FDataReferenceCollection OutputDataReferences;
-				OutputDataReferences.AddDataReadReference(TEXT("Out"), TDataClassReadRef(ValueRef));
+				OutputDataReferences.AddDataReadReference(METASOUND_GET_PARAM_NAME(OutMath), TDataClassReadRef(ValueRef));
 
 				return OutputDataReferences;
 			}
@@ -307,13 +315,33 @@ namespace Metasound
 
 		static const FVertexInterface& GetVertexInterface()
 		{
+			using namespace MathOpNames;
+
+			static const FDataVertexMetadata PrimaryOperandMetadata
+			{
+				  METASOUND_LOCTEXT("MathOpAddNode_InitialTooltip", "Initial addend.") // description
+				, METASOUND_GET_PARAM_DISPLAYNAME(PrimaryOperand) // display name
+			};
+
+			static const FDataVertexMetadata AdditionalAddendsMetadata
+			{
+				  METASOUND_LOCTEXT("MathOpAddNode_AddendsTooltip", "Additional addend(s).") // description
+				, METASOUND_GET_PARAM_DISPLAYNAME(AdditionalOperands) // display name
+			};
+
+			static const FDataVertexMetadata OutputMetadata
+			{
+				  METASOUND_LOCTEXT("MathOpAddNode_OutTooltip", "Math operation result") // description
+				, METASOUND_GET_PARAM_DISPLAYNAME(OutMath) // display name
+			};
+
 			static const FVertexInterface Interface(
 				FInputVertexInterface(
-					TInputDataVertexModel<TDataClass>(MathOpNames::PrimaryOperandName, METASOUND_LOCTEXT("MathOpAddNode_InitialTooltip", "Initial addend."), static_cast<TDataClass>(0)),
-					TInputDataVertexModel<TOperandDataClass>(MathOpNames::AdditionalOperandsName, METASOUND_LOCTEXT("MathOpAddNode_AddendsTooltip", "Additional addend(s)."), static_cast<TOperandDataClass>(0))
+					TInputDataVertexModel<TDataClass>(METASOUND_GET_PARAM_NAME(PrimaryOperand), PrimaryOperandMetadata, static_cast<TDataClass>(0)),
+					TInputDataVertexModel<TOperandDataClass>(METASOUND_GET_PARAM_NAME(AdditionalOperands), AdditionalAddendsMetadata, static_cast<TOperandDataClass>(0))
 				),
 				FOutputVertexInterface(
-					TOutputDataVertexModel<TDataClass>(TEXT("Out"), METASOUND_LOCTEXT("MathOpAddNode_OutTooltip", "Math operation result"))
+					TOutputDataVertexModel<TDataClass>(METASOUND_GET_PARAM_NAME(OutMath), OutputMetadata)
 				)
 			);
 
@@ -356,13 +384,33 @@ namespace Metasound
 
 		static const FVertexInterface& GetVertexInterface()
 		{
+			using namespace MathOpNames; 
+
+			static const FDataVertexMetadata PrimaryOperandMetadata
+			{
+				  METASOUND_LOCTEXT("MathOpSubtractNode_MinuendTooltip", "Minuend.") // description
+				, METASOUND_GET_PARAM_DISPLAYNAME(PrimaryOperand) // display name
+			};
+
+			static const FDataVertexMetadata AdditionalOperandsMetadata
+			{
+				  METASOUND_LOCTEXT("MathOpSubtractNode_SubtrahendsTooltip", "Subtrahend(s).") // description
+				, METASOUND_GET_PARAM_DISPLAYNAME(AdditionalOperands) // display name
+			};
+
+			static const FDataVertexMetadata OutputMetadata
+			{
+				  METASOUND_LOCTEXT("MathOpSubtractNode_OutTooltip", "Subtraction result") // description
+				, METASOUND_GET_PARAM_DISPLAYNAME(OutMath) // display name
+			};
+
 			static const FVertexInterface Interface(
 				FInputVertexInterface(
-					TInputDataVertexModel<TDataClass>(MathOpNames::PrimaryOperandName, METASOUND_LOCTEXT("MathOpSubtractNode_MinuendTooltip", "Minuend."), static_cast<TDataClass>(0)),
-					TInputDataVertexModel<TOperandDataClass>(MathOpNames::AdditionalOperandsName, METASOUND_LOCTEXT("MathOpSubtractNode_SubtrahendsTooltip", "Subtrahend(s)."), static_cast<TOperandDataClass>(0))
+					TInputDataVertexModel<TDataClass>(METASOUND_GET_PARAM_NAME(PrimaryOperand), PrimaryOperandMetadata, static_cast<TDataClass>(0)),
+					TInputDataVertexModel<TOperandDataClass>(METASOUND_GET_PARAM_NAME(AdditionalOperands), AdditionalOperandsMetadata, static_cast<TOperandDataClass>(0))
 				),
 				FOutputVertexInterface(
-					TOutputDataVertexModel<TDataClass>(TEXT("Out"), METASOUND_LOCTEXT("MathOpSubtractNode_OutTooltip", "Subtraction result"))
+					TOutputDataVertexModel<TDataClass>(METASOUND_GET_PARAM_NAME(OutMath), OutputMetadata)
 				)
 			);
 
@@ -405,13 +453,33 @@ namespace Metasound
 
 		static const FVertexInterface& GetVertexInterface()
 		{
+			using namespace MathOpNames;
+
+			static const FDataVertexMetadata PrimaryOperandMetadata
+			{
+				  METASOUND_LOCTEXT("MathOpMultiplyNode_InitMultiplicandTooltip", "Initial multiplicand.") // description
+				, METASOUND_GET_PARAM_DISPLAYNAME(PrimaryOperand) // display name
+			};
+
+			static const FDataVertexMetadata AdditionalOperandsMetadata
+			{
+				 METASOUND_LOCTEXT("MathOpMultiplyNode_MultiplicandsTooltip", "Additional multiplicand(s).") // description
+				, METASOUND_GET_PARAM_DISPLAYNAME(AdditionalOperands) // display name
+			};
+
+			static const FDataVertexMetadata OutputMetadata
+			{
+				  METASOUND_LOCTEXT("MathOpMultiplyNode_ResultTooltip", "Multiplication result") // description
+				, METASOUND_GET_PARAM_DISPLAYNAME(OutMath) // display name
+			};
+
 			static const FVertexInterface Interface(
 				FInputVertexInterface(
-					TInputDataVertexModel<TDataClass>(MathOpNames::PrimaryOperandName, METASOUND_LOCTEXT("MathOpMultiplyNode_InitMultiplicandTooltip", "Initial multiplicand."), static_cast<TDataClass>(1)),
-					TInputDataVertexModel<TOperandDataClass>(MathOpNames::AdditionalOperandsName, METASOUND_LOCTEXT("MathOpMultiplyNode_MultiplicandsTooltip", "Additional multiplicand(s)."), static_cast<TOperandDataClass>(1))
+					TInputDataVertexModel<TDataClass>(METASOUND_GET_PARAM_NAME(PrimaryOperand), PrimaryOperandMetadata, static_cast<TDataClass>(1)),
+					TInputDataVertexModel<TOperandDataClass>(METASOUND_GET_PARAM_NAME(AdditionalOperands), AdditionalOperandsMetadata, static_cast<TOperandDataClass>(1))
 				),
 				FOutputVertexInterface(
-					TOutputDataVertexModel<TDataClass>(TEXT("Out"), METASOUND_LOCTEXT("MathOpMultiplyNode_ResultTooltip", "Multiplication result"))
+					TOutputDataVertexModel<TDataClass>(METASOUND_GET_PARAM_NAME(OutMath), OutputMetadata)
 				)
 			);
 
@@ -454,13 +522,33 @@ namespace Metasound
 
 		static const FVertexInterface& GetVertexInterface()
 		{
+			using namespace MathOpNames;
+
+			static const FDataVertexMetadata PrimaryOperandMetadata
+			{
+				  METASOUND_LOCTEXT("MathOpDivideNode_DividendTooltip", "Dividend.") // description
+				, METASOUND_GET_PARAM_DISPLAYNAME(PrimaryOperand) // display name
+			};
+
+			static const FDataVertexMetadata AdditionalOperandsMetadata
+			{
+				  METASOUND_LOCTEXT("MathOpDivideNode_DivisorsTooltip", "Divisor(s).") // description
+				, METASOUND_GET_PARAM_DISPLAYNAME(AdditionalOperands) // display name
+			};
+
+			static const FDataVertexMetadata OutputMetadata
+			{
+				  METASOUND_LOCTEXT("MathOpDivideNode_OutTooltip", "Division result") // description
+				, METASOUND_GET_PARAM_DISPLAYNAME(OutMath) // display name
+			};
+
 			static const FVertexInterface Interface(
 				FInputVertexInterface(
-					TInputDataVertexModel<TDataClass>(MathOpNames::PrimaryOperandName, METASOUND_LOCTEXT("MathOpDivideNode_DividendTooltip", "Dividend."), static_cast<TDataClass>(1)),
-					TInputDataVertexModel<TOperandDataClass>(MathOpNames::AdditionalOperandsName, METASOUND_LOCTEXT("MathOpDivideNode_DivisorsTooltip", "Divisor(s)."), static_cast<TOperandDataClass>(1))
+					TInputDataVertexModel<TDataClass>(METASOUND_GET_PARAM_NAME(PrimaryOperand), PrimaryOperandMetadata, static_cast<TDataClass>(1)),
+					TInputDataVertexModel<TOperandDataClass>(METASOUND_GET_PARAM_NAME(AdditionalOperands), AdditionalOperandsMetadata, static_cast<TOperandDataClass>(1))
 				),
 				FOutputVertexInterface(
-					TOutputDataVertexModel<TDataClass>(TEXT("Out"), METASOUND_LOCTEXT("MathOpDivideNode_OutTooltip", "Division result"))
+					TOutputDataVertexModel<TDataClass>(METASOUND_GET_PARAM_NAME(OutMath), OutputMetadata)
 				)
 			);
 
@@ -510,13 +598,33 @@ namespace Metasound
 
 		static const FVertexInterface& GetVertexInterface()
 		{
+			using namespace MathOpNames;
+
+			static const FDataVertexMetadata PrimaryOperandMetadata
+			{
+				  METASOUND_LOCTEXT("MathOpModuloNode_DividendTooltip", "Dividend.") // description
+				, METASOUND_GET_PARAM_DISPLAYNAME(PrimaryOperand) // display name
+			};
+
+			static const FDataVertexMetadata AdditionalOperandsMetadata
+			{
+				  METASOUND_LOCTEXT("MathOpModuloNode_DivisorsTooltip", "Divisor(s).") // description
+				, METASOUND_GET_PARAM_DISPLAYNAME(AdditionalOperands) // display name
+			};
+
+			static const FDataVertexMetadata OutputMetadata
+			{
+				  METASOUND_LOCTEXT("MathOpModuloNode_OutTooltip", "Modulo result") // description
+				, METASOUND_GET_PARAM_DISPLAYNAME(OutMath) // display name
+			};
+
 			static const FVertexInterface Interface(
 				FInputVertexInterface(
-					TInputDataVertexModel<TDataClass>(MathOpNames::PrimaryOperandName, METASOUND_LOCTEXT("MathOpModuloNode_DividendTooltip", "Dividend."), static_cast<TDataClass>(1)),
-					TInputDataVertexModel<TOperandDataClass>(MathOpNames::AdditionalOperandsName, METASOUND_LOCTEXT("MathOpModuloNode_DivisorsTooltip", "Divisor(s)."), static_cast<TOperandDataClass>(1))
+					TInputDataVertexModel<TDataClass>(METASOUND_GET_PARAM_NAME(PrimaryOperand), PrimaryOperandMetadata, static_cast<TDataClass>(1)),
+					TInputDataVertexModel<TOperandDataClass>(METASOUND_GET_PARAM_NAME(AdditionalOperands), AdditionalOperandsMetadata, static_cast<TOperandDataClass>(1))
 				),
 				FOutputVertexInterface(
-					TOutputDataVertexModel<TDataClass>(TEXT("Out"), METASOUND_LOCTEXT("MathOpModuloNode_OutTooltip", "Modulo result"))
+					TOutputDataVertexModel<TDataClass>(METASOUND_GET_PARAM_NAME(OutMath), OutputMetadata)
 				)
 			);
 
@@ -568,13 +676,33 @@ namespace Metasound
 
 		static const FVertexInterface& GetVertexInterface()
 		{
+			using namespace MathOpNames;
+
+			static const FDataVertexMetadata PrimaryOperandMetadata
+			{
+				  METASOUND_LOCTEXT("MathOpPowerNode_Base", "The base of the power") // description
+				, METASOUND_GET_PARAM_DISPLAYNAME(PrimaryOperand) // display name
+			};
+
+			static const FDataVertexMetadata AdditionalOperandsMetadata
+			{
+				  METASOUND_LOCTEXT("MathOpPowerNode_Exponent", "The exponent to take the base to the power of") // description
+				, METASOUND_GET_PARAM_DISPLAYNAME(AdditionalOperands) // display name
+					  };
+
+			static const FDataVertexMetadata OutputMetadata
+			{
+				  METASOUND_LOCTEXT("MathOpPowerNode_Result", "Returns Base to the Exponent power") // description
+				, METASOUND_GET_PARAM_DISPLAYNAME(OutMath) // display name
+			};
+
 			static const FVertexInterface DefaultInterface(
 				FInputVertexInterface(
-					TInputDataVertexModel<float>(MathOpNames::PrimaryOperandName, METASOUND_LOCTEXT("MathOpPowerNode_Base", "The base of the power")),
-					TInputDataVertexModel<float>(MathOpNames::AdditionalOperandsName, METASOUND_LOCTEXT("MathOpPowerNode_Exponent", "The exponent to take the base to the power of"))
+					TInputDataVertexModel<float>(METASOUND_GET_PARAM_NAME(PrimaryOperand), PrimaryOperandMetadata),
+					TInputDataVertexModel<float>(METASOUND_GET_PARAM_NAME(AdditionalOperands), AdditionalOperandsMetadata)
 				),
 				FOutputVertexInterface(
-					TOutputDataVertexModel<float>(TEXT("Out"), METASOUND_LOCTEXT("MathOpPowerNode_Result", "Returns Base to the Exponent power"))
+					TOutputDataVertexModel<float>(METASOUND_GET_PARAM_NAME(OutMath), OutputMetadata)
 				)
 			);
 
@@ -620,13 +748,33 @@ namespace Metasound
 
 		static const FVertexInterface& GetVertexInterface()
 		{
+			using namespace MathOpNames;
+
+			static const FDataVertexMetadata PrimaryOperandMetadata
+			{
+				   METASOUND_LOCTEXT("MathOpLogarithmNode_Base", "The base of the logarithm") // description
+				, METASOUND_GET_PARAM_DISPLAYNAME(PrimaryOperand) // display name
+			};
+
+			static const FDataVertexMetadata AdditionalOperandsMetadata
+			{
+				  METASOUND_LOCTEXT("MathOpLogarithmNode_Value", "The value to find the logarithm of") // description
+				, METASOUND_GET_PARAM_DISPLAYNAME(AdditionalOperands) // display name
+			};
+
+			static const FDataVertexMetadata OutputMetadata
+			{
+				  METASOUND_LOCTEXT("MathOpLogarithmNode_Result", "The logarithm of the inputted value") // description
+				, METASOUND_GET_PARAM_DISPLAYNAME(OutMath) // display name
+			};
+
 			static const FVertexInterface DefaultInterface(
 				FInputVertexInterface(
-					TInputDataVertexModel<float>(MathOpNames::PrimaryOperandName, METASOUND_LOCTEXT("MathOpLogarithmNode_Base", "The base of the logarithm")),
-					TInputDataVertexModel<float>(MathOpNames::AdditionalOperandsName, METASOUND_LOCTEXT("MathOpLogarithmNode_Value", "The value to find the logarithm of"))
+					TInputDataVertexModel<float>(METASOUND_GET_PARAM_NAME(PrimaryOperand), PrimaryOperandMetadata),
+					TInputDataVertexModel<float>(METASOUND_GET_PARAM_NAME(AdditionalOperands), AdditionalOperandsMetadata)
 				),
 				FOutputVertexInterface(
-					TOutputDataVertexModel<float>(TEXT("Out"), METASOUND_LOCTEXT("MathOpLogarithmNode_Result", "The logarithm of the inputted value"))
+					TOutputDataVertexModel<float>(METASOUND_GET_PARAM_NAME(OutMath), OutputMetadata)
 				)
 			);
 
@@ -669,13 +817,27 @@ namespace Metasound
 	public:
 		static const FVertexInterface& GetVertexInterface()
 		{
+			using namespace MathOpNames; 
+
+			static const FDataVertexMetadata PrimaryOperandMetadata
+			{
+				  METASOUND_LOCTEXT("MathOpAddAudioBufferNameTooltip", "First addend.") // description
+				, METASOUND_GET_PARAM_DISPLAYNAME(PrimaryOperand) // display name
+			};
+
+			static const FDataVertexMetadata AdditionalOperandsMetadata
+			{
+				  METASOUND_LOCTEXT("MathOpAddAudioBufferAdditionalTooltip", "Additional addends.") // description
+				, METASOUND_GET_PARAM_DISPLAYNAME(AdditionalOperands) // display name
+			};
+
 			static const FVertexInterface Interface(
 				FInputVertexInterface(
-					TInputDataVertexModel<FAudioBuffer>(MathOpNames::PrimaryOperandName, METASOUND_LOCTEXT("MathOpAddAudioBufferNameTooltip", "First addend.")),
-					TInputDataVertexModel<FAudioBuffer>(MathOpNames::AdditionalOperandsName, METASOUND_LOCTEXT("MathOpAddAudioBufferAdditionalTooltip", "Additional addends."))
+					TInputDataVertexModel<FAudioBuffer>(METASOUND_GET_PARAM_NAME(PrimaryOperand), PrimaryOperandMetadata),
+					TInputDataVertexModel<FAudioBuffer>(METASOUND_GET_PARAM_NAME(AdditionalOperands), AdditionalOperandsMetadata)
 				),
 				FOutputVertexInterface(
-					TOutputDataVertexModel<FAudioBuffer>(TEXT("Out"), METASOUND_LOCTEXT("MathOpOutTooltip", "Math operation result"))
+					TOutputDataVertexModel<FAudioBuffer>(METASOUND_GET_PARAM_NAME_AND_METADATA(OutAudio))
 				)
 			);
 
@@ -746,13 +908,33 @@ namespace Metasound
 	public:
 		static const FVertexInterface& GetVertexInterface()
 		{
+			using namespace MathOpNames;
+
+			static const FDataVertexMetadata PrimaryOperandMetadata
+			{
+				  METASOUND_LOCTEXT("MathOpAddTimeTooltip", "First addend.") // description
+				, METASOUND_GET_PARAM_DISPLAYNAME(PrimaryOperand) // display name
+			};
+
+			static const FDataVertexMetadata AdditionalOperandsMetadata
+			{
+				  METASOUND_LOCTEXT("MathOpAddTimeAdditionalalTooltip", "Additional addends.") // description
+				, METASOUND_GET_PARAM_DISPLAYNAME(AdditionalOperands) // display name
+			};
+
+			static const FDataVertexMetadata OutputMetadata
+			{
+				  METASOUND_LOCTEXT("MathOpOutTooltip", "Math operation result") // description
+				, METASOUND_GET_PARAM_DISPLAYNAME(OutMath) // display name
+			};
+
 			static const FVertexInterface Interface(
 				FInputVertexInterface(
-					TInputDataVertexModel<FTime>(MathOpNames::PrimaryOperandName, METASOUND_LOCTEXT("MathOpAddTimeTooltip", "First addend."), 0.0f),
-					TInputDataVertexModel<FTime>(MathOpNames::AdditionalOperandsName, METASOUND_LOCTEXT("MathOpAddTimeAdditionalalTooltip", "Additional addends."), 0.0f)
+					TInputDataVertexModel<FTime>(METASOUND_GET_PARAM_NAME(PrimaryOperand), PrimaryOperandMetadata, 0.0f),
+					TInputDataVertexModel<FTime>(METASOUND_GET_PARAM_NAME(AdditionalOperands), AdditionalOperandsMetadata, 0.0f)
 				),
 				FOutputVertexInterface(
-					TOutputDataVertexModel<FTime>(TEXT("Out"), METASOUND_LOCTEXT("MathOpOutTooltip", "Math operation result"))
+					TOutputDataVertexModel<FTime>(METASOUND_GET_PARAM_NAME(OutMath), OutputMetadata)
 				)
 			);
 
@@ -790,13 +972,27 @@ namespace Metasound
 	public:
 		static const FVertexInterface& GetVertexInterface()
 		{
+			using namespace MathOpNames; 
+
+			static const FDataVertexMetadata PrimaryOperandMetadata
+			{
+				  METASOUND_LOCTEXT("MathOpAddAudioTooltip", "Audio Buffer to add offset(s) to.") // description
+				, METASOUND_GET_PARAM_DISPLAYNAME(PrimaryOperand) // display name
+			};
+
+			static const FDataVertexMetadata AdditionalOperandsMetadata
+			{
+				  METASOUND_LOCTEXT("MathOpAddAdditionalTooltip", "Float addends of which to offset buffer samples.") // description
+				, METASOUND_GET_PARAM_DISPLAYNAME(AdditionalOperands) // display name
+			};
+
 			static const FVertexInterface Interface(
 				FInputVertexInterface(
-					TInputDataVertexModel<FAudioBuffer>(MathOpNames::PrimaryOperandName, METASOUND_LOCTEXT("MathOpAddAudioTooltip", "Audio Buffer to add offset(s) to.")),
-					TInputDataVertexModel<float>(MathOpNames::AdditionalOperandsName, METASOUND_LOCTEXT("MathOpAddAdditionalTooltip", "Float addends of which to offset buffer samples."), 0.0f)
+					TInputDataVertexModel<FAudioBuffer>(METASOUND_GET_PARAM_NAME(PrimaryOperand), PrimaryOperandMetadata),
+					TInputDataVertexModel<float>(METASOUND_GET_PARAM_NAME(AdditionalOperands), AdditionalOperandsMetadata, 0.0f)
 				),
 				FOutputVertexInterface(
-					TOutputDataVertexModel<FAudioBuffer>(TEXT("Out"), METASOUND_LOCTEXT("MathOpAudioFloatAddOutTooltip", "Resulting buffer"))
+					TOutputDataVertexModel<FAudioBuffer>(METASOUND_GET_PARAM_NAME_AND_METADATA(OutAudio))
 				)
 			);
 
@@ -838,13 +1034,27 @@ namespace Metasound
 	public:
 		static const FVertexInterface& GetVertexInterface()
 		{
+			using namespace MathOpNames;
+
+			static const FDataVertexMetadata PrimaryOperandMetadata
+			{
+				  METASOUND_LOCTEXT("MathOpBuffersMinuendTooltip", "Initial buffer to act as minuend.") // description
+				, METASOUND_GET_PARAM_DISPLAYNAME(PrimaryOperand) // display name
+			};
+
+			static const FDataVertexMetadata AdditionalAddendsMetadata
+			{
+				 METASOUND_LOCTEXT("MathOpSubtractBuffersSubtrahendsTooltip", "Additional buffers to act as subtrahend(s).") // description
+				, METASOUND_GET_PARAM_DISPLAYNAME(AdditionalOperands) // display name
+			};
+
 			static const FVertexInterface Interface(
 				FInputVertexInterface(
-					TInputDataVertexModel<FAudioBuffer>(MathOpNames::PrimaryOperandName, METASOUND_LOCTEXT("MathOpBuffersMinuendTooltip", "Initial buffer to act as minuend.")),
-					TInputDataVertexModel<FAudioBuffer>(MathOpNames::AdditionalOperandsName, METASOUND_LOCTEXT("MathOpSubtractBuffersSubtrahendsTooltip", "Additional buffers to act as subtrahend(s)."))
+					TInputDataVertexModel<FAudioBuffer>(METASOUND_GET_PARAM_NAME(PrimaryOperand), PrimaryOperandMetadata),
+					TInputDataVertexModel<FAudioBuffer>(METASOUND_GET_PARAM_NAME(AdditionalOperands), AdditionalAddendsMetadata)
 				),
 				FOutputVertexInterface(
-					TOutputDataVertexModel<FAudioBuffer>(TEXT("Out"), METASOUND_LOCTEXT("MathOpSubtractBuffersOutTooltip", "Resulting buffer"))
+					TOutputDataVertexModel<FAudioBuffer>(METASOUND_GET_PARAM_NAME_AND_METADATA(OutAudio))
 				)
 			);
 
@@ -913,13 +1123,33 @@ namespace Metasound
 	public:
 		static const FVertexInterface& GetVertexInterface()
 		{
+			using namespace MathOpNames;
+			
+			static const FDataVertexMetadata PrimaryOperandMetadata
+			{
+				 METASOUND_LOCTEXT("MathOpTimeNode_MinuendTooltip", "Time minuend.") // description
+				, METASOUND_GET_PARAM_DISPLAYNAME(PrimaryOperand) // display name
+			};
+
+			static const FDataVertexMetadata AdditionalAddendsMetadata
+			{
+				  METASOUND_LOCTEXT("MathOpTimeNode_SubtrahendsTooltip", "Time subtrahends.") // description
+				, METASOUND_GET_PARAM_DISPLAYNAME(AdditionalOperands) // display name
+			};
+
+			static const FDataVertexMetadata OutputMetadata
+			{
+				  METASOUND_LOCTEXT("MathOpTimeNode_SubtractOutTooltip", "Resulting time value") // description
+				, METASOUND_GET_PARAM_DISPLAYNAME(OutMath) // display name
+			};
+
 			static const FVertexInterface Interface(
 				FInputVertexInterface(
-					TInputDataVertexModel<FTime>(MathOpNames::PrimaryOperandName, METASOUND_LOCTEXT("MathOpTimeNode_MinuendTooltip", "Time minuend."), 0.0f),
-					TInputDataVertexModel<FTime>(MathOpNames::AdditionalOperandsName, METASOUND_LOCTEXT("MathOpTimeNode_SubtrahendsTooltip", "Time subtrahends."), 0.0f)
+					TInputDataVertexModel<FTime>(METASOUND_GET_PARAM_NAME(PrimaryOperand), PrimaryOperandMetadata, 0.0f),
+					TInputDataVertexModel<FTime>(METASOUND_GET_PARAM_NAME(AdditionalOperands), AdditionalAddendsMetadata, 0.0f)
 				),
 				FOutputVertexInterface(
-					TOutputDataVertexModel<FTime>(TEXT("Out"), METASOUND_LOCTEXT("MathOpTimeNode_SubtractOutTooltip", "Resulting time value"))
+					TOutputDataVertexModel<FTime>(METASOUND_GET_PARAM_NAME(OutMath), OutputMetadata)
 				)
 			);
 
@@ -957,13 +1187,27 @@ namespace Metasound
 	public:
 		static const FVertexInterface& GetVertexInterface()
 		{
+			using namespace MathOpNames;
+
+			static const FDataVertexMetadata PrimaryOperandMetadata
+			{
+				  METASOUND_LOCTEXT("MathOpAudioMultiplyNode_InitMultiplicandTooltip", "Initial audio to multiply.") // description
+				, METASOUND_GET_PARAM_DISPLAYNAME(PrimaryOperand) // display name
+			};
+
+			static const FDataVertexMetadata AdditionalAddendsMetadata
+			{
+				 METASOUND_LOCTEXT("MathOpAudioMultiplyNode_MultiplicandsTooltip", "Additional audio to multiply sample-by-sample.") // description
+				, METASOUND_GET_PARAM_DISPLAYNAME(AdditionalOperands) // display name
+			};
+
 			static const FVertexInterface Interface(
 				FInputVertexInterface(
-					TInputDataVertexModel<FAudioBuffer>(MathOpNames::PrimaryOperandName, METASOUND_LOCTEXT("MathOpAudioMultiplyNode_InitMultiplicandTooltip", "Initial audio to multiply.")),
-					TInputDataVertexModel<FAudioBuffer>(MathOpNames::AdditionalOperandsName, METASOUND_LOCTEXT("MathOpAudioMultiplyNode_MultiplicandsTooltip", "Additional audio to multiply sample-by-sample."))
+					TInputDataVertexModel<FAudioBuffer>(METASOUND_GET_PARAM_NAME(PrimaryOperand), PrimaryOperandMetadata),
+					TInputDataVertexModel<FAudioBuffer>(METASOUND_GET_PARAM_NAME(AdditionalOperands), AdditionalAddendsMetadata)
 				),
 				FOutputVertexInterface(
-					TOutputDataVertexModel<FAudioBuffer>(TEXT("Out"), METASOUND_LOCTEXT("MathOpAudioMultiplyNode_OutTooltip", "Resulting buffer"))
+					TOutputDataVertexModel<FAudioBuffer>(METASOUND_GET_PARAM_NAME_AND_METADATA(OutAudio))
 				)
 			);
 
@@ -1035,13 +1279,27 @@ namespace Metasound
 	public:
 		static const FVertexInterface& GetVertexInterface()
 		{
+			using namespace MathOpNames; 
+
+			static const FDataVertexMetadata PrimaryOperandMetadata
+			{
+				  METASOUND_LOCTEXT("MathOpAudioFloatMultiplyNode_FloatTooltip", "Audio multiplicand.") // description
+				, METASOUND_GET_PARAM_DISPLAYNAME(PrimaryOperand) // display name
+			};
+
+			static const FDataVertexMetadata AdditionalOperandsMetadata
+			{
+				  METASOUND_LOCTEXT("MathOpAudioFloatMultiplyNode_MultiplicandTooltip", "Float multiplicand to apply sample-by-sample to audio. Interpolates over buffer size on value change.") // description
+				, METASOUND_GET_PARAM_DISPLAYNAME(AdditionalOperands) // display name
+			};
+
 			static const FVertexInterface Interface(
 				FInputVertexInterface(
-					TInputDataVertexModel<FAudioBuffer>(MathOpNames::PrimaryOperandName, METASOUND_LOCTEXT("MathOpAudioFloatMultiplyNode_FloatTooltip", "Audio multiplicand.")),
-					TInputDataVertexModel<float>(MathOpNames::AdditionalOperandsName, METASOUND_LOCTEXT("MathOpAudioFloatMultiplyNode_MultiplicandTooltip", "Float multiplicand to apply sample-by-sample to audio. Interpolates over buffer size on value change."), 1.0f)
+					TInputDataVertexModel<FAudioBuffer>(METASOUND_GET_PARAM_NAME(PrimaryOperand), PrimaryOperandMetadata),
+					TInputDataVertexModel<float>(METASOUND_GET_PARAM_NAME(AdditionalOperands), AdditionalOperandsMetadata, 1.0f)
 				),
 				FOutputVertexInterface(
-					TOutputDataVertexModel<FAudioBuffer>(TEXT("Out"), METASOUND_LOCTEXT("MathOpAudioFloatMultiplyNode_OutTooltip", "Resulting buffer"))
+					TOutputDataVertexModel<FAudioBuffer>(METASOUND_GET_PARAM_NAME_AND_METADATA(OutAudio))
 				)
 			);
 
@@ -1099,13 +1357,33 @@ namespace Metasound
 	public:
 		static const FVertexInterface& GetVertexInterface()
 		{
+			using namespace MathOpNames;
+			
+			static const FDataVertexMetadata PrimaryOperandMetadata
+			{
+				  METASOUND_LOCTEXT("MathOpTimeFloatMultiplyNode_MultiplyFloatTooltip", "Time multiplicand.") // description
+				, METASOUND_GET_PARAM_DISPLAYNAME(PrimaryOperand) // display name
+			};
+
+			static const FDataVertexMetadata AdditionalOperandsMetadata
+			{
+				  METASOUND_LOCTEXT("MathOpTimeFloatMultiplyNode_MultiplicandsTooltip", "Float multiplicand(s).") // description
+				, METASOUND_GET_PARAM_DISPLAYNAME(AdditionalOperands) // display name
+			};
+
+			static const FDataVertexMetadata OutputMetadata
+			{
+				  METASOUND_LOCTEXT("MathOpTimeFloatMultiplyNode_OutTooltip", "Time float multiplication output") // description
+				, METASOUND_GET_PARAM_DISPLAYNAME(OutMath) // display name
+			};
+
 			static const FVertexInterface Interface(
 				FInputVertexInterface(
-					TInputDataVertexModel<FTime>(MathOpNames::PrimaryOperandName, METASOUND_LOCTEXT("MathOpTimeFloatMultiplyNode_MultiplyFloatTooltip", "Time multiplicand."), 1.0f),
-					TInputDataVertexModel<float>(MathOpNames::AdditionalOperandsName, METASOUND_LOCTEXT("MathOpTimeFloatMultiplyNode_MultiplicandsTooltip", "Float multiplicand(s)."), 1.0f)
+					TInputDataVertexModel<FTime>(METASOUND_GET_PARAM_NAME(PrimaryOperand), PrimaryOperandMetadata, 1.0f),
+					TInputDataVertexModel<float>(METASOUND_GET_PARAM_NAME(AdditionalOperands), AdditionalOperandsMetadata, 1.0f)
 				),
 				FOutputVertexInterface(
-					TOutputDataVertexModel<FTime>(TEXT("Out"), METASOUND_LOCTEXT("MathOpTimeFloatMultiplyNode_OutTooltip", "Time float multiplication output"))
+					TOutputDataVertexModel<FTime>(METASOUND_GET_PARAM_NAME(OutMath), OutputMetadata)
 				)
 			);
 
@@ -1145,13 +1423,33 @@ namespace Metasound
 	public:
 		static const FVertexInterface& GetVertexInterface()
 		{
+			using namespace MathOpNames; 
+
+			static const FDataVertexMetadata PrimaryOperandMetadata
+			{
+				  METASOUND_LOCTEXT("MathOpTimeFloatDivideNode_DividendFloatTooltip", "Time dividend.") // description
+				, METASOUND_GET_PARAM_DISPLAYNAME(PrimaryOperand) // display name
+			};
+
+			static const FDataVertexMetadata AdditionalOperandsMetadata
+			{
+				  METASOUND_LOCTEXT("MathOpTimeFloatDivideNode_DivisorsTooltip", "Float divisor(s).") // description
+				, METASOUND_GET_PARAM_DISPLAYNAME(AdditionalOperands) // display name
+			};
+
+			static const FDataVertexMetadata OutputMetadata
+			{
+				  METASOUND_LOCTEXT("MathOpTimeFloatDivideNode_OutTooltip", "Time divide-by-float output") // description
+				, METASOUND_GET_PARAM_DISPLAYNAME(OutMath) // display name
+			};
+
 			static const FVertexInterface Interface(
 				FInputVertexInterface(
-					TInputDataVertexModel<FTime>(MathOpNames::PrimaryOperandName, METASOUND_LOCTEXT("MathOpTimeFloatDivideNode_DividendFloatTooltip", "Time dividend."), 1.0f),
-					TInputDataVertexModel<float>(MathOpNames::AdditionalOperandsName, METASOUND_LOCTEXT("MathOpTimeFloatDivideNode_DivisorsTooltip", "Float divisor(s)."), 1.0f)
+					TInputDataVertexModel<FTime>(METASOUND_GET_PARAM_NAME(PrimaryOperand), PrimaryOperandMetadata, 1.0f),
+					TInputDataVertexModel<float>(METASOUND_GET_PARAM_NAME(AdditionalOperands), AdditionalOperandsMetadata, 1.0f)
 				),
 				FOutputVertexInterface(
-					TOutputDataVertexModel<FTime>(TEXT("Out"), METASOUND_LOCTEXT("MathOpTimeFloatDivideNode_OutTooltip", "Time divide-by-float output"))
+					TOutputDataVertexModel<FTime>(METASOUND_GET_PARAM_NAME(OutMath), OutputMetadata)
 				)
 			);
 
diff --git a/Engine/Plugins/Runtime/Metasound/Source/MetasoundStandardNodes/Private/MetasoundMixerNode.cpp b/Engine/Plugins/Runtime/Metasound/Source/MetasoundStandardNodes/Private/MetasoundMixerNode.cpp
--- a/Engine/Plugins/Runtime/Metasound/Source/MetasoundStandardNodes/Private/MetasoundMixerNode.cpp
+++ b/Engine/Plugins/Runtime/Metasound/Source/MetasoundStandardNodes/Private/MetasoundMixerNode.cpp
@@ -58,7 +58,12 @@ namespace Metasound
 					// audio channels
 					for (uint32 ChanIndex = 0; ChanIndex < NumChannels; ++ChanIndex)
 					{
-						InputInterface.Add(TInputDataVertexModel<FAudioBuffer>(GetAudioInputName(InputIndex, ChanIndex), GetAudioInputDescription(InputIndex, ChanIndex)));
+						const FDataVertexMetadata AudioInputMetadata
+						{
+							GetAudioInputDescription(InputIndex, ChanIndex),
+							GetAudioInputDisplayName(InputIndex, ChanIndex)
+						};
+						InputInterface.Add(TInputDataVertexModel<FAudioBuffer>(GetAudioInputName(InputIndex, ChanIndex), AudioInputMetadata));
 					}
 
 					// gain scalar
@@ -76,7 +81,13 @@ namespace Metasound
 				FOutputVertexInterface OutputInterface;
 				for (uint32 i = 0; i < NumChannels; ++i)
 				{
-					OutputInterface.Add(TOutputDataVertexModel<FAudioBuffer>(GetAudioOutputName(i), GetAudioOutputDescription(i)));
+					const FDataVertexMetadata AudioOutputMetadata
+					{
+						GetAudioOutputDescription(i),
+						GetAudioOutputDisplayName(i)
+					};
+
+					OutputInterface.Add(TOutputDataVertexModel<FAudioBuffer>(GetAudioOutputName(i), AudioOutputMetadata));
 				}
 
 				return FVertexInterface(InputInterface, OutputInterface);
@@ -252,6 +263,26 @@ namespace Metasound
 			return METASOUND_LOCTEXT_FORMAT("AudioMixerAudioInputDescription", "Audio Input #: {0}, Channel: {1}", InputIndex, ChannelIndex);
 		}
 
+		static const FText GetAudioInputDisplayName(uint32 InputIndex, uint32 ChannelIndex)
+		{
+			if (NumChannels == 1)
+			{
+				return METASOUND_LOCTEXT_FORMAT("AudioMixerAudioInput1In", "In {0}", InputIndex);
+			}
+			else if (NumChannels == 2)
+			{
+				if (ChannelIndex == 0)
+				{
+					return METASOUND_LOCTEXT_FORMAT("AudioMixerAudioInput2InL", "In {0} L", InputIndex);
+				}
+				else
+				{
+					return METASOUND_LOCTEXT_FORMAT("AudioMixerAudioInput2InR", "In {0} R", InputIndex);
+				}
+			}
+			return METASOUND_LOCTEXT_FORMAT("AudioMixerAudioInputIn", "In {0}, {0}", InputIndex, ChannelIndex);
+		}
+
 		static const FVertexName GetGainInputName(uint32 InputIndex)
 		{
 			return *FString::Printf(TEXT("Gain %i"), InputIndex);
@@ -281,6 +312,26 @@ namespace Metasound
 			return *FString::Printf(TEXT("Out %i"), ChannelIndex);
 		}
 
+		static const FText GetAudioOutputDisplayName(uint32 ChannelIndex)
+		{
+			if (NumChannels == 1)
+			{
+				return METASOUND_LOCTEXT("AudioMixerAudioOutput1Out", "Out");
+			}
+			else if (NumChannels == 2)
+			{
+				if (ChannelIndex == 0)
+				{
+					return METASOUND_LOCTEXT("AudioMixerAudioOutput2OutL", "Out L");
+				}
+				else
+				{
+					return METASOUND_LOCTEXT("AudioMixerAudioOutput2OutR", "Out R");
+				}
+			}
+			return METASOUND_LOCTEXT_FORMAT("AudioMixerAudioOutputOut", "Out {0}", ChannelIndex);
+		}
+
 		static const FText GetAudioOutputDescription(uint32 ChannelIndex)
 		{
 			return METASOUND_LOCTEXT_FORMAT("AudioMixerAudioOutputDescription", "Summed output for channel: {0}", ChannelIndex);
diff --git a/Engine/Plugins/Runtime/Metasound/Source/MetasoundStandardNodes/Private/MetasoundOscillatorNodes.cpp b/Engine/Plugins/Runtime/Metasound/Source/MetasoundStandardNodes/Private/MetasoundOscillatorNodes.cpp
--- a/Engine/Plugins/Runtime/Metasound/Source/MetasoundStandardNodes/Private/MetasoundOscillatorNodes.cpp
+++ b/Engine/Plugins/Runtime/Metasound/Source/MetasoundStandardNodes/Private/MetasoundOscillatorNodes.cpp
@@ -9,6 +9,7 @@
 #include "MetasoundFacade.h"
 #include "MetasoundNodeRegistrationMacro.h"
 #include "MetasoundDataTypeRegistrationMacro.h"
+#include "MetasoundParamHelper.h"
 #include "MetasoundPrimitives.h"
 #include "MetasoundStandardNodesCategories.h"
 #include "MetasoundStandardNodesNames.h"
@@ -517,39 +518,40 @@ namespace Metasound
 		using FTrianglePolysmoothWithFm = TGenerateBlockFM<FTriangleGenerator>;
 	}
 
+	namespace OscillatorCommonVertexNames
+	{
+		METASOUND_PARAM(EnabledPin, "Enabled", "Enable the oscillator.")
+		METASOUND_PARAM(BiPolarPin, "Bi Polar", "If the output is Bi-Polar (-1..1) or Uni-Polar (0..1)")
+		METASOUND_PARAM(FrequencyModPin, "Modulation","Modulation Frequency Input (for doing FM)")
+		METASOUND_PARAM(OscBaseFrequencyPin, "Frequency", "Base Frequency of Oscillator in Hz.")
+		METASOUND_PARAM(OscPhaseResetPin, "Sync", "Phase Reset")
+		METASOUND_PARAM(PhaseOffsetPin, "Phase Offset", "Phase Offset In Degrees (0..360)")
+		METASOUND_PARAM(GlideFactorPin, "Glide", "The amount of glide to use when changing frequencies. 0.0 = no glide, 1.0 = lots of glide.")
+		METASOUND_PARAM(AudioOutPin, "Audio", "The output audio")
+	}
+
 	// Base class of Oscillator factories which holds common the interface.
 	class FOscilatorFactoryBase : public IOperatorFactory
 	{
 	public:
-		// Common pins
-		static constexpr const TCHAR* EnabledPinName = TEXT("Enabled");
-		static constexpr const TCHAR* FrequencyModPinName = TEXT("Modulation");
-		static constexpr const TCHAR* BaseFrequencyPinName = TEXT("Frequency");
-		static constexpr const TCHAR* PhaseResetPinName = TEXT("Sync");
-		static constexpr const TCHAR* PhaseOffsetPinName = TEXT("Phase Offset");
-		static constexpr const TCHAR* GlideFactorPinName = TEXT("Glide");
-		static constexpr const TCHAR* AudioOutPinName = TEXT("Audio");
-		static constexpr const TCHAR* BiPolarPinName = TEXT("Bi Polar");
-		static constexpr const TCHAR* MinOutputValuePinName = TEXT("Min Value");
-		static constexpr const TCHAR* MaxOutputValuePinName = TEXT("Max Value");
-
-
 		// Common to all Oscillators.
 		static FVertexInterface GetCommmonVertexInterface()
 		{
+			using namespace OscillatorCommonVertexNames; 
+
 			static const FVertexInterface Interface
 			{
 				FInputVertexInterface{
-					TInputDataVertexModel<bool>(EnabledPinName, METASOUND_LOCTEXT("OscActivateDescription", "Enable the oscillator."), true),
-					TInputDataVertexModel<bool>(BiPolarPinName, METASOUND_LOCTEXT("OscBaseBiPolarDescription", "If the output is Bi-Polar (-1..1) or Uni-Polar (0..1)"), true),
-					TInputDataVertexModel<float>(BaseFrequencyPinName, METASOUND_LOCTEXT("OscBaseFrequencyDescription", "Base Frequency of Oscillator in HZ"), 440.f),
-					TInputDataVertexModel<FAudioBuffer>(FrequencyModPinName, METASOUND_LOCTEXT("OscModFrequencyDescription", "Modulation Frequency Input (for doing FM)")),
-					TInputDataVertexModel<FTrigger>(PhaseResetPinName, METASOUND_LOCTEXT("OscPhaseResetDescription", "Phase Reset")),
-					TInputDataVertexModel<float>(PhaseOffsetPinName, METASOUND_LOCTEXT("OscPhaseOffsetDescription", "Phase Offset In Degrees (0..360)"), 0.f),
-					TInputDataVertexModel<float>(GlideFactorPinName, METASOUND_LOCTEXT("OscGlideFactorDescription", "The amount of glide to use when changing frequencies. 0.0 = no glide, 1.0 = lots of glide."), 0.f)
+					TInputDataVertexModel<bool>(METASOUND_GET_PARAM_NAME_AND_METADATA(EnabledPin), true),
+					TInputDataVertexModel<bool>(METASOUND_GET_PARAM_NAME_AND_METADATA(BiPolarPin), true),
+					TInputDataVertexModel<float>(METASOUND_GET_PARAM_NAME_AND_METADATA(OscBaseFrequencyPin), 440.f),
+					TInputDataVertexModel<FAudioBuffer>(METASOUND_GET_PARAM_NAME_AND_METADATA(FrequencyModPin)),
+					TInputDataVertexModel<FTrigger>(METASOUND_GET_PARAM_NAME_AND_METADATA(OscPhaseResetPin)),
+					TInputDataVertexModel<float>(METASOUND_GET_PARAM_NAME_AND_METADATA(PhaseOffsetPin), 0.f),
+					TInputDataVertexModel<float>(METASOUND_GET_PARAM_NAME_AND_METADATA(GlideFactorPin), 0.f)
 				},
 				FOutputVertexInterface{
-					TOutputDataVertexModel<FAudioBuffer>(TEXT("Audio"), METASOUND_LOCTEXT("AudioTooltip", "The output audio"))
+					TOutputDataVertexModel<FAudioBuffer>(METASOUND_GET_PARAM_NAME_AND_METADATA(AudioOutPin))
 				}
 			};
 			return Interface;
@@ -591,20 +593,24 @@ namespace Metasound
 
 		FDataReferenceCollection GetInputs() const override
 		{
+			using namespace OscillatorCommonVertexNames; 
+
 			FDataReferenceCollection InputDataReferences;
-			InputDataReferences.AddDataReadReference(FOscilatorFactoryBase::EnabledPinName, Enabled);
-			InputDataReferences.AddDataReadReference(FOscilatorFactoryBase::BaseFrequencyPinName, BaseFrequency);
-			InputDataReferences.AddDataReadReference(FOscilatorFactoryBase::PhaseOffsetPinName, PhaseOffset);
-			InputDataReferences.AddDataReadReference(FOscilatorFactoryBase::PhaseResetPinName, PhaseReset);
-			InputDataReferences.AddDataReadReference(FOscilatorFactoryBase::GlideFactorPinName, GlideFactor);
-			InputDataReferences.AddDataReadReference(FOscilatorFactoryBase::BiPolarPinName, BiPolar);
+			InputDataReferences.AddDataReadReference(METASOUND_GET_PARAM_NAME(EnabledPin), Enabled);
+			InputDataReferences.AddDataReadReference(METASOUND_GET_PARAM_NAME(OscBaseFrequencyPin), BaseFrequency);
+			InputDataReferences.AddDataReadReference(METASOUND_GET_PARAM_NAME(PhaseOffsetPin), PhaseOffset);
+			InputDataReferences.AddDataReadReference(METASOUND_GET_PARAM_NAME(OscPhaseResetPin), PhaseReset);
+			InputDataReferences.AddDataReadReference(METASOUND_GET_PARAM_NAME(GlideFactorPin), GlideFactor);
+			InputDataReferences.AddDataReadReference(METASOUND_GET_PARAM_NAME(BiPolarPin), BiPolar);
 			return InputDataReferences;
 		}
 
 		FDataReferenceCollection GetOutputs() const override
 		{
+			using namespace OscillatorCommonVertexNames;
+			
 			FDataReferenceCollection OutputDataReferences;
-			OutputDataReferences.AddDataReadReference(FOscilatorFactoryBase::AudioOutPinName, AudioBuffer);
+			OutputDataReferences.AddDataReadReference(METASOUND_GET_PARAM_NAME(AudioOutPin), AudioBuffer);
 			return OutputDataReferences;
 		}
 
@@ -694,8 +700,10 @@ namespace Metasound
 
 		FDataReferenceCollection GetInputs() const override
 		{
+			using namespace OscillatorCommonVertexNames; 
+
 			FDataReferenceCollection Inputs = Super::GetInputs();
-			Inputs.AddDataReadReference(FOscilatorFactoryBase::FrequencyModPinName, Fm);
+			Inputs.AddDataReadReference(METASOUND_GET_PARAM_NAME(FrequencyModPin), Fm);
 			return Inputs;
 		}
 
@@ -752,13 +760,16 @@ namespace Metasound
 		DEFINE_METASOUND_ENUM_ENTRY(ESineGenerationType::Wavetable, "WavetableDescription", "Wavetable", "WavetableDescriptionTT", "Uses a wavetable to generate the sine"),
 	DEFINE_METASOUND_ENUM_END()
 
+	namespace SineOscilatorVertexNames
+	{
+		METASOUND_PARAM(SineType, "Type", "Type of the Sinewave Generator")
+	}
+
 	class FSineOscilatorNode::FFactory : public FOscilatorFactoryBase
 	{
 	public:
 		FFactory() = default;
 		
-		static constexpr const TCHAR* SineTypePin = TEXT("Type");
-
 		static const FNodeClassMetadata& GetNodeInfo()
 		{
 			auto InitNodeInfo = []() -> FNodeClassMetadata
@@ -780,11 +791,13 @@ namespace Metasound
 		}
 		static const FVertexInterface& GetVertexInterface()
 		{
+			using namespace SineOscilatorVertexNames; 
+
 			auto MakeInterface = []() -> FVertexInterface
 			{
 				FVertexInterface Interface = GetCommmonVertexInterface();
 				Interface.GetInputInterface().Add(
-					TInputDataVertexModel<FEnumSineGenerationType>(SineTypePin, METASOUND_LOCTEXT("SineTypeDescription", "Type of the Sinewave Generator"), static_cast<int32>(ESineGenerationType::Wavetable))
+					TInputDataVertexModel<FEnumSineGenerationType>(METASOUND_GET_PARAM_NAME_AND_METADATA(SineType), static_cast<int32>(ESineGenerationType::Wavetable))
 				);
 				return Interface;
 			};
@@ -798,27 +811,30 @@ namespace Metasound
 			const FDataReferenceCollection& InputCol = InParams.InputDataReferences;
 			const FOperatorSettings& Settings = InParams.OperatorSettings;
 			using namespace Generators;
+			using namespace OscillatorCommonVertexNames; 
+			using namespace SineOscilatorVertexNames;
+			
 			FOscillatorOperatorConstructParams OpParams
 			{
 				Settings,
-				InputCol.GetDataReadReferenceOrConstruct<bool>(EnabledPinName, SineNode.GetDefaultEnablement()),
-				InputCol.GetDataReadReferenceOrConstruct<float>(BaseFrequencyPinName, SineNode.GetDefaultFrequency()),
-				InputCol.GetDataReadReferenceOrConstruct<float>(PhaseOffsetPinName, SineNode.GetDefaultPhaseOffset()),
-				InputCol.GetDataReadReferenceOrConstruct<FTrigger>(PhaseResetPinName, Settings),
-				InputCol.GetDataReadReferenceOrConstruct<float>(GlideFactorPinName, SineNode.GetDefaultGlideFactor()),
-				InputCol.GetDataReadReferenceOrConstruct<bool>(BiPolarPinName, true)
+				InputCol.GetDataReadReferenceOrConstruct<bool>(METASOUND_GET_PARAM_NAME(EnabledPin), SineNode.GetDefaultEnablement()),
+				InputCol.GetDataReadReferenceOrConstruct<float>(METASOUND_GET_PARAM_NAME(OscBaseFrequencyPin), SineNode.GetDefaultFrequency()),
+				InputCol.GetDataReadReferenceOrConstruct<float>(METASOUND_GET_PARAM_NAME(PhaseOffsetPin), SineNode.GetDefaultPhaseOffset()),
+				InputCol.GetDataReadReferenceOrConstruct<FTrigger>(METASOUND_GET_PARAM_NAME(OscPhaseResetPin), Settings),
+				InputCol.GetDataReadReferenceOrConstruct<float>(METASOUND_GET_PARAM_NAME(GlideFactorPin), SineNode.GetDefaultGlideFactor()),
+				InputCol.GetDataReadReferenceOrConstruct<bool>(METASOUND_GET_PARAM_NAME(BiPolarPin), true)
 			};
 
 			// TODO: Make this a static prop. For now its a pin.
 			
 			// Check to see if we have an FM input connected.
-			bool bHasFM = InputCol.ContainsDataReadReference<FAudioBuffer>(FrequencyModPinName);			
-			FEnumSineGenerationTypeReadRef Type = InputCol.GetDataReadReferenceOrConstruct<FEnumSineGenerationType>(SineTypePin, ESineGenerationType::Wavetable);
+			bool bHasFM = InputCol.ContainsDataReadReference<FAudioBuffer>(METASOUND_GET_PARAM_NAME(FrequencyModPin));
+			FEnumSineGenerationTypeReadRef Type = InputCol.GetDataReadReferenceOrConstruct<FEnumSineGenerationType>(METASOUND_GET_PARAM_NAME(SineType), ESineGenerationType::Wavetable);
 			if (bHasFM)
 			{
 
 				// FM Oscillators.
-				FAudioBufferReadRef FmBuffer = InputCol.GetDataReadReference<FAudioBuffer>(FrequencyModPinName);
+				FAudioBufferReadRef FmBuffer = InputCol.GetDataReadReference<FAudioBuffer>(METASOUND_GET_PARAM_NAME(FrequencyModPin));
 				switch (*Type)
 				{
 				default:
@@ -872,12 +888,15 @@ namespace Metasound
 		DEFINE_METASOUND_ENUM_ENTRY(ESawGenerationType::Trivial, "SawTrivialDescription", "Trivial", "TrivialDescriptionTT", "The most basic raw implementation"),
 		//DEFINE_METASOUND_ENUM_ENTRY(ESawGenerationType::Wavetable, "SawWavetableDescription", "Wavetable", "SawWavetableDescriptionTT", "Use a Wavetable iterpolation to generate the Waveform")
 	DEFINE_METASOUND_ENUM_END()
+	
+	namespace SawOscilatorVertexNames
+	{
+		METASOUND_PARAM(SawType, "Type", "Type of the Saw Generator")
+	}
 
 	class FSawOscilatorNode::FFactory : public FOscilatorFactoryBase
 	{
-	public:
-		static constexpr const TCHAR* SawTypePin = TEXT("Type");
-		
+	public:		
 		FFactory() = default;
 
 		TUniquePtr<IOperator> CreateOperator(const FCreateOperatorParams& InParams, FBuildErrorArray& OutErrors) override;
@@ -903,11 +922,13 @@ namespace Metasound
 		}
 		static const FVertexInterface& GetVertexInterface()
 		{
+			using namespace SawOscilatorVertexNames;
+
 			auto MakeInterface = []() -> FVertexInterface
 			{
 				FVertexInterface Interface = GetCommmonVertexInterface();
 				Interface.GetInputInterface().Add(
-					TInputDataVertexModel<FEnumSawGenerationType>(SawTypePin, METASOUND_LOCTEXT("Saw Type", "Which type of Saw Generator to Create"))
+					TInputDataVertexModel<FEnumSawGenerationType>(METASOUND_GET_PARAM_NAME_AND_METADATA(SawType))
 				);
 				return Interface;
 			};
@@ -924,25 +945,27 @@ namespace Metasound
 		const FDataReferenceCollection& InputCol = InParams.InputDataReferences;
 		const FOperatorSettings& Settings = InParams.OperatorSettings;
 		using namespace Generators;		
+		using namespace OscillatorCommonVertexNames; 
+		using namespace SawOscilatorVertexNames; 
 
-		FSawGenerationTypeReadRef Type = InputCol.GetDataReadReferenceOrConstruct<FEnumSawGenerationType>(SawTypePin);
+		FSawGenerationTypeReadRef Type = InputCol.GetDataReadReferenceOrConstruct<FEnumSawGenerationType>(METASOUND_GET_PARAM_NAME(SawType));
 
 		FOscillatorOperatorConstructParams OpParams
 		{
 			Settings,
-			InputCol.GetDataReadReferenceOrConstruct<bool>(EnabledPinName, Node.GetDefaultEnablement()),
-			InputCol.GetDataReadReferenceOrConstruct<float>(BaseFrequencyPinName, Node.GetDefaultFrequency()),
-			InputCol.GetDataReadReferenceOrConstruct<float>(PhaseOffsetPinName, Node.GetDefaultPhaseOffset()),
-			InputCol.GetDataReadReferenceOrConstruct<FTrigger>(PhaseResetPinName, Settings),
-			InputCol.GetDataReadReferenceOrConstruct<float>(GlideFactorPinName, Node.GetDefaultGlideFactor()),
-			InputCol.GetDataReadReferenceOrConstruct<bool>(BiPolarPinName, true)
+			InputCol.GetDataReadReferenceOrConstruct<bool>(METASOUND_GET_PARAM_NAME(EnabledPin), Node.GetDefaultEnablement()),
+			InputCol.GetDataReadReferenceOrConstruct<float>(METASOUND_GET_PARAM_NAME(OscBaseFrequencyPin), Node.GetDefaultFrequency()),
+			InputCol.GetDataReadReferenceOrConstruct<float>(METASOUND_GET_PARAM_NAME(PhaseOffsetPin), Node.GetDefaultPhaseOffset()),
+			InputCol.GetDataReadReferenceOrConstruct<FTrigger>(METASOUND_GET_PARAM_NAME(OscPhaseResetPin), Settings),
+			InputCol.GetDataReadReferenceOrConstruct<float>(METASOUND_GET_PARAM_NAME(GlideFactorPin), Node.GetDefaultGlideFactor()),
+			InputCol.GetDataReadReferenceOrConstruct<bool>(METASOUND_GET_PARAM_NAME(BiPolarPin), true)
 		};
 
-		bool bHasFM = InputCol.ContainsDataReadReference<FAudioBuffer>(FrequencyModPinName);
+		bool bHasFM = InputCol.ContainsDataReadReference<FAudioBuffer>(METASOUND_GET_PARAM_NAME(FrequencyModPin));
 
 		if (bHasFM)
 		{
-			FAudioBufferReadRef FmBuffer = InputCol.GetDataReadReference<FAudioBuffer>(FrequencyModPinName);
+			FAudioBufferReadRef FmBuffer = InputCol.GetDataReadReference<FAudioBuffer>(METASOUND_GET_PARAM_NAME(FrequencyModPin));
 			switch (*Type)
 			{
 			default:
@@ -1060,40 +1083,46 @@ namespace Metasound
 		FAudioBufferReadRef FM;
 	};
 
+	namespace SquareOscillatorVertexNames
+	{
+		METASOUND_PARAM(SquarePulseWidthPin, "Pulse Width", "The Width of the square part of the wave")
+		METASOUND_PARAM(SquareTypePin, "Type", "The generator type to make the squarewave")
+	}
+
 	class FSquareOscilatorNode::FFactory : public FOscilatorFactoryBase
 	{
 	public:
 		FFactory() = default;
 
-		static constexpr const TCHAR* PulseWidthPinName = TEXT("Pulse Width");
-		static constexpr const TCHAR* SquareTypePinName = TEXT("Type");
-
 		TUniquePtr<IOperator> CreateOperator(const FCreateOperatorParams& InParams, FBuildErrorArray& OutErrors) override
 		{
+			using namespace SquareOscillatorVertexNames; 
+
 			const FSquareOscilatorNode& Node = static_cast<const FSquareOscilatorNode&>(InParams.Node);
 			const FDataReferenceCollection& InputCol = InParams.InputDataReferences;
 			const FOperatorSettings& Settings = InParams.OperatorSettings;
 			using namespace Generators;
+			using namespace OscillatorCommonVertexNames; 
 
 			FOscillatorOperatorConstructParams OpParams
 			{
 				Settings,
-				InputCol.GetDataReadReferenceOrConstruct<bool>(EnabledPinName, Node.GetDefaultEnablement()),
-				InputCol.GetDataReadReferenceOrConstruct<float>(BaseFrequencyPinName, Node.GetDefaultFrequency()),
-				InputCol.GetDataReadReferenceOrConstruct<float>(PhaseOffsetPinName, Node.GetDefaultPhaseOffset()),
-				InputCol.GetDataReadReferenceOrConstruct<FTrigger>(PhaseResetPinName, Settings),
-				InputCol.GetDataReadReferenceOrConstruct<float>(GlideFactorPinName, Node.GetDefaultGlideFactor()),
-				InputCol.GetDataReadReferenceOrConstruct<bool>(BiPolarPinName, true)
+				InputCol.GetDataReadReferenceOrConstruct<bool>(METASOUND_GET_PARAM_NAME(EnabledPin), Node.GetDefaultEnablement()),
+				InputCol.GetDataReadReferenceOrConstruct<float>(METASOUND_GET_PARAM_NAME(OscBaseFrequencyPin), Node.GetDefaultFrequency()),
+				InputCol.GetDataReadReferenceOrConstruct<float>(METASOUND_GET_PARAM_NAME(PhaseOffsetPin), Node.GetDefaultPhaseOffset()),
+				InputCol.GetDataReadReferenceOrConstruct<FTrigger>(METASOUND_GET_PARAM_NAME(OscPhaseResetPin), Settings),
+				InputCol.GetDataReadReferenceOrConstruct<float>(METASOUND_GET_PARAM_NAME(GlideFactorPin), Node.GetDefaultGlideFactor()),
+				InputCol.GetDataReadReferenceOrConstruct<bool>(METASOUND_GET_PARAM_NAME(BiPolarPin), true)
 			};
+			
+			FSquareGenerationTypeReadRef Type = InputCol.GetDataReadReferenceOrConstruct<FEnumSquareGenerationType>(METASOUND_GET_PARAM_NAME(SquareTypePin));
 
-			FSquareGenerationTypeReadRef Type = InputCol.GetDataReadReferenceOrConstruct<FEnumSquareGenerationType>(SquareTypePinName);
-
-			bool bHasFM = InputCol.ContainsDataReadReference<FAudioBuffer>(FrequencyModPinName);
-			FFloatReadRef PulseWidth = InputCol.GetDataReadReferenceOrConstruct<float>(PulseWidthPinName, 0.5f);
+			bool bHasFM = InputCol.ContainsDataReadReference<FAudioBuffer>(METASOUND_GET_PARAM_NAME(FrequencyModPin));
+			FFloatReadRef PulseWidth = InputCol.GetDataReadReferenceOrConstruct<float>(METASOUND_GET_PARAM_NAME(SquarePulseWidthPin), 0.5f);
 
 			if (bHasFM)
 			{
-				FAudioBufferReadRef FmBuffer = InputCol.GetDataReadReference<FAudioBuffer>(FrequencyModPinName);
+				FAudioBufferReadRef FmBuffer = InputCol.GetDataReadReference<FAudioBuffer>(METASOUND_GET_PARAM_NAME(FrequencyModPin));
 				switch (*Type)
 				{
 				default:				
@@ -1134,11 +1163,13 @@ namespace Metasound
 		}
 		static const FVertexInterface& GetVertexInterface()
 		{
+			using namespace SquareOscillatorVertexNames; 
+
 			auto MakeInterface = []() -> FVertexInterface
 			{
 				FVertexInterface Interface = GetCommmonVertexInterface();
-				Interface.GetInputInterface().Add(TInputDataVertexModel<FEnumSquareGenerationType>(SquareTypePinName, METASOUND_LOCTEXT("SquareTypeDescription", "The generator type to make the squarewave")));
-				Interface.GetInputInterface().Add(TInputDataVertexModel<float>(PulseWidthPinName, METASOUND_LOCTEXT("PhaseWidthDescription", "The Width of the square part of the wave"), 0.5f));
+				Interface.GetInputInterface().Add(TInputDataVertexModel<FEnumSquareGenerationType>(METASOUND_GET_PARAM_NAME_AND_METADATA(SquareTypePin)));
+				Interface.GetInputInterface().Add(TInputDataVertexModel<float>(METASOUND_GET_PARAM_NAME_AND_METADATA(SquarePulseWidthPin), 0.5f));
 				return Interface;
 			};
 			static const FVertexInterface Interface = MakeInterface();
@@ -1176,38 +1207,43 @@ namespace Metasound
 		//DEFINE_METASOUND_ENUM_ENTRY(ETriangleGenerationType::Wavetable, "TriangleWavetableDescription", "Wavetable", "TriangleWavetableDescriptionTT", "Use a Wavetable iterpolation to generate the Waveform")
 	DEFINE_METASOUND_ENUM_END()
 
+	namespace TriangleOscilatorVertexNames
+	{
+		METASOUND_PARAM(TriangeTypePin, "Type", "The generator type to make the triangle wave")
+	}
+
 	class FTriangleOscilatorNode::FFactory : public FOscilatorFactoryBase
 	{
 	public:
 		FFactory() = default;
 
-		static constexpr const TCHAR* TriangeTypePinName = TEXT("Type");
-
 		TUniquePtr<IOperator> CreateOperator(const FCreateOperatorParams& InParams, FBuildErrorArray& OutErrors) override
 		{
 			const FTriangleOscilatorNode& Node = static_cast<const FTriangleOscilatorNode&>(InParams.Node);
 			const FDataReferenceCollection& InputCol = InParams.InputDataReferences;
 			const FOperatorSettings& Settings = InParams.OperatorSettings;
 			using namespace Generators;
+			using namespace OscillatorCommonVertexNames;
+			using namespace TriangleOscilatorVertexNames;
 			
-			FTriangleGenerationTypeReadRef Type = InputCol.GetDataReadReferenceOrConstruct<FEnumTriangleGenerationType>(TriangeTypePinName);
+			FTriangleGenerationTypeReadRef Type = InputCol.GetDataReadReferenceOrConstruct<FEnumTriangleGenerationType>(METASOUND_GET_PARAM_NAME(TriangeTypePin));
 		
 			FOscillatorOperatorConstructParams OpParams
 			{
 				Settings,
-				InputCol.GetDataReadReferenceOrConstruct<bool>(EnabledPinName, Node.GetDefaultEnablement()),
-				InputCol.GetDataReadReferenceOrConstruct<float>(BaseFrequencyPinName, Node.GetDefaultFrequency()),
-				InputCol.GetDataReadReferenceOrConstruct<float>(PhaseOffsetPinName, Node.GetDefaultPhaseOffset()),
-				InputCol.GetDataReadReferenceOrConstruct<FTrigger>(PhaseResetPinName, Settings),
-				InputCol.GetDataReadReferenceOrConstruct<float>(GlideFactorPinName, Node.GetDefaultGlideFactor()),
-				InputCol.GetDataReadReferenceOrConstruct<bool>(BiPolarPinName, true)
+				InputCol.GetDataReadReferenceOrConstruct<bool>(METASOUND_GET_PARAM_NAME(EnabledPin), Node.GetDefaultEnablement()),
+				InputCol.GetDataReadReferenceOrConstruct<float>(METASOUND_GET_PARAM_NAME(OscBaseFrequencyPin), Node.GetDefaultFrequency()),
+				InputCol.GetDataReadReferenceOrConstruct<float>(METASOUND_GET_PARAM_NAME(PhaseOffsetPin), Node.GetDefaultPhaseOffset()),
+				InputCol.GetDataReadReferenceOrConstruct<FTrigger>(METASOUND_GET_PARAM_NAME(OscPhaseResetPin), Settings),
+				InputCol.GetDataReadReferenceOrConstruct<float>(METASOUND_GET_PARAM_NAME(GlideFactorPin), Node.GetDefaultGlideFactor()),
+				InputCol.GetDataReadReferenceOrConstruct<bool>(METASOUND_GET_PARAM_NAME(BiPolarPin), true)
 			};
 
-			bool bHasFM = InputCol.ContainsDataReadReference<FAudioBuffer>(FrequencyModPinName);
+			bool bHasFM = InputCol.ContainsDataReadReference<FAudioBuffer>(METASOUND_GET_PARAM_NAME(FrequencyModPin));
 			
 			if (bHasFM)
 			{
-				FAudioBufferReadRef FmBuffer = InputCol.GetDataReadReference<FAudioBuffer>(FrequencyModPinName);
+				FAudioBufferReadRef FmBuffer = InputCol.GetDataReadReference<FAudioBuffer>(METASOUND_GET_PARAM_NAME(FrequencyModPin));
 				switch (*Type)
 				{
 				default:
@@ -1248,11 +1284,13 @@ namespace Metasound
 		}
 		static const FVertexInterface& GetVertexInterface()
 		{
+			using namespace TriangleOscilatorVertexNames; 
+
 			auto MakeInterface = []() -> FVertexInterface
 			{
 				FVertexInterface Interface = GetCommmonVertexInterface();
 				Interface.GetInputInterface().Add(
-					TInputDataVertexModel<FEnumTriangleGenerationType>(TriangeTypePinName, METASOUND_LOCTEXT("TriangleTypeDescription", "The generator type to make the triangle wave"))
+					TInputDataVertexModel<FEnumTriangleGenerationType>(METASOUND_GET_PARAM_NAME_AND_METADATA(TriangeTypePin))
 				);
 				return Interface;
 			};
@@ -1295,29 +1333,41 @@ namespace Metasound
 		DEFINE_METASOUND_ENUM_ENTRY(ELfoWaveshapeType::Square, "LfoWaveShapeSquareDescription", "Square", "LfoWaveShapeSquareDescriptionTT", "Square shape Low Frequency Oscillator")
 	DEFINE_METASOUND_ENUM_END()
 
+	namespace LfoVertexNames
+	{
+		// Common pins
+		METASOUND_PARAM(WaveshapePin, "Shape", "Waveshape of the LFO")
+		METASOUND_PARAM(LfoOutPin, "Out", "Output of the LFO (blockrate)")
+		METASOUND_PARAM(LfoBaseFrequencyPin, "Frequency", "Frequency of LFO (Hz), clamped at blockrate")
+
+		METASOUND_PARAM(MinOutputValuePin, "Min Value", "The minimum output value.")
+		METASOUND_PARAM(MaxOutputValuePin, "Max Value", "The maximum output value.")
+		METASOUND_PARAM(LfoPhaseResetPin, "Sync", "Phase Reset (block rate only)")
+		METASOUND_PARAM(PhaseOffsetPin, "Phase Offset", "Phase Offset In Degrees (0..360)")
+		METASOUND_PARAM(LfoPulseWidthPin, "Pulse Width", "Pulse Width (0..1)")
+	}
+
 	// Blockrate All-Purpose Oscillator
 	class FLfoOperator : public TExecutableOperator<FLfoOperator>
 	{
 	public:
-		// Common pins
-		static constexpr const TCHAR* WaveshapePinName = TEXT("Shape");
-		static constexpr const TCHAR* LfoOutPinName = TEXT("Out");
-	
 		static const FVertexInterface& GetVertexInterface()
 		{
+			using namespace LfoVertexNames; 
+
 			static const FVertexInterface Interface
 			{
 				FInputVertexInterface{
-					TInputDataVertexModel<float>(FOscilatorFactoryBase::BaseFrequencyPinName, METASOUND_LOCTEXT("LfoFrequencyDescription", "Frequency of LFO (Hz), clamped at blockrate"), 5.f),
-					TInputDataVertexModel<FEnumLfoWaveshapeType>(WaveshapePinName, METASOUND_LOCTEXT("LfoShapeDescription", "Waveshape of the LFO")),
-					TInputDataVertexModel<float>(FOscilatorFactoryBase::MinOutputValuePinName, METASOUND_LOCTEXT("LfoMinValueDescription", "The minimum output value."), -1.0f),
-					TInputDataVertexModel<float>(FOscilatorFactoryBase::MaxOutputValuePinName, METASOUND_LOCTEXT("LfoMaxValueDescription", "The maximum output value."), 1.0f),
-					TInputDataVertexModel<FTrigger>(FOscilatorFactoryBase::PhaseResetPinName, METASOUND_LOCTEXT("LfoPhaseResetDescription", "Phase Reset (block rate only)")),
-					TInputDataVertexModel<float>(FOscilatorFactoryBase::PhaseOffsetPinName, METASOUND_LOCTEXT("LfoPhaseOffsetDescription", "Phase Offset In Degrees (0..360)"), 0.f),
-					TInputDataVertexModel<float>(FSquareOscilatorNode::FFactory::PulseWidthPinName, METASOUND_LOCTEXT("LfoPulseWidthDescription", "Pulse Width (0..1)"), 0.5f)
+					TInputDataVertexModel<float>(METASOUND_GET_PARAM_NAME_AND_METADATA(LfoBaseFrequencyPin), 5.f),
+					TInputDataVertexModel<FEnumLfoWaveshapeType>(METASOUND_GET_PARAM_NAME_AND_METADATA(WaveshapePin)),
+					TInputDataVertexModel<float>(METASOUND_GET_PARAM_NAME_AND_METADATA(MinOutputValuePin), -1.0f),
+					TInputDataVertexModel<float>(METASOUND_GET_PARAM_NAME_AND_METADATA(MaxOutputValuePin), 1.0f),
+					TInputDataVertexModel<FTrigger>(METASOUND_GET_PARAM_NAME_AND_METADATA(LfoPhaseResetPin)),
+					TInputDataVertexModel<float>(METASOUND_GET_PARAM_NAME_AND_METADATA(PhaseOffsetPin), 0.f),
+					TInputDataVertexModel<float>(METASOUND_GET_PARAM_NAME_AND_METADATA(LfoPulseWidthPin), 0.5f)
 				},
 				FOutputVertexInterface{
-					TOutputDataVertexModel<float>(LfoOutPinName, METASOUND_LOCTEXT("LfoOutputDescription", "Output of the LFO (blockrate)"))
+					TOutputDataVertexModel<float>(METASOUND_GET_PARAM_NAME_AND_METADATA(LfoOutPin))
 				}
 			};
 			return Interface;
@@ -1350,15 +1400,17 @@ namespace Metasound
 			const FOperatorSettings& Settings = InParams.OperatorSettings;
 			const FInputVertexInterface& InputInterface = GetVertexInterface().GetInputInterface();
 
+			using namespace LfoVertexNames;
+
 			return MakeUnique<FLfoOperator>(
 				  Settings
-				, InputCol.GetDataReadReferenceOrConstructWithVertexDefault<float>(InputInterface, FOscilatorFactoryBase::BaseFrequencyPinName, Settings)
-				, InputCol.GetDataReadReferenceOrConstruct<FEnumLfoWaveshapeType>(WaveshapePinName)
-				, InputCol.GetDataReadReferenceOrConstructWithVertexDefault<float>(InputInterface, FOscilatorFactoryBase::MinOutputValuePinName, Settings)
-				, InputCol.GetDataReadReferenceOrConstructWithVertexDefault<float>(InputInterface, FOscilatorFactoryBase::MaxOutputValuePinName, Settings)
-				, InputCol.GetDataReadReferenceOrConstruct<FTrigger>(FOscilatorFactoryBase::PhaseResetPinName, Settings)
-				, InputCol.GetDataReadReferenceOrConstructWithVertexDefault<float>(InputInterface,FOscilatorFactoryBase::PhaseOffsetPinName, Settings)
-				, InputCol.GetDataReadReferenceOrConstructWithVertexDefault<float>(InputInterface,FSquareOscilatorNode::FFactory::PulseWidthPinName, Settings)
+				, InputCol.GetDataReadReferenceOrConstructWithVertexDefault<float>(InputInterface, METASOUND_GET_PARAM_NAME(LfoBaseFrequencyPin), Settings)
+				, InputCol.GetDataReadReferenceOrConstruct<FEnumLfoWaveshapeType>(METASOUND_GET_PARAM_NAME(WaveshapePin))
+				, InputCol.GetDataReadReferenceOrConstructWithVertexDefault<float>(InputInterface, METASOUND_GET_PARAM_NAME(MinOutputValuePin), Settings)
+				, InputCol.GetDataReadReferenceOrConstructWithVertexDefault<float>(InputInterface, METASOUND_GET_PARAM_NAME(MaxOutputValuePin), Settings)
+				, InputCol.GetDataReadReferenceOrConstruct<FTrigger>(METASOUND_GET_PARAM_NAME(LfoPhaseResetPin), Settings)
+				, InputCol.GetDataReadReferenceOrConstructWithVertexDefault<float>(InputInterface, METASOUND_GET_PARAM_NAME(PhaseOffsetPin), Settings)
+				, InputCol.GetDataReadReferenceOrConstructWithVertexDefault<float>(InputInterface, METASOUND_GET_PARAM_NAME(LfoPulseWidthPin), Settings)
 			);
 		}
 
@@ -1380,21 +1432,25 @@ namespace Metasound
 
 		FDataReferenceCollection GetInputs() const override
 		{
+			using namespace LfoVertexNames;
+
 			FDataReferenceCollection InputDataReferences;
-			InputDataReferences.AddDataReadReference(FOscilatorFactoryBase::BaseFrequencyPinName, Frequency);
-			InputDataReferences.AddDataReadReference(WaveshapePinName, Waveshape);
-			InputDataReferences.AddDataReadReference(FOscilatorFactoryBase::MinOutputValuePinName, MinValue);
-			InputDataReferences.AddDataReadReference(FOscilatorFactoryBase::MaxOutputValuePinName, MaxValue);
-			InputDataReferences.AddDataReadReference(FOscilatorFactoryBase::PhaseOffsetPinName, PhaseOffset);
-			InputDataReferences.AddDataReadReference(FOscilatorFactoryBase::PhaseResetPinName, PhaseReset);
-			InputDataReferences.AddDataReadReference(FSquareOscilatorNode::FFactory::PulseWidthPinName, PulseWidth);
+			InputDataReferences.AddDataReadReference(METASOUND_GET_PARAM_NAME(LfoBaseFrequencyPin), Frequency);
+			InputDataReferences.AddDataReadReference(METASOUND_GET_PARAM_NAME(WaveshapePin), Waveshape);
+			InputDataReferences.AddDataReadReference(METASOUND_GET_PARAM_NAME(MinOutputValuePin), MinValue);
+			InputDataReferences.AddDataReadReference(METASOUND_GET_PARAM_NAME(MaxOutputValuePin), MaxValue);
+			InputDataReferences.AddDataReadReference(METASOUND_GET_PARAM_NAME(PhaseOffsetPin), PhaseOffset);
+			InputDataReferences.AddDataReadReference(METASOUND_GET_PARAM_NAME(LfoPhaseResetPin), PhaseReset);
+			InputDataReferences.AddDataReadReference(METASOUND_GET_PARAM_NAME(LfoPulseWidthPin), PulseWidth);
 			return InputDataReferences;
 		}
 
 		FDataReferenceCollection GetOutputs() const override
 		{
+			using namespace LfoVertexNames;
+
 			FDataReferenceCollection OutputDataReferences;
-			OutputDataReferences.AddDataReadReference(LfoOutPinName, Output);
+			OutputDataReferences.AddDataReadReference(METASOUND_GET_PARAM_NAME(LfoOutPin), Output);
 			return OutputDataReferences;
 		}
 
diff --git a/Engine/Plugins/Runtime/Metasound/Source/MetasoundStandardNodes/Private/MetasoundRandomNode.cpp b/Engine/Plugins/Runtime/Metasound/Source/MetasoundStandardNodes/Private/MetasoundRandomNode.cpp
--- a/Engine/Plugins/Runtime/Metasound/Source/MetasoundStandardNodes/Private/MetasoundRandomNode.cpp
+++ b/Engine/Plugins/Runtime/Metasound/Source/MetasoundStandardNodes/Private/MetasoundRandomNode.cpp
@@ -9,97 +9,6 @@
 
 namespace Metasound
 {
-	namespace RandomNodeNames
-	{
-		const FVertexName& GetInputNextTriggerName()
-		{
-			static const FVertexName Name = TEXT("Next");
-			return Name;
-		}
-
-		const FVertexName& GetInputResetTriggerName()
-		{
-			static const FVertexName Name = TEXT("Reset");
-			return Name;
-		}
-
-		const FVertexName& GetInputSeedName()
-		{
-			static const FVertexName Name = TEXT("Seed");
-			return Name;
-		}
-
-		const FVertexName& GetInputMinName()
-		{
-			static const FVertexName Name = TEXT("Min");
-			return Name;
-		}
-
-		const FVertexName& GetInputMaxName()
-		{
-			static const FVertexName Name = TEXT("Max");
-			return Name;
-		}
-
-		const FVertexName& GetOutputOnNextTriggerName()
-		{
-			static const FVertexName Name = TEXT("On Next");
-			return Name;
-		}
-
-		const FVertexName& GetOutputOnResetTriggerName()
-		{
-			static const FVertexName Name = TEXT("On Reset");
-			return Name;
-		}
-
-		const FVertexName& GetOutputValueName()
-		{
-			static const FVertexName Name = TEXT("Value");
-			return Name;
-		}
-
-		static FText GetNextTriggerDescription()
-		{
-			return METASOUND_LOCTEXT("RandomNode_NextTT", "Trigger to generate the next random integer.");
-		}
-
-		static FText GetResetDescription()
-		{
-			return METASOUND_LOCTEXT("RandomNode_ResetTT", "Trigger to reset the random sequence with the supplied seed. Useful to get randomized repetition.");
-		}
-
-		static FText GetSeedDescription()
-		{
-			return METASOUND_LOCTEXT("RandomNode_SeedTT", "The seed value to use for the random node. Set to -1 to use a random seed.");
-		}
-
-		static FText GetMinDescription()
-		{
-			return METASOUND_LOCTEXT("RandomNode_MinTT", "Min random value.");
-		}
-
-		static FText GetMaxDescription()
-		{
-			return METASOUND_LOCTEXT("RandomNode_MaxTT", "Max random value.");
-		}
-
-		static FText GetOutputDescription()
-		{
-			return METASOUND_LOCTEXT("RandomNode_OutputTT", "The randomly generated value.");
-		}
-
-		static FText GetOutputOnNextDescription()
-		{
-			return METASOUND_LOCTEXT("RandomNode_OutputNextTT", "Triggers when next is triggered.");
-		}
-
-		static FText GetOutputOnResetDescription()
-		{
-			return METASOUND_LOCTEXT("RandomNode_OutputResetTT", "Triggers when reset is triggered.");
-		}
-	}
-
 	// Mac Clang require linkage for constexpr
 	template<typename ValueType>
 	constexpr int32 Metasound::TRandomNodeOperator<ValueType>::DefaultSeed;
diff --git a/Engine/Plugins/Runtime/Metasound/Source/MetasoundStandardNodes/Private/MetasoundStereoDelayNode.cpp b/Engine/Plugins/Runtime/Metasound/Source/MetasoundStandardNodes/Private/MetasoundStereoDelayNode.cpp
--- a/Engine/Plugins/Runtime/Metasound/Source/MetasoundStandardNodes/Private/MetasoundStereoDelayNode.cpp
+++ b/Engine/Plugins/Runtime/Metasound/Source/MetasoundStandardNodes/Private/MetasoundStereoDelayNode.cpp
@@ -5,6 +5,7 @@
 #include "MetasoundExecutableOperator.h"
 #include "MetasoundNodeRegistrationMacro.h"
 #include "MetasoundDataTypeRegistrationMacro.h"
+#include "MetasoundParamHelper.h"
 #include "MetasoundPrimitives.h"
 #include "MetasoundStandardNodesNames.h"
 #include "MetasoundTrigger.h"
@@ -20,16 +21,16 @@ namespace Metasound
 {
 	namespace StereoDelay
 	{
-		static const TCHAR* InParamNameAudioInputLeft = TEXT("In Left");
-		static const TCHAR* InParamNameAudioInputRight = TEXT("In Right");
-		static const TCHAR* InParamNameDelayMode = TEXT("Delay Mode");
-		static const TCHAR* InParamNameDelayTime = TEXT("Delay Time");
-		static const TCHAR* InParamNameDelayRatio = TEXT("Delay Ratio");
-		static const TCHAR* InParamNameDryLevel = TEXT("Dry Level");
-		static const TCHAR* InParamNameWetLevel = TEXT("Wet Level");
-		static const TCHAR* InParamNameFeedbackAmount = TEXT("Feedback");
-		static const TCHAR* OutParamNameAudioLeft = TEXT("Out Left");
-		static const TCHAR* OutParamNameAudioRight = TEXT("Out Right");
+		METASOUND_PARAM(InAudioInputLeft, "In Left", "Left channel audio input.")
+		METASOUND_PARAM(InAudioInputRight, "In Right", "Right channel audio input.")
+		METASOUND_PARAM(InDelayMode, "Delay Mode", "Delay mode.")
+		METASOUND_PARAM(InDelayTime, "Delay Time", "The amount of time to delay the audio.")
+		METASOUND_PARAM(InDelayRatio, "Delay Ratio", "Delay spread for left and right channels. Allows left and right channels to have differential delay amounts. Useful for stereo channel decorrelation")
+		METASOUND_PARAM(InDryLevel, "Dry Level", "The dry level of the delay.")
+		METASOUND_PARAM(InWetLevel, "Wet Level", "The wet level of the delay.")
+		METASOUND_PARAM(InFeedbackAmount, "Feedback", "Feedback amount.")
+		METASOUND_PARAM(OutAudioLeft, "Out Left", "Left channel audio output.")
+		METASOUND_PARAM(OutAudioRight, "Out Right", "Right channel audio output.")
 
 		static float MaxDelaySeconds = 5.0f;
 	}
@@ -142,24 +143,28 @@ namespace Metasound
 
 	FDataReferenceCollection FStereoDelayOperator::GetInputs() const
 	{
+		using namespace StereoDelay; 
+
 		FDataReferenceCollection InputDataReferences;
-		InputDataReferences.AddDataReadReference(StereoDelay::InParamNameAudioInputLeft, FAudioBufferReadRef(LeftAudioInput));
-		InputDataReferences.AddDataReadReference(StereoDelay::InParamNameAudioInputRight, FAudioBufferReadRef(RightAudioInput));
-		InputDataReferences.AddDataReadReference(StereoDelay::InParamNameDelayMode, FStereoDelayModeReadRef(StereoDelayMode));
-		InputDataReferences.AddDataReadReference(StereoDelay::InParamNameDelayTime, FTimeReadRef(DelayTime));
-		InputDataReferences.AddDataReadReference(StereoDelay::InParamNameDelayRatio, FFloatReadRef(DelayRatio));
-		InputDataReferences.AddDataReadReference(StereoDelay::InParamNameDryLevel, FFloatReadRef(DryLevel));
-		InputDataReferences.AddDataReadReference(StereoDelay::InParamNameWetLevel, FFloatReadRef(WetLevel));
-		InputDataReferences.AddDataReadReference(StereoDelay::InParamNameFeedbackAmount, FFloatReadRef(Feedback));
+		InputDataReferences.AddDataReadReference(METASOUND_GET_PARAM_NAME(InAudioInputLeft), FAudioBufferReadRef(LeftAudioInput));
+		InputDataReferences.AddDataReadReference(METASOUND_GET_PARAM_NAME(InAudioInputRight), FAudioBufferReadRef(RightAudioInput));
+		InputDataReferences.AddDataReadReference(METASOUND_GET_PARAM_NAME(InDelayMode), FStereoDelayModeReadRef(StereoDelayMode));
+		InputDataReferences.AddDataReadReference(METASOUND_GET_PARAM_NAME(InDelayTime), FTimeReadRef(DelayTime));
+		InputDataReferences.AddDataReadReference(METASOUND_GET_PARAM_NAME(InDelayRatio), FFloatReadRef(DelayRatio));
+		InputDataReferences.AddDataReadReference(METASOUND_GET_PARAM_NAME(InDryLevel), FFloatReadRef(DryLevel));
+		InputDataReferences.AddDataReadReference(METASOUND_GET_PARAM_NAME(InWetLevel), FFloatReadRef(WetLevel));
+		InputDataReferences.AddDataReadReference(METASOUND_GET_PARAM_NAME(InFeedbackAmount), FFloatReadRef(Feedback));
 
 		return InputDataReferences;
 	}
 
 	FDataReferenceCollection FStereoDelayOperator::GetOutputs() const
 	{
+		using namespace StereoDelay; 
+
 		FDataReferenceCollection OutputDataReferences;
-		OutputDataReferences.AddDataReadReference(StereoDelay::OutParamNameAudioLeft, FAudioBufferReadRef(LeftAudioOutput));
-		OutputDataReferences.AddDataReadReference(StereoDelay::OutParamNameAudioRight, FAudioBufferReadRef(RightAudioOutput));
+		OutputDataReferences.AddDataReadReference(METASOUND_GET_PARAM_NAME(OutAudioLeft), FAudioBufferReadRef(LeftAudioOutput));
+		OutputDataReferences.AddDataReadReference(METASOUND_GET_PARAM_NAME(OutAudioRight), FAudioBufferReadRef(RightAudioOutput));
 		return OutputDataReferences;
 	}
 
@@ -282,20 +287,22 @@ namespace Metasound
 
 	const FVertexInterface& FStereoDelayOperator::GetVertexInterface()
 	{
+		using namespace StereoDelay; 
+
 		static const FVertexInterface Interface(
 			FInputVertexInterface(
-				TInputDataVertexModel<FAudioBuffer>(StereoDelay::InParamNameAudioInputLeft, METASOUND_LOCTEXT("LeftAudioInputTooltip", "Left channel audio input.")),
-				TInputDataVertexModel<FAudioBuffer>(StereoDelay::InParamNameAudioInputRight, METASOUND_LOCTEXT("RightAudioInputTooltip", "Right channel audio input.")),
-				TInputDataVertexModel<FEnumStereoDelayMode>(StereoDelay::InParamNameDelayMode, METASOUND_LOCTEXT("DelayModeTooltip", "Delay mode.")),
-				TInputDataVertexModel<FTime>(StereoDelay::InParamNameDelayTime, METASOUND_LOCTEXT("DelayTimeTooltip", "The amount of time to delay the audio."), 1.0f),
-				TInputDataVertexModel<float>(StereoDelay::InParamNameDelayRatio, METASOUND_LOCTEXT("DelayRatioTooltip", "Delay spread for left and right channels. Allows left and right channels to have differential delay amounts. Useful for stereo channel decorrelation"), 0.0f),
-				TInputDataVertexModel<float>(StereoDelay::InParamNameDryLevel, METASOUND_LOCTEXT("DryLevelTooltip", "The dry level of the delay."), 0.0f),
-				TInputDataVertexModel<float>(StereoDelay::InParamNameWetLevel, METASOUND_LOCTEXT("WetLevelTooltip", "The wet level of the delay."), 1.0f),
-				TInputDataVertexModel<float>(StereoDelay::InParamNameFeedbackAmount, METASOUND_LOCTEXT("FeedbackTooltip", "Feedback amount."), 0.0f)
+				TInputDataVertexModel<FAudioBuffer>(METASOUND_GET_PARAM_NAME_AND_METADATA(InAudioInputLeft)),
+				TInputDataVertexModel<FAudioBuffer>(METASOUND_GET_PARAM_NAME_AND_METADATA(InAudioInputRight)),
+				TInputDataVertexModel<FEnumStereoDelayMode>(METASOUND_GET_PARAM_NAME_AND_METADATA(InDelayMode)),
+				TInputDataVertexModel<FTime>(METASOUND_GET_PARAM_NAME_AND_METADATA(InDelayTime), 1.0f),
+				TInputDataVertexModel<float>(METASOUND_GET_PARAM_NAME_AND_METADATA(InDelayRatio), 0.0f),
+				TInputDataVertexModel<float>(METASOUND_GET_PARAM_NAME_AND_METADATA(InDryLevel), 0.0f),
+				TInputDataVertexModel<float>(METASOUND_GET_PARAM_NAME_AND_METADATA(InWetLevel), 1.0f),
+				TInputDataVertexModel<float>(METASOUND_GET_PARAM_NAME_AND_METADATA(InFeedbackAmount), 0.0f)
 			),
 			FOutputVertexInterface(
-				TOutputDataVertexModel<FAudioBuffer>(StereoDelay::OutParamNameAudioLeft, METASOUND_LOCTEXT("LeftDelayOutputTooltip", "Left channel audio output.")),
-				TOutputDataVertexModel<FAudioBuffer>(StereoDelay::OutParamNameAudioRight, METASOUND_LOCTEXT("RightDelayOutputTooltip", "Right channel audio output."))
+				TOutputDataVertexModel<FAudioBuffer>(METASOUND_GET_PARAM_NAME_AND_METADATA(OutAudioLeft)),
+				TOutputDataVertexModel<FAudioBuffer>(METASOUND_GET_PARAM_NAME_AND_METADATA(OutAudioRight))
 			)
 		);
 
@@ -326,17 +333,19 @@ namespace Metasound
 
 	TUniquePtr<IOperator> FStereoDelayOperator::CreateOperator(const FCreateOperatorParams& InParams, FBuildErrorArray& OutErrors)
 	{
+		using namespace StereoDelay; 
+
 		const FDataReferenceCollection& InputCollection = InParams.InputDataReferences;
 		const FInputVertexInterface& InputInterface = GetVertexInterface().GetInputInterface();
 
-		FAudioBufferReadRef LeftAudioIn = InputCollection.GetDataReadReferenceOrConstruct<FAudioBuffer>(StereoDelay::InParamNameAudioInputLeft, InParams.OperatorSettings);
-		FAudioBufferReadRef RightAudioIn = InputCollection.GetDataReadReferenceOrConstruct<FAudioBuffer>(StereoDelay::InParamNameAudioInputRight, InParams.OperatorSettings);
-		FStereoDelayModeReadRef StereoDelayMode = InputCollection.GetDataReadReferenceOrConstruct<FEnumStereoDelayMode>(StereoDelay::InParamNameDelayMode);
-		FTimeReadRef DelayTime = InputCollection.GetDataReadReferenceOrConstructWithVertexDefault<FTime>(InputInterface, StereoDelay::InParamNameDelayTime, InParams.OperatorSettings);
-		FFloatReadRef DelayRatio = InputCollection.GetDataReadReferenceOrConstructWithVertexDefault<float>(InputInterface, StereoDelay::InParamNameDelayRatio, InParams.OperatorSettings);
-		FFloatReadRef DryLevel = InputCollection.GetDataReadReferenceOrConstructWithVertexDefault<float>(InputInterface, StereoDelay::InParamNameDryLevel, InParams.OperatorSettings);
-		FFloatReadRef WetLevel = InputCollection.GetDataReadReferenceOrConstructWithVertexDefault<float>(InputInterface, StereoDelay::InParamNameWetLevel, InParams.OperatorSettings);
-		FFloatReadRef Feedback = InputCollection.GetDataReadReferenceOrConstructWithVertexDefault<float>(InputInterface, StereoDelay::InParamNameFeedbackAmount, InParams.OperatorSettings);
+		FAudioBufferReadRef LeftAudioIn = InputCollection.GetDataReadReferenceOrConstruct<FAudioBuffer>(METASOUND_GET_PARAM_NAME(InAudioInputLeft), InParams.OperatorSettings);
+		FAudioBufferReadRef RightAudioIn = InputCollection.GetDataReadReferenceOrConstruct<FAudioBuffer>(METASOUND_GET_PARAM_NAME(InAudioInputRight), InParams.OperatorSettings);
+		FStereoDelayModeReadRef StereoDelayMode = InputCollection.GetDataReadReferenceOrConstruct<FEnumStereoDelayMode>(METASOUND_GET_PARAM_NAME(InDelayMode));
+		FTimeReadRef DelayTime = InputCollection.GetDataReadReferenceOrConstructWithVertexDefault<FTime>(InputInterface, METASOUND_GET_PARAM_NAME(InDelayTime), InParams.OperatorSettings);
+		FFloatReadRef DelayRatio = InputCollection.GetDataReadReferenceOrConstructWithVertexDefault<float>(InputInterface, METASOUND_GET_PARAM_NAME(InDelayRatio), InParams.OperatorSettings);
+		FFloatReadRef DryLevel = InputCollection.GetDataReadReferenceOrConstructWithVertexDefault<float>(InputInterface, METASOUND_GET_PARAM_NAME(InDryLevel), InParams.OperatorSettings);
+		FFloatReadRef WetLevel = InputCollection.GetDataReadReferenceOrConstructWithVertexDefault<float>(InputInterface, METASOUND_GET_PARAM_NAME(InWetLevel), InParams.OperatorSettings);
+		FFloatReadRef Feedback = InputCollection.GetDataReadReferenceOrConstructWithVertexDefault<float>(InputInterface, METASOUND_GET_PARAM_NAME(InFeedbackAmount), InParams.OperatorSettings);
 
 		return MakeUnique<FStereoDelayOperator>(InParams.OperatorSettings, LeftAudioIn, RightAudioIn, StereoDelayMode, DelayTime, DelayRatio, DryLevel, WetLevel, Feedback);
 	}
diff --git a/Engine/Plugins/Runtime/Metasound/Source/MetasoundStandardNodes/Private/MetasoundStereoPannerNode.cpp b/Engine/Plugins/Runtime/Metasound/Source/MetasoundStandardNodes/Private/MetasoundStereoPannerNode.cpp
--- a/Engine/Plugins/Runtime/Metasound/Source/MetasoundStandardNodes/Private/MetasoundStereoPannerNode.cpp
+++ b/Engine/Plugins/Runtime/Metasound/Source/MetasoundStandardNodes/Private/MetasoundStereoPannerNode.cpp
@@ -5,6 +5,7 @@
 #include "MetasoundExecutableOperator.h"
 #include "MetasoundNodeRegistrationMacro.h"
 #include "MetasoundDataTypeRegistrationMacro.h"
+#include "MetasoundParamHelper.h"
 #include "MetasoundPrimitives.h"
 #include "MetasoundStandardNodesNames.h"
 #include "MetasoundTrigger.h"
@@ -21,65 +22,11 @@ namespace Metasound
 {
 	namespace StereoPannerVertexNames
 	{
-		const FVertexName& GetInputAudioName()
-		{
-			static FVertexName Name = TEXT("In");
-			return Name;
-		}
-
-		const FText& GetInputAudioDescription()
-		{
-			static FText Desc = METASOUND_LOCTEXT("StereoPannerNodeInDesc", "The input audio to pan.");
-			return Desc;
-		}
-
-		const FVertexName& GetInputPanAmountName()
-		{
-			static FVertexName Name = TEXT("Pan Amount");
-			return Name;
-		}
-
-		const FText& GetInputPanAmountDescription()
-		{
-			static FText Desc = METASOUND_LOCTEXT("StereoPannerNodePanAmountDesc", "The amount of pan. -1.0 is full left, 1.0 is full right.");
-			return Desc;
-		}
-
-		const FVertexName& GetInputPanningLawName()
-		{
-			static FVertexName Name = TEXT("Panning Law");
-			return Name;
-		}
-
-		const FText& GetInputPanningLawDescription()
-		{
-			static FText Desc = METASOUND_LOCTEXT("StereoPannerNodePanningLawDescription", "Which panning law should be used for the stereo panner.");
-			return Desc;
-		}
-
-		const FVertexName& GetOutputAudioLeftName()
-		{
-			static FVertexName Name = TEXT("Out Left");
-			return Name;
-		}
-
-		const FText& GetOutputAudioLeftDescription()
-		{
-			static FText Desc = METASOUND_LOCTEXT("StereoPannerNodeOutputLeftDescription", "Left channel audio output.");
-			return Desc;
-		}
-
-		const FVertexName& GetOutputAudioRightName()
-		{
-			static FVertexName Name = TEXT("Out Right");
-			return Name;
-		}
-
-		const FText& GetOutputAudioRightDescription()
-		{
-			static FText Desc = METASOUND_LOCTEXT("StereoPannerNodeOutputRightDescription", "Right channel audio output.");
-			return Desc;
-		}
+		METASOUND_PARAM(InputAudio, "In", "The input audio to pan.")
+		METASOUND_PARAM(InputPanAmount, "Pan Amount", "The amount of pan. -1.0 is full left, 1.0 is full right.")
+		METASOUND_PARAM(InputPanningLaw, "Panning Law", "Which panning law should be used for the stereo panner.")
+		METASOUND_PARAM(OutputAudioLeft, "Out Left", "Left channel audio output.")
+		METASOUND_PARAM(OutputAudioRight, "Out Right", "Right channel audio output.")
 	}
 
 	enum class EPanningLaw
@@ -155,9 +102,9 @@ namespace Metasound
 		using namespace StereoPannerVertexNames;
 
 		FDataReferenceCollection InputDataReferences;
-		InputDataReferences.AddDataReadReference(GetInputAudioName(), AudioInput);
-		InputDataReferences.AddDataReadReference(GetInputPanAmountName(), PanningAmount);
-		InputDataReferences.AddDataReadReference(GetInputPanningLawName(), PanningLaw);
+		InputDataReferences.AddDataReadReference(METASOUND_GET_PARAM_NAME(InputAudio), AudioInput);
+		InputDataReferences.AddDataReadReference(METASOUND_GET_PARAM_NAME(InputPanAmount), PanningAmount);
+		InputDataReferences.AddDataReadReference(METASOUND_GET_PARAM_NAME(InputPanningLaw), PanningLaw);
 
 		return InputDataReferences;
 	}
@@ -167,8 +114,8 @@ namespace Metasound
 		using namespace StereoPannerVertexNames;
 
 		FDataReferenceCollection OutputDataReferences;
-		OutputDataReferences.AddDataReadReference(GetOutputAudioLeftName(), AudioLeftOutput);
-		OutputDataReferences.AddDataReadReference(GetOutputAudioRightName(), AudioRightOutput);
+		OutputDataReferences.AddDataReadReference(METASOUND_GET_PARAM_NAME(OutputAudioLeft), AudioLeftOutput);
+		OutputDataReferences.AddDataReadReference(METASOUND_GET_PARAM_NAME(OutputAudioRight), AudioRightOutput);
 
 		return OutputDataReferences;
 	}
@@ -233,13 +180,13 @@ namespace Metasound
 
 		static const FVertexInterface Interface(
 			FInputVertexInterface(
-				TInputDataVertexModel<FAudioBuffer>(GetInputAudioName(), GetInputAudioDescription()),
-				TInputDataVertexModel<float>(GetInputPanAmountName(), GetInputPanAmountDescription(), 0.0f),
-				TInputDataVertexModel<FEnumPanningLaw>(GetInputPanningLawName(), GetInputPanningLawDescription())
+				TInputDataVertexModel<FAudioBuffer>(METASOUND_GET_PARAM_NAME_AND_METADATA(InputAudio)),
+				TInputDataVertexModel<float>(METASOUND_GET_PARAM_NAME_AND_METADATA(InputPanAmount), 0.0f),
+				TInputDataVertexModel<FEnumPanningLaw>(METASOUND_GET_PARAM_NAME_AND_METADATA(InputPanningLaw))
 			),
 			FOutputVertexInterface(
-				TOutputDataVertexModel<FAudioBuffer>(GetOutputAudioLeftName(), GetOutputAudioLeftDescription()),
-				TOutputDataVertexModel<FAudioBuffer>(GetOutputAudioRightName(), GetOutputAudioRightDescription())
+				TOutputDataVertexModel<FAudioBuffer>(METASOUND_GET_PARAM_NAME_AND_METADATA(OutputAudioLeft)),
+				TOutputDataVertexModel<FAudioBuffer>(METASOUND_GET_PARAM_NAME_AND_METADATA(OutputAudioRight))
 			)
 		);
 
@@ -275,9 +222,9 @@ namespace Metasound
 
 		using namespace StereoPannerVertexNames;
 
-		FAudioBufferReadRef AudioIn = InputCollection.GetDataReadReferenceOrConstruct<FAudioBuffer>(GetInputAudioName(), InParams.OperatorSettings);
-		FFloatReadRef PanningAmount = InputCollection.GetDataReadReferenceOrConstructWithVertexDefault<float>(InputInterface, GetInputPanAmountName(), InParams.OperatorSettings);
-		FPanningLawReadRef PanningLaw = InputCollection.GetDataReadReferenceOrConstruct<FEnumPanningLaw>(GetInputPanningLawName());
+		FAudioBufferReadRef AudioIn = InputCollection.GetDataReadReferenceOrConstruct<FAudioBuffer>(METASOUND_GET_PARAM_NAME(InputAudio), InParams.OperatorSettings);
+		FFloatReadRef PanningAmount = InputCollection.GetDataReadReferenceOrConstructWithVertexDefault<float>(InputInterface, METASOUND_GET_PARAM_NAME(InputPanAmount), InParams.OperatorSettings);
+		FPanningLawReadRef PanningLaw = InputCollection.GetDataReadReferenceOrConstruct<FEnumPanningLaw>(METASOUND_GET_PARAM_NAME(InputPanningLaw));
 
 		return MakeUnique<FStereoPannerOperator>(InParams.OperatorSettings, AudioIn, PanningAmount, PanningLaw);
 	}
diff --git a/Engine/Plugins/Runtime/Metasound/Source/MetasoundStandardNodes/Private/MetasoundTriggerAnyNode.cpp b/Engine/Plugins/Runtime/Metasound/Source/MetasoundStandardNodes/Private/MetasoundTriggerAnyNode.cpp
--- a/Engine/Plugins/Runtime/Metasound/Source/MetasoundStandardNodes/Private/MetasoundTriggerAnyNode.cpp
+++ b/Engine/Plugins/Runtime/Metasound/Source/MetasoundStandardNodes/Private/MetasoundTriggerAnyNode.cpp
@@ -6,6 +6,7 @@
 #include "MetasoundNodeRegistrationMacro.h"
 #include "MetasoundStandardNodesCategories.h"
 #include "MetasoundStandardNodesNames.h"
+#include "MetasoundParamHelper.h"
 #include "MetasoundTrigger.h"
 #include "MetasoundVertex.h"
 
@@ -43,6 +44,8 @@ namespace Metasound
 
 	namespace TriggerAnyVertexNames
 	{
+		METASOUND_PARAM(OutputTrigger, "Out", "Triggered when any of the input triggers have been triggered. ")
+
 		METASOUNDSTANDARDNODES_API const FVertexName GetInputTriggerName(uint32 InIndex)
 		{
 			return *FString::Format(TEXT("In {0}"), { InIndex });
@@ -53,16 +56,9 @@ namespace Metasound
 			return METASOUND_LOCTEXT_FORMAT("TriggerAnyInputTriggerDesc", "Trigger {0} input. The output trigger is hit when any of the input triggers are hit.", InIndex);
 		}
 
-		METASOUNDSTANDARDNODES_API const FVertexName& GetOutputTriggerName()
-		{
-			static const FVertexName Name = TEXT("Out");
-			return Name;
-		}
-
-		METASOUNDSTANDARDNODES_API const FText& GetOutputTriggerDescription()
+		METASOUNDSTANDARDNODES_API const FText GetInputTriggerDisplayName(uint32 InIndex)
 		{
-			static const FText Desc = METASOUND_LOCTEXT("TriggerAnyOutputTriggerDesc", "Triggered when any of the input triggers have been triggered. ");
-			return Desc;
+			return METASOUND_LOCTEXT_FORMAT("TriggerAnyInputTriggerDisplayName", "In {0}", InIndex);
 		}
 	}
 
@@ -80,11 +76,17 @@ namespace Metasound
 
 				for (uint32 i = 0; i < NumInputs; ++i)
 				{
-					InputInterface.Add(TInputDataVertexModel<FTrigger>(GetInputTriggerName(i), GetInputTriggerDescription(i)));
+					const FDataVertexMetadata InputTriggerMetadata
+					{
+						  GetInputTriggerDescription(i) // description
+						, GetInputTriggerDisplayName(i) // display name
+					};
+
+					InputInterface.Add(TInputDataVertexModel<FTrigger>(GetInputTriggerName(i), InputTriggerMetadata));
 				}
 
 				FOutputVertexInterface OutputInterface;
-				OutputInterface.Add(TOutputDataVertexModel<FTrigger>(GetOutputTriggerName(), GetOutputTriggerDescription()));
+				OutputInterface.Add(TOutputDataVertexModel<FTrigger>(METASOUND_GET_PARAM_NAME_AND_METADATA(OutputTrigger)));
 
 				return FVertexInterface(InputInterface, OutputInterface);
 			};
@@ -153,7 +155,7 @@ namespace Metasound
 			using namespace TriggerAnyVertexNames;
 
 			FDataReferenceCollection Outputs;
-			Outputs.AddDataReadReference(GetOutputTriggerName(), OutputTrigger);
+			Outputs.AddDataReadReference(METASOUND_GET_PARAM_NAME(OutputTrigger), OutputTrigger);
 
 			return Outputs;
 		}
diff --git a/Engine/Plugins/Runtime/Metasound/Source/MetasoundStandardNodes/Private/MetasoundTriggerOnThresholdNode.cpp b/Engine/Plugins/Runtime/Metasound/Source/MetasoundStandardNodes/Private/MetasoundTriggerOnThresholdNode.cpp
--- a/Engine/Plugins/Runtime/Metasound/Source/MetasoundStandardNodes/Private/MetasoundTriggerOnThresholdNode.cpp
+++ b/Engine/Plugins/Runtime/Metasound/Source/MetasoundStandardNodes/Private/MetasoundTriggerOnThresholdNode.cpp
@@ -32,7 +32,15 @@ namespace Metasound
 		DEFINE_METASOUND_ENUM_ENTRY(EBufferTriggerType::RisingEdge, "RisingEdgeDescription", "Rising Edge", "RisingEdgeDescriptionTT", ""),
 		DEFINE_METASOUND_ENUM_ENTRY(EBufferTriggerType::FallingEdge, "FallingEdgeDescription", "Falling Edge", "FallingEdgeDescriptionTT", ""),
 		DEFINE_METASOUND_ENUM_ENTRY(EBufferTriggerType::AbsThreshold, "AbsThresholdDescription", "Abs Threshold", "AbsThresholdDescriptionTT", "")
-		DEFINE_METASOUND_ENUM_END()
+	DEFINE_METASOUND_ENUM_END()
+
+	namespace TriggerOnThresholdVertexNames
+	{
+		METASOUND_PARAM(OutPin, "Out", "Output");
+		METASOUND_PARAM(InPin, "In", "Input");
+		METASOUND_PARAM(InThresholdPin, "Threshold", "Trigger Threshold");
+		METASOUND_PARAM(InTriggerType, "Type", "Trigger Threshold Type");
+	}
 
 	class FTriggerOnThresholdOperator : public IOperator
 	{
@@ -54,11 +62,6 @@ namespace Metasound
 		FTriggerWriteRef Out;
 		bool bTriggered = false;
 		float LastSample = 0.f;
-
-		static constexpr const TCHAR* OutPinName = TEXT("Out");
-		static constexpr const TCHAR* InPinName = TEXT("In");
-		static constexpr const TCHAR* ThresholdPinName = TEXT("Threshold");
-		static constexpr const TCHAR* TriggerType = TEXT("Type");
 	};
 
 	// Mac Clang requires linkage on constexpr
@@ -166,29 +169,35 @@ namespace Metasound
 
 	FDataReferenceCollection FTriggerOnThresholdOperator::GetInputs() const
 	{
+		using namespace TriggerOnThresholdVertexNames;
+		
 		FDataReferenceCollection InputDataReferences;
-		InputDataReferences.AddDataReadReference(ThresholdPinName, FFloatReadRef(Threshold));
-		InputDataReferences.AddDataReadReference(InPinName, FAudioBufferReadRef(In));
+		InputDataReferences.AddDataReadReference(METASOUND_GET_PARAM_NAME(InThresholdPin), FFloatReadRef(Threshold));
+		InputDataReferences.AddDataReadReference(METASOUND_GET_PARAM_NAME(InPin), FAudioBufferReadRef(In));
 		return InputDataReferences;
 	}
 
 	FDataReferenceCollection FTriggerOnThresholdOperator::GetOutputs() const
 	{
+		using namespace TriggerOnThresholdVertexNames; 
+
 		FDataReferenceCollection OutputDataReferences;
-		OutputDataReferences.AddDataReadReference(OutPinName, FTriggerWriteRef(Out));
+		OutputDataReferences.AddDataReadReference(METASOUND_GET_PARAM_NAME(OutPin), FTriggerWriteRef(Out));
 		return OutputDataReferences;
 	}
 
 	FVertexInterface FTriggerOnThresholdOperator::DeclareVertexInterface()
 	{
+		using namespace TriggerOnThresholdVertexNames; 
+
 		static const FVertexInterface Interface(
 			FInputVertexInterface(
-				TInputDataVertexModel<FAudioBuffer>(InPinName, METASOUND_LOCTEXT("BufferInDescription", "Input")),
-				TInputDataVertexModel<float>(ThresholdPinName, METASOUND_LOCTEXT("ThresholdDescription", "Trigger Threshold"), DefaultThreshold),
-				TInputDataVertexModel<FEnumBufferTriggerType>(TriggerType, METASOUND_LOCTEXT("ThresholdDescription", "Trigger Threshold"))
+				TInputDataVertexModel<FAudioBuffer>(METASOUND_GET_PARAM_NAME_AND_METADATA(InPin)),
+				TInputDataVertexModel<float>(METASOUND_GET_PARAM_NAME_AND_METADATA(InThresholdPin), DefaultThreshold),
+				TInputDataVertexModel<FEnumBufferTriggerType>(METASOUND_GET_PARAM_NAME_AND_METADATA(InTriggerType))
 			),
 			FOutputVertexInterface(
-				TOutputDataVertexModel<FTrigger>(OutPinName, METASOUND_LOCTEXT("TriggerOutDescription", "Output"))
+				TOutputDataVertexModel<FTrigger>(METASOUND_GET_PARAM_NAME_AND_METADATA(OutPin))
 			)
 		);
 		return Interface;
@@ -228,16 +237,18 @@ namespace Metasound
 
 	TUniquePtr<IOperator> FTriggerOnThresholdOperator::CreateOperator(const FCreateOperatorParams& InParams, FBuildErrorArray& OutErrors)
 	{
+		using namespace TriggerOnThresholdVertexNames; 
+
 		const FTriggerOnThresholdNode& Node = static_cast<const FTriggerOnThresholdNode&>(InParams.Node);
 		const FDataReferenceCollection& InputCol = InParams.InputDataReferences;
 		const FOperatorSettings& Settings = InParams.OperatorSettings;
 		const FInputVertexInterface& InputInterface = DeclareVertexInterface().GetInputInterface();
 		
 		// Static property pin, only used for factory.
-		FBufferTriggerTypeReadRef Type = InputCol.GetDataReadReferenceOrConstruct<FEnumBufferTriggerType>(TriggerType);
+		FBufferTriggerTypeReadRef Type = InputCol.GetDataReadReferenceOrConstruct<FEnumBufferTriggerType>(METASOUND_GET_PARAM_NAME(InTriggerType));
 
-		FAudioBufferReadRef InputBuffer = InputCol.GetDataReadReferenceOrConstruct<FAudioBuffer>(InPinName, Settings);
-		FFloatReadRef Threshold = InputCol.GetDataReadReferenceOrConstructWithVertexDefault<float>(InputInterface, ThresholdPinName, Settings);
+		FAudioBufferReadRef InputBuffer = InputCol.GetDataReadReferenceOrConstruct<FAudioBuffer>(METASOUND_GET_PARAM_NAME(InPin), Settings);
+		FFloatReadRef Threshold = InputCol.GetDataReadReferenceOrConstructWithVertexDefault<float>(InputInterface, METASOUND_GET_PARAM_NAME(InThresholdPin), Settings);
 
 		switch (*Type)
 		{
diff --git a/Engine/Plugins/Runtime/Metasound/Source/MetasoundStandardNodes/Private/MetasoundTriggerRouteNode.cpp b/Engine/Plugins/Runtime/Metasound/Source/MetasoundStandardNodes/Private/MetasoundTriggerRouteNode.cpp
--- a/Engine/Plugins/Runtime/Metasound/Source/MetasoundStandardNodes/Private/MetasoundTriggerRouteNode.cpp
+++ b/Engine/Plugins/Runtime/Metasound/Source/MetasoundStandardNodes/Private/MetasoundTriggerRouteNode.cpp
@@ -9,6 +9,7 @@
 #include "MetasoundNodeInterface.h"
 #include "MetasoundNodeRegistrationMacro.h"
 #include "MetasoundOperatorInterface.h"
+#include "MetasoundParamHelper.h"
 #include "MetasoundPrimitives.h"
 #include "MetasoundStandardNodesNames.h"
 #include "MetasoundStandardNodesCategories.h"
@@ -25,11 +26,19 @@ namespace Metasound
 {
 	namespace TriggerRouteVertexNames
 	{
+		METASOUND_PARAM(OutputTrigger, "On Set", "Triggered when any of the input triggers are set.")
+		METASOUND_PARAM(OutputValue, "Value", "The output value set by the input triggers.")
+
 		const FVertexName GetInputTriggerName(uint32 InIndex)
 		{
 			return *FString::Format(TEXT("Set {0}"), { InIndex });
 		}
 
+		const FText GetInputTriggerDisplayName(uint32 InIndex)
+		{
+			return METASOUND_LOCTEXT_FORMAT("InputTriggerDisplayName", "Set {0}", InIndex);
+		}
+
 		const FText GetInputTriggerDescription(uint32 InIndex)
 		{
 			if (InIndex == 0)
@@ -49,28 +58,9 @@ namespace Metasound
 			return METASOUND_LOCTEXT_FORMAT("TriggerRouteValueDesc", "The input value ({0}) to route to the output when triggered by Set {0}.", InIndex);
 		}
 
-		const FVertexName& GetOutputTriggerName()
+		const FText GetInputValueDisplayName(uint32 InIndex)
 		{
-			static const FVertexName Name = TEXT("On Set");
-			return Name;
-		}
-
-		const FText& GetOutputTriggerDescription()
-		{
-			static const FText Desc = METASOUND_LOCTEXT("TriggerRouteOnSetDesc", "Triggered when any of the input triggers are set.");
-			return Desc;
-		}
-
-		const FVertexName& GetOutputValueName()
-		{
-			static const FVertexName Name = TEXT("Value");
-			return Name;
-		}
-
-		const FText& GetOutputValueDescription()
-		{
-			static const FText Desc = METASOUND_LOCTEXT("TriggerRouteOutputValueDesc", "The output value set by the input triggers.");
-			return Desc;
+			return METASOUND_LOCTEXT_FORMAT("InputValueDisplayName", "Value {0}", InIndex);
 		}
 	}
 
@@ -111,13 +101,25 @@ namespace Metasound
 				FInputVertexInterface InputInterface;
 				for (uint32 i = 0; i < NumInputs; ++i)
 				{
-					InputInterface.Add(TInputDataVertexModel<FTrigger>(GetInputTriggerName(i), GetInputTriggerDescription(i)));
-					InputInterface.Add(TInputDataVertexModel<ValueType>(GetInputValueName(i), GetInputValueDescription(i)));
+					const FDataVertexMetadata InputTriggerMetadata
+					{
+						  GetInputTriggerDescription(i) // description
+						, GetInputTriggerDisplayName(i) // display name
+					};
+
+					const FDataVertexMetadata InputValueMetadata
+					{
+						  GetInputValueDescription(i) // description
+						, GetInputValueDisplayName(i) // display name
+					};
+
+					InputInterface.Add(TInputDataVertexModel<FTrigger>(GetInputTriggerName(i), InputTriggerMetadata));
+					InputInterface.Add(TInputDataVertexModel<ValueType>(GetInputValueName(i), InputValueMetadata));
 				}
 
 				FOutputVertexInterface OutputInterface;
-				OutputInterface.Add(TOutputDataVertexModel<FTrigger>(GetOutputTriggerName(), GetOutputTriggerDescription()));
-				OutputInterface.Add(TOutputDataVertexModel<ValueType>(GetOutputValueName(), GetOutputValueDescription()));
+				OutputInterface.Add(TOutputDataVertexModel<FTrigger>(METASOUND_GET_PARAM_NAME_AND_METADATA(OutputTrigger)));
+				OutputInterface.Add(TOutputDataVertexModel<ValueType>(METASOUND_GET_PARAM_NAME_AND_METADATA(OutputValue)));
 				return FVertexInterface(InputInterface, OutputInterface);
 			};
 
@@ -194,8 +196,8 @@ namespace Metasound
 			using namespace TriggerRouteVertexNames;
 
 			FDataReferenceCollection Outputs;
-			Outputs.AddDataReadReference(GetOutputTriggerName(), OutputTrigger);
-			Outputs.AddDataReadReference(GetOutputValueName(), OutputValue);
+			Outputs.AddDataReadReference(METASOUND_GET_PARAM_NAME(OutputTrigger), OutputTrigger);
+			Outputs.AddDataReadReference(METASOUND_GET_PARAM_NAME(OutputValue), OutputValue);
 
 			return Outputs;
 		}
diff --git a/Engine/Plugins/Runtime/Metasound/Source/MetasoundStandardNodes/Private/MetasoundTriggerToggleNode.cpp b/Engine/Plugins/Runtime/Metasound/Source/MetasoundStandardNodes/Private/MetasoundTriggerToggleNode.cpp
--- a/Engine/Plugins/Runtime/Metasound/Source/MetasoundStandardNodes/Private/MetasoundTriggerToggleNode.cpp
+++ b/Engine/Plugins/Runtime/Metasound/Source/MetasoundStandardNodes/Private/MetasoundTriggerToggleNode.cpp
@@ -11,6 +11,7 @@
 #include "MetasoundNodeInterface.h"
 #include "MetasoundNodeRegistrationMacro.h"
 #include "MetasoundOperatorInterface.h"
+#include "MetasoundParamHelper.h"
 #include "MetasoundPrimitives.h"
 #include "MetasoundStandardNodesNames.h"
 #include "MetasoundStandardNodesCategories.h"
@@ -22,11 +23,11 @@ namespace Metasound
 {
 	namespace TriggerToggle
 	{
-		const FVertexName GetInputOnTriggerName() { return TEXT("On"); }
-		const FVertexName GetInputOffTriggerName() { return TEXT("Off"); }
-		const FVertexName GetInputInitName() { return TEXT("Init"); }
-		const FVertexName GetOutputTriggerName() { return TEXT("Out"); }
-		const FVertexName GetOutputValueName() { return TEXT("Value"); }
+		METASOUND_PARAM(InputOnTrigger, "On", "Trigger to toggle gate output to 1.")
+		METASOUND_PARAM(InputOffTrigger, "Off", "Trigger to toggle gate output to 0.")
+		METASOUND_PARAM(InputInit, "Init", "Initial value of the output gate.")
+		METASOUND_PARAM(OutputTrigger, "Out", "Triggers output when gate is toggled.")
+		METASOUND_PARAM(OutputValue, "Value", "Current output value of the toggle.")
 	}
 
 	class FTriggerToggleOperator : public TExecutableOperator<FTriggerToggleOperator>
@@ -66,9 +67,9 @@ namespace Metasound
 		using namespace TriggerToggle;
 
 		FDataReferenceCollection InputDataReferences;
-		InputDataReferences.AddDataReadReference(GetInputOnTriggerName(), TriggerOn);
-		InputDataReferences.AddDataReadReference(GetInputOffTriggerName(), TriggerOff);
-		InputDataReferences.AddDataReadReference(GetInputInitName(), InitValue);
+		InputDataReferences.AddDataReadReference(METASOUND_GET_PARAM_NAME(InputOnTrigger), TriggerOn);
+		InputDataReferences.AddDataReadReference(METASOUND_GET_PARAM_NAME(InputOffTrigger), TriggerOff);
+		InputDataReferences.AddDataReadReference(METASOUND_GET_PARAM_NAME(InputInit), InitValue);
 		return InputDataReferences;
 	}
 
@@ -77,8 +78,8 @@ namespace Metasound
 		using namespace TriggerToggle;
 
 		FDataReferenceCollection OutputDataReferences;
-		OutputDataReferences.AddDataReadReference(GetOutputTriggerName(), TriggerOutput);
-		OutputDataReferences.AddDataReadReference(GetOutputValueName(), ValueOutput);
+		OutputDataReferences.AddDataReadReference(METASOUND_GET_PARAM_NAME(OutputTrigger), TriggerOutput);
+		OutputDataReferences.AddDataReadReference(METASOUND_GET_PARAM_NAME(OutputValue), ValueOutput);
 
 		return OutputDataReferences;
 	}
@@ -115,9 +116,9 @@ namespace Metasound
 		using namespace TriggerToggle;
 
 		const FInputVertexInterface& InputInterface = GetVertexInterface().GetInputInterface();
-		FTriggerReadRef InTriggerOn = InParams.InputDataReferences.GetDataReadReferenceOrConstruct<FTrigger>(GetInputOnTriggerName(), InParams.OperatorSettings);
-		FTriggerReadRef InTriggerOff = InParams.InputDataReferences.GetDataReadReferenceOrConstruct<FTrigger>(GetInputOffTriggerName(), InParams.OperatorSettings);
-		FBoolReadRef InInitValue = InParams.InputDataReferences.GetDataReadReferenceOrConstructWithVertexDefault<bool>(InputInterface, GetInputInitName(), InParams.OperatorSettings);
+		FTriggerReadRef InTriggerOn = InParams.InputDataReferences.GetDataReadReferenceOrConstruct<FTrigger>(METASOUND_GET_PARAM_NAME(InputOnTrigger), InParams.OperatorSettings);
+		FTriggerReadRef InTriggerOff = InParams.InputDataReferences.GetDataReadReferenceOrConstruct<FTrigger>(METASOUND_GET_PARAM_NAME(InputOffTrigger), InParams.OperatorSettings);
+		FBoolReadRef InInitValue = InParams.InputDataReferences.GetDataReadReferenceOrConstructWithVertexDefault<bool>(InputInterface, METASOUND_GET_PARAM_NAME(InputInit), InParams.OperatorSettings);
 
 		return MakeUnique<FTriggerToggleOperator>(InParams.OperatorSettings, InTriggerOn, InTriggerOff, InInitValue);
 	}
@@ -128,13 +129,13 @@ namespace Metasound
 
 		static const FVertexInterface Interface(
 			FInputVertexInterface(
-				TInputDataVertexModel<FTrigger>(GetInputOnTriggerName(), METASOUND_LOCTEXT("TriggerGateOnTT", "Trigger to toggle gate output to 1.")),
-				TInputDataVertexModel<FTrigger>(GetInputOffTriggerName(), METASOUND_LOCTEXT("TriggerGateOffTT", "Trigger to toggle gate output to 0.")),
-				TInputDataVertexModel<bool>(GetInputInitName(), METASOUND_LOCTEXT("TriggerGateInitTT", "Initial value of the output gate."))
+				TInputDataVertexModel<FTrigger>(METASOUND_GET_PARAM_NAME_AND_METADATA(InputOnTrigger)),
+				TInputDataVertexModel<FTrigger>(METASOUND_GET_PARAM_NAME_AND_METADATA(InputOffTrigger)),
+				TInputDataVertexModel<bool>(METASOUND_GET_PARAM_NAME_AND_METADATA(InputInit))
 			),
 			FOutputVertexInterface(
-				TOutputDataVertexModel<FTrigger>(GetOutputTriggerName(), METASOUND_LOCTEXT("TriggerGateOutputTriggerTT", "Triggers output when gate is toggled.")),
-				TOutputDataVertexModel<bool>(GetOutputValueName(), METASOUND_LOCTEXT("BoolOutputTT", "Current output value of the toggle."))
+				TOutputDataVertexModel<FTrigger>(METASOUND_GET_PARAM_NAME_AND_METADATA(OutputTrigger)),
+				TOutputDataVertexModel<bool>(METASOUND_GET_PARAM_NAME_AND_METADATA(OutputValue))
 			)
 		);
 
diff --git a/Engine/Plugins/Runtime/Metasound/Source/MetasoundStandardNodes/Private/MetasoundWaveShaperNode.cpp b/Engine/Plugins/Runtime/Metasound/Source/MetasoundStandardNodes/Private/MetasoundWaveShaperNode.cpp
--- a/Engine/Plugins/Runtime/Metasound/Source/MetasoundStandardNodes/Private/MetasoundWaveShaperNode.cpp
+++ b/Engine/Plugins/Runtime/Metasound/Source/MetasoundStandardNodes/Private/MetasoundWaveShaperNode.cpp
@@ -5,6 +5,7 @@
 #include "MetasoundExecutableOperator.h"
 #include "MetasoundNodeRegistrationMacro.h"
 #include "MetasoundDataTypeRegistrationMacro.h"
+#include "MetasoundParamHelper.h"
 #include "MetasoundPrimitives.h"
 #include "MetasoundStandardNodesNames.h"
 #include "MetasoundAudioBuffer.h"
@@ -16,12 +17,12 @@ namespace Metasound
 {
 	namespace WaveShaperNode
 	{
-		static const TCHAR* InParamNameAudioInput = TEXT("In");
-		static const TCHAR* InParamNameWaveShapeAmount = TEXT("Amount");
-		static const TCHAR* InParamNameWaveShapeBias = TEXT("Bias");
-		static const TCHAR* InParamNameOutputGain = TEXT("OutputGain");
-		static const TCHAR* InParamNameType = TEXT("Type");
-		static const TCHAR* OutParamNameAudio = TEXT("Out");
+		METASOUND_PARAM(InAudioInput, "In", "Audio input.");
+		METASOUND_PARAM(InWaveShapeAmount, "Amount", "The amount of wave shaping to apply.");
+		METASOUND_PARAM(InWaveShapeBias, "Bias", "DC offset to apply before wave shaping.");
+		METASOUND_PARAM(InOutputGain, "OutputGain", "The amount of gain to apply after processing.");
+		METASOUND_PARAM(InType, "Type", "Algorithm to use to process the audio.");
+		METASOUND_PARAM(OutAudio, "Out", "Audio output.");
 	}
 
 	using namespace WaveShaperNode;
@@ -96,20 +97,24 @@ namespace Metasound
 
 	FDataReferenceCollection FWaveShaperOperator::GetInputs() const
 	{
+		using namespace WaveShaperNode;
+		
 		FDataReferenceCollection InputDataReferences;
-		InputDataReferences.AddDataReadReference(InParamNameAudioInput, AudioInput);
-		InputDataReferences.AddDataReadReference(InParamNameWaveShapeAmount, Amount);
-		InputDataReferences.AddDataReadReference(InParamNameWaveShapeBias, Bias);
-		InputDataReferences.AddDataReadReference(InParamNameOutputGain, OutputGain);
-		InputDataReferences.AddDataReadReference(InParamNameType, Type);
+		InputDataReferences.AddDataReadReference(METASOUND_GET_PARAM_NAME(InAudioInput), AudioInput);
+		InputDataReferences.AddDataReadReference(METASOUND_GET_PARAM_NAME(InWaveShapeAmount), Amount);
+		InputDataReferences.AddDataReadReference(METASOUND_GET_PARAM_NAME(InWaveShapeBias), Bias);
+		InputDataReferences.AddDataReadReference(METASOUND_GET_PARAM_NAME(InOutputGain), OutputGain);
+		InputDataReferences.AddDataReadReference(METASOUND_GET_PARAM_NAME(InType), Type);
 
 		return InputDataReferences;
 	}
 
 	FDataReferenceCollection FWaveShaperOperator::GetOutputs() const
 	{
+		using namespace WaveShaperNode; 
+
 		FDataReferenceCollection OutputDataReferences;
-		OutputDataReferences.AddDataReadReference(OutParamNameAudio, AudioOutput);
+		OutputDataReferences.AddDataReadReference(METASOUND_GET_PARAM_NAME(OutAudio), AudioOutput);
 		return OutputDataReferences;
 	}
 
@@ -132,14 +137,14 @@ namespace Metasound
 	{
 		static const FVertexInterface Interface(
 			FInputVertexInterface(
-				TInputDataVertexModel<FAudioBuffer>(InParamNameAudioInput, METASOUND_LOCTEXT("AudioInputTT", "Audio input.")),
-				TInputDataVertexModel<float>(InParamNameWaveShapeAmount, METASOUND_LOCTEXT("WaveShapeAmountTT", "The amount of wave shaping to apply."), 1.0f),
-				TInputDataVertexModel<float>(InParamNameWaveShapeBias, METASOUND_LOCTEXT("WaveShaperBiasTT", "DC offset to apply before wave shaping."), 0.0f),
-				TInputDataVertexModel<float>(InParamNameOutputGain, METASOUND_LOCTEXT("WaveShaperGainTT", "The amount of gain to apply after processing."), 1.0f),
-				TInputDataVertexModel<FEnumEWaveShaperType>(InParamNameType, METASOUND_LOCTEXT("WaveShaperTypeTT", "Algorithm to use to process the audio."))
+				TInputDataVertexModel<FAudioBuffer>(METASOUND_GET_PARAM_NAME_AND_METADATA(InAudioInput)),
+				TInputDataVertexModel<float>(METASOUND_GET_PARAM_NAME_AND_METADATA(InWaveShapeAmount), 1.0f),
+				TInputDataVertexModel<float>(METASOUND_GET_PARAM_NAME_AND_METADATA(InWaveShapeBias), 0.0f),
+				TInputDataVertexModel<float>(METASOUND_GET_PARAM_NAME_AND_METADATA(InOutputGain), 1.0f),
+				TInputDataVertexModel<FEnumEWaveShaperType>(METASOUND_GET_PARAM_NAME_AND_METADATA(InType))
 			),
 			FOutputVertexInterface(
-				TOutputDataVertexModel<FAudioBuffer>(OutParamNameAudio, METASOUND_LOCTEXT("WaveShaperOutputTT", "Audio output."))
+				TOutputDataVertexModel<FAudioBuffer>(METASOUND_GET_PARAM_NAME_AND_METADATA(OutAudio))
 			)
 		);
 
@@ -173,11 +178,11 @@ namespace Metasound
 		const FDataReferenceCollection& InputCollection = InParams.InputDataReferences;
 		const FInputVertexInterface& InputInterface = GetVertexInterface().GetInputInterface();
 
-		FAudioBufferReadRef AudioIn = InputCollection.GetDataReadReferenceOrConstruct<FAudioBuffer>(InParamNameAudioInput, InParams.OperatorSettings);
-		FFloatReadRef InAmount = InputCollection.GetDataReadReferenceOrConstructWithVertexDefault<float>(InputInterface, InParamNameWaveShapeAmount, InParams.OperatorSettings);
-		FFloatReadRef InBias = InputCollection.GetDataReadReferenceOrConstructWithVertexDefault<float>(InputInterface, InParamNameWaveShapeBias, InParams.OperatorSettings);
-		FFloatReadRef OutGain = InputCollection.GetDataReadReferenceOrConstructWithVertexDefault<float>(InputInterface, InParamNameOutputGain, InParams.OperatorSettings);
-		FEnumWaveShaperReadRef InType = InputCollection.GetDataReadReferenceOrConstruct<FEnumEWaveShaperType>(InParamNameType);
+		FAudioBufferReadRef AudioIn = InputCollection.GetDataReadReferenceOrConstruct<FAudioBuffer>(METASOUND_GET_PARAM_NAME(InAudioInput), InParams.OperatorSettings);
+		FFloatReadRef InAmount = InputCollection.GetDataReadReferenceOrConstructWithVertexDefault<float>(InputInterface, METASOUND_GET_PARAM_NAME(InWaveShapeAmount), InParams.OperatorSettings);
+		FFloatReadRef InBias = InputCollection.GetDataReadReferenceOrConstructWithVertexDefault<float>(InputInterface, METASOUND_GET_PARAM_NAME(InWaveShapeBias), InParams.OperatorSettings);
+		FFloatReadRef OutGain = InputCollection.GetDataReadReferenceOrConstructWithVertexDefault<float>(InputInterface, METASOUND_GET_PARAM_NAME(InOutputGain), InParams.OperatorSettings);
+		FEnumWaveShaperReadRef InType = InputCollection.GetDataReadReferenceOrConstruct<FEnumEWaveShaperType>(METASOUND_GET_PARAM_NAME(InType));
 
 		return MakeUnique<FWaveShaperOperator>(InParams.OperatorSettings, AudioIn, InAmount, InBias, OutGain, InType);
 	}
diff --git a/Engine/Plugins/Runtime/Metasound/Source/MetasoundStandardNodes/Private/MetasoundWaveWriterNode.cpp b/Engine/Plugins/Runtime/Metasound/Source/MetasoundStandardNodes/Private/MetasoundWaveWriterNode.cpp
--- a/Engine/Plugins/Runtime/Metasound/Source/MetasoundStandardNodes/Private/MetasoundWaveWriterNode.cpp
+++ b/Engine/Plugins/Runtime/Metasound/Source/MetasoundStandardNodes/Private/MetasoundWaveWriterNode.cpp
@@ -5,6 +5,7 @@
 #include "MetasoundEnumRegistrationMacro.h"
 #include "MetasoundExecutableOperator.h"
 #include "MetasoundNodeRegistrationMacro.h"
+#include "MetasoundParamHelper.h"
 #include "MetasoundPrimitives.h"
 #include "MetasoundStandardNodesCategories.h"
 #include "MetasoundStandardNodesNames.h"
@@ -17,6 +18,12 @@
 
 namespace Metasound
 {
+	namespace WaveWriterVertexNames
+	{
+		METASOUND_PARAM(InEnabledPin, "Enabled", "If this wave writer is enabled or not.");
+		METASOUND_PARAM(InFilenamePrefixPin, "Filename Prefix", "Filename Prefix of file you are writing.");
+	}
+
 	// Incremental wave writer class.
 	class FWaveWriter
 	{
@@ -210,16 +217,6 @@ namespace Metasound
 			static const FString DefaultFileName = TEXT("Output");
 			return DefaultFileName;
 		}
-		static const TCHAR* GetEnabledPinName()
-		{
-			static const TCHAR* EnabledPinName = TEXT("Enabled");
-			return EnabledPinName;
-		}
-		static const TCHAR* GetFilenamePrefixPinName()
-		{
-			static const TCHAR* FilenamePrefixPinName = TEXT("Filename Prefix");
-			return FilenamePrefixPinName;
-		}
 	}
 
 	template<int32 NumInputChannels>
@@ -270,17 +267,25 @@ namespace Metasound
 			auto CreateDefaultInterface = []()-> FVertexInterface
 			{
 				using namespace WaveWriterOperatorPrivate;
+				using namespace WaveWriterVertexNames;
 
 				// inputs
 				FInputVertexInterface InputInterface(
-					TInputDataVertexModel<FString>(GetFilenamePrefixPinName(), METASOUND_LOCTEXT("WaveWriterFilenamePrefixDescription", "Filename Prefix of file you are writing."), GetDefaultFileName()),
-					TInputDataVertexModel<bool>(GetEnabledPinName(), METASOUND_LOCTEXT("WaveWriterEnabledDescription", "If this wave writer is enabled or not."), true)
+					TInputDataVertexModel<FString>(METASOUND_GET_PARAM_NAME_AND_METADATA(InFilenamePrefixPin), GetDefaultFileName()),
+					TInputDataVertexModel<bool>(METASOUND_GET_PARAM_NAME_AND_METADATA(InEnabledPin), true)
 				);
 
+
 				// For backwards compatibility with previous (mono) node, in the case of 1 channels, just provide the old interface.
 				for (int32 InputIndex = 0; InputIndex < NumInputChannels; ++InputIndex)
 				{
-					InputInterface.Add(TInputDataVertexModel<FAudioBuffer>(GetAudioInputName(InputIndex), GetAudioInputDescription(InputIndex)));
+					const FDataVertexMetadata AudioInputMetadata
+					{
+						  GetAudioInputDescription(InputIndex) // description
+						, GetAudioInputDisplayName(InputIndex) // display name
+					};
+
+					InputInterface.Add(TInputDataVertexModel<FAudioBuffer>(GetAudioInputName(InputIndex), AudioInputMetadata));
 				}
 				FOutputVertexInterface OutputInterface;
 
@@ -380,6 +385,28 @@ namespace Metasound
 			return *FString::Printf(TEXT("In %d"), InInputIndex);
 		}
 
+		static const FText GetAudioInputDisplayName(int32 InInputIndex)
+		{
+			if (NumInputChannels == 1)
+			{
+				// To maintain backwards compatibility with previous implementation keep the pin name the same.
+				static const FText AudioInputPinName = METASOUND_LOCTEXT("AudioInputPinNameIn", "In");
+				return AudioInputPinName;
+			}
+			else if (NumInputChannels == 2)
+			{
+				if (InInputIndex == 0)
+				{
+					return METASOUND_LOCTEXT_FORMAT("AudioInputIn2ChannelName", "In {0} L", InInputIndex);
+				}
+				else
+				{
+					return METASOUND_LOCTEXT_FORMAT("AudioInputIn2ChannelName", "In {0} R", InInputIndex);
+				}
+			}
+			return METASOUND_LOCTEXT_FORMAT("AudioInputInChannelName", "In {0}", InInputIndex);
+		}
+
 		static const FText GetAudioInputDescription(int32 InputIndex)
 		{
 			return METASOUND_LOCTEXT_FORMAT("WaveWriterAudioInputDescription", "Audio Input #: {0}", InputIndex);
@@ -453,12 +480,13 @@ namespace Metasound
 	TUniquePtr<Metasound::IOperator> TWaveWriterOperator<NumInputChannels>::CreateOperator(const FCreateOperatorParams& InParams, FBuildErrorArray& OutErrors)
 	{
 		using namespace WaveWriterOperatorPrivate;
+		using namespace WaveWriterVertexNames;
 
 		const FDataReferenceCollection& InputCol = InParams.InputDataReferences;
 		const FOperatorSettings& Settings = InParams.OperatorSettings;
 		const FInputVertexInterface& InputInterface = DeclareVertexInterface().GetInputInterface();
 
-		FStringReadRef FilenamePrefix = InputCol.GetDataReadReferenceOrConstructWithVertexDefault<FString>(InputInterface, GetFilenamePrefixPinName(), Settings);
+		FStringReadRef FilenamePrefix = InputCol.GetDataReadReferenceOrConstructWithVertexDefault<FString>(InputInterface, METASOUND_GET_PARAM_NAME(InFilenamePrefixPin), Settings);
 
 		int32 NumConnectedAudioPins = 0;
 		TArray<FAudioBufferReadRef> InputBuffers;
@@ -475,7 +503,7 @@ namespace Metasound
 			return MakeUnique<TWaveWriterOperator>(
 				Settings,
 				MoveTemp(InputBuffers),
-				InputCol.GetDataReadReferenceOrConstructWithVertexDefault<bool>(InputInterface, GetEnabledPinName(), Settings),
+				InputCol.GetDataReadReferenceOrConstructWithVertexDefault<bool>(InputInterface, METASOUND_GET_PARAM_NAME(InEnabledPin), Settings),
 				GetNameCache(),
 				*FilenamePrefix
 			);
diff --git a/Engine/Plugins/Runtime/Metasound/Source/MetasoundStandardNodes/Public/MetasoundRandomNode.h b/Engine/Plugins/Runtime/Metasound/Source/MetasoundStandardNodes/Public/MetasoundRandomNode.h
--- a/Engine/Plugins/Runtime/Metasound/Source/MetasoundStandardNodes/Public/MetasoundRandomNode.h
+++ b/Engine/Plugins/Runtime/Metasound/Source/MetasoundStandardNodes/Public/MetasoundRandomNode.h
@@ -5,6 +5,7 @@
 #include "CoreMinimal.h"
 #include "MetasoundFacade.h"
 #include "MetasoundExecutableOperator.h"
+#include "MetasoundParamHelper.h"
 #include "MetasoundPrimitives.h"
 #include "MetasoundStandardNodesNames.h"
 #include "MetasoundTrigger.h"
@@ -19,27 +20,15 @@ namespace Metasound
 {
 	namespace RandomNodeNames
 	{
-		/** Input vertex names. */
-		const FVertexName& GetInputNextTriggerName();
-		const FVertexName& GetInputResetTriggerName();
-		const FVertexName& GetInputSeedName();
-		const FVertexName& GetInputMinName();
-		const FVertexName& GetInputMaxName();
-
-		/** Output vertex names. */
-		const FVertexName& GetOutputOnNextTriggerName();
-		const FVertexName& GetOutputOnResetTriggerName();
-		const FVertexName& GetOutputValueName();
-
-		/** Descriptions. */
-		static FText GetNextTriggerDescription();
-		static FText GetResetDescription();
-		static FText GetSeedDescription();
-		static FText GetMinDescription();
-		static FText GetMaxDescription();
-		static FText GetOutputDescription();
-		static FText GetOutputOnNextDescription();
-		static FText GetOutputOnResetDescription();
+		METASOUND_PARAM(InputNextTrigger, "Next", "Trigger to generate the next random integer.")
+		METASOUND_PARAM(InputResetTrigger, "Reset", "Trigger to reset the random sequence with the supplied seed. Useful to get randomized repetition.")
+		METASOUND_PARAM(InputSeed, "Seed", "The seed value to use for the random node. Set to -1 to use a random seed.")
+		METASOUND_PARAM(InputMin, "Min", "Min random value.")
+		METASOUND_PARAM(InputMax, "Max", "Max random value.")
+		
+		METASOUND_PARAM(OutputOnNextTrigger, "On Next", "Triggers when next is triggered.")
+		METASOUND_PARAM(OutputOnResetTrigger, "On Reset", "Triggers when reset is triggered.")
+		METASOUND_PARAM(OutputValue, "Value", "The randomly generated value.")
 	}
 
 	template<typename ValueType>
@@ -88,12 +77,16 @@ namespace Metasound
 
 		static TDataReadReference<int32> CreateMinValueRef(const FDataReferenceCollection& InputCollection, const FInputVertexInterface& InputInterface, const FOperatorSettings& InOperatorSettings)
 		{
-			return InputCollection.GetDataReadReferenceOrConstructWithVertexDefault<int32>(InputInterface, RandomNodeNames::GetInputMinName(), InOperatorSettings);
+			using namespace RandomNodeNames; 
+
+			return InputCollection.GetDataReadReferenceOrConstructWithVertexDefault<int32>(InputInterface, METASOUND_GET_PARAM_NAME(InputMin), InOperatorSettings);
 		}
 
 		static TDataReadReference<int32> CreateMaxValueRef(const FDataReferenceCollection& InputCollection, const FInputVertexInterface& InputInterface, const FOperatorSettings& InOperatorSettings)
 		{
-			return InputCollection.GetDataReadReferenceOrConstructWithVertexDefault<int32>(InputInterface, RandomNodeNames::GetInputMaxName(), InOperatorSettings);
+			using namespace RandomNodeNames; 
+
+			return InputCollection.GetDataReadReferenceOrConstructWithVertexDefault<int32>(InputInterface, METASOUND_GET_PARAM_NAME(InputMax), InOperatorSettings);
 		}
 	};
 
@@ -137,12 +130,16 @@ namespace Metasound
 
 		static TDataReadReference<float> CreateMinValueRef(const FDataReferenceCollection& InputCollection, const FInputVertexInterface& InputInterface, const FOperatorSettings& InOperatorSettings)
 		{
-			return InputCollection.GetDataReadReferenceOrConstructWithVertexDefault<float>(InputInterface, RandomNodeNames::GetInputMinName(), InOperatorSettings);
+			using namespace RandomNodeNames;
+
+			return InputCollection.GetDataReadReferenceOrConstructWithVertexDefault<float>(InputInterface, METASOUND_GET_PARAM_NAME(InputMin), InOperatorSettings);
 		}
 
 		static TDataReadReference<float> CreateMaxValueRef(const FDataReferenceCollection& InputCollection, const FInputVertexInterface& InputInterface, const FOperatorSettings& InOperatorSettings)
 		{
-			return InputCollection.GetDataReadReferenceOrConstructWithVertexDefault<float>(InputInterface, RandomNodeNames::GetInputMaxName(), InOperatorSettings);
+			using namespace RandomNodeNames;
+			
+			return InputCollection.GetDataReadReferenceOrConstructWithVertexDefault<float>(InputInterface, METASOUND_GET_PARAM_NAME(InputMax), InOperatorSettings);
 		}
 	};
 
@@ -237,12 +234,16 @@ namespace Metasound
 
 		static TDataReadReference<FTime> CreateMinValueRef(const FDataReferenceCollection& InputCollection, const FInputVertexInterface& InputInterface, const FOperatorSettings& InOperatorSettings)
 		{
-			return InputCollection.GetDataReadReferenceOrConstructWithVertexDefault<FTime>(InputInterface, RandomNodeNames::GetInputMinName(), InOperatorSettings);
+			using namespace RandomNodeNames; 
+
+			return InputCollection.GetDataReadReferenceOrConstructWithVertexDefault<FTime>(InputInterface, METASOUND_GET_PARAM_NAME(InputMin), InOperatorSettings);
 		}
 
 		static TDataReadReference<FTime> CreateMaxValueRef(const FDataReferenceCollection& InputCollection, const FInputVertexInterface& InputInterface, const FOperatorSettings& InOperatorSettings)
 		{
-			return InputCollection.GetDataReadReferenceOrConstructWithVertexDefault<FTime>(InputInterface, RandomNodeNames::GetInputMaxName(), InOperatorSettings);
+			using namespace RandomNodeNames;
+
+			return InputCollection.GetDataReadReferenceOrConstructWithVertexDefault<FTime>(InputInterface, METASOUND_GET_PARAM_NAME(InputMax), InOperatorSettings);
 		}
 	};
 
@@ -262,16 +263,16 @@ namespace Metasound
 			{
 				static const FVertexInterface DefaultInterface(
 					FInputVertexInterface(
-						TInputDataVertexModel<FTrigger>(GetInputNextTriggerName(), GetNextTriggerDescription()),
-						TInputDataVertexModel<FTrigger>(GetInputResetTriggerName(), GetResetDescription()),
-						TInputDataVertexModel<int32>(GetInputSeedName(), GetSeedDescription(), DefaultSeed),
-						TInputDataVertexModel<ValueType>(GetInputMinName(), GetMinDescription(), TRandomNodeSpecialization<ValueType>::GetDefaultMin()),
-						TInputDataVertexModel<ValueType>(GetInputMaxName(), GetMaxDescription(), TRandomNodeSpecialization<ValueType>::GetDefaultMax())
+						TInputDataVertexModel<FTrigger>(METASOUND_GET_PARAM_NAME_AND_METADATA(InputNextTrigger)),
+						TInputDataVertexModel<FTrigger>(METASOUND_GET_PARAM_NAME_AND_METADATA(InputResetTrigger)),
+						TInputDataVertexModel<int32>(METASOUND_GET_PARAM_NAME_AND_METADATA(InputSeed), DefaultSeed),
+						TInputDataVertexModel<ValueType>(METASOUND_GET_PARAM_NAME_AND_METADATA(InputMin), TRandomNodeSpecialization<ValueType>::GetDefaultMin()),
+						TInputDataVertexModel<ValueType>(METASOUND_GET_PARAM_NAME_AND_METADATA(InputMax), TRandomNodeSpecialization<ValueType>::GetDefaultMax())
 					),
 					FOutputVertexInterface(
-						TOutputDataVertexModel<FTrigger>(GetOutputOnNextTriggerName(), GetOutputOnNextDescription()),
-						TOutputDataVertexModel<FTrigger>(GetOutputOnResetTriggerName(), GetOutputOnResetDescription()),
-						TOutputDataVertexModel<ValueType>(GetOutputValueName(), GetOutputDescription())
+						TOutputDataVertexModel<FTrigger>(METASOUND_GET_PARAM_NAME_AND_METADATA(OutputOnNextTrigger)),
+						TOutputDataVertexModel<FTrigger>(METASOUND_GET_PARAM_NAME_AND_METADATA(OutputOnResetTrigger)),
+						TOutputDataVertexModel<ValueType>(METASOUND_GET_PARAM_NAME_AND_METADATA(OutputValue))
 					)
 				);
 
@@ -281,14 +282,14 @@ namespace Metasound
 			{
 				static const FVertexInterface DefaultInterface(
 					FInputVertexInterface(
-						TInputDataVertexModel<FTrigger>(GetInputNextTriggerName(), GetNextTriggerDescription()),
-						TInputDataVertexModel<FTrigger>(GetInputResetTriggerName(), GetResetDescription()),
-						TInputDataVertexModel<int32>(GetInputSeedName(), GetSeedDescription(), DefaultSeed)
+						TInputDataVertexModel<FTrigger>(METASOUND_GET_PARAM_NAME_AND_METADATA(InputNextTrigger)),
+						TInputDataVertexModel<FTrigger>(METASOUND_GET_PARAM_NAME_AND_METADATA(InputResetTrigger)),
+						TInputDataVertexModel<int32>(METASOUND_GET_PARAM_NAME_AND_METADATA(InputSeed), DefaultSeed)
 					),
 					FOutputVertexInterface(
-						TOutputDataVertexModel<FTrigger>(GetOutputOnNextTriggerName(), GetOutputOnNextDescription()),
-						TOutputDataVertexModel<FTrigger>(GetOutputOnResetTriggerName(), GetOutputOnResetDescription()),
-						TOutputDataVertexModel<ValueType>(GetOutputValueName(), GetOutputDescription())
+						TOutputDataVertexModel<FTrigger>(METASOUND_GET_PARAM_NAME_AND_METADATA(OutputOnNextTrigger)),
+						TOutputDataVertexModel<FTrigger>(METASOUND_GET_PARAM_NAME_AND_METADATA(OutputOnResetTrigger)),
+						TOutputDataVertexModel<ValueType>(METASOUND_GET_PARAM_NAME_AND_METADATA(OutputValue))
 					)
 				);
 
@@ -324,9 +325,9 @@ namespace Metasound
  			const FDataReferenceCollection& InputCollection = InParams.InputDataReferences;
 			using namespace RandomNodeNames;
 
- 			FTriggerReadRef InNextTrigger = InputCollection.GetDataReadReferenceOrConstruct<FTrigger>(GetInputNextTriggerName(), InParams.OperatorSettings);
-			FTriggerReadRef InResetTrigger = InputCollection.GetDataReadReferenceOrConstruct<FTrigger>(GetInputResetTriggerName(), InParams.OperatorSettings);
-			FInt32ReadRef InSeedValue = InputCollection.GetDataReadReferenceOrConstructWithVertexDefault<int32>(InputInterface, GetInputSeedName(), InParams.OperatorSettings);
+ 			FTriggerReadRef InNextTrigger = InputCollection.GetDataReadReferenceOrConstruct<FTrigger>(METASOUND_GET_PARAM_NAME(InputNextTrigger), InParams.OperatorSettings);
+			FTriggerReadRef InResetTrigger = InputCollection.GetDataReadReferenceOrConstruct<FTrigger>(METASOUND_GET_PARAM_NAME(InputResetTrigger), InParams.OperatorSettings);
+			FInt32ReadRef InSeedValue = InputCollection.GetDataReadReferenceOrConstructWithVertexDefault<int32>(InputInterface, METASOUND_GET_PARAM_NAME(InputSeed), InParams.OperatorSettings);
 
 			// note: random bool does not have a range
 			if (TRandomNodeSpecialization<ValueType>::HasRange())
@@ -397,15 +398,15 @@ namespace Metasound
 			using namespace RandomNodeNames;
 
 			FDataReferenceCollection Inputs;
-			Inputs.AddDataReadReference(GetInputNextTriggerName(), NextTrigger);
- 			Inputs.AddDataReadReference(GetInputResetTriggerName(), ResetTrigger);
- 			Inputs.AddDataReadReference(GetInputSeedName(), SeedValue);
+			Inputs.AddDataReadReference(METASOUND_GET_PARAM_NAME(InputNextTrigger), NextTrigger);
+ 			Inputs.AddDataReadReference(METASOUND_GET_PARAM_NAME(InputResetTrigger), ResetTrigger);
+ 			Inputs.AddDataReadReference(METASOUND_GET_PARAM_NAME(InputSeed), SeedValue);
 
 			// If the type doesn't have a range no need for input pins to define it
 			if (TRandomNodeSpecialization<ValueType>::HasRange())
 			{
-				Inputs.AddDataReadReference(GetInputMinName(), MinValue);
-				Inputs.AddDataReadReference(GetInputMaxName(), MinValue);
+				Inputs.AddDataReadReference(METASOUND_GET_PARAM_NAME(InputMin), MinValue);
+				Inputs.AddDataReadReference(METASOUND_GET_PARAM_NAME(InputMax), MinValue);
 			}
 			return Inputs;
 		}
@@ -415,9 +416,9 @@ namespace Metasound
 			using namespace RandomNodeNames;
 
 			FDataReferenceCollection Outputs;
-			Outputs.AddDataReadReference(GetOutputOnNextTriggerName(), TriggerOutOnNext);
-			Outputs.AddDataReadReference(GetOutputOnResetTriggerName(), TriggerOutOnReset);
-			Outputs.AddDataReadReference(GetOutputValueName(), OutputValue);
+			Outputs.AddDataReadReference(METASOUND_GET_PARAM_NAME(OutputOnNextTrigger), TriggerOutOnNext);
+			Outputs.AddDataReadReference(METASOUND_GET_PARAM_NAME(OutputOnResetTrigger), TriggerOutOnReset);
+			Outputs.AddDataReadReference(METASOUND_GET_PARAM_NAME(OutputValue), OutputValue);
 			return Outputs;
 		}
 
