Fix additional spot missing rename feedback from prior CL: predicate --> projection
#rnx

diff --git a/Engine/Plugins/Runtime/Metasound/Source/MetasoundEditor/Private/MetasoundEditorModule.cpp b/Engine/Plugins/Runtime/Metasound/Source/MetasoundEditor/Private/MetasoundEditorModule.cpp
--- a/Engine/Plugins/Runtime/Metasound/Source/MetasoundEditor/Private/MetasoundEditorModule.cpp
+++ b/Engine/Plugins/Runtime/Metasound/Source/MetasoundEditor/Private/MetasoundEditorModule.cpp
@@ -713,7 +713,7 @@ namespace Metasound
 
 				// Required to ensure logic to order nodes for presets exclusive to
 				// editor is propagated to transform instances while editing in editor.
-				Frontend::DocumentTransform::RegisterNodeDisplayNamePredicate([](const Frontend::FNodeHandle& NodeHandle)
+				Frontend::DocumentTransform::RegisterNodeDisplayNameProjection([](const Frontend::FNodeHandle& NodeHandle)
 				{
 					constexpr bool bIncludeNamespace = false;
 					return FGraphBuilder::GetDisplayName(*NodeHandle, bIncludeNamespace);
diff --git a/Engine/Plugins/Runtime/Metasound/Source/MetasoundFrontend/Private/MetasoundFrontendTransform.cpp b/Engine/Plugins/Runtime/Metasound/Source/MetasoundFrontend/Private/MetasoundFrontendTransform.cpp
--- a/Engine/Plugins/Runtime/Metasound/Source/MetasoundFrontend/Private/MetasoundFrontendTransform.cpp
+++ b/Engine/Plugins/Runtime/Metasound/Source/MetasoundFrontend/Private/MetasoundFrontendTransform.cpp
@@ -19,16 +19,16 @@ namespace Metasound
 		namespace DocumentTransform
 		{
 #if WITH_EDITOR
-			FGetNodeDisplayNamePredicate NodeDisplayNamePredicate;
+			FGetNodeDisplayNameProjection NodeDisplayNameProjection;
 
-			void RegisterNodeDisplayNamePredicate(FGetNodeDisplayNamePredicate&& InNamePredicate)
+			void RegisterNodeDisplayNameProjection(FGetNodeDisplayNameProjection&& InNameProjection)
 			{
-				NodeDisplayNamePredicate = MoveTemp(InNamePredicate);
+				NodeDisplayNameProjection = MoveTemp(InNameProjection);
 			}
 
-			FGetNodeDisplayNamePredicateRef GetNodeDisplayNamePredicate()
+			FGetNodeDisplayNameProjectionRef GetNodeDisplayNameProjection()
 			{
-				return NodeDisplayNamePredicate;
+				return NodeDisplayNameProjection;
 			}
 #endif // WITH_EDITOR
 		} // namespace DocumentTransform
@@ -757,7 +757,7 @@ namespace Metasound
 			FMetasoundFrontendInterfaceStyle Style = ReferencedGraphHandle->GetInputStyle();
 			InPresetGraphHandle->SetInputStyle(Style);
 
-			Style.SortDefaults(NodeHandles, DocumentTransform::GetNodeDisplayNamePredicate());
+			Style.SortDefaults(NodeHandles, DocumentTransform::GetNodeDisplayNameProjection());
 
 			for (const FNodeHandle& InputNode : NodeHandles)
 			{
@@ -797,7 +797,7 @@ namespace Metasound
 			FMetasoundFrontendInterfaceStyle Style = ReferencedGraphHandle->GetOutputStyle();
 			InPresetGraphHandle->SetOutputStyle(Style);
 
-			Style.SortDefaults(NodeHandles, DocumentTransform::GetNodeDisplayNamePredicate());
+			Style.SortDefaults(NodeHandles, DocumentTransform::GetNodeDisplayNameProjection());
 
 			// Set input node location
 			for (const FNodeHandle& OutputNode : NodeHandles)
diff --git a/Engine/Plugins/Runtime/Metasound/Source/MetasoundFrontend/Public/MetasoundFrontendTransform.h b/Engine/Plugins/Runtime/Metasound/Source/MetasoundFrontend/Public/MetasoundFrontendTransform.h
--- a/Engine/Plugins/Runtime/Metasound/Source/MetasoundFrontend/Public/MetasoundFrontendTransform.h
+++ b/Engine/Plugins/Runtime/Metasound/Source/MetasoundFrontend/Public/MetasoundFrontendTransform.h
@@ -17,11 +17,11 @@ namespace Metasound
 		namespace DocumentTransform
 		{
 #if WITH_EDITOR
-			using FGetNodeDisplayNamePredicate = TFunction<FText(const FNodeHandle&)>;
-			using FGetNodeDisplayNamePredicateRef = TFunctionRef<FText(const FNodeHandle&)>;
+			using FGetNodeDisplayNameProjection = TFunction<FText(const FNodeHandle&)>;
+			using FGetNodeDisplayNameProjectionRef = TFunctionRef<FText(const FNodeHandle&)>;
 
-			METASOUNDFRONTEND_API void RegisterNodeDisplayNamePredicate(FGetNodeDisplayNamePredicate&& InNamePredicate);
-			METASOUNDFRONTEND_API FGetNodeDisplayNamePredicateRef GetNodeDisplayNamePredicate();
+			METASOUNDFRONTEND_API void RegisterNodeDisplayNameProjection(FGetNodeDisplayNameProjection&& InNameProjection);
+			METASOUNDFRONTEND_API FGetNodeDisplayNameProjectionRef GetNodeDisplayNameProjection();
 #endif // WITH_EDITOR
 		}
 
