MetaSound Wave Player node typo fix

#rnx

diff --git a/Engine/Plugins/Runtime/Metasound/Source/MetasoundEngine/Private/MetasoundWavePlayerNode.cpp b/Engine/Plugins/Runtime/Metasound/Source/MetasoundEngine/Private/MetasoundWavePlayerNode.cpp
--- a/Engine/Plugins/Runtime/Metasound/Source/MetasoundEngine/Private/MetasoundWavePlayerNode.cpp
+++ b/Engine/Plugins/Runtime/Metasound/Source/MetasoundEngine/Private/MetasoundWavePlayerNode.cpp
@@ -42,7 +42,7 @@ namespace Metasound
 		METASOUND_PARAM(OutputCuePointID, "Cue Point ID", "The cue point ID that was triggered.")
 		METASOUND_PARAM(OutputCuePointLabel, "Cue Point Label", "The cue point label that was triggered (if there was a label parsed in the imported .wav file).")
 		METASOUND_PARAM(OutputLoopRatio, "Loop Ratio", "Returns the current playback location as a ratio of the loop (0-1) if looping is enabled.")
-		METASOUND_PARAM(OutputPlaybackLocation, "OutputPlaybackLocation", "Returns the absolute position of the wave playback as a ratio of wave duration (0-1).")
+		METASOUND_PARAM(OutputPlaybackLocation, "Playback Location", "Returns the absolute position of the wave playback as a ratio of wave duration (0-1).")
 		METASOUND_PARAM(OutputAudioLeft, "Out Left", "The left channel audio output. Mono wave assets will be upmixed to dual stereo.")
 		METASOUND_PARAM(OutputAudioRight, "Out Right", "The right channel audio output. Mono wave assets will be upmixed to dual stereo.")
 	}
