Default to 'direct' ESelectInfo in SSearchableComboBox::SetSelectedItem
- Internal SelectedItem member was not getting set after last fix to avoid selection on key navigation while the menu was open
- Minor clean-up
#rnx

diff --git a/Engine/Plugins/Runtime/Metasound/Source/MetasoundEditor/Private/MetasoundNodeDetailCustomization.cpp b/Engine/Plugins/Runtime/Metasound/Source/MetasoundEditor/Private/MetasoundNodeDetailCustomization.cpp
--- a/Engine/Plugins/Runtime/Metasound/Source/MetasoundEditor/Private/MetasoundNodeDetailCustomization.cpp
+++ b/Engine/Plugins/Runtime/Metasound/Source/MetasoundEditor/Private/MetasoundNodeDetailCustomization.cpp
@@ -16,6 +16,7 @@
 #include "MetasoundFrontendController.h"
 #include "MetasoundFrontendDataTypeRegistry.h"
 #include "MetasoundFrontendRegistries.h"
+#include "Modules/ModuleManager.h"
 #include "PropertyCustomizationHelpers.h"
 #include "PropertyEditorDelegates.h"
 #include "SAssetDropTarget.h"
@@ -25,6 +26,7 @@
 #include "SSearchableComboBox.h"
 #include "Templates/Casts.h"
 #include "Templates/SharedPointer.h"
+#include "Templates/UniquePtr.h"
 #include "UObject/WeakObjectPtr.h"
 #include "UObject/WeakObjectPtrTemplates.h"
 #include "Widgets/Images/SImage.h"
@@ -786,7 +788,7 @@ namespace Metasound
 			const TSharedPtr<FString>* SelectedItem = ComboOptions.FindByPredicate(NameMatchesPredicate);
 			if (ensure(SelectedItem))
 			{
-				DataTypeComboBox->SetSelectedItem(*SelectedItem);
+				DataTypeComboBox->SetSelectedItem(*SelectedItem, ESelectInfo::Direct);
 			}
 		}
 
@@ -838,9 +840,9 @@ namespace Metasound
 					{
 						Member->OnRenameRequested.Clear();
 						RenameRequestedHandle = Member->OnRenameRequested.AddLambda([this]()
-							{
-								FSlateApplication::Get().SetKeyboardFocus(NameEditableTextBox.ToSharedRef(), EFocusCause::SetDirectly);
-							});
+						{
+							FSlateApplication::Get().SetKeyboardFocus(NameEditableTextBox.ToSharedRef(), EFocusCause::SetDirectly);
+						});
 					}
 				}
 			}
diff --git a/Engine/Plugins/Runtime/Metasound/Source/MetasoundEditor/Private/MetasoundNodeDetailCustomization.h b/Engine/Plugins/Runtime/Metasound/Source/MetasoundEditor/Private/MetasoundNodeDetailCustomization.h
--- a/Engine/Plugins/Runtime/Metasound/Source/MetasoundEditor/Private/MetasoundNodeDetailCustomization.h
+++ b/Engine/Plugins/Runtime/Metasound/Source/MetasoundEditor/Private/MetasoundNodeDetailCustomization.h
@@ -2,6 +2,7 @@
 #pragma once
 
 #include "Components/Widget.h"
+#include "Delegates/IDelegateInstance.h"
 #include "DetailCategoryBuilder.h"
 #include "DetailLayoutBuilder.h"
 #include "DetailWidgetRow.h"
@@ -16,8 +17,8 @@
 #include "MetasoundEditorGraphBuilder.h"
 #include "MetasoundEditorGraphMemberDefaults.h"
 #include "MetasoundEditorGraphNode.h"
-#include "MetasoundFrontendLiteral.h"
 #include "MetasoundEditorModule.h"
+#include "MetasoundFrontendLiteral.h"
 #include "MetasoundUObjectRegistry.h"
 #include "PropertyHandle.h"
 #include "PropertyRestriction.h"
@@ -35,6 +36,7 @@
 
 #define LOCTEXT_NAMESPACE "MetaSoundEditor"
 
+
 namespace Metasound
 {
 	namespace Editor
diff --git a/Engine/Source/Developer/ToolWidgets/Private/SSearchableComboBox.cpp b/Engine/Source/Developer/ToolWidgets/Private/SSearchableComboBox.cpp
--- a/Engine/Source/Developer/ToolWidgets/Private/SSearchableComboBox.cpp
+++ b/Engine/Source/Developer/ToolWidgets/Private/SSearchableComboBox.cpp
@@ -109,15 +109,15 @@ void SSearchableComboBox::ClearSelection()
 	ComboListView->ClearSelection();
 }
 
-void SSearchableComboBox::SetSelectedItem(TSharedPtr<FString> InSelectedItem)
+void SSearchableComboBox::SetSelectedItem(TSharedPtr<FString> InSelectedItem, ESelectInfo::Type InSelectInfo)
 {
 	if (TListTypeTraits<TSharedPtr<FString>>::IsPtrValid(InSelectedItem))
 	{
-		ComboListView->SetSelection(InSelectedItem, ESelectInfo::OnNavigation);
+		ComboListView->SetSelection(InSelectedItem, InSelectInfo);
 	}
 	else
 	{
-		ComboListView->SetSelection(SelectedItem, ESelectInfo::OnNavigation);
+		ComboListView->SetSelection(SelectedItem, InSelectInfo);
 	}
 }
 
diff --git a/Engine/Source/Developer/ToolWidgets/Public/SSearchableComboBox.h b/Engine/Source/Developer/ToolWidgets/Public/SSearchableComboBox.h
--- a/Engine/Source/Developer/ToolWidgets/Public/SSearchableComboBox.h
+++ b/Engine/Source/Developer/ToolWidgets/Public/SSearchableComboBox.h
@@ -90,7 +90,14 @@ public:
 
 	void ClearSelection();
 
-	void SetSelectedItem(TSharedPtr<FString> InSelectedItem);
+	/**
+	 * Sets the selected item.  By default, registers as a navigation request, which *does not*
+	 * set the cached selected item, only updating the visually selected item.
+	 * 
+	 * @param InSelectedItem Item to select
+	 * @param InSelectInfo (optional) How the selected item is being committed (default: OnNavigation)
+	 */
+	void SetSelectedItem(TSharedPtr<FString> InSelectedItem, ESelectInfo::Type InSelectInfo = ESelectInfo::OnNavigation);
 
 	/** @return the item currently selected by the combo box. */
 	TSharedPtr<FString> GetSelectedItem();
