Add specific metasound LLM tags
#rnx

diff --git a/Engine/Plugins/Runtime/Metasound/Source/MetasoundEngine/Private/MetasoundSource.cpp b/Engine/Plugins/Runtime/Metasound/Source/MetasoundEngine/Private/MetasoundSource.cpp
--- a/Engine/Plugins/Runtime/Metasound/Source/MetasoundEngine/Private/MetasoundSource.cpp
+++ b/Engine/Plugins/Runtime/Metasound/Source/MetasoundEngine/Private/MetasoundSource.cpp
@@ -273,6 +273,8 @@ void UMetaSoundSource::SetRegistryAssetClassInfo(const Metasound::Frontend::FNod
 
 void UMetaSoundSource::InitParameters(TArray<FAudioParameter>& InParametersToInit, FName InFeatureName)
 {
+	METASOUND_LLM_SCOPE;
+
 	// Have to call cache vs a simple get as the source may have yet to start playing/has not been registered
 	// via InitResources. If it has, this call is fast and returns the already cached RuntimeData.
 	const FRuntimeData& RuntimeData = CacheRuntimeData();
@@ -463,6 +465,8 @@ void UMetaSoundSource::InitParameters(TArray<FAudioParameter>& InParametersToIni
 void UMetaSoundSource::InitResources()
 {
 	using namespace Metasound::Frontend;
+
+	METASOUND_LLM_SCOPE;
 	METASOUND_TRACE_CPUPROFILER_EVENT_SCOPE(UMetaSoundSource::InitResources);
 
 	FMetaSoundAssetRegistrationOptions RegOptions;
@@ -510,6 +514,7 @@ ISoundGeneratorPtr UMetaSoundSource::CreateSoundGenerator(const FSoundGeneratorI
 	using namespace Metasound::Frontend;
 	using namespace Metasound::Engine;
 
+	METASOUND_LLM_SCOPE;
 	METASOUND_TRACE_CPUPROFILER_EVENT_SCOPE(UMetaSoundSource::CreateSoundGenerator);
 
 	SampleRate = InParams.SampleRate;
@@ -791,6 +796,8 @@ bool UMetaSoundSource::IsOneShot() const
 
 TUniquePtr<Audio::IParameterTransmitter> UMetaSoundSource::CreateParameterTransmitter(Audio::FParameterTransmitterInitParams&& InParams) const
 {
+	METASOUND_LLM_SCOPE;
+
 	Metasound::FMetaSoundParameterTransmitter::FInitParams InitParams(GetOperatorSettings(InParams.SampleRate), InParams.InstanceID);
 
 	for (const FSendInfoAndVertexName& InfoAndName : FMetasoundAssetBase::GetSendInfos(InParams.InstanceID))
diff --git a/Engine/Plugins/Runtime/Metasound/Source/MetasoundEngine/Private/MetasoundUObjectRegistry.cpp b/Engine/Plugins/Runtime/Metasound/Source/MetasoundEngine/Private/MetasoundUObjectRegistry.cpp
--- a/Engine/Plugins/Runtime/Metasound/Source/MetasoundEngine/Private/MetasoundUObjectRegistry.cpp
+++ b/Engine/Plugins/Runtime/Metasound/Source/MetasoundEngine/Private/MetasoundUObjectRegistry.cpp
@@ -18,7 +18,6 @@
 #include "MetasoundTrace.h"
 #include "UObject/Object.h"
 
-
 namespace Metasound
 {
 	class FMetasoundUObjectRegistry : public IMetasoundUObjectRegistry
@@ -28,6 +27,8 @@ namespace Metasound
 
 			void RegisterUClassInterface(TUniquePtr<IMetasoundUObjectRegistryEntry>&& InEntry) override
 			{
+				METASOUND_LLM_SCOPE;
+
 				if (InEntry.IsValid())
 				{
 					Frontend::FInterfaceRegistryKey Key = Frontend::GetInterfaceRegistryKey(InEntry->GetInterfaceVersion());
@@ -70,6 +71,8 @@ namespace Metasound
 
 			UObject* NewObject(UClass* InClass, const FMetasoundFrontendDocument& InDocument, const FString& InPath) const override
 			{
+				METASOUND_LLM_SCOPE;
+
 				TArray<const IMetasoundUObjectRegistryEntry*> AllInterfaceEntries;
 
 				for (const FMetasoundFrontendVersion& InterfaceVersion : InDocument.Interfaces)
@@ -220,3 +223,4 @@ namespace Metasound
 		return Registry;
 	}
 }
+
diff --git a/Engine/Plugins/Runtime/Metasound/Source/MetasoundFrontend/Private/MetasoundFrontendArchetypeRegistry.cpp b/Engine/Plugins/Runtime/Metasound/Source/MetasoundFrontend/Private/MetasoundFrontendArchetypeRegistry.cpp
--- a/Engine/Plugins/Runtime/Metasound/Source/MetasoundFrontend/Private/MetasoundFrontendArchetypeRegistry.cpp
+++ b/Engine/Plugins/Runtime/Metasound/Source/MetasoundFrontend/Private/MetasoundFrontendArchetypeRegistry.cpp
@@ -4,6 +4,7 @@
 
 #include "HAL/PlatformTime.h"
 #include "MetasoundFrontendRegistryTransaction.h"
+#include "MetasoundTrace.h"
 
 namespace Metasound
 {
@@ -18,6 +19,8 @@ namespace Metasound
 
 				virtual void RegisterInterface(TUniquePtr<IInterfaceRegistryEntry>&& InEntry) override
 				{
+					METASOUND_LLM_SCOPE;
+
 					FInterfaceRegistryTransaction::FTimeType TransactionTime = FPlatformTime::Cycles64();
 					if (InEntry.IsValid())
 					{
diff --git a/Engine/Plugins/Runtime/Metasound/Source/MetasoundFrontend/Private/MetasoundFrontendDataTypeRegistry.cpp b/Engine/Plugins/Runtime/Metasound/Source/MetasoundFrontend/Private/MetasoundFrontendDataTypeRegistry.cpp
--- a/Engine/Plugins/Runtime/Metasound/Source/MetasoundFrontend/Private/MetasoundFrontendDataTypeRegistry.cpp
+++ b/Engine/Plugins/Runtime/Metasound/Source/MetasoundFrontend/Private/MetasoundFrontendDataTypeRegistry.cpp
@@ -1,6 +1,7 @@
 // Copyright Epic Games, Inc. All Rights Reserved.
 #include "MetasoundFrontendDataTypeRegistry.h"
 
+#include "MetasoundTrace.h"
 
 namespace Metasound
 {
@@ -516,6 +517,8 @@ namespace Metasound
 
 			bool FDataTypeRegistry::RegisterDataType(TUniquePtr<IDataTypeRegistryEntry>&& InEntry)
 			{
+				METASOUND_LLM_SCOPE;
+
 				if (InEntry.IsValid())
 				{
 					const FName Name = InEntry->GetDataTypeInfo().DataTypeName;
@@ -975,3 +978,4 @@ namespace Metasound
 
 	}
 }
+
diff --git a/Engine/Plugins/Runtime/Metasound/Source/MetasoundFrontend/Private/MetasoundFrontendRegistries.cpp b/Engine/Plugins/Runtime/Metasound/Source/MetasoundFrontend/Private/MetasoundFrontendRegistries.cpp
--- a/Engine/Plugins/Runtime/Metasound/Source/MetasoundFrontend/Private/MetasoundFrontendRegistries.cpp
+++ b/Engine/Plugins/Runtime/Metasound/Source/MetasoundFrontend/Private/MetasoundFrontendRegistries.cpp
@@ -13,9 +13,6 @@
 #include "MetasoundTrace.h"
 #include "Misc/ScopeLock.h"
 
-#ifndef WITH_METASOUND_FRONTEND
-#define WITH_METASOUND_FRONTEND 0
-#endif
 
 namespace Metasound
 {
@@ -159,6 +156,7 @@ namespace Metasound
 
 			void FRegistryContainerImpl::RegisterPendingNodes()
 			{
+				METASOUND_LLM_SCOPE;
 				METASOUND_TRACE_CPUPROFILER_EVENT_SCOPE(metasound::FRegistryContainerImpl::RegisterPendingNodes);
 				{
 					FScopeLock ScopeLock(&LazyInitCommandCritSection);
@@ -259,6 +257,8 @@ namespace Metasound
 
 			FNodeRegistryKey FRegistryContainerImpl::RegisterNode(TUniquePtr<INodeRegistryEntry>&& InEntry)
 			{
+				METASOUND_LLM_SCOPE;
+
 				FNodeRegistryKey Key;
 
 				if (InEntry.IsValid())
@@ -288,6 +288,8 @@ namespace Metasound
 
 			bool FRegistryContainerImpl::UnregisterNode(const FNodeRegistryKey& InKey)
 			{
+				METASOUND_LLM_SCOPE;
+
 				if (NodeRegistryKey::IsValid(InKey))
 				{
 					if (const INodeRegistryEntry* Entry = FindNodeEntry(InKey))
@@ -647,3 +649,4 @@ bool FMetasoundFrontendRegistryContainer::GetOutputNodeRegistryKeyForDataType(co
 	return false;
 }
 
+
diff --git a/Engine/Plugins/Runtime/Metasound/Source/MetasoundFrontend/Private/MetasoundFrontendSearchEngine.cpp b/Engine/Plugins/Runtime/Metasound/Source/MetasoundFrontend/Private/MetasoundFrontendSearchEngine.cpp
--- a/Engine/Plugins/Runtime/Metasound/Source/MetasoundFrontend/Private/MetasoundFrontendSearchEngine.cpp
+++ b/Engine/Plugins/Runtime/Metasound/Source/MetasoundFrontend/Private/MetasoundFrontendSearchEngine.cpp
@@ -327,6 +327,7 @@ namespace Metasound
 
 			virtual void Prime() override
 			{
+				METASOUND_LLM_SCOPE;
 				FScopeLock Lock(&QueryCriticalSection);
 				Update();
 			}
@@ -335,6 +336,7 @@ namespace Metasound
 			// result.
 			ResultType UpdateAndFindResult(const FFrontendQueryKey& InKey)
 			{
+				METASOUND_LLM_SCOPE;
 				FScopeLock Lock(&QueryCriticalSection);
 
 				Update();
@@ -349,6 +351,7 @@ namespace Metasound
 			// to OutResult. Returns true on success, false on failure. 
 			bool UpdateAndFindResult(const FFrontendQueryKey& InKey, ResultType& OutResult)
 			{
+				METASOUND_LLM_SCOPE;
 				FScopeLock Lock(&QueryCriticalSection);
 
 				Update();
@@ -757,6 +760,7 @@ namespace Metasound
 
 		void FSearchEngine::Prime()
 		{
+			METASOUND_LLM_SCOPE;
 			METASOUND_TRACE_CPUPROFILER_EVENT_SCOPE(metasound::FSearchEngine::Prime);
 
 			FindAllClassesQuery.Prime();
@@ -884,3 +888,4 @@ namespace Metasound
 		}
 	}
 }
+
diff --git a/Engine/Plugins/Runtime/Metasound/Source/MetasoundFrontend/Private/MetasoundRouter.cpp b/Engine/Plugins/Runtime/Metasound/Source/MetasoundFrontend/Private/MetasoundRouter.cpp
--- a/Engine/Plugins/Runtime/Metasound/Source/MetasoundFrontend/Private/MetasoundRouter.cpp
+++ b/Engine/Plugins/Runtime/Metasound/Source/MetasoundFrontend/Private/MetasoundRouter.cpp
@@ -5,9 +5,9 @@
 #include "MetasoundFrontendDataTypeRegistry.h"
 #include "MetasoundFrontendRegistries.h"
 #include "MetasoundOperatorInterface.h"
+#include "MetasoundTrace.h"
 #include "HAL/IConsoleManager.h"
 
-
 // Convenience exec commands to push values to global params.
 static FAutoConsoleCommand GPushFloatCommand(
 	TEXT("au.MetaSound.SetFloat"),
@@ -176,6 +176,7 @@ namespace Metasound
 
 	TSharedPtr<IDataChannel, ESPMode::ThreadSafe> FAddressRouter::GetDataChannel(const FSendAddress& InAddress, const FOperatorSettings& InOperatorSettings)
 	{
+		METASOUND_LLM_SCOPE;
 		TSharedPtr<IDataChannel, ESPMode::ThreadSafe> DataChannel = FindDataChannel(InAddress);
 
 		if (!DataChannel.IsValid())
@@ -195,6 +196,7 @@ namespace Metasound
 
 	TUniquePtr<ISender> FAddressRouter::RegisterNewSender(const FSendAddress& InAddress, const FSenderInitParams& InitParams)
 	{
+		METASOUND_LLM_SCOPE;
 		TSharedPtr<IDataChannel, ESPMode::ThreadSafe> DataChannel = GetDataChannel(InAddress, InitParams.OperatorSettings);
 
 		if (DataChannel.IsValid())
@@ -209,6 +211,7 @@ namespace Metasound
 
 	bool FAddressRouter::UnregisterDataChannel(const FSendAddress& InAddress)
 	{
+		METASOUND_LLM_SCOPE;
 		FScopeLock ScopeLock(&DataChannelMapMutationLock);
 
 		if (TSharedRef<IDataChannel, ESPMode::ThreadSafe>* Channel = DataChannelMap.Find(InAddress))
@@ -229,6 +232,7 @@ namespace Metasound
 
 	bool FAddressRouter::UnregisterDataChannelIfUnconnected(const FSendAddress& InAddress)
 	{
+		METASOUND_LLM_SCOPE;
 		FScopeLock ScopeLock(&DataChannelMapMutationLock);
 
 		if (TSharedRef<IDataChannel, ESPMode::ThreadSafe>* Channel = DataChannelMap.Find(InAddress))
@@ -247,6 +251,7 @@ namespace Metasound
 
 	TUniquePtr<IReceiver> FAddressRouter::RegisterNewReceiver(const FSendAddress& InAddress, const FReceiverInitParams& InitParams)
 	{
+		METASOUND_LLM_SCOPE;
 		TSharedPtr<IDataChannel, ESPMode::ThreadSafe> DataChannel = GetDataChannel(InAddress, InitParams.OperatorSettings);
 
 		if (DataChannel.IsValid())
diff --git a/Engine/Plugins/Runtime/Metasound/Source/MetasoundFrontend/Public/MetasoundArrayNodesRegistration.h b/Engine/Plugins/Runtime/Metasound/Source/MetasoundFrontend/Public/MetasoundArrayNodesRegistration.h
--- a/Engine/Plugins/Runtime/Metasound/Source/MetasoundFrontend/Public/MetasoundArrayNodesRegistration.h
+++ b/Engine/Plugins/Runtime/Metasound/Source/MetasoundFrontend/Public/MetasoundArrayNodesRegistration.h
@@ -174,7 +174,6 @@ namespace Metasound
 	template<typename ArrayType>
 	bool RegisterArrayNodes()
 	{
-		LLM_SCOPE_BYNAME(TEXT("Metasound"));
 		using namespace MetasoundArrayNodesPrivate;
 
 		bool bSuccess = RegisterArrayNumNode<ArrayType>();
diff --git a/Engine/Plugins/Runtime/Metasound/Source/MetasoundGraphCore/Public/MetasoundTrace.h b/Engine/Plugins/Runtime/Metasound/Source/MetasoundGraphCore/Public/MetasoundTrace.h
--- a/Engine/Plugins/Runtime/Metasound/Source/MetasoundGraphCore/Public/MetasoundTrace.h
+++ b/Engine/Plugins/Runtime/Metasound/Source/MetasoundGraphCore/Public/MetasoundTrace.h
@@ -28,3 +28,7 @@
 #define METASOUND_TRACE_CPUPROFILER_EVENT_SCOPE_TEXT(Name)
 
 #endif
+
+// Convenience macro for MetaSound LLM scope to avoid misspells. 
+#define METASOUND_LLM_SCOPE LLM_SCOPE_BYNAME(TEXT("Audio/MetaSound"));
+
