Fix for MetaSound node crash when accessing a destructed temporary.
#rnx

diff --git a/Engine/Plugins/Runtime/Metasound/Source/MetasoundStandardNodes/Private/MetasoundNoiseGenerator.cpp b/Engine/Plugins/Runtime/Metasound/Source/MetasoundStandardNodes/Private/MetasoundNoiseGenerator.cpp
--- a/Engine/Plugins/Runtime/Metasound/Source/MetasoundStandardNodes/Private/MetasoundNoiseGenerator.cpp
+++ b/Engine/Plugins/Runtime/Metasound/Source/MetasoundStandardNodes/Private/MetasoundNoiseGenerator.cpp
@@ -218,7 +218,6 @@ namespace Metasound
 		const FNoiseNode& Node = static_cast<const FNoiseNode&>(InParams.Node);
 		const FDataReferenceCollection& InputCol = InParams.InputDataReferences;
 		const FOperatorSettings& Settings = InParams.OperatorSettings;
-		const FInputVertexInterface& InputInterface = DeclareVertexInterface().GetInputInterface();
 
 		// Static property pin, only used for factory.
 		FEnumNoiseTypeReadRef Type = InputCol.GetDataReadReferenceOrConstruct<FEnumNoiseType>(METASOUND_GET_PARAM_NAME(InputType));
diff --git a/Engine/Plugins/Runtime/Metasound/Source/MetasoundStandardNodes/Private/MetasoundTriggerOnThresholdNode.cpp b/Engine/Plugins/Runtime/Metasound/Source/MetasoundStandardNodes/Private/MetasoundTriggerOnThresholdNode.cpp
--- a/Engine/Plugins/Runtime/Metasound/Source/MetasoundStandardNodes/Private/MetasoundTriggerOnThresholdNode.cpp
+++ b/Engine/Plugins/Runtime/Metasound/Source/MetasoundStandardNodes/Private/MetasoundTriggerOnThresholdNode.cpp
@@ -242,7 +242,7 @@ namespace Metasound
 		const FTriggerOnThresholdNode& Node = static_cast<const FTriggerOnThresholdNode&>(InParams.Node);
 		const FDataReferenceCollection& InputCol = InParams.InputDataReferences;
 		const FOperatorSettings& Settings = InParams.OperatorSettings;
-		const FInputVertexInterface& InputInterface = DeclareVertexInterface().GetInputInterface();
+		const FInputVertexInterface& InputInterface = InParams.Node.GetVertexInterface().GetInputInterface();
 		
 		// Static property pin, only used for factory.
 		FBufferTriggerTypeReadRef Type = InputCol.GetDataReadReferenceOrConstruct<FEnumBufferTriggerType>(METASOUND_GET_PARAM_NAME(InTriggerType));
