Remove unicode 1/f from the tooltip and replace with ansi one.

diff --git a/Engine/Plugins/Runtime/Metasound/Source/MetasoundStandardNodes/Private/MetasoundNoiseGenerator.cpp b/Engine/Plugins/Runtime/Metasound/Source/MetasoundStandardNodes/Private/MetasoundNoiseGenerator.cpp
--- a/Engine/Plugins/Runtime/Metasound/Source/MetasoundStandardNodes/Private/MetasoundNoiseGenerator.cpp
+++ b/Engine/Plugins/Runtime/Metasound/Source/MetasoundStandardNodes/Private/MetasoundNoiseGenerator.cpp
@@ -30,7 +30,7 @@ namespace Metasound
 		FEnumNoiseType, FEnumNoiseTypeInfo, FEnumNoiseTypeReadRef, FEnumNoiseTypeWriteRef);
 
 	DEFINE_METASOUND_ENUM_BEGIN(ENoiseType, FEnumNoiseType, "NoiseType")
-		DEFINE_METASOUND_ENUM_ENTRY(ENoiseType::Pink, "PinkDescription", "Pink Noise", "PinkDescriptionTT", "Pink noise or ​1⁄f noise, spectral density is inversely proportional to the frequency of the signal"),
+		DEFINE_METASOUND_ENUM_ENTRY(ENoiseType::Pink, "PinkDescription", "Pink Noise", "PinkDescriptionTT", "Pink noise or 1/f noise, spectral density is inversely proportional to the frequency of the signal"),
 		DEFINE_METASOUND_ENUM_ENTRY(ENoiseType::White, "WhiteDescription", "White Noise", "WhiteDescriptionTT", "A random signal having equal intensity at different frequencies"),
 	DEFINE_METASOUND_ENUM_END()
 
