Fix for changed names in metasound math and array nodes
#rnx

diff --git a/Engine/Plugins/Runtime/Metasound/Source/MetasoundFrontend/Public/MetasoundArrayNodes.h b/Engine/Plugins/Runtime/Metasound/Source/MetasoundFrontend/Public/MetasoundArrayNodes.h
--- a/Engine/Plugins/Runtime/Metasound/Source/MetasoundFrontend/Public/MetasoundArrayNodes.h
+++ b/Engine/Plugins/Runtime/Metasound/Source/MetasoundFrontend/Public/MetasoundArrayNodes.h
@@ -56,7 +56,7 @@ namespace Metasound
 		METASOUND_PARAM(InputValue, "Value", "Value to set.")
 
 		METASOUND_PARAM(OutputNum, "Num", "Number of elements in the array.")
-		METASOUND_PARAM(OutputValue, "Value", "Value of element at array index.")
+		METASOUND_PARAM(OutputValue, "Element", "Value of element at array index.")
 		METASOUND_PARAM(OutputArrayConcat, "Array", "Array after concatenation.")
 		METASOUND_PARAM(OutputArraySet, "Array", "Array after setting.")
 		METASOUND_PARAM(OutputArraySubset, "Array", "Subset of input array.")
diff --git a/Engine/Plugins/Runtime/Metasound/Source/MetasoundStandardNodes/Private/MetasoundMathNodes.cpp b/Engine/Plugins/Runtime/Metasound/Source/MetasoundStandardNodes/Private/MetasoundMathNodes.cpp
--- a/Engine/Plugins/Runtime/Metasound/Source/MetasoundStandardNodes/Private/MetasoundMathNodes.cpp
+++ b/Engine/Plugins/Runtime/Metasound/Source/MetasoundStandardNodes/Private/MetasoundMathNodes.cpp
@@ -53,8 +53,8 @@ namespace Metasound
 {
 	namespace MathOpNames
 	{
-		METASOUND_PARAM(PrimaryOperand, "Primary Operand", "Initial operand")
-		METASOUND_PARAM(AdditionalOperands, "Additional Operand", "Additional operand(s)")
+		METASOUND_PARAM(PrimaryOperand, "PrimaryOperand", "Initial operand")
+		METASOUND_PARAM(AdditionalOperands, "AdditionalOperands", "Additional operand(s)")
 
 		METASOUND_PARAM(OutMath, "Out", "Math operation result.")
 		METASOUND_PARAM(OutAudio, "Out", "Resulting buffer.")
