MetaSound FVertex simplification
- Simplifying FVertexInteface by removing VertexModel concept.
- Deprecating unwanted functions in interface.
- Simplifying storage mechanism in TVertexGroup

diff --git a/Engine/Plugins/Runtime/AudioModulation/Source/AudioModulation/Private/MetaSoundNodes/GetModulatorValueNode.cpp b/Engine/Plugins/Runtime/AudioModulation/Source/AudioModulation/Private/MetaSoundNodes/GetModulatorValueNode.cpp
--- a/Engine/Plugins/Runtime/AudioModulation/Source/AudioModulation/Private/MetaSoundNodes/GetModulatorValueNode.cpp
+++ b/Engine/Plugins/Runtime/AudioModulation/Source/AudioModulation/Private/MetaSoundNodes/GetModulatorValueNode.cpp
@@ -19,13 +19,15 @@ namespace AudioModulation
 	public:
 		static const Metasound::FVertexInterface& GetDefaultInterface()
 		{
-			static const Metasound::FVertexInterface DefaultInterface(
-				Metasound::FInputVertexInterface(
-					Metasound::TInputDataVertexModel<FSoundModulatorAsset>("Modulator", LOCTEXT("MetasoundValueNode_InputModulatorName", "Modulator")),
-					Metasound::TInputDataVertexModel<bool>("Normalized", LOCTEXT("MixModulatorsNode_InputNormalizedName", "Normalized"), true)
+			using namespace Metasound;
+
+			static const FVertexInterface DefaultInterface(
+				FInputVertexInterface(
+					TInputDataVertex<FSoundModulatorAsset>("Modulator", FDataVertexMetadata{ LOCTEXT("MetasoundValueNode_InputModulatorName", "Modulator") }),
+					TInputDataVertex<bool>("Normalized", FDataVertexMetadata{LOCTEXT("MixModulatorsNode_InputNormalizedName", "Normalized")}, true )
 				),
-				Metasound::FOutputVertexInterface(
-					Metasound::TOutputDataVertexModel<float>("Out", LOCTEXT("MetasoundValueNode_OutputModulatorValue", "Out"))
+				FOutputVertexInterface(
+					TOutputDataVertex<float>("Out", FDataVertexMetadata{LOCTEXT("MetasoundValueNode_OutputModulatorValue", "Out")})
 				)
 			);
 
diff --git a/Engine/Plugins/Runtime/AudioModulation/Source/AudioModulation/Private/MetaSoundNodes/MixModulatorsNode.cpp b/Engine/Plugins/Runtime/AudioModulation/Source/AudioModulation/Private/MetaSoundNodes/MixModulatorsNode.cpp
--- a/Engine/Plugins/Runtime/AudioModulation/Source/AudioModulation/Private/MetaSoundNodes/MixModulatorsNode.cpp
+++ b/Engine/Plugins/Runtime/AudioModulation/Source/AudioModulation/Private/MetaSoundNodes/MixModulatorsNode.cpp
@@ -19,15 +19,17 @@ namespace AudioModulation
 	public:
 		static const Metasound::FVertexInterface& GetDefaultInterface()
 		{
-			static const Metasound::FVertexInterface DefaultInterface(
-				Metasound::FInputVertexInterface(
-					Metasound::TInputDataVertexModel<FSoundModulatorAsset>("In1", LOCTEXT("MixModulatorsNode_InputModulator1Name", "In 1")),
-					Metasound::TInputDataVertexModel<FSoundModulatorAsset>("In2", LOCTEXT("MixModulatorsNode_InputModulator2Name", "In 2")),
-					Metasound::TInputDataVertexModel<FSoundModulationParameterAsset>("MixParameter", LOCTEXT("MixModulatorsNode_InputMixParameterName", "Mix Parameter")),
-					Metasound::TInputDataVertexModel<bool>("Normalized", LOCTEXT("MixModulatorsNode_InputNormalizedName", "Normalized"), true)
+			using namespace Metasound;
+
+			static const FVertexInterface DefaultInterface(
+				FInputVertexInterface(
+					TInputDataVertex<FSoundModulatorAsset>("In1", FDataVertexMetadata{ LOCTEXT("MixModulatorsNode_InputModulator1Name", "In 1") }),
+					TInputDataVertex<FSoundModulatorAsset>("In2", FDataVertexMetadata{ LOCTEXT("MixModulatorsNode_InputModulator2Name", "In 2") }),
+					TInputDataVertex<FSoundModulationParameterAsset>("MixParameter", FDataVertexMetadata{ LOCTEXT("MixModulatorsNode_InputMixParameterName", "Mix Parameter") }),
+					TInputDataVertex<bool>("Normalized", FDataVertexMetadata{ LOCTEXT("MixModulatorsNode_InputNormalizedName", "Normalized") }, true)
 				),
-				Metasound::FOutputVertexInterface(
-					Metasound::TOutputDataVertexModel<float>("Out", LOCTEXT("MixModulatorsNode_OutputModulatorValue", "Out"))
+				FOutputVertexInterface(
+					TOutputDataVertex<float>("Out", FDataVertexMetadata{ LOCTEXT("MixModulatorsNode_OutputModulatorValue", "Out") })
 				)
 			);
 
diff --git a/Engine/Plugins/Runtime/Metasound/Source/MetasoundEngine/Private/MetasoundWaveInfoNode.cpp b/Engine/Plugins/Runtime/Metasound/Source/MetasoundEngine/Private/MetasoundWaveInfoNode.cpp
--- a/Engine/Plugins/Runtime/Metasound/Source/MetasoundEngine/Private/MetasoundWaveInfoNode.cpp
+++ b/Engine/Plugins/Runtime/Metasound/Source/MetasoundEngine/Private/MetasoundWaveInfoNode.cpp
@@ -104,10 +104,10 @@ namespace Metasound
 		using namespace WaveInfoNodeParameterNames;
 		static const FVertexInterface Interface(
 			FInputVertexInterface(
-				TInputDataVertexModel<FWaveAsset>(METASOUND_GET_PARAM_NAME_AND_METADATA(ParamWaveAsset))
+				TInputDataVertex<FWaveAsset>(METASOUND_GET_PARAM_NAME_AND_METADATA(ParamWaveAsset))
 			),
 			FOutputVertexInterface(
-				TOutputDataVertexModel<FTime>(METASOUND_GET_PARAM_NAME_AND_METADATA(ParamDurationSeconds))
+				TOutputDataVertex<FTime>(METASOUND_GET_PARAM_NAME_AND_METADATA(ParamDurationSeconds))
 			)
 		);
 
@@ -181,4 +181,4 @@ namespace Metasound
 
 } // namespace Metasound
 
-#undef LOCTEXT_NAMESPACE //MetasoundWaveInfo
\ No newline at end of file
+#undef LOCTEXT_NAMESPACE //MetasoundWaveInfo
diff --git a/Engine/Plugins/Runtime/Metasound/Source/MetasoundEngine/Private/MetasoundWavePlayerNode.cpp b/Engine/Plugins/Runtime/Metasound/Source/MetasoundEngine/Private/MetasoundWavePlayerNode.cpp
--- a/Engine/Plugins/Runtime/Metasound/Source/MetasoundEngine/Private/MetasoundWavePlayerNode.cpp
+++ b/Engine/Plugins/Runtime/Metasound/Source/MetasoundEngine/Private/MetasoundWavePlayerNode.cpp
@@ -1015,27 +1015,27 @@ namespace Metasound
 
 		return FVertexInterface(
 			FInputVertexInterface(
-				TInputDataVertexModel<FTrigger>(METASOUND_GET_PARAM_NAME_AND_METADATA(InputTriggerPlay)),
-				TInputDataVertexModel<FTrigger>(METASOUND_GET_PARAM_NAME_AND_METADATA(InputTriggerStop)),
-				TInputDataVertexModel<FWaveAsset>(METASOUND_GET_PARAM_NAME_AND_METADATA(InputWaveAsset)),
-				TInputDataVertexModel<FTime>(METASOUND_GET_PARAM_NAME_AND_METADATA(InputStartTime), 0.0f),
-				TInputDataVertexModel<float>(METASOUND_GET_PARAM_NAME_AND_METADATA(InputPitchShift), 0.0f),
-				TInputDataVertexModel<bool>(METASOUND_GET_PARAM_NAME_AND_METADATA(InputLoop), false),
-				TInputDataVertexModel<FTime>(METASOUND_GET_PARAM_NAME_AND_METADATA(InputLoopStart), 0.0f),
-				TInputDataVertexModel<FTime>(METASOUND_GET_PARAM_NAME_AND_METADATA(InputLoopDuration), -1.0f)
+				TInputDataVertex<FTrigger>(METASOUND_GET_PARAM_NAME_AND_METADATA(InputTriggerPlay)),
+				TInputDataVertex<FTrigger>(METASOUND_GET_PARAM_NAME_AND_METADATA(InputTriggerStop)),
+				TInputDataVertex<FWaveAsset>(METASOUND_GET_PARAM_NAME_AND_METADATA(InputWaveAsset)),
+				TInputDataVertex<FTime>(METASOUND_GET_PARAM_NAME_AND_METADATA(InputStartTime), 0.0f),
+				TInputDataVertex<float>(METASOUND_GET_PARAM_NAME_AND_METADATA(InputPitchShift), 0.0f),
+				TInputDataVertex<bool>(METASOUND_GET_PARAM_NAME_AND_METADATA(InputLoop), false),
+				TInputDataVertex<FTime>(METASOUND_GET_PARAM_NAME_AND_METADATA(InputLoopStart), 0.0f),
+				TInputDataVertex<FTime>(METASOUND_GET_PARAM_NAME_AND_METADATA(InputLoopDuration), -1.0f)
 				),
 			FOutputVertexInterface(
-				TOutputDataVertexModel<FTrigger>(METASOUND_GET_PARAM_NAME_AND_METADATA(OutputTriggerOnPlay)),
-				TOutputDataVertexModel<FTrigger>(METASOUND_GET_PARAM_NAME_AND_METADATA(OutputTriggerOnDone)),
-				TOutputDataVertexModel<FTrigger>(METASOUND_GET_PARAM_NAME_AND_METADATA(OutputTriggerOnNearlyDone)),
-				TOutputDataVertexModel<FTrigger>(METASOUND_GET_PARAM_NAME_AND_METADATA(OutputTriggerOnLooped)),
-				TOutputDataVertexModel<FTrigger>(METASOUND_GET_PARAM_NAME_AND_METADATA(OutputTriggerOnCuePoint)),
-				TOutputDataVertexModel<int32>(METASOUND_GET_PARAM_NAME_AND_METADATA(OutputCuePointID)),
-				TOutputDataVertexModel<FString>(METASOUND_GET_PARAM_NAME_AND_METADATA(OutputCuePointLabel)),
-				TOutputDataVertexModel<float>(METASOUND_GET_PARAM_NAME_AND_METADATA(OutputLoopRatio)),
-				TOutputDataVertexModel<float>(METASOUND_GET_PARAM_NAME_AND_METADATA(OutputPlaybackLocation)),
-				TOutputDataVertexModel<FAudioBuffer>(METASOUND_GET_PARAM_NAME_AND_METADATA(OutputAudioLeft)),
-				TOutputDataVertexModel<FAudioBuffer>(METASOUND_GET_PARAM_NAME_AND_METADATA(OutputAudioRight))
+				TOutputDataVertex<FTrigger>(METASOUND_GET_PARAM_NAME_AND_METADATA(OutputTriggerOnPlay)),
+				TOutputDataVertex<FTrigger>(METASOUND_GET_PARAM_NAME_AND_METADATA(OutputTriggerOnDone)),
+				TOutputDataVertex<FTrigger>(METASOUND_GET_PARAM_NAME_AND_METADATA(OutputTriggerOnNearlyDone)),
+				TOutputDataVertex<FTrigger>(METASOUND_GET_PARAM_NAME_AND_METADATA(OutputTriggerOnLooped)),
+				TOutputDataVertex<FTrigger>(METASOUND_GET_PARAM_NAME_AND_METADATA(OutputTriggerOnCuePoint)),
+				TOutputDataVertex<int32>(METASOUND_GET_PARAM_NAME_AND_METADATA(OutputCuePointID)),
+				TOutputDataVertex<FString>(METASOUND_GET_PARAM_NAME_AND_METADATA(OutputCuePointLabel)),
+				TOutputDataVertex<float>(METASOUND_GET_PARAM_NAME_AND_METADATA(OutputLoopRatio)),
+				TOutputDataVertex<float>(METASOUND_GET_PARAM_NAME_AND_METADATA(OutputPlaybackLocation)),
+				TOutputDataVertex<FAudioBuffer>(METASOUND_GET_PARAM_NAME_AND_METADATA(OutputAudioLeft)),
+				TOutputDataVertex<FAudioBuffer>(METASOUND_GET_PARAM_NAME_AND_METADATA(OutputAudioRight))
 			)
 		);
 	}
diff --git a/Engine/Plugins/Runtime/Metasound/Source/MetasoundFrontend/Private/MetasoundFrontendDocument.cpp b/Engine/Plugins/Runtime/Metasound/Source/MetasoundFrontend/Private/MetasoundFrontendDocument.cpp
--- a/Engine/Plugins/Runtime/Metasound/Source/MetasoundFrontend/Private/MetasoundFrontendDocument.cpp
+++ b/Engine/Plugins/Runtime/Metasound/Source/MetasoundFrontend/Private/MetasoundFrontendDocument.cpp
@@ -198,16 +198,15 @@ FMetasoundFrontendClassInterface FMetasoundFrontendClassInterface::GenerateClass
 		FMetasoundFrontendInterfaceStyle InputStyle;
 #endif // WITH_EDITOR
 
-		for (const TPair<FVertexName, FInputDataVertex>& InputTuple : InputInterface)
+		for (const FInputDataVertex& InputVertex : InputInterface)
 		{
 			FMetasoundFrontendClassInput ClassInput;
 
-			const FInputDataVertex& InputVertex = InputTuple.Value;
-			ClassInput.Name = InputVertex.GetVertexName();
-			ClassInput.TypeName = InputVertex.GetDataTypeName();
+			ClassInput.Name = InputVertex.VertexName;
+			ClassInput.TypeName = InputVertex.DataTypeName;
 			ClassInput.VertexID = FGuid::NewGuid();
 
-			const FDataVertexMetadata& VertexMetadata = InputVertex.GetMetadata();
+			const FDataVertexMetadata& VertexMetadata = InputVertex.Metadata;
 
 #if WITH_EDITOR
 			ClassInput.Metadata.SetSerializeText(false);
@@ -216,7 +215,7 @@ FMetasoundFrontendClassInterface FMetasoundFrontendClassInterface::GenerateClass
 			ClassInput.Metadata.bIsAdvancedDisplay = VertexMetadata.bIsAdvancedDisplay;
 
 			// Advanced display items are pushed to bottom of sort order
-			ClassInput.Metadata.SortOrderIndex = InputInterface.GetSortOrderIndex(InputTuple.Key);
+			ClassInput.Metadata.SortOrderIndex = InputInterface.GetSortOrderIndex(InputVertex.VertexName);
 			if (ClassInput.Metadata.bIsAdvancedDisplay)
 			{
 				ClassInput.Metadata.SortOrderIndex += InputInterface.Num();
@@ -251,15 +250,16 @@ FMetasoundFrontendClassInterface FMetasoundFrontendClassInterface::GenerateClass
 		FMetasoundFrontendInterfaceStyle OutputStyle;
 #endif // WITH_EDITOR
 
-		for (const TPair<FVertexName, FOutputDataVertex>& OutputTuple : OutputInterface)
+		for (const FOutputDataVertex& OutputVertex: OutputInterface)
 		{
 			FMetasoundFrontendClassOutput ClassOutput;
 
-			ClassOutput.Name = OutputTuple.Value.GetVertexName();
-			ClassOutput.TypeName = OutputTuple.Value.GetDataTypeName();
+
+			ClassOutput.Name = OutputVertex.VertexName;
+			ClassOutput.TypeName = OutputVertex.DataTypeName;
 			ClassOutput.VertexID = FGuid::NewGuid();
 
-			const FDataVertexMetadata& VertexMetadata = OutputTuple.Value.GetMetadata();
+			const FDataVertexMetadata& VertexMetadata = OutputVertex.Metadata;
 
 #if WITH_EDITOR
 			ClassOutput.Metadata.SetSerializeText(false);
@@ -268,7 +268,7 @@ FMetasoundFrontendClassInterface FMetasoundFrontendClassInterface::GenerateClass
 			ClassOutput.Metadata.bIsAdvancedDisplay = VertexMetadata.bIsAdvancedDisplay;
 
 			// Advanced display items are pushed to bottom below non-advanced
-			ClassOutput.Metadata.SortOrderIndex = OutputInterface.GetSortOrderIndex(OutputTuple.Key);
+			ClassOutput.Metadata.SortOrderIndex = OutputInterface.GetSortOrderIndex(OutputVertex.VertexName);
 			if (ClassOutput.Metadata.bIsAdvancedDisplay)
 			{
 				ClassOutput.Metadata.SortOrderIndex += OutputInterface.Num();
@@ -288,11 +288,11 @@ FMetasoundFrontendClassInterface FMetasoundFrontendClassInterface::GenerateClass
 #endif // WITH_EDITOR
 	}
 
-	for (auto& EnvTuple : InVertexInterface.GetEnvironmentInterface())
+	for (const FEnvironmentVertex& EnvVertex : InVertexInterface.GetEnvironmentInterface())
 	{
 		FMetasoundFrontendClassEnvironmentVariable EnvVar;
 
-		EnvVar.Name = EnvTuple.Value.GetVertexName();
+		EnvVar.Name = EnvVertex.VertexName;
 		EnvVar.bIsRequired = true;
 
 		ClassInterface.Environment.Add(EnvVar);
diff --git a/Engine/Plugins/Runtime/Metasound/Source/MetasoundFrontend/Private/MetasoundFrontendInjectReceiveNodes.cpp b/Engine/Plugins/Runtime/Metasound/Source/MetasoundFrontend/Private/MetasoundFrontendInjectReceiveNodes.cpp
--- a/Engine/Plugins/Runtime/Metasound/Source/MetasoundFrontend/Private/MetasoundFrontendInjectReceiveNodes.cpp
+++ b/Engine/Plugins/Runtime/Metasound/Source/MetasoundFrontend/Private/MetasoundFrontendInjectReceiveNodes.cpp
@@ -83,7 +83,7 @@ namespace Metasound
 						FInputVertexInterface{},
 						FOutputVertexInterface
 						{
-							TOutputDataVertexModel<FSendAddress>(METASOUND_GET_PARAM_NAME_AND_METADATA(OutputAddress))
+							TOutputDataVertex<FSendAddress>(METASOUND_GET_PARAM_NAME_AND_METADATA(OutputAddress))
 						}
 					};
 
@@ -133,8 +133,8 @@ namespace Metasound
 			// should never contain null nodes for input destination.
 			check(InputDestination.Node != nullptr);
 
-			const FVertexName& VertexKey = InputDestination.Vertex.GetVertexName();
-			const FName& DataType = InputDestination.Vertex.GetDataTypeName();
+			const FVertexName& VertexKey = InputDestination.Vertex.VertexName;
+			const FName& DataType = InputDestination.Vertex.DataTypeName;
 
 			// Create a receive node.
 			const FGuid ReceiveNodeID = FGuid::NewGuid();
@@ -142,7 +142,7 @@ namespace Metasound
 
 			if (!ReceiveNode.IsValid())
 			{
-				UE_LOG(LogMetaSound, Warning, TEXT("Failed to create receive node while injecting receive node for graph input [VertexKey:%s, VertexDescription:%s, DataTypeName:%s]"), *VertexKey.ToString(), *InputDestination.Vertex.GetMetadata().Description.ToString(), *DataType.ToString());
+				UE_LOG(LogMetaSound, Warning, TEXT("Failed to create receive node while injecting receive node for graph input [VertexKey:%s, VertexDescription:%s, DataTypeName:%s]"), *VertexKey.ToString(), *InputDestination.Vertex.Metadata.Description.ToString(), *DataType.ToString());
 				return false;
 			}
 
@@ -152,7 +152,7 @@ namespace Metasound
 
 			if (!AddressNode.IsValid())
 			{
-				UE_LOG(LogMetaSound, Warning, TEXT("Failed to create address node while injecting receive node for graph input [VertexKey:%s, VertexDescription:%s, DataTypeName:%s]"), *VertexKey.ToString(), *InputDestination.Vertex.GetMetadata().Description.ToString(), *DataType.ToString());
+				UE_LOG(LogMetaSound, Warning, TEXT("Failed to create address node while injecting receive node for graph input [VertexKey:%s, VertexDescription:%s, DataTypeName:%s]"), *VertexKey.ToString(), *InputDestination.Vertex.Metadata.Description.ToString(), *DataType.ToString());
 				return false;
 			}
 
@@ -165,7 +165,7 @@ namespace Metasound
 
 			auto IsEdgeConnectedToCurrentInput = [&InputDestination](const FDataEdge& InEdge)
 			{
-				return (InEdge.From.Node == InputDestination.Node) && (InEdge.From.Vertex.GetVertexName() == InputDestination.Vertex.GetVertexName());
+				return (InEdge.From.Node == InputDestination.Node) && (InEdge.From.Vertex.VertexName == InputDestination.Vertex.VertexName);
 			};
 
 			// Cache previous connections from input node.
@@ -205,7 +205,7 @@ namespace Metasound
 					continue;
 				}
 
-				const FVertexName& VertexKey = InputDestination.Vertex.GetVertexName();
+				const FVertexName& VertexKey = InputDestination.Vertex.VertexName;
 				if (InInputVertexNames.Contains(VertexKey))
 				{
 					const bool bInjectionSuccess = InjectReceiveNode(InGraph, InAddressFunction, InputDestination);
diff --git a/Engine/Plugins/Runtime/Metasound/Source/MetasoundFrontend/Public/MetasoundArrayNodes.h b/Engine/Plugins/Runtime/Metasound/Source/MetasoundFrontend/Public/MetasoundArrayNodes.h
--- a/Engine/Plugins/Runtime/Metasound/Source/MetasoundFrontend/Public/MetasoundArrayNodes.h
+++ b/Engine/Plugins/Runtime/Metasound/Source/MetasoundFrontend/Public/MetasoundArrayNodes.h
@@ -78,10 +78,10 @@ namespace Metasound
 
 			static const FVertexInterface DefaultInterface(
 				FInputVertexInterface(
-					TInputDataVertexModel<ArrayType>(METASOUND_GET_PARAM_NAME_AND_METADATA(InputArray))
+					TInputDataVertex<ArrayType>(METASOUND_GET_PARAM_NAME_AND_METADATA(InputArray))
 				),
 				FOutputVertexInterface(
-					TOutputDataVertexModel<int32>(METASOUND_GET_PARAM_NAME_AND_METADATA(OutputNum))
+					TOutputDataVertex<int32>(METASOUND_GET_PARAM_NAME_AND_METADATA(OutputNum))
 				)
 			);
 
@@ -190,12 +190,12 @@ namespace Metasound
 			using namespace ArrayNodeVertexNames;
 			static const FVertexInterface DefaultInterface(
 				FInputVertexInterface(
-					TInputDataVertexModel<FTrigger>(METASOUND_GET_PARAM_NAME_AND_METADATA(InputTriggerGet)),
-					TInputDataVertexModel<ArrayType>(METASOUND_GET_PARAM_NAME_AND_METADATA(InputArray)),
-					TInputDataVertexModel<int32>(METASOUND_GET_PARAM_NAME_AND_METADATA(InputIndex))
+					TInputDataVertex<FTrigger>(METASOUND_GET_PARAM_NAME_AND_METADATA(InputTriggerGet)),
+					TInputDataVertex<ArrayType>(METASOUND_GET_PARAM_NAME_AND_METADATA(InputArray)),
+					TInputDataVertex<int32>(METASOUND_GET_PARAM_NAME_AND_METADATA(InputIndex))
 				),
 				FOutputVertexInterface(
-					TOutputDataVertexModel<ElementType>(METASOUND_GET_PARAM_NAME_AND_METADATA(OutputValue))
+					TOutputDataVertex<ElementType>(METASOUND_GET_PARAM_NAME_AND_METADATA(OutputValue))
 				)
 			);
 
@@ -327,13 +327,13 @@ namespace Metasound
 			using namespace ArrayNodeVertexNames;
 			static const FVertexInterface DefaultInterface(
 				FInputVertexInterface(
-					TInputDataVertexModel<FTrigger>(METASOUND_GET_PARAM_NAME_AND_METADATA(InputTriggerSet)),
-					TInputDataVertexModel<ArrayType>(METASOUND_GET_PARAM_NAME_AND_METADATA(InputArray)),
-					TInputDataVertexModel<int32>(METASOUND_GET_PARAM_NAME_AND_METADATA(InputIndex)),
-					TInputDataVertexModel<ElementType>(METASOUND_GET_PARAM_NAME_AND_METADATA(InputValue))
+					TInputDataVertex<FTrigger>(METASOUND_GET_PARAM_NAME_AND_METADATA(InputTriggerSet)),
+					TInputDataVertex<ArrayType>(METASOUND_GET_PARAM_NAME_AND_METADATA(InputArray)),
+					TInputDataVertex<int32>(METASOUND_GET_PARAM_NAME_AND_METADATA(InputIndex)),
+					TInputDataVertex<ElementType>(METASOUND_GET_PARAM_NAME_AND_METADATA(InputValue))
 				),
 				FOutputVertexInterface(
-					TOutputDataVertexModel<ArrayType>(METASOUND_GET_PARAM_NAME_AND_METADATA(OutputArraySet))
+					TOutputDataVertex<ArrayType>(METASOUND_GET_PARAM_NAME_AND_METADATA(OutputArraySet))
 				)
 			);
 
@@ -468,12 +468,12 @@ namespace Metasound
 
 			static const FVertexInterface DefaultInterface(
 				FInputVertexInterface(
-					TInputDataVertexModel<FTrigger>(METASOUND_GET_PARAM_NAME_AND_METADATA(InputTriggerGet)),
-					TInputDataVertexModel<ArrayType>(METASOUND_GET_PARAM_NAME_AND_METADATA(InputLeftArray)),
-					TInputDataVertexModel<ArrayType>(METASOUND_GET_PARAM_NAME_AND_METADATA(InputRightArray))
+					TInputDataVertex<FTrigger>(METASOUND_GET_PARAM_NAME_AND_METADATA(InputTriggerGet)),
+					TInputDataVertex<ArrayType>(METASOUND_GET_PARAM_NAME_AND_METADATA(InputLeftArray)),
+					TInputDataVertex<ArrayType>(METASOUND_GET_PARAM_NAME_AND_METADATA(InputRightArray))
 				),
 				FOutputVertexInterface(
-					TOutputDataVertexModel<ArrayType>(METASOUND_GET_PARAM_NAME_AND_METADATA(OutputArrayConcat))
+					TOutputDataVertex<ArrayType>(METASOUND_GET_PARAM_NAME_AND_METADATA(OutputArrayConcat))
 				)
 			);
 
@@ -592,14 +592,14 @@ namespace Metasound
 
 			static const FVertexInterface DefaultInterface(
 				FInputVertexInterface(
-					TInputDataVertexModel<FTrigger>(METASOUND_GET_PARAM_NAME_AND_METADATA(InputTriggerGet)),
-					TInputDataVertexModel<ArrayType>(METASOUND_GET_PARAM_NAME_AND_METADATA(InputArray)),
-					TInputDataVertexModel<int32>(METASOUND_GET_PARAM_NAME_AND_METADATA(InputStartIndex)),
-					TInputDataVertexModel<int32>(METASOUND_GET_PARAM_NAME_AND_METADATA(InputEndIndex))
+					TInputDataVertex<FTrigger>(METASOUND_GET_PARAM_NAME_AND_METADATA(InputTriggerGet)),
+					TInputDataVertex<ArrayType>(METASOUND_GET_PARAM_NAME_AND_METADATA(InputArray)),
+					TInputDataVertex<int32>(METASOUND_GET_PARAM_NAME_AND_METADATA(InputStartIndex)),
+					TInputDataVertex<int32>(METASOUND_GET_PARAM_NAME_AND_METADATA(InputEndIndex))
 
 				),
 				FOutputVertexInterface(
-					TOutputDataVertexModel<ArrayType>(METASOUND_GET_PARAM_NAME_AND_METADATA(OutputArraySubset))
+					TOutputDataVertex<ArrayType>(METASOUND_GET_PARAM_NAME_AND_METADATA(OutputArraySubset))
 				)
 			);
 
diff --git a/Engine/Plugins/Runtime/Metasound/Source/MetasoundFrontend/Public/MetasoundArrayRandomNode.h b/Engine/Plugins/Runtime/Metasound/Source/MetasoundFrontend/Public/MetasoundArrayRandomNode.h
--- a/Engine/Plugins/Runtime/Metasound/Source/MetasoundFrontend/Public/MetasoundArrayRandomNode.h
+++ b/Engine/Plugins/Runtime/Metasound/Source/MetasoundFrontend/Public/MetasoundArrayRandomNode.h
@@ -109,18 +109,18 @@ namespace Metasound
 
 			static const FVertexInterface DefaultInterface(
 				FInputVertexInterface(
-					TInputDataVertexModel<FTrigger>(METASOUND_GET_PARAM_NAME_AND_METADATA(InputTriggerNextValue)),
-					TInputDataVertexModel<FTrigger>(METASOUND_GET_PARAM_NAME_AND_METADATA(InputTriggerResetSeed)),
-					TInputDataVertexModel<ArrayType>(METASOUND_GET_PARAM_NAME_AND_METADATA(InputRandomArray)),
-					TInputDataVertexModel<WeightsArrayType>(METASOUND_GET_PARAM_NAME_AND_METADATA(InputWeights)),
-					TInputDataVertexModel<int32>(METASOUND_GET_PARAM_NAME_AND_METADATA(InputSeed), -1),
-					TInputDataVertexModel<int32>(METASOUND_GET_PARAM_NAME_AND_METADATA(InputNoRepeatOrder), 1),
-					TInputDataVertexModel<bool>(METASOUND_GET_PARAM_NAME_AND_METADATA(InputEnableSharedState), false)
+					TInputDataVertex<FTrigger>(METASOUND_GET_PARAM_NAME_AND_METADATA(InputTriggerNextValue)),
+					TInputDataVertex<FTrigger>(METASOUND_GET_PARAM_NAME_AND_METADATA(InputTriggerResetSeed)),
+					TInputDataVertex<ArrayType>(METASOUND_GET_PARAM_NAME_AND_METADATA(InputRandomArray)),
+					TInputDataVertex<WeightsArrayType>(METASOUND_GET_PARAM_NAME_AND_METADATA(InputWeights)),
+					TInputDataVertex<int32>(METASOUND_GET_PARAM_NAME_AND_METADATA(InputSeed), -1),
+					TInputDataVertex<int32>(METASOUND_GET_PARAM_NAME_AND_METADATA(InputNoRepeatOrder), 1),
+					TInputDataVertex<bool>(METASOUND_GET_PARAM_NAME_AND_METADATA(InputEnableSharedState), false)
 				),
 				FOutputVertexInterface(
-					TOutputDataVertexModel<FTrigger>(METASOUND_GET_PARAM_NAME_AND_METADATA(OutputTriggerOnNext)),
-					TOutputDataVertexModel<FTrigger>(METASOUND_GET_PARAM_NAME_AND_METADATA(OutputTriggerOnReset)),
-					TOutputDataVertexModel<ElementType>(METASOUND_GET_PARAM_NAME_AND_METADATA(ShuffleOutputValue))
+					TOutputDataVertex<FTrigger>(METASOUND_GET_PARAM_NAME_AND_METADATA(OutputTriggerOnNext)),
+					TOutputDataVertex<FTrigger>(METASOUND_GET_PARAM_NAME_AND_METADATA(OutputTriggerOnReset)),
+					TOutputDataVertex<ElementType>(METASOUND_GET_PARAM_NAME_AND_METADATA(ShuffleOutputValue))
 				)
 			);
 
diff --git a/Engine/Plugins/Runtime/Metasound/Source/MetasoundFrontend/Public/MetasoundArrayShuffleNode.h b/Engine/Plugins/Runtime/Metasound/Source/MetasoundFrontend/Public/MetasoundArrayShuffleNode.h
--- a/Engine/Plugins/Runtime/Metasound/Source/MetasoundFrontend/Public/MetasoundArrayShuffleNode.h
+++ b/Engine/Plugins/Runtime/Metasound/Source/MetasoundFrontend/Public/MetasoundArrayShuffleNode.h
@@ -147,19 +147,19 @@ namespace Metasound
 
 			static const FVertexInterface DefaultInterface(
 				FInputVertexInterface(
-					TInputDataVertexModel<FTrigger>(METASOUND_GET_PARAM_NAME_AND_METADATA(InputTriggerNext)),
-					TInputDataVertexModel<FTrigger>(METASOUND_GET_PARAM_NAME_AND_METADATA(InputTriggerShuffle)),
-					TInputDataVertexModel<FTrigger>(METASOUND_GET_PARAM_NAME_AND_METADATA(InputTriggerReset)),
-					TInputDataVertexModel<ArrayType>(METASOUND_GET_PARAM_NAME_AND_METADATA(InputShuffleArray)),
-					TInputDataVertexModel<int32>(METASOUND_GET_PARAM_NAME_AND_METADATA(InputShuffleSeed), -1),
-					TInputDataVertexModel<bool>(METASOUND_GET_PARAM_NAME_AND_METADATA(InputAutoShuffle), true),
-					TInputDataVertexModel<bool>(METASOUND_GET_PARAM_NAME_AND_METADATA(InputShuffleEnableSharedState), false)
+					TInputDataVertex<FTrigger>(METASOUND_GET_PARAM_NAME_AND_METADATA(InputTriggerNext)),
+					TInputDataVertex<FTrigger>(METASOUND_GET_PARAM_NAME_AND_METADATA(InputTriggerShuffle)),
+					TInputDataVertex<FTrigger>(METASOUND_GET_PARAM_NAME_AND_METADATA(InputTriggerReset)),
+					TInputDataVertex<ArrayType>(METASOUND_GET_PARAM_NAME_AND_METADATA(InputShuffleArray)),
+					TInputDataVertex<int32>(METASOUND_GET_PARAM_NAME_AND_METADATA(InputShuffleSeed), -1),
+					TInputDataVertex<bool>(METASOUND_GET_PARAM_NAME_AND_METADATA(InputAutoShuffle), true),
+					TInputDataVertex<bool>(METASOUND_GET_PARAM_NAME_AND_METADATA(InputShuffleEnableSharedState), false)
 					),
 				FOutputVertexInterface(
-					TOutputDataVertexModel<FTrigger>(METASOUND_GET_PARAM_NAME_AND_METADATA(OutputTriggerOnNext)),
-					TOutputDataVertexModel<FTrigger>(METASOUND_GET_PARAM_NAME_AND_METADATA(OutputTriggerOnShuffle)),
-					TOutputDataVertexModel<FTrigger>(METASOUND_GET_PARAM_NAME_AND_METADATA(OutputTriggerOnResetSeed)),
-					TOutputDataVertexModel<ElementType>(METASOUND_GET_PARAM_NAME_AND_METADATA(OutputShuffledValue))
+					TOutputDataVertex<FTrigger>(METASOUND_GET_PARAM_NAME_AND_METADATA(OutputTriggerOnNext)),
+					TOutputDataVertex<FTrigger>(METASOUND_GET_PARAM_NAME_AND_METADATA(OutputTriggerOnShuffle)),
+					TOutputDataVertex<FTrigger>(METASOUND_GET_PARAM_NAME_AND_METADATA(OutputTriggerOnResetSeed)),
+					TOutputDataVertex<ElementType>(METASOUND_GET_PARAM_NAME_AND_METADATA(OutputShuffledValue))
 				)
 			);
 
diff --git a/Engine/Plugins/Runtime/Metasound/Source/MetasoundFrontend/Public/MetasoundAutoConverterNode.h b/Engine/Plugins/Runtime/Metasound/Source/MetasoundFrontend/Public/MetasoundAutoConverterNode.h
--- a/Engine/Plugins/Runtime/Metasound/Source/MetasoundFrontend/Public/MetasoundAutoConverterNode.h
+++ b/Engine/Plugins/Runtime/Metasound/Source/MetasoundFrontend/Public/MetasoundAutoConverterNode.h
@@ -64,10 +64,10 @@ namespace Metasound
 
 			return FVertexInterface(
 				FInputVertexInterface(
-					TInputDataVertexModel<FromDataType>(GetInputName(), InputDesc)
+					TInputDataVertex<FromDataType>(GetInputName(), FDataVertexMetadata{ InputDesc })
 				),
 				FOutputVertexInterface(
-					TOutputDataVertexModel<ToDataType>(GetOutputName(), OutputDesc)
+					TOutputDataVertex<ToDataType>(GetOutputName(), FDataVertexMetadata{ OutputDesc })
 				)
 			);
 		}
diff --git a/Engine/Plugins/Runtime/Metasound/Source/MetasoundFrontend/Public/MetasoundInputNode.h b/Engine/Plugins/Runtime/Metasound/Source/MetasoundFrontend/Public/MetasoundInputNode.h
--- a/Engine/Plugins/Runtime/Metasound/Source/MetasoundFrontend/Public/MetasoundInputNode.h
+++ b/Engine/Plugins/Runtime/Metasound/Source/MetasoundFrontend/Public/MetasoundInputNode.h
@@ -265,10 +265,10 @@ namespace Metasound
 			{
 				return FVertexInterface(
 					FInputVertexInterface(
-						TInputDataVertexModel<DataType>(InVertexName, FText::GetEmpty())
+						TInputDataVertex<DataType>(InVertexName, FDataVertexMetadata{ FText::GetEmpty() })
 					),
 					FOutputVertexInterface(
-						TOutputDataVertexModel<DataType>(InVertexName, FText::GetEmpty())
+						TOutputDataVertex<DataType>(InVertexName, FDataVertexMetadata{ FText::GetEmpty() })
 					)
 				);
 			}
diff --git a/Engine/Plugins/Runtime/Metasound/Source/MetasoundFrontend/Public/MetasoundOutputNode.h b/Engine/Plugins/Runtime/Metasound/Source/MetasoundFrontend/Public/MetasoundOutputNode.h
--- a/Engine/Plugins/Runtime/Metasound/Source/MetasoundFrontend/Public/MetasoundOutputNode.h
+++ b/Engine/Plugins/Runtime/Metasound/Source/MetasoundFrontend/Public/MetasoundOutputNode.h
@@ -105,10 +105,10 @@ namespace Metasound
 
 			return FVertexInterface(
 				FInputVertexInterface(
-					TInputDataVertexModel<DataType>(InVertexName, VertexDescription)
+					TInputDataVertex<DataType>(InVertexName, FDataVertexMetadata{VertexDescription})
 				),
 				FOutputVertexInterface(
-					TOutputDataVertexModel<DataType>(InVertexName, VertexDescription)
+					TOutputDataVertex<DataType>(InVertexName, FDataVertexMetadata{VertexDescription})
 				)
 			);
 		}
diff --git a/Engine/Plugins/Runtime/Metasound/Source/MetasoundFrontend/Public/MetasoundReceiveNode.h b/Engine/Plugins/Runtime/Metasound/Source/MetasoundFrontend/Public/MetasoundReceiveNode.h
--- a/Engine/Plugins/Runtime/Metasound/Source/MetasoundFrontend/Public/MetasoundReceiveNode.h
+++ b/Engine/Plugins/Runtime/Metasound/Source/MetasoundFrontend/Public/MetasoundReceiveNode.h
@@ -55,11 +55,11 @@ namespace Metasound
 			};
 			return FVertexInterface(
 				FInputVertexInterface(
-					TInputDataVertexModel<FSendAddress>(METASOUND_GET_PARAM_NAME(AddressInput), AddressInputMetadata),
-					TInputDataVertexModel<TDataType>(METASOUND_GET_PARAM_NAME(DefaultDataInput), DefaultDataInputMetadata)
+					TInputDataVertex<FSendAddress>(METASOUND_GET_PARAM_NAME(AddressInput), AddressInputMetadata),
+					TInputDataVertex<TDataType>(METASOUND_GET_PARAM_NAME(DefaultDataInput), DefaultDataInputMetadata)
 				),
 				FOutputVertexInterface(
-					TOutputDataVertexModel<TDataType>(METASOUND_GET_PARAM_NAME(Output), OutputMetadata)
+					TOutputDataVertex<TDataType>(METASOUND_GET_PARAM_NAME(Output), OutputMetadata)
 				)
 			);
 		}
diff --git a/Engine/Plugins/Runtime/Metasound/Source/MetasoundFrontend/Public/MetasoundSendNode.h b/Engine/Plugins/Runtime/Metasound/Source/MetasoundFrontend/Public/MetasoundSendNode.h
--- a/Engine/Plugins/Runtime/Metasound/Source/MetasoundFrontend/Public/MetasoundSendNode.h
+++ b/Engine/Plugins/Runtime/Metasound/Source/MetasoundFrontend/Public/MetasoundSendNode.h
@@ -44,8 +44,8 @@ namespace Metasound
 
 			return FVertexInterface(
 				FInputVertexInterface(
-					TInputDataVertexModel<FSendAddress>(METASOUND_GET_PARAM_NAME(AddressInput), AddressInputMetadata),
-					TInputDataVertexModel<TDataType>(GetSendInputName(), FText::GetEmpty())
+					TInputDataVertex<FSendAddress>(METASOUND_GET_PARAM_NAME(AddressInput), AddressInputMetadata),
+					TInputDataVertex<TDataType>(GetSendInputName(), FDataVertexMetadata{ FText::GetEmpty() })
 				),
 				FOutputVertexInterface(
 				)
diff --git a/Engine/Plugins/Runtime/Metasound/Source/MetasoundGraphCore/Private/MetasoundGraph.cpp b/Engine/Plugins/Runtime/Metasound/Source/MetasoundGraphCore/Private/MetasoundGraph.cpp
--- a/Engine/Plugins/Runtime/Metasound/Source/MetasoundGraphCore/Private/MetasoundGraph.cpp
+++ b/Engine/Plugins/Runtime/Metasound/Source/MetasoundGraphCore/Private/MetasoundGraph.cpp
@@ -101,7 +101,7 @@ namespace Metasound
 		const FInputDataVertex& ToVertex = ToVertexInterface.GetInputVertex(ToKey);
 
 
-		if (FromVertex.GetDataTypeName() != ToVertex.GetDataTypeName())
+		if (FromVertex.DataTypeName != ToVertex.DataTypeName)
 		{
 			return false;
 		}
diff --git a/Engine/Plugins/Runtime/Metasound/Source/MetasoundGraphCore/Private/MetasoundGraphLinter.cpp b/Engine/Plugins/Runtime/Metasound/Source/MetasoundGraphCore/Private/MetasoundGraphLinter.cpp
--- a/Engine/Plugins/Runtime/Metasound/Source/MetasoundGraphCore/Private/MetasoundGraphLinter.cpp
+++ b/Engine/Plugins/Runtime/Metasound/Source/MetasoundGraphCore/Private/MetasoundGraphLinter.cpp
@@ -22,7 +22,7 @@ namespace Metasound
 			if (InDest.Node)
 			{
 				const FInputVertexInterface& Interface = InDest.Node->GetVertexInterface().GetInputInterface();
-				if (const FInputDataVertex* Vertex = Interface.Find(InDest.Vertex.GetVertexName()))
+				if (const FInputDataVertex* Vertex = Interface.Find(InDest.Vertex.VertexName))
 				{
 					return *Vertex == InDest.Vertex;
 				}
@@ -37,7 +37,7 @@ namespace Metasound
 			if (InSource.Node)
 			{
 				const FOutputVertexInterface& Interface = InSource.Node->GetVertexInterface().GetOutputInterface();
-				if (const FOutputDataVertex* Vertex = Interface.Find(InSource.Vertex.GetVertexName()))
+				if (const FOutputDataVertex* Vertex = Interface.Find(InSource.Vertex.VertexName))
 				{
 					return *Vertex == InSource.Vertex;
 				}
@@ -66,7 +66,7 @@ namespace Metasound
 
 		for (const FDataEdge& Edge : InGraph.GetDataEdges())
 		{
-			if (Edge.From.Vertex.GetDataTypeName() != Edge.To.Vertex.GetDataTypeName())
+			if (Edge.From.Vertex.DataTypeName != Edge.To.Vertex.DataTypeName)
 			{
 				AddBuildError<FInvalidConnectionDataTypeError>(OutErrors, Edge);
 
diff --git a/Engine/Plugins/Runtime/Metasound/Source/MetasoundGraphCore/Private/MetasoundOperatorBuilder.cpp b/Engine/Plugins/Runtime/Metasound/Source/MetasoundGraphCore/Private/MetasoundOperatorBuilder.cpp
--- a/Engine/Plugins/Runtime/Metasound/Source/MetasoundGraphCore/Private/MetasoundOperatorBuilder.cpp
+++ b/Engine/Plugins/Runtime/Metasound/Source/MetasoundGraphCore/Private/MetasoundOperatorBuilder.cpp
@@ -338,15 +338,15 @@ namespace Metasound
 
 			bool bSuccess = false;
 
-			if (FromDataReferenceCollection.ContainsDataWriteReference(Edge->From.Vertex.GetVertexName(), Edge->From.Vertex.GetDataTypeName()))
+			if (FromDataReferenceCollection.ContainsDataWriteReference(Edge->From.Vertex.VertexName, Edge->From.Vertex.DataTypeName))
 			{
 				// Contains data write reference
-				bSuccess = OutCollection.AddDataWriteReferenceFrom(Edge->To.Vertex.GetVertexName(), FromDataReferenceCollection, Edge->From.Vertex.GetVertexName(), Edge->From.Vertex.GetDataTypeName());
+				bSuccess = OutCollection.AddDataWriteReferenceFrom(Edge->To.Vertex.VertexName, FromDataReferenceCollection, Edge->From.Vertex.VertexName, Edge->From.Vertex.DataTypeName);
 			}
-			else if (FromDataReferenceCollection.ContainsDataReadReference(Edge->From.Vertex.GetVertexName(), Edge->From.Vertex.GetDataTypeName()))
+			else if (FromDataReferenceCollection.ContainsDataReadReference(Edge->From.Vertex.VertexName, Edge->From.Vertex.DataTypeName))
 			{
 				// Contains data read reference
-				bSuccess = OutCollection.AddDataReadReferenceFrom(Edge->To.Vertex.GetVertexName(), FromDataReferenceCollection, Edge->From.Vertex.GetVertexName(), Edge->From.Vertex.GetDataTypeName());
+				bSuccess = OutCollection.AddDataReadReferenceFrom(Edge->To.Vertex.VertexName, FromDataReferenceCollection, Edge->From.Vertex.VertexName, Edge->From.Vertex.DataTypeName);
 			}
 
 			if (!bSuccess)
@@ -376,10 +376,10 @@ namespace Metasound
 		{
 			if (const FInputDataDestination* Destination = *DestinationPtr)
 			{
-				const bool bReferenceExists = InExternalCollection.ContainsDataReadReference(Destination->Vertex.GetVertexName(), Destination->Vertex.GetDataTypeName());
+				const bool bReferenceExists = InExternalCollection.ContainsDataReadReference(Destination->Vertex.VertexName, Destination->Vertex.DataTypeName);
 				if (bReferenceExists)
 				{
-					const bool bSuccess = OutDataReferences.AddDataReadReferenceFrom(Destination->Vertex.GetVertexName(), InExternalCollection, Destination->Vertex.GetVertexName(), Destination->Vertex.GetDataTypeName());
+					const bool bSuccess = OutDataReferences.AddDataReadReferenceFrom(Destination->Vertex.VertexName, InExternalCollection, Destination->Vertex.VertexName, Destination->Vertex.DataTypeName);
 
 					if (!bSuccess)
 					{
@@ -477,8 +477,8 @@ namespace Metasound
 			// Get either readable or writable values for inputs. 
 			// Readable inputs can occur when the data reference is provided outside
 			// of the graph being built. 
-			const bool bContainsWritableInput = Collection.ContainsDataWriteReference(InputDestination.Vertex.GetVertexName(), InputDestination.Vertex.GetDataTypeName());
-			const bool bContainsReadableInput = Collection.ContainsDataReadReference(InputDestination.Vertex.GetVertexName(), InputDestination.Vertex.GetDataTypeName());
+			const bool bContainsWritableInput = Collection.ContainsDataWriteReference(InputDestination.Vertex.VertexName, InputDestination.Vertex.DataTypeName);
+			const bool bContainsReadableInput = Collection.ContainsDataReadReference(InputDestination.Vertex.VertexName, InputDestination.Vertex.DataTypeName);
 
 			const bool bContainsInput = bContainsWritableInput || bContainsReadableInput;
 
@@ -493,11 +493,11 @@ namespace Metasound
 			bool bSuccess = true;
 			if (bContainsWritableInput)
 			{
-				bSuccess = OutGraphInputs.AddDataWriteReferenceFrom(InputDestination.Vertex.GetVertexName(), Collection, InputDestination.Vertex.GetVertexName(), InputDestination.Vertex.GetDataTypeName());
+				bSuccess = OutGraphInputs.AddDataWriteReferenceFrom(InputDestination.Vertex.VertexName, Collection, InputDestination.Vertex.VertexName, InputDestination.Vertex.DataTypeName);
 			}
 			else
 			{
-				bSuccess = OutGraphInputs.AddDataReadReferenceFrom(InputDestination.Vertex.GetVertexName(), Collection, InputDestination.Vertex.GetVertexName(), InputDestination.Vertex.GetDataTypeName());
+				bSuccess = OutGraphInputs.AddDataReadReferenceFrom(InputDestination.Vertex.VertexName, Collection, InputDestination.Vertex.VertexName, InputDestination.Vertex.DataTypeName);
 			}
 
 			if (!bSuccess)
@@ -526,7 +526,7 @@ namespace Metasound
 
 			FDataReferenceCollection& Collection = InNodeDataReferences[OutputSource.Node].Outputs;
 
-			if (!Collection.ContainsDataReadReference(OutputSource.Vertex.GetVertexName(), OutputSource.Vertex.GetDataTypeName()))
+			if (!Collection.ContainsDataReadReference(OutputSource.Vertex.VertexName, OutputSource.Vertex.DataTypeName))
 			{
 				// This will likely produce an IOperator which does not work as
 				// expected.
@@ -537,7 +537,7 @@ namespace Metasound
 				continue;
 			}
 
-			bool bSuccess = OutGraphOutputs.AddDataReadReferenceFrom(OutputSource.Vertex.GetVertexName(), Collection, OutputSource.Vertex.GetVertexName(), OutputSource.Vertex.GetDataTypeName());
+			bool bSuccess = OutGraphOutputs.AddDataReadReferenceFrom(OutputSource.Vertex.VertexName, Collection, OutputSource.Vertex.VertexName, OutputSource.Vertex.DataTypeName);
 
 			if (!bSuccess)
 			{
diff --git a/Engine/Plugins/Runtime/Metasound/Source/MetasoundGraphCore/Private/MetasoundVertex.cpp b/Engine/Plugins/Runtime/Metasound/Source/MetasoundGraphCore/Private/MetasoundVertex.cpp
--- a/Engine/Plugins/Runtime/Metasound/Source/MetasoundGraphCore/Private/MetasoundVertex.cpp
+++ b/Engine/Plugins/Runtime/Metasound/Source/MetasoundGraphCore/Private/MetasoundVertex.cpp
@@ -6,76 +6,9 @@
 
 namespace Metasound
 {
-	FInputDataVertex::FInputDataVertex()
-	:	VertexModel(MakeUnique<FEmptyVertexModel>(TEXT(""), FText::GetEmpty()))
-	{}
-
-	FInputDataVertex::FInputDataVertex(const FInputDataVertex& InOther)
-	:	VertexModel()
-	{
-		// Copy constructor has to call underlying model's clone.
-		if (InOther.VertexModel.IsValid())
-		{
-			VertexModel = InOther.VertexModel->Clone();
-		}
-
-		if (!VertexModel.IsValid())
-		{
-			// Default to empty vertex if there was a failure. 
-			VertexModel = MakeUnique<FEmptyVertexModel>(TEXT(""), FText::GetEmpty());
-		}
-	}
-
-	FInputDataVertex& FInputDataVertex::operator=(const FInputDataVertex& InOther)
-	{
-		// Assignment operator has to call underlying model's clone.
-		if (InOther.VertexModel.IsValid())
-		{
-			VertexModel = InOther.VertexModel->Clone();
-		}
-
-		if (!VertexModel.IsValid())
-		{
-			// Default to empty vertex if there was a failure. 
-			VertexModel = MakeUnique<FEmptyVertexModel>(TEXT(""), FText::GetEmpty());
-		}
-
-		return *this;
-	}
-
-	const FVertexName& FInputDataVertex::GetVertexName() const
-	{
-		return VertexModel->VertexName;
-	}
-
-	const FName& FInputDataVertex::GetDataTypeName() const
-	{
-		return VertexModel->DataTypeName;
-	}
-
-	FLiteral FInputDataVertex::GetDefaultLiteral() const
-	{
-		return VertexModel->CreateDefaultLiteral();
-	}
-
-	const FDataVertexMetadata& FInputDataVertex::GetMetadata() const
-	{
-		return VertexModel->VertexMetadata;
-	}
-
-	bool FInputDataVertex::IsReferenceOfSameType(const IDataReference& InReference) const
-	{
-		return VertexModel->IsReferenceOfSameType(InReference);
-	}
-
-	const FName& FInputDataVertex::GetVertexTypeName() const 
-	{
-		return VertexModel->GetVertexTypeName();
-	}
-
 	bool operator==(const FInputDataVertex& InLHS, const FInputDataVertex& InRHS)
 	{
-		return (InLHS.VertexModel->IsEqual(*InRHS.VertexModel) && (InRHS.VertexModel->IsEqual(*InLHS.VertexModel)));
+		return (InLHS.VertexName == InRHS.VertexName) && (InLHS.DataTypeName == InRHS.DataTypeName);
 	}
 
 	bool operator!=(const FInputDataVertex& InLHS, const FInputDataVertex& InRHS)
@@ -85,87 +18,19 @@ namespace Metasound
 
 	bool operator<(const FInputDataVertex& InLHS, const FInputDataVertex& InRHS)
 	{
-		if (InLHS == InRHS)
-		{
-			return false;
-		}
-		
-		if (InLHS.GetVertexName() == InRHS.GetVertexName())
+		if (InLHS.VertexName == InRHS.VertexName)
 		{
-			return InLHS.GetDataTypeName().FastLess(InRHS.GetDataTypeName());
+			return InLHS.DataTypeName.FastLess(InRHS.DataTypeName);
 		}
 		else
 		{
-			return InLHS.GetVertexName().FastLess(InRHS.GetVertexName());
-		}
-	}
-
-	FOutputDataVertex::FOutputDataVertex()
-	:	VertexModel(MakeUnique<FEmptyVertexModel>(TEXT(""), FText::GetEmpty()))
-	{}
-
-	FOutputDataVertex::FOutputDataVertex(const FOutputDataVertex& InOther)
-	:	VertexModel()
-	{
-		// Copy constructor has to call underlying model's clone.
-		if (InOther.VertexModel.IsValid())
-		{
-			VertexModel = InOther.VertexModel->Clone();
-		}
-
-		if (!VertexModel.IsValid())
-		{
-			// Default to empty vertex if there was a failure. 
-			VertexModel = MakeUnique<FEmptyVertexModel>(TEXT(""), FText::GetEmpty());
-		}
-	}
-
-	FOutputDataVertex& FOutputDataVertex::operator=(const FOutputDataVertex& InOther)
-	{
-		// assignment operator has to call underlying model's clone.
-		if (InOther.VertexModel.IsValid())
-		{
-			VertexModel = InOther.VertexModel->Clone();
-		}
-
-		if (!VertexModel.IsValid())
-		{
-			// Default to empty vertex if there was a failure. 
-			VertexModel = MakeUnique<FEmptyVertexModel>(TEXT(""), FText::GetEmpty());
+			return InLHS.VertexName.FastLess(InRHS.VertexName);
 		}
-
-		return *this;
-	}
-
-	FVertexName FOutputDataVertex::GetVertexName() const
-	{
-		return VertexModel->VertexName;
-	}
-
-	const FName& FOutputDataVertex::GetDataTypeName() const
-	{
-		return VertexModel->DataTypeName;
 	}
 
-	const FDataVertexMetadata& FOutputDataVertex::GetMetadata() const
-	{
-		return VertexModel->VertexMetadata;
-	}
-
-	bool FOutputDataVertex::IsReferenceOfSameType(const IDataReference& InReference) const
-	{
-		return VertexModel->IsReferenceOfSameType(InReference);
-	}
-
-	const FName& FOutputDataVertex::GetVertexTypeName() const
-	{
-		return VertexModel->GetVertexTypeName();
-	}
-
-
 	bool operator==(const FOutputDataVertex& InLHS, const FOutputDataVertex& InRHS)
 	{
-		return (InLHS.VertexModel->IsEqual(*InRHS.VertexModel) && (InRHS.VertexModel->IsEqual(*InLHS.VertexModel)));
+		return (InLHS.VertexName == InRHS.VertexName) && (InLHS.DataTypeName == InRHS.DataTypeName);
 	}
 
 	bool operator!=(const FOutputDataVertex& InLHS, const FOutputDataVertex& InRHS)
@@ -175,75 +40,19 @@ namespace Metasound
 
 	bool operator<(const FOutputDataVertex& InLHS, const FOutputDataVertex& InRHS)
 	{
-		if (InLHS == InRHS)
-		{
-			return false;
-		}
-		
-		if (InLHS.GetVertexName() == InRHS.GetVertexName())
+		if (InLHS.VertexName == InRHS.VertexName)
 		{
-			return InLHS.GetDataTypeName().FastLess(InRHS.GetDataTypeName());
+			return InLHS.DataTypeName.FastLess(InRHS.DataTypeName);
 		}
 		else
 		{
-			return InLHS.GetVertexName().FastLess(InRHS.GetVertexName());
+			return InLHS.VertexName.FastLess(InRHS.VertexName);
 		}
 	}
 
-	FEnvironmentVertex::FEnvironmentVertex()
-	:	VertexModel(MakeUnique<FEmptyVertexModel>(TEXT(""), FText::GetEmpty()))
-	{
-	}
-
-	/** Copy constructor */
-	FEnvironmentVertex::FEnvironmentVertex(const FEnvironmentVertex& InOther)
-	{
-		if (InOther.VertexModel.IsValid())
-		{
-			// call underlying model's clone to copy model.
-			VertexModel = InOther.VertexModel->Clone();
-		}
-		else
-		{
-			// Make an empty model if no valid model exists.
-			VertexModel = MakeUnique<FEmptyVertexModel>(TEXT(""), FText::GetEmpty());
-		}
-	}
-
-	FEnvironmentVertex& FEnvironmentVertex::operator=(const FEnvironmentVertex& InOther)
-	{
-		if (InOther.VertexModel.IsValid())
-		{
-			// call underlying model's clone to copy model.
-			VertexModel = InOther.VertexModel->Clone();
-		}
-		else
-		{
-			// Make an empty model if no valid model exists.
-			VertexModel = MakeUnique<FEmptyVertexModel>(TEXT(""), FText::GetEmpty());
-		}
-
-		return *this;
-	}
-
-	const FVertexName& FEnvironmentVertex::GetVertexName() const
-	{
-		return VertexModel->VertexName;
-	}
-
-	const FText& FEnvironmentVertex::GetDescription() const
-	{
-		return VertexModel->Description;
-	}
-
-	bool FEnvironmentVertex::IsVariableOfSameType(const IMetasoundEnvironmentVariable& InVariable) const
-	{
-		return VertexModel->IsVariableOfSameType(InVariable);
-	}
-
 	bool operator==(const FEnvironmentVertex& InLHS, const FEnvironmentVertex& InRHS)
 	{
-		return (InLHS.VertexModel->IsEqual(*InRHS.VertexModel) && InRHS.VertexModel->IsEqual(*InLHS.VertexModel));
+		return (InLHS.VertexName == InRHS.VertexName);
 	}
 
 	bool operator!=(const FEnvironmentVertex& InLHS, const FEnvironmentVertex& InRHS)
@@ -253,12 +62,7 @@ namespace Metasound
 
 	bool operator<(const FEnvironmentVertex& InLHS, const FEnvironmentVertex& InRHS)
 	{
-		if (InLHS == InRHS)
-		{
-			return false;
-		}
-
-		return InLHS.GetVertexName().FastLess(InRHS.GetVertexName());
+		return InLHS.VertexName.FastLess(InRHS.VertexName);
 	}
 
 	FVertexInterface::FVertexInterface(const FInputVertexInterface& InInputs, const FOutputVertexInterface& InOutputs)
diff --git a/Engine/Plugins/Runtime/Metasound/Source/MetasoundGraphCore/Public/MetasoundLiteralNode.h b/Engine/Plugins/Runtime/Metasound/Source/MetasoundGraphCore/Public/MetasoundLiteralNode.h
--- a/Engine/Plugins/Runtime/Metasound/Source/MetasoundGraphCore/Public/MetasoundLiteralNode.h
+++ b/Engine/Plugins/Runtime/Metasound/Source/MetasoundGraphCore/Public/MetasoundLiteralNode.h
@@ -113,7 +113,7 @@ namespace Metasound
 			return FVertexInterface(
 				FInputVertexInterface(),
 				FOutputVertexInterface(
-					TOutputDataVertexModel<DataType>(METASOUND_GET_PARAM_NAME(OutputValue), OutputMetadata)
+					TOutputDataVertex<DataType>(METASOUND_GET_PARAM_NAME(OutputValue), OutputMetadata)
 				)
 			);
 		}
diff --git a/Engine/Plugins/Runtime/Metasound/Source/MetasoundGraphCore/Public/MetasoundNodeInterface.h b/Engine/Plugins/Runtime/Metasound/Source/MetasoundGraphCore/Public/MetasoundNodeInterface.h
--- a/Engine/Plugins/Runtime/Metasound/Source/MetasoundGraphCore/Public/MetasoundNodeInterface.h
+++ b/Engine/Plugins/Runtime/Metasound/Source/MetasoundGraphCore/Public/MetasoundNodeInterface.h
@@ -268,12 +268,12 @@ namespace Metasound
 	/** Make a FNodeDataVertexKey from an FOutputDataSource. */
 	FORCEINLINE FNodeDataVertexKey MakeSourceDataVertexKey(const FOutputDataSource& InSource)
 	{
-		return FNodeDataVertexKey(InSource.Node, InSource.Vertex.GetVertexName());
+		return FNodeDataVertexKey(InSource.Node, InSource.Vertex.VertexName);
 	}
 
 	FORCEINLINE FNodeDataVertexKey MakeDestinationDataVertexKey(const FInputDataDestination& InDestination)
 	{
-		return FNodeDataVertexKey(InDestination.Node, InDestination.Vertex.GetVertexName());
+		return FNodeDataVertexKey(InDestination.Node, InDestination.Vertex.VertexName);
 	}
 
 	/** FDataEdge
diff --git a/Engine/Plugins/Runtime/Metasound/Source/MetasoundGraphCore/Public/MetasoundParamHelper.h b/Engine/Plugins/Runtime/Metasound/Source/MetasoundGraphCore/Public/MetasoundParamHelper.h
--- a/Engine/Plugins/Runtime/Metasound/Source/MetasoundGraphCore/Public/MetasoundParamHelper.h
+++ b/Engine/Plugins/Runtime/Metasound/Source/MetasoundGraphCore/Public/MetasoundParamHelper.h
@@ -70,7 +70,10 @@
 #if WITH_EDITOR
 #define METASOUND_GET_PARAM_NAME_WITH_INDEX_AND_TT(NAME, INDEX)  *FString::Format(NAME##Name, {INDEX}), FText::Format(NAME##Tooltip, INDEX)
 #define METASOUND_GET_PARAM_TT_WITH_INDEX(NAME, INDEX)  FText::Format(NAME##Tooltip, INDEX)
+#define METASOUND_GET_PARAM_NAME_WITH_INDEX_AND_METADATA(NAME, INDEX)  *FString::Format(NAME##Name, {INDEX}), FDataVertexMetadata {FText::Format(NAME##Tooltip, INDEX)}
 #else 
 #define METASOUND_GET_PARAM_NAME_WITH_INDEX_AND_TT(NAME, INDEX)  *FString::Format(NAME##Name, {INDEX}), FText::GetEmpty()
 #define METASOUND_GET_PARAM_TT_WITH_INDEX(NAME, INDEX)  FText::GetEmpty();
+#define METASOUND_GET_PARAM_NAME_WITH_INDEX_AND_METADATA(NAME, INDEX)  *FString::Format(NAME##Name, {INDEX}), FDataVertexMetadata {FText::GetEmpty()}
 #endif // WITH_EDITOR
+
diff --git a/Engine/Plugins/Runtime/Metasound/Source/MetasoundGraphCore/Public/MetasoundVariableNodes.h b/Engine/Plugins/Runtime/Metasound/Source/MetasoundGraphCore/Public/MetasoundVariableNodes.h
--- a/Engine/Plugins/Runtime/Metasound/Source/MetasoundGraphCore/Public/MetasoundVariableNodes.h
+++ b/Engine/Plugins/Runtime/Metasound/Source/MetasoundGraphCore/Public/MetasoundVariableNodes.h
@@ -191,7 +191,7 @@ namespace Metasound
 				FInputVertexInterface(
 				),
 				FOutputVertexInterface(
-					TOutputDataVertexModel<FVariable>(METASOUND_GET_PARAM_NAME(OutputVariable), OutputVariableMetadata)
+					TOutputDataVertex<FVariable>(METASOUND_GET_PARAM_NAME(OutputVariable), OutputVariableMetadata)
 				)
 			);
 		}
@@ -338,11 +338,11 @@ namespace Metasound
 
 		return FVertexInterface(
 			FInputVertexInterface(
-				TInputDataVertexModel<DataType>(METASOUND_GET_PARAM_NAME(InputData), InputDataMetadata),
-				TInputDataVertexModel<FVariable>(METASOUND_GET_PARAM_NAME(InputVariable), InputVariableMetadata)
+				TInputDataVertex<DataType>(METASOUND_GET_PARAM_NAME(InputData), InputDataMetadata),
+				TInputDataVertex<FVariable>(METASOUND_GET_PARAM_NAME(InputVariable), InputVariableMetadata)
 			),
 			FOutputVertexInterface(
-				TOutputDataVertexModel<FVariable>(METASOUND_GET_PARAM_NAME(OutputVariable), OutputVariableMetadata)
+				TOutputDataVertex<FVariable>(METASOUND_GET_PARAM_NAME(OutputVariable), OutputVariableMetadata)
 			)
 		);
 	}
@@ -454,11 +454,11 @@ namespace Metasound
 
 		return FVertexInterface(
 			FInputVertexInterface(
-				TInputDataVertexModel<FVariable>(METASOUND_GET_PARAM_NAME(InputVariable), InputVariableMetadata)
+				TInputDataVertex<FVariable>(METASOUND_GET_PARAM_NAME(InputVariable), InputVariableMetadata)
 			),
 			FOutputVertexInterface(
-				TOutputDataVertexModel<FVariable>(METASOUND_GET_PARAM_NAME(OutputVariable), OutputVariableMetadata),
-				TOutputDataVertexModel<DataType>(METASOUND_GET_PARAM_NAME(OutputData), OutputDataMetadata)
+				TOutputDataVertex<FVariable>(METASOUND_GET_PARAM_NAME(OutputVariable), OutputVariableMetadata),
+				TOutputDataVertex<DataType>(METASOUND_GET_PARAM_NAME(OutputData), OutputDataMetadata)
 			)
 		);
 	}
@@ -568,11 +568,11 @@ namespace Metasound
 
 		return FVertexInterface(
 			FInputVertexInterface(
-				TInputDataVertexModel<FVariable>(METASOUND_GET_PARAM_NAME(InputVariable), InputDataMetadata)
+				TInputDataVertex<FVariable>(METASOUND_GET_PARAM_NAME(InputVariable), InputDataMetadata)
 			),
 			FOutputVertexInterface(
-				TOutputDataVertexModel<DataType>(METASOUND_GET_PARAM_NAME(OutputData), OutputDataMetadata),
-				TOutputDataVertexModel<FVariable>(METASOUND_GET_PARAM_NAME(OutputVariable), OutputVariableMetadata)
+				TOutputDataVertex<DataType>(METASOUND_GET_PARAM_NAME(OutputData), OutputDataMetadata),
+				TOutputDataVertex<FVariable>(METASOUND_GET_PARAM_NAME(OutputVariable), OutputVariableMetadata)
 			)
 		);
 	}
diff --git a/Engine/Plugins/Runtime/Metasound/Source/MetasoundGraphCore/Public/MetasoundVertex.h b/Engine/Plugins/Runtime/Metasound/Source/MetasoundGraphCore/Public/MetasoundVertex.h
--- a/Engine/Plugins/Runtime/Metasound/Source/MetasoundGraphCore/Public/MetasoundVertex.h
+++ b/Engine/Plugins/Runtime/Metasound/Source/MetasoundGraphCore/Public/MetasoundVertex.h
@@ -2,7 +2,7 @@
 
 #pragma once
 
-#include "Containers/SortedMap.h"
+#include "Containers/Array.h"
 #include "MetasoundDataFactory.h"
 #include "MetasoundDataReference.h"
 #include "MetasoundEnvironment.h"
@@ -12,15 +12,9 @@
 
 namespace Metasound
 {
-	/** Name of a given vertex.  Only unique for a given node interface.
-	 */
+	/** Name of a given vertex.  Only unique for a given node interface. */
 	using FVertexName = FName;
 
-	// Forward declarations
-	class FInputDataVertex;
-	class FOutputDataVertex;
-	class FEnvironmentVertex;
-
 	// Vertex metadata
 	struct FDataVertexMetadata
 	{
@@ -29,618 +23,285 @@ namespace Metasound
 		bool bIsAdvancedDisplay = false;
 	};
 
-	/** FDataVertexModel
-	 *
-	 * An FDataVertexModel implements various vertex behaviors. 
+	/** FDataVertex
 	 *
-	 * Essentially, an FDataVertexModel implements the underlying behavior of a
-	 * FInputDataVertex or FOutputDataVertex.  
-	 *
-	 * The FInputDataVertex and FOutputDataVertex supply copy constructors and 
-	 * assignment operators by utilizing the Clone() operator of the 
-	 * FDataVertexModel. This allows FDataVertexModel behavior to be passed 
-	 * to other objects without having to know the concrete implementation of the
-	 * FDataVertexModel.
+	 * An FDataVertex is a named vertex of a MetaSound node which can contain data.
 	 */
-	struct FDataVertexModel
+	class FDataVertex
 	{
-		/** FDataVertexModel Constructor
+
+	public:
+
+		FDataVertex() = default;
+
+		/** FDataVertex Constructor
 		 *
 		 * @InVertexName - Name of vertex.
 		 * @InDataTypeName - Name of data type.
 		 * @InMetadata - Metadata pertaining to given vertex.
 		 */
-		FDataVertexModel(const FVertexName& InVertexName, const FName& InDataTypeName, const FDataVertexMetadata& InMetadata)
+		FDataVertex(const FVertexName& InVertexName, const FName& InDataTypeName, const FDataVertexMetadata& InMetadata)
 		:	VertexName(InVertexName)
 		,	DataTypeName(InDataTypeName)
-		,	VertexMetadata(InMetadata)
+		,	Metadata(InMetadata)
 		{
 		}
 
-		virtual ~FDataVertexModel() = default;
+		virtual ~FDataVertex() = default;
 
 		/** Name of vertex. */
-		const FVertexName VertexName;
+		FVertexName VertexName;
 
 		/** Type name of data. */
-		const FName DataTypeName;
+		FName DataTypeName;
 
 		/** Metadata associated with vertex. */
-		const FDataVertexMetadata VertexMetadata;
+		FDataVertexMetadata Metadata;
 
-		/** Test if a IDataReference matches the DataType of this vertex model.
-		 *
-		 * @param InReference - Data reference object.
-		 *
-		 * @return True if the types are equal, false otherwise.
-		 */
-		virtual bool IsReferenceOfSameType(const IDataReference& InReference) const = 0;
+		UE_DEPRECATED(5.1, "GetVertexName() is deprecated. Use VertexName")
+		const FVertexName& GetVertexName() const
+		{
+			return VertexName;
+		}
 
-		/** Return the vertex type name (not to be confused with the data type name).
-		 *
-		 * @return Name of this vertex type.
-		 */
-		virtual const FName& GetVertexTypeName() const = 0;
-	};
+		UE_DEPRECATED(5.1, "GetDataTypeName() is deprecated. Use DataTypeName")
+		const FName& GetDataTypeName() const
+		{
+			return DataTypeName;
+		}
 
+		UE_DEPRECATED(5.1, "GetMetadata() is deprecated. Use Metadata")
+		const FDataVertexMetadata& GetMetadata() const
+		{
+			return Metadata;
+		}
+	};
 
-	/** FInputDataVertexModel
-	 *
-	 * Vertex model for inputs.
-	 */
-	struct FInputDataVertexModel : public FDataVertexModel
+	/** FInputDataVertex */
+	class FInputDataVertex : public FDataVertex
 	{
-	private:
-		
-		// Factory for creating a literal. 
-		struct ILiteralFactory
+		// FLiteralFactory supports creation of an FLiteral for a vertex given a
+		// copyable literal value type.
+		//
+		// An FLiteral cannot be assigned or copy-constructed because it can contain
+		// a TUniquePtr<> which does not support the assignment operator and copy constructor.
+		//
+		// For factory can create a literal given a copyable type. The factory can
+		// also be copied.
+		struct FLiteralFactory
 		{
-			virtual ~ILiteralFactory() = default;
-			virtual FLiteral CreateLiteral() const = 0;
-			virtual TUniquePtr<ILiteralFactory> Clone() const = 0;
-		};
+			struct ILiteralFactoryImpl
+			{
+				virtual ~ILiteralFactoryImpl() = default;
+				virtual FLiteral CreateLiteral() const = 0;
+				virtual TUniquePtr<ILiteralFactoryImpl> Clone() const = 0;
+			};
+
+			// Create a literal with a copyable type.
+			template<typename LiteralValueType>
+			struct TLiteralFactoryImpl : ILiteralFactoryImpl
+			{
+				static_assert(std::is_copy_constructible<LiteralValueType>::value && std::is_copy_assignable<LiteralValueType>::value, "Literals can only be assigned for copyable types");
 
-		// Create a literal with a copyable type.
-		template<typename LiteralValueType>
-		struct TLiteralFactory : ILiteralFactory
-		{
-			LiteralValueType LiteralValue;
+				TLiteralFactoryImpl(const LiteralValueType& InValue)
+				: LiteralValue(InValue)
+				{
+				}
+				
+				FLiteral CreateLiteral() const override
+				{
+					return FLiteral(LiteralValue);
+				}
 
-			TLiteralFactory(const LiteralValueType& InValue)
-			: LiteralValue(InValue)
-			{
-			}
-			
-			FLiteral CreateLiteral() const override
+				TUniquePtr<ILiteralFactoryImpl> Clone() const override
+				{
+					return MakeUnique<TLiteralFactoryImpl<LiteralValueType>>(*this);
+				}
+			private:
+				LiteralValueType LiteralValue;
+			};
+
+			FLiteralFactory()
+			: FactoryImpl(nullptr)
 			{
-				return FLiteral(LiteralValue);
 			}
 
-			TUniquePtr<ILiteralFactory> Clone() const override
+			template<typename ValueType>
+			FLiteralFactory(const ValueType& InValue)
+			: FactoryImpl(MakeUnique<TLiteralFactoryImpl<ValueType>>(InValue))
 			{
-				return MakeUnique<TLiteralFactory<LiteralValueType>>(*this);
 			}
-		};
-
-	public:
-
-		FInputDataVertexModel(const FVertexName& InVertexName, const FName& InDataTypeName, const FText& InDescription)
-			: FDataVertexModel(InVertexName, InDataTypeName, { InDescription })
-			, LiteralFactory(MakeUnique<TLiteralFactory<FLiteral::FNone>>(FLiteral::FNone{}))
-		{
-		}
-
-		FInputDataVertexModel(const FVertexName& InVertexName, const FName& InDataTypeName, const FDataVertexMetadata& InMetadata)
-			: FDataVertexModel(InVertexName, InDataTypeName, InMetadata)
-			, LiteralFactory(MakeUnique<TLiteralFactory<FLiteral::FNone>>(FLiteral::FNone{}))
-		{
-		}
 
-		template<typename LiteralValueType>
-		FInputDataVertexModel(const FVertexName& InVertexName, const FName& InDataTypeName, const FText& InDescription, const LiteralValueType& InLiteralValue)
-			: FDataVertexModel(InVertexName, InDataTypeName, { InDescription })
-			, LiteralFactory(MakeUnique<TLiteralFactory<LiteralValueType>>(InLiteralValue))
-		{
-		}
-
-		template<typename LiteralValueType>
-		FInputDataVertexModel(const FVertexName& InVertexName, const FName& InDataTypeName, const FDataVertexMetadata& InMetadata, const LiteralValueType& InLiteralValue)
-			: FDataVertexModel(InVertexName, InDataTypeName, InMetadata)
-			, LiteralFactory(MakeUnique<TLiteralFactory<LiteralValueType>>(InLiteralValue))
-		{
-		}
-
-		FInputDataVertexModel(const FInputDataVertexModel& InOther)
-		: FDataVertexModel(InOther)
-		{
-			if (InOther.LiteralFactory)
+			FLiteralFactory(const FLiteralFactory& InOther)
 			{
-				LiteralFactory = InOther.LiteralFactory->Clone();
+				*this = InOther;
 			}
-		}
-
-		/** Create a clone of this FInputDataVertexModel */
-		virtual TUniquePtr<FInputDataVertexModel> Clone() const = 0;
 
-		/** Creates the default literal for this vertex. If a default does not exist,
-		 * then an invalid literal is returned. 
-		 */
-		FLiteral CreateDefaultLiteral() const 
-		{
-			if (LiteralFactory.IsValid())
+			FLiteralFactory& operator=(const FLiteralFactory& InOther)
 			{
-				return LiteralFactory->CreateLiteral();
+				FactoryImpl.Reset();
+				if (InOther.FactoryImpl.IsValid())
+				{
+					FactoryImpl = InOther.FactoryImpl->Clone();
+				}
+				return *this;
 			}
 
-			return FLiteral::CreateInvalid();
-		}
+			FLiteral CreateLiteral() const
+			{
+				if (FactoryImpl.IsValid())
+				{
+					return FactoryImpl->CreateLiteral();
+				}
 
-		friend bool METASOUNDGRAPHCORE_API operator==(const FInputDataVertex& InLHS, const FInputDataVertex& InRHS);
-		friend bool METASOUNDGRAPHCORE_API operator!=(const FInputDataVertex& InLHS, const FInputDataVertex& InRHS);
+				return FLiteral::CreateInvalid();
+			}
 
 		private:
 
-		/** Compare another FInputDataVertexModel is equal to this
-		 *
-		 * @return True if models are equal, false otherwise.
-		 */
-		virtual bool IsEqual(const FInputDataVertexModel& InOther) const = 0;
-
-		TUniquePtr<ILiteralFactory> LiteralFactory;
-	};
-
-	/** FOutputDataVertexModel
-	 *
-	 * Vertex model for outputs.
-	 */
-	struct FOutputDataVertexModel : public FDataVertexModel
-	{
-		using FDataVertexModel::FDataVertexModel;
-
-		/** Create a clone of this FOutputDataVertexModel */
-		virtual TUniquePtr<FOutputDataVertexModel> Clone() const = 0;
-
-		friend bool METASOUNDGRAPHCORE_API operator==(const FOutputDataVertex& InLHS, const FOutputDataVertex& InRHS);
-		friend bool METASOUNDGRAPHCORE_API operator!=(const FOutputDataVertex& InLHS, const FOutputDataVertex& InRHS);
-
-		private:
+			TUniquePtr<ILiteralFactoryImpl> FactoryImpl;
+		};
 
-		/** Compare another FOutputDataVertexModel is equal to this
-		 *
-		 * @return True if models are equal, false otherwise.
-		 */
-		virtual bool IsEqual(const FOutputDataVertexModel& InOther) const = 0;
+	public:
 
-	};
+		FInputDataVertex() = default;
 
-	/** TBaseVertexModel provides basic functionality of vertex models. */
-	template<typename DataType, typename VertexModelType>
-	struct TBaseVertexModel : public VertexModelType
-	{
-		/** TBaseVertexModel Constructor
-		 *
-		 * @InVertexName - Name of vertex.
-		 * @InDescription - Human readable vertex description.
-		 */
-		template<typename... ModelArgTypes>
-		TBaseVertexModel(const FVertexName& InVertexName, const FText& InDescription, ModelArgTypes&&... ModelArgs)
-		:	VertexModelType(InVertexName, GetMetasoundDataTypeName<DataType>(), { InDescription }, Forward<ModelArgTypes>(ModelArgs)...)
+		/** Construct an FInputDataVertex. */
+		FInputDataVertex(const FVertexName& InVertexName, const FName& InDataTypeName, const FDataVertexMetadata& InMetadata)
+			: FDataVertex(InVertexName, InDataTypeName, InMetadata)
 		{
 		}
 
-		/** TBaseVertexModel Constructor
-		 *
-		 * @InVertexName - Name of vertex.
-		 * @InVertexMetadata - Vertex metadata, used primarily for debugging or live edit.
-		 */
-		template<typename... ModelArgTypes>
-		TBaseVertexModel(const FVertexName& InVertexName, const FDataVertexMetadata& InMetadata, ModelArgTypes&&... ModelArgs)
-		:	VertexModelType(InVertexName, GetMetasoundDataTypeName<DataType>(), InMetadata, Forward<ModelArgTypes>(ModelArgs)...)
+		/** Construct an FInputDataVertex with a default literal. */
+		template<typename LiteralValueType>
+		FInputDataVertex(const FVertexName& InVertexName, const FName& InDataTypeName, const FDataVertexMetadata& InMetadata, const LiteralValueType& InLiteralValue)
+			: FDataVertex(InVertexName, InDataTypeName, InMetadata)
+			, LiteralFactory(InLiteralValue)
 		{
 		}
 
-		/** Test if a IDataReference matches the DataType of this vertex model.
-		 *
-		 * @param InReference - Data reference object.
-		 *
-		 * @return True if the types are equal, false otherwise.
-		 */
-		virtual bool IsReferenceOfSameType(const IDataReference& InReference) const override
+		/** Returns the default literal associated with this input. */
+		FLiteral GetDefaultLiteral() const 
 		{
-			return IsDataReferenceOfType<DataType>(InReference);
+			return LiteralFactory.CreateLiteral();
 		}
 
+		friend bool METASOUNDGRAPHCORE_API operator==(const FInputDataVertex& InLHS, const FInputDataVertex& InRHS);
+		friend bool METASOUNDGRAPHCORE_API operator!=(const FInputDataVertex& InLHS, const FInputDataVertex& InRHS);
+		friend bool METASOUNDGRAPHCORE_API operator<(const FInputDataVertex& InLHS, const FInputDataVertex& InRHS);
 
-		private:
-
-		/** Compare if another VertexModelType is equal to this.
-		 *
-		 * @return True if models are equal, false otherwise.
-		 */
-		virtual bool IsEqual(const VertexModelType& InOther) const override
-		{
-			bool bIsEqual = this->GetVertexTypeName() == InOther.GetVertexTypeName();
-
-			bIsEqual &= this->VertexName == InOther.VertexName;
-			bIsEqual &= this->DataTypeName == InOther.DataTypeName;
+	private:
 
-			return bIsEqual;
-		}
+		FLiteralFactory LiteralFactory;
 	};
 
-	/** TOuputDataVertexModel creates a simple, unchanging, input vertex. */
+	/** Create a FInputDataVertex with a templated MetaSound data type. */
 	template<typename DataType>
-	struct TInputDataVertexModel : public TBaseVertexModel<DataType, FInputDataVertexModel>
+	class TInputDataVertex : public FInputDataVertex
 	{
+	public:
+		TInputDataVertex() = default;
 
-		using TBaseVertexModel<DataType, FInputDataVertexModel>::TBaseVertexModel;
-
-		/** TInputDataVertexModel Constructor
-		 *
-		 * @InVertexName - Name of vertex.
-		 * @InDescription - Human readable vertex description.
-		 */
-		TInputDataVertexModel(const FVertexName& InVertexName, const FText& InDescription)
-		:	TBaseVertexModel<DataType, FInputDataVertexModel>(InVertexName, InDescription)
-		{
-		}
-
-		/** TInputDataVertexModel Constructor
-		 *
-		 * @InVertexName - Name of vertex.
-		 * @InDescription - Human readable vertex description.
-		 * @InDefaultValue - Default Value of vertex
-		 */
-		template<typename LiteralValueType>
-		TInputDataVertexModel(const FVertexName& InVertexName, const FText& InDescription, const LiteralValueType& InDefaultValue)
-		:	TBaseVertexModel<DataType, FInputDataVertexModel>(InVertexName, InDescription, InDefaultValue)
-		{
-		}
-
-		/** TInputDataVertexModel Constructor
-		 *
-		 * @InVertexName - Name of vertex.
-		 * @InDescription - Human readable vertex description.
-		 * @InDefaultValue - Default Value of vertex
-		 */
-		template<typename LiteralValueType>
-		TInputDataVertexModel(const FVertexName& InVertexName, const FDataVertexMetadata& InMetadata)
-		:	TBaseVertexModel<DataType, FInputDataVertexModel>(InVertexName, InMetadata)
-		{
-		}
-
-		TInputDataVertexModel(const TInputDataVertexModel& InOther) = default;
-
-		/** Return the vertex type name (not to be confused with the data type name).
-		 *
-		 * @return Name of this vertex type.
-		 */
-		virtual const FName& GetVertexTypeName() const override
+		template<typename... ArgTypes>
+		TInputDataVertex(const FVertexName& InVertexName, const FDataVertexMetadata& InMetadata, ArgTypes&&... Args)
+		: FInputDataVertex(InVertexName, GetMetasoundDataTypeName<DataType>(), InMetadata, Forward<ArgTypes>(Args)...)
 		{
-			static const FString VertexTypeString = TEXT("InputDataVertex_") + GetMetasoundDataTypeName<DataType>().ToString();
-			static const FName VertexTypeName(*VertexTypeString);
-
-			return VertexTypeName;
 		}
 
-		/** Create a clone of this VertexModelType */
-		virtual TUniquePtr<FInputDataVertexModel> Clone() const override
+		template<typename... ArgTypes>
+ 		UE_DEPRECATED(5.1, "TInputDataVertex constructor without explicit FDataVertexMatadata is deprecated. Use constructor which accepts TDataVertexMetadata")
+		TInputDataVertex(const FVertexName& InVertexName, const FText& InDescription, ArgTypes&&... Args)
+		: FInputDataVertex(InVertexName, GetMetasoundDataTypeName<DataType>(), { InDescription }, Forward<ArgTypes>(Args)...)
 		{
-			return MakeUnique<TInputDataVertexModel<DataType>>(*this);
 		}
 	};
 
-
-	/** FInputDataVertex represents an input vertex to an interface. It uses
-	 * a FInputDataVertexModel to determine its behavior.
+	template<typename DataType>
+	using TInputDataVertexModel UE_DEPRECATED(5.1, "TInputDataVertexModel<DataType> is deprecated. Replace with TInputDataVertex<DataType>") = TInputDataVertex<DataType>;
+	
+	/** FOutputDataVertex
+	 *
+	 * Vertex for outputs.
 	 */
-	class METASOUNDGRAPHCORE_API FInputDataVertex
+	class FOutputDataVertex : public FDataVertex
 	{
-			using FEmptyVertexModel = TInputDataVertexModel<void>;
-
-		public:
-
-			/** Construct an FInputDataVertex with a given vertex model.
-			 *
-			 * @param InVertexModel - A model subclass of FInputDataVertexModel.
-			 */
-			template<typename VertexModelType, typename = typename TEnableIf<TIsDerivedFrom<VertexModelType, FInputDataVertexModel>::Value>::Type>
-			FInputDataVertex(const VertexModelType& InVertexModel)
-			:	VertexModel(InVertexModel.Clone())
-			{
-				static_assert(TIsDerivedFrom<VertexModelType, FInputDataVertexModel>::Value, "Vertex implementation must be a subclass of FInputDataVertexModel");
-
-				if (!VertexModel.IsValid())
-				{
-					VertexModel = MakeUnique<FEmptyVertexModel>(TEXT(""), FText::GetEmpty());
-				}
-			}
-
-			/** Construct an empty FInputDataVertex. */
-			FInputDataVertex();
-
-			/** Copy constructor */
-			FInputDataVertex(const FInputDataVertex& InOther);
-
-			/** Assignment operator. */
-			FInputDataVertex& operator=(const FInputDataVertex& InOther);
-
-			/** Name of vertex. */
-			const FVertexName& GetVertexName() const;
-
-			/** Type name of data reference. */
-			const FName& GetDataTypeName() const;
-
-			/** Description of the vertex. */
-			const FDataVertexMetadata& GetMetadata() const;
-
-			/** Default value of the vertex. */
-			FLiteral GetDefaultLiteral() const;
-
-			/** Determine if vertex refers to same data type. 
-			 *
-			 * @param InReference - Data reference object.
-			 *
-			 * @return True if the types are equal, false otherwise.
-			 */
-			bool IsReferenceOfSameType(const IDataReference& InReference) const;
-
-			/** Return the vertex type name (not to be confused with the data type name).
-			 *
-			 * @return Name of this vertex type.
-			 */
-			const FName& GetVertexTypeName() const;
-
-			friend bool METASOUNDGRAPHCORE_API operator==(const FInputDataVertex& InLHS, const FInputDataVertex& InRHS);
-			friend bool METASOUNDGRAPHCORE_API operator!=(const FInputDataVertex& InLHS, const FInputDataVertex& InRHS);
-			friend bool METASOUNDGRAPHCORE_API operator<(const FInputDataVertex& InLHS, const FInputDataVertex& InRHS);
-
-		private:
+	public:
+		using FDataVertex::FDataVertex;
 
-			TUniquePtr<FInputDataVertexModel> VertexModel;
+		friend bool METASOUNDGRAPHCORE_API operator==(const FOutputDataVertex& InLHS, const FOutputDataVertex& InRHS);
+		friend bool METASOUNDGRAPHCORE_API operator!=(const FOutputDataVertex& InLHS, const FOutputDataVertex& InRHS);
+		friend bool METASOUNDGRAPHCORE_API operator<(const FOutputDataVertex& InLHS, const FOutputDataVertex& InRHS);
 	};
 
-	/** TOuputDataVertexModel creates a simple, unchanging, output vertex. */
+	/** Create a FOutputDataVertex with a templated MetaSound data type. */
 	template<typename DataType>
-	struct TOutputDataVertexModel : TBaseVertexModel<DataType, FOutputDataVertexModel>
+	class TOutputDataVertex : public FOutputDataVertex
 	{
-		using TBaseVertexModel<DataType, FOutputDataVertexModel>::TBaseVertexModel;
+	public:
 
-		/** Return the vertex type name (not to be confused with the data type name).
-		 *
-		 * @return Name of this vertex type.
-		 */
-		virtual const FName& GetVertexTypeName() const override
-		{
-			static const FString VertexTypeString = TEXT("OutputDataVertex_") + GetMetasoundDataTypeName<DataType>().ToString();
-			static const FName VertexTypeName(*VertexTypeString);
+		TOutputDataVertex() = default;
 
-			return VertexTypeName;
+		template<typename... ArgTypes>
+		TOutputDataVertex(const FVertexName& InVertexName, const FDataVertexMetadata& InMetadata, ArgTypes&&... Args)
+		: FOutputDataVertex(InVertexName, GetMetasoundDataTypeName<DataType>(), InMetadata, Forward<ArgTypes>(Args)...)
+		{
 		}
 
-		/** Create a clone of this TOutputDataVertexModel*/
-		virtual TUniquePtr<FOutputDataVertexModel> Clone() const override
+		template<typename... ArgTypes>
+ 		UE_DEPRECATED(5.1, "TOutputDataVertex constructor without explicit FDataVertexMatadata is deprecated. Use constructor which accepts TDataVertexMetadata")
+		TOutputDataVertex(const FVertexName& InVertexName, const FText& InDescription, ArgTypes&&... Args)
+		: FOutputDataVertex(InVertexName, GetMetasoundDataTypeName<DataType>(), { InDescription }, Forward<ArgTypes>(Args)...)
 		{
-			return MakeUnique<TOutputDataVertexModel<DataType>>(this->VertexName, this->VertexMetadata);
 		}
 	};
 
-	/** FOutputDataVertex represents an input vertex to an interface. It uses
-	 * a FOutputDataVertexModel to determine its behavior.
-	 */
-	class METASOUNDGRAPHCORE_API FOutputDataVertex
-	{
-		using FEmptyVertexModel = TOutputDataVertexModel<void>;
-
-		public:
-
-			/** Construct an FOutputDataVertex with a given vertex model.
-			 *
-			 * @param InVertexModel - A model subclass of FOutputDataVertexModel.
-			 */
-			template<typename VertexModelType>
-			FOutputDataVertex(const VertexModelType& InVertexModel)
-			:	VertexModel(InVertexModel.Clone())
-			{
-				static_assert(TIsDerivedFrom<VertexModelType, FOutputDataVertexModel>::Value, "Vertex implementation must be a subclass of FOutputDataVertexModel");
-
-				if (!VertexModel.IsValid())
-				{
-					VertexModel = MakeUnique<FEmptyVertexModel>(TEXT(""), FText::GetEmpty());
-				}
-			}
-
-			/** Construct an empty FInputDataVertex. */
-			FOutputDataVertex();
-
-			/** Copy constructor */
-			FOutputDataVertex(const FOutputDataVertex& InOther);
-
-			/** Assignment operator. */
-			FOutputDataVertex& operator=(const FOutputDataVertex& InOther);
-
-			/** Name of vertex. */
-			FVertexName GetVertexName() const;
-
-			/** Type name of data reference. */
-			const FName& GetDataTypeName() const;
-
-			/** Metadata of the vertex. */
-			const FDataVertexMetadata& GetMetadata() const;
-
-			/** Test if a IDataReference matches the DataType of this vertex.
-			 *
-			 * @param InReference - Data reference object.
-			 *
-			 * @return True if the types are equal, false otherwise.
-			 */
-			bool IsReferenceOfSameType(const IDataReference& InReference) const;
-
-			/** Return the vertex type name (not to be confused with the data type name).
-			 *
-			 * @return Name of this vertex type.
-			 */
-			const FName& GetVertexTypeName() const;
-
-			friend bool METASOUNDGRAPHCORE_API operator==(const FOutputDataVertex& InLHS, const FOutputDataVertex& InRHS);
-			friend bool METASOUNDGRAPHCORE_API operator!=(const FOutputDataVertex& InLHS, const FOutputDataVertex& InRHS);
-			friend bool METASOUNDGRAPHCORE_API operator<(const FOutputDataVertex& InLHS, const FOutputDataVertex& InRHS);
+	template<typename DataType>
+	using TOutputDataVertexModel UE_DEPRECATED(5.1, "TOutputDataVertexModel<DataType> is deprecated. Replace with TOutputDataVertex<DataType>") = TOutputDataVertex<DataType>;
 
-		private:
 
-			TUniquePtr<FOutputDataVertexModel> VertexModel;
-	};
-
-	/** FEnvironmentVertexModel
-	 *
-	 * An FEnvironmentVertexModel implements various vertex behaviors. 
+	/** FEnvironmentVertex
 	 *
-	 * The FEnvironmentVertex supplies copy constructors and assignment operators 
-	 * by utilizing the Clone() operator of the FEnvironmentVertexModel. 
-	 * This allows FEnvironmentVertexModel behavior to be passed to other objects 
-	 * without having to know the concrete implementation of the FEnvironmentVertexModel.
+	 * A vertex for environment variables. 
 	 */
-	struct FEnvironmentVertexModel
+	class FEnvironmentVertex
 	{
-		/** FEnvironmentVertexModel Construtor
+	public:
+		/** FEnvironmentVertex Construtor
 		 *
 		 * @InVertexName - Name of vertex.
 		 * @InDescription - Human readable vertex description.
 		 */
-		FEnvironmentVertexModel(FVertexName InVertexName, const FText& InDescription)
+		FEnvironmentVertex(const FVertexName& InVertexName, const FText& InDescription)
 		:	VertexName(InVertexName)
 		,	Description(InDescription)
 		{
 		}
 
-		virtual ~FEnvironmentVertexModel() = default;
+		virtual ~FEnvironmentVertex() = default;
 
 		/** Name of vertex. */
-		const FVertexName VertexName;
+		FVertexName VertexName;
 
 		/** Description of the vertex. */
-		const FText Description;
-
-		/** Test if a IMetasoundEnvironmentVariable matches the DataType of this vertex model.
-		 *
-		 * @param InVariable - Environment variable object.
-		 *
-		 * @return True if the types are equal, false otherwise.
-		 */
-		virtual bool IsVariableOfSameType(const IMetasoundEnvironmentVariable& InVariable) const = 0;
-
-		virtual FMetasoundEnvironmentVariableTypeId GetVariableTypeId() const = 0;
-
-		/** Create a clone of this FEnvironmentVertexModel */
-		virtual TUniquePtr<FEnvironmentVertexModel> Clone() const = 0;
-
-		friend bool METASOUNDGRAPHCORE_API operator==(const FEnvironmentVertex& InLHS, const FEnvironmentVertex& InRHS);
-		friend bool METASOUNDGRAPHCORE_API operator!=(const FEnvironmentVertex& InLHS, const FEnvironmentVertex& InRHS);
-
-		private:
-
-		/** Compare another FEnvironmentVertexModel is equal to this
-		 *
-		 * @return True if models are equal, false otherwise.
-		 */
-		virtual bool IsEqual(const FEnvironmentVertexModel& InOther) const = 0;
-	};
-
-	/** TEnvironmentVertexModel creates a simple, unchanging, environment vertex. */
-	template<typename VarType>
-	struct TEnvironmentVertexModel : FEnvironmentVertexModel
-	{
-		/** Inherited constructor. */
-		using FEnvironmentVertexModel::FEnvironmentVertexModel;
+		FText Description;
 
-		/** Test if a IMetasoundEnvironmentVariable matches the DataType of this vertex model.
-		 *
-		 * @param InVariable - Environment variable object.
-		 *
-		 * @return True if the types are equal, false otherwise.
-		 */
-		bool IsVariableOfSameType(const IMetasoundEnvironmentVariable& InVariable) const override
-		{
-			return IsEnvironmentVariableOfType<VarType>(InVariable);
-		}
-		
-		/** Return the type id of the vertex. */
-		FMetasoundEnvironmentVariableTypeId GetVariableTypeId() const override
+		/** Name of vertex. */
+		UE_DEPRECATED(5.1, "GetVertexName() is deprecated. Use VertexName")
+		const FVertexName& GetVertexName() const
 		{
-			return GetMetasoundEnvironmentVariableTypeId<VarType>();
+			return VertexName;
 		}
 
-		/** Create a clone of this TEnvironmentVertexModel*/
-		virtual TUniquePtr<FEnvironmentVertexModel> Clone() const override
+		/** Description of the vertex. */
+		UE_DEPRECATED(5.1, "GetDescription() is deprecated. Use Description")
+		const FText& GetDescription() const
 		{
-			return MakeUnique<TEnvironmentVertexModel<VarType>>(this->VertexName, this->Description);
+			return Description;
 		}
 
 		friend bool METASOUNDGRAPHCORE_API operator==(const FEnvironmentVertex& InLHS, const FEnvironmentVertex& InRHS);
 		friend bool METASOUNDGRAPHCORE_API operator!=(const FEnvironmentVertex& InLHS, const FEnvironmentVertex& InRHS);
-	private:
-
-		bool IsEqual(const FEnvironmentVertexModel& InOther) const override
-		{
-			return (GetVariableTypeId() == InOther.GetVariableTypeId()) && (VertexName == InOther.VertexName);
-		}
-	};
-
-	/** FEnvironmentVertex represents an input vertex to an interface. It uses
-	 * a FEnvironmentVertexModel to determine its behavior.
-	 */
-	class METASOUNDGRAPHCORE_API FEnvironmentVertex
-	{
-		using FEmptyVertexModel = TEnvironmentVertexModel<void>;
-
-		public:
-
-			/** Construct an FEnvironmentVertex with a given vertex model.
-			 *
-			 * @param InVertexModel - A model subclass of FEnvironmentVertexModel.
-			 */
-			template<typename VertexModelType>
-			FEnvironmentVertex(const VertexModelType& InVertexModel)
-			:	VertexModel(InVertexModel.Clone())
-			{
-				static_assert(TIsDerivedFrom<VertexModelType, FEnvironmentVertexModel>::Value, "Vertex implementation must be a subclass of FEnvironmentVertexModel");
-
-				if (!VertexModel.IsValid())
-				{
-					VertexModel = MakeUnique<FEmptyVertexModel>(TEXT(""), FText::GetEmpty());
-				}
-			}
-
-			/** Construct an empty FInputDataVertex. */
-			FEnvironmentVertex();
-
-			/** Copy constructor */
-			FEnvironmentVertex(const FEnvironmentVertex& InOther);
-
-			/** Assignment operator. */
-			FEnvironmentVertex& operator=(const FEnvironmentVertex& InOther);
-
-			/** Name of vertex. */
-			const FVertexName& GetVertexName() const;
-
-			/** Description of the vertex. */
-			const FText& GetDescription() const;
-
-			/** Test if a IMetasoundEnvironmentVariable matches the DataType of this vertex.
-			 *
-			 * @param InVariable - Environment variable object.
-			 *
-			 * @return True if the types are equal, false otherwise.
-			 */
-			bool IsVariableOfSameType(const IMetasoundEnvironmentVariable& InVariable) const;
-
-			friend bool METASOUNDGRAPHCORE_API operator==(const FEnvironmentVertex& InLHS, const FEnvironmentVertex& InRHS);
-			friend bool METASOUNDGRAPHCORE_API operator!=(const FEnvironmentVertex& InLHS, const FEnvironmentVertex& InRHS);
-			friend bool METASOUNDGRAPHCORE_API operator<(const FEnvironmentVertex& InLHS, const FEnvironmentVertex& InRHS);
-
-		private:
-
-			TUniquePtr<FEnvironmentVertexModel> VertexModel;
+		friend bool METASOUNDGRAPHCORE_API operator<(const FEnvironmentVertex& InLHS, const FEnvironmentVertex& InRHS);
 	};
 
 	/** TVertexInterfaceGroups encapsulates multiple related data vertices. It 
@@ -656,11 +317,42 @@ namespace Metasound
 
 			static_assert(bIsSupportedVertexType, "VertexType must be derived from FInputDataVertex, FOutputDataVertex, or FEnvironmentVertex");
 
-			using FContainerType = TSortedMap<FVertexName, VertexType, FDefaultAllocator, FNameFastLess>;
-			using FOrderContainerType = TArray<FVertexName>;
+			using FContainerType = TArray<VertexType>;
+
+			struct FEqualVertexNamePredicate
+			{
+				const FVertexName& NameRef;
+
+				FEqualVertexNamePredicate(const FVertexName& InName)
+				: NameRef(InName)
+				{
+				}
+
+				FEqualVertexNamePredicate(const VertexType& InVertex)
+				: NameRef(InVertex.VertexName)
+				{
+				}
+
+				bool operator()(const VertexType& InOther)
+				{
+					return InOther.VertexName == NameRef;
+				}
+			};
+
+			void AddOrUpdateVertex(VertexType&& InVertex)
+			{
+				if (VertexType* Vertex = Find(InVertex.VertexName))
+				{
+					*Vertex = MoveTemp(InVertex);
+				}
+				else
+				{
+					Vertices.Add(MoveTemp(InVertex));
+				}
+			}
 
 			// Required for end of recursion.
-			static void CopyInputs(FContainerType& InStorage, FOrderContainerType& InOrder)
+			void CopyInputs()
 			{
 			}
 
@@ -668,19 +360,17 @@ namespace Metasound
 			// each object in the template parameter pack into the groups internal
 			// storage.
 			//
-			// Assume that InputType is either a sublcass of FDataVertexModel
+			// Assume that InputType is either a sublcass of FDataVertex
 			// and that VertexType is either FInputDataVertex or FOutputDataVertex
-			template<typename VertexModelType, typename... RemainingVertexModelTypes>
-			static void CopyInputs(FContainerType& InStorage, FOrderContainerType& InKeyOrder, const VertexModelType& InInput, const RemainingVertexModelTypes&... InRemainingInputs)
+			template<typename CurrentVertexType, typename... RemainingVertexTypes>
+			void CopyInputs(CurrentVertexType&& InInput, RemainingVertexTypes&&... InRemainingInputs)
 			{
 				// Create vertex out of vertex model
-				VertexType Vertex(InInput);
-
-				InStorage.Add(Vertex.GetVertexName(), Vertex);
-				InKeyOrder.Add(Vertex.GetVertexName());
-				CopyInputs(InStorage, InKeyOrder, InRemainingInputs...);
+				AddOrUpdateVertex(MoveTemp(InInput));
+				CopyInputs(InRemainingInputs...);
 			}
 
+
 		public:
 
 			using RangedForConstIteratorType = typename FContainerType::RangedForConstIteratorType;
@@ -688,57 +378,65 @@ namespace Metasound
 			/** TVertexInterfaceGroup constructor with variadic list of vertex
 			 * models.
 			 */
-			template<typename... VertexModelTypes>
-			TVertexInterfaceGroup(VertexModelTypes&&... InVertexModels)
+			template<typename... VertexTypes>
+			TVertexInterfaceGroup(VertexTypes&&... InVertexs)
 			{
-				CopyInputs(Vertices, OrderedKeys, Forward<VertexModelTypes>(InVertexModels)...);
+				CopyInputs(Forward<VertexTypes>(InVertexs)...);
 			}
 
-			/** Add a vertex model to the group. */
-			template<typename VertexModelType, typename = typename TEnableIf<TIsDerivedFrom<VertexModelType, FDataVertexModel>::Value>::Type>
-			void Add(const VertexModelType& InVertexModel)
+			/** Add a vertex to the group. */
+			void Add(const VertexType& InVertex)
 			{
-				VertexType Vertex(InVertexModel);
-				Add(Vertex);
+				AddOrUpdateVertex(VertexType(InVertex));
 			}
 
-			/** Add a vertex to the group. */
-			void Add(const VertexType& InVertex)
+			void Add(VertexType&& InVertex)
 			{
-				Vertices.Add(InVertex.GetVertexName(), InVertex);
-				OrderedKeys.Add(InVertex.GetVertexName());
+				AddOrUpdateVertex(MoveTemp(InVertex));
 			}
 
 			/** Remove a vertex by key. */
 			bool Remove(const FVertexName& InKey)
 			{
-				int32 NumRemoved = Vertices.Remove(InKey);
-				OrderedKeys.Remove(InKey);
+				int32 NumRemoved = Vertices.RemoveAll(FEqualVertexNamePredicate(InKey));
 				return (NumRemoved > 0);
 			}
 
 			/** Returns true if the group contains a vertex with a matching key. */
 			bool Contains(const FVertexName& InKey) const
 			{
-				return Vertices.Contains(InKey);
+				return Vertices.ContainsByPredicate(FEqualVertexNamePredicate(InKey));
 			}
 
+			/** Find a vertex with a given VertexName */
+			VertexType* Find(const FVertexName& InKey)
+			{
+				return Vertices.FindByPredicate(FEqualVertexNamePredicate(InKey));
+			}
+
+			/** Find a vertex with a given VertexName */
 			const VertexType* Find(const FVertexName& InKey) const
 			{
-				return Vertices.Find(InKey);
+				return Vertices.FindByPredicate(FEqualVertexNamePredicate(InKey));
 			}
 
+			/** Return the sort order index of a vertex with the given name.
+			 *
+			 * @param InKey - FVertexName of vertex of interest.
+			 *
+			 * @return The index of the vertex. INDEX_NONE if the vertex does not exist. 
+			 */
 			int32 GetSortOrderIndex(const FVertexName& InKey) const
 			{
-				int32 OutIndex = INDEX_NONE;
-				OrderedKeys.Find(InKey, OutIndex);
-				return OutIndex;
+				return Vertices.IndexOfByPredicate(FEqualVertexNamePredicate(InKey));
 			}
 
 			/** Return the vertex for a given vertex key. */
 			const VertexType& operator[](const FVertexName& InName) const
 			{
-				return Vertices[InName];
+				const VertexType* Vertex = Find(InName);
+				checkf(nullptr != Vertex, TEXT("Vertex with name '%s' does not exist"), *InName.ToString());
+				return *Vertex;
 			}
 
 			/** Iterator for ranged for loops. */
@@ -774,7 +472,6 @@ namespace Metasound
 		private:
 
 			FContainerType Vertices;
-			FOrderContainerType OrderedKeys;
 	};
 
 	
diff --git a/Engine/Plugins/Runtime/Metasound/Source/MetasoundStandardNodes/Private/MetasoundADEnvelopeNode.cpp b/Engine/Plugins/Runtime/Metasound/Source/MetasoundStandardNodes/Private/MetasoundADEnvelopeNode.cpp
--- a/Engine/Plugins/Runtime/Metasound/Source/MetasoundStandardNodes/Private/MetasoundADEnvelopeNode.cpp
+++ b/Engine/Plugins/Runtime/Metasound/Source/MetasoundStandardNodes/Private/MetasoundADEnvelopeNode.cpp
@@ -206,17 +206,17 @@ namespace Metasound
 
 			static const FVertexInterface DefaultInterface(
 				FInputVertexInterface(
-					TInputDataVertexModel<FTrigger>(METASOUND_GET_PARAM_NAME_AND_METADATA(InputTrigger)),
-					TInputDataVertexModel<FTime>(METASOUND_GET_PARAM_NAME_AND_METADATA(InputAttackTime), 0.01f),
-					TInputDataVertexModel<FTime>(METASOUND_GET_PARAM_NAME_AND_METADATA(InputDecayTime), 1.0f),
-					TInputDataVertexModel<float>(METASOUND_GET_PARAM_NAME_AND_METADATA(InputAttackCurve), 1.0f),
-					TInputDataVertexModel<float>(METASOUND_GET_PARAM_NAME_AND_METADATA(InputDecayCurve), 1.0f),
-					TInputDataVertexModel<bool>(METASOUND_GET_PARAM_NAME_AND_METADATA(InputLooping), false)
+					TInputDataVertex<FTrigger>(METASOUND_GET_PARAM_NAME_AND_METADATA(InputTrigger)),
+					TInputDataVertex<FTime>(METASOUND_GET_PARAM_NAME_AND_METADATA(InputAttackTime), 0.01f),
+					TInputDataVertex<FTime>(METASOUND_GET_PARAM_NAME_AND_METADATA(InputDecayTime), 1.0f),
+					TInputDataVertex<float>(METASOUND_GET_PARAM_NAME_AND_METADATA(InputAttackCurve), 1.0f),
+					TInputDataVertex<float>(METASOUND_GET_PARAM_NAME_AND_METADATA(InputDecayCurve), 1.0f),
+					TInputDataVertex<bool>(METASOUND_GET_PARAM_NAME_AND_METADATA(InputLooping), false)
 				),
 				FOutputVertexInterface(
-					TOutputDataVertexModel<FTrigger>(METASOUND_GET_PARAM_NAME_AND_METADATA(OutputOnTrigger)),
-					TOutputDataVertexModel<FTrigger>(METASOUND_GET_PARAM_NAME_AND_METADATA(OutputOnDone)),
-					TOutputDataVertexModel<ValueType>(METASOUND_GET_PARAM_NAME_AND_METADATA(OutputEnvelopeValue))
+					TOutputDataVertex<FTrigger>(METASOUND_GET_PARAM_NAME_AND_METADATA(OutputOnTrigger)),
+					TOutputDataVertex<FTrigger>(METASOUND_GET_PARAM_NAME_AND_METADATA(OutputOnDone)),
+					TOutputDataVertex<ValueType>(METASOUND_GET_PARAM_NAME_AND_METADATA(OutputEnvelopeValue))
 				)
 			);
 
diff --git a/Engine/Plugins/Runtime/Metasound/Source/MetasoundStandardNodes/Private/MetasoundADSREnvelopeNode.cpp b/Engine/Plugins/Runtime/Metasound/Source/MetasoundStandardNodes/Private/MetasoundADSREnvelopeNode.cpp
--- a/Engine/Plugins/Runtime/Metasound/Source/MetasoundStandardNodes/Private/MetasoundADSREnvelopeNode.cpp
+++ b/Engine/Plugins/Runtime/Metasound/Source/MetasoundStandardNodes/Private/MetasoundADSREnvelopeNode.cpp
@@ -330,23 +330,23 @@ namespace Metasound
 
 			static const FVertexInterface DefaultInterface(
 				FInputVertexInterface(
-					TInputDataVertexModel<FTrigger>(METASOUND_GET_PARAM_NAME_AND_METADATA(InputAttackTrigger)),
-					TInputDataVertexModel<FTrigger>(METASOUND_GET_PARAM_NAME_AND_METADATA(InputReleaseTrigger)),
-					TInputDataVertexModel<FTime>(METASOUND_GET_PARAM_NAME_AND_METADATA(InputAttackTime), 0.01f),
-					TInputDataVertexModel<FTime>(METASOUND_GET_PARAM_NAME_AND_METADATA(InputDecayTime), 0.2f),
-					TInputDataVertexModel<float>(METASOUND_GET_PARAM_NAME_AND_METADATA(InputSustainLevel), 0.5f),
-					TInputDataVertexModel<FTime>(METASOUND_GET_PARAM_NAME_AND_METADATA(InputReleaseTime), 1.0f),
-					TInputDataVertexModel<float>(METASOUND_GET_PARAM_NAME_AND_METADATA(InputAttackCurve), 1.0f),
-					TInputDataVertexModel<float>(METASOUND_GET_PARAM_NAME_AND_METADATA(InputDecayCurve), 1.0f),
-					TInputDataVertexModel<float>(METASOUND_GET_PARAM_NAME_AND_METADATA(InputReleaseCurve), 1.0f)
+					TInputDataVertex<FTrigger>(METASOUND_GET_PARAM_NAME_AND_METADATA(InputAttackTrigger)),
+					TInputDataVertex<FTrigger>(METASOUND_GET_PARAM_NAME_AND_METADATA(InputReleaseTrigger)),
+					TInputDataVertex<FTime>(METASOUND_GET_PARAM_NAME_AND_METADATA(InputAttackTime), 0.01f),
+					TInputDataVertex<FTime>(METASOUND_GET_PARAM_NAME_AND_METADATA(InputDecayTime), 0.2f),
+					TInputDataVertex<float>(METASOUND_GET_PARAM_NAME_AND_METADATA(InputSustainLevel), 0.5f),
+					TInputDataVertex<FTime>(METASOUND_GET_PARAM_NAME_AND_METADATA(InputReleaseTime), 1.0f),
+					TInputDataVertex<float>(METASOUND_GET_PARAM_NAME_AND_METADATA(InputAttackCurve), 1.0f),
+					TInputDataVertex<float>(METASOUND_GET_PARAM_NAME_AND_METADATA(InputDecayCurve), 1.0f),
+					TInputDataVertex<float>(METASOUND_GET_PARAM_NAME_AND_METADATA(InputReleaseCurve), 1.0f)
 				),
 				FOutputVertexInterface(
-					TOutputDataVertexModel<FTrigger>(METASOUND_GET_PARAM_NAME_AND_METADATA(OutputOnAttackTrigger)),
-					TOutputDataVertexModel<FTrigger>(METASOUND_GET_PARAM_NAME_AND_METADATA(OutputOnDecayTrigger)),
-					TOutputDataVertexModel<FTrigger>(METASOUND_GET_PARAM_NAME_AND_METADATA(OutputOnSustainTrigger)),
-					TOutputDataVertexModel<FTrigger>(METASOUND_GET_PARAM_NAME_AND_METADATA(OutputOnReleaseTrigger)),
-					TOutputDataVertexModel<FTrigger>(METASOUND_GET_PARAM_NAME_AND_METADATA(OutputOnDone)),
-					TOutputDataVertexModel<ValueType>(METASOUND_GET_PARAM_NAME_AND_METADATA(OutputEnvelopeValue))
+					TOutputDataVertex<FTrigger>(METASOUND_GET_PARAM_NAME_AND_METADATA(OutputOnAttackTrigger)),
+					TOutputDataVertex<FTrigger>(METASOUND_GET_PARAM_NAME_AND_METADATA(OutputOnDecayTrigger)),
+					TOutputDataVertex<FTrigger>(METASOUND_GET_PARAM_NAME_AND_METADATA(OutputOnSustainTrigger)),
+					TOutputDataVertex<FTrigger>(METASOUND_GET_PARAM_NAME_AND_METADATA(OutputOnReleaseTrigger)),
+					TOutputDataVertex<FTrigger>(METASOUND_GET_PARAM_NAME_AND_METADATA(OutputOnDone)),
+					TOutputDataVertex<ValueType>(METASOUND_GET_PARAM_NAME_AND_METADATA(OutputEnvelopeValue))
 				)
 			);
 
diff --git a/Engine/Plugins/Runtime/Metasound/Source/MetasoundStandardNodes/Private/MetasoundAdditiveSynthNode.cpp b/Engine/Plugins/Runtime/Metasound/Source/MetasoundStandardNodes/Private/MetasoundAdditiveSynthNode.cpp
--- a/Engine/Plugins/Runtime/Metasound/Source/MetasoundStandardNodes/Private/MetasoundAdditiveSynthNode.cpp
+++ b/Engine/Plugins/Runtime/Metasound/Source/MetasoundStandardNodes/Private/MetasoundAdditiveSynthNode.cpp
@@ -64,15 +64,15 @@ namespace Metasound
 
 			static const FVertexInterface Interface(
 				FInputVertexInterface(
-					TInputDataVertexModel<float>(METASOUND_GET_PARAM_NAME_AND_METADATA(BaseFrequency), 440.0f),
-					TInputDataVertexModel<FInputFloatArrayType>(METASOUND_GET_PARAM_NAME_AND_METADATA(HarmonicMultipliers)),
-					TInputDataVertexModel<FInputFloatArrayType>(METASOUND_GET_PARAM_NAME_AND_METADATA(Amplitudes)),
-					TInputDataVertexModel<FInputFloatArrayType>(METASOUND_GET_PARAM_NAME_AND_METADATA(Phases)),
-					TInputDataVertexModel<FInputFloatArrayType>(METASOUND_GET_PARAM_NAME_AND_METADATA(PanAmounts))
+					TInputDataVertex<float>(METASOUND_GET_PARAM_NAME_AND_METADATA(BaseFrequency), 440.0f),
+					TInputDataVertex<FInputFloatArrayType>(METASOUND_GET_PARAM_NAME_AND_METADATA(HarmonicMultipliers)),
+					TInputDataVertex<FInputFloatArrayType>(METASOUND_GET_PARAM_NAME_AND_METADATA(Amplitudes)),
+					TInputDataVertex<FInputFloatArrayType>(METASOUND_GET_PARAM_NAME_AND_METADATA(Phases)),
+					TInputDataVertex<FInputFloatArrayType>(METASOUND_GET_PARAM_NAME_AND_METADATA(PanAmounts))
 				),
 				FOutputVertexInterface(
-					TOutputDataVertexModel<FAudioBuffer>(METASOUND_GET_PARAM_NAME_AND_METADATA(LeftAudioOut)),
-					TOutputDataVertexModel<FAudioBuffer>(METASOUND_GET_PARAM_NAME_AND_METADATA(RightAudioOut))
+					TOutputDataVertex<FAudioBuffer>(METASOUND_GET_PARAM_NAME_AND_METADATA(LeftAudioOut)),
+					TOutputDataVertex<FAudioBuffer>(METASOUND_GET_PARAM_NAME_AND_METADATA(RightAudioOut))
 				)
 			);
 
diff --git a/Engine/Plugins/Runtime/Metasound/Source/MetasoundStandardNodes/Private/MetasoundAudioFormats.cpp b/Engine/Plugins/Runtime/Metasound/Source/MetasoundStandardNodes/Private/MetasoundAudioFormats.cpp
--- a/Engine/Plugins/Runtime/Metasound/Source/MetasoundStandardNodes/Private/MetasoundAudioFormats.cpp
+++ b/Engine/Plugins/Runtime/Metasound/Source/MetasoundStandardNodes/Private/MetasoundAudioFormats.cpp
@@ -176,11 +176,11 @@ namespace Metasound
 
 			return FVertexInterface(
 				FInputVertexInterface(
-					TInputDataVertexModel<FAudioBuffer>(METASOUND_GET_PARAM_NAME_AND_METADATA(LeftChannelVertex)),
-					TInputDataVertexModel<FAudioBuffer>(METASOUND_GET_PARAM_NAME_AND_METADATA(RightChannelVertex))
+					TInputDataVertex<FAudioBuffer>(METASOUND_GET_PARAM_NAME_AND_METADATA(LeftChannelVertex)),
+					TInputDataVertex<FAudioBuffer>(METASOUND_GET_PARAM_NAME_AND_METADATA(RightChannelVertex))
 				),
 				FOutputVertexInterface(
-					TOutputDataVertexModel<FStereoAudioFormat>(InVertexName, METASOUND_LOCTEXT("Metasound_StereoOutputVertexDescription", "Stereo Output."))
+					TOutputDataVertex<FStereoAudioFormat>(InVertexName, FDataVertexMetadata{METASOUND_LOCTEXT("Metasound_StereoOutputVertexDescription", "Stereo Output.")})
 				)
 			);
 		}
@@ -318,11 +318,11 @@ namespace Metasound
 			
 			return FVertexInterface(
 				FInputVertexInterface(
-					TInputDataVertexModel<FStereoAudioFormat>(InVertexName, METASOUND_LOCTEXT("Metasound_StereoInputVertexDescription", "Stereo Input."))
+					TInputDataVertex<FStereoAudioFormat>(InVertexName, FDataVertexMetadata{METASOUND_LOCTEXT("Metasound_StereoInputVertexDescription", "Stereo Input.")})
 				),
 				FOutputVertexInterface(
-					TOutputDataVertexModel<FAudioBuffer>(METASOUND_GET_PARAM_NAME_AND_METADATA(LeftChannelVertex)),
-					TOutputDataVertexModel<FAudioBuffer>(METASOUND_GET_PARAM_NAME_AND_METADATA(RightChannelVertex))
+					TOutputDataVertex<FAudioBuffer>(METASOUND_GET_PARAM_NAME_AND_METADATA(LeftChannelVertex)),
+					TOutputDataVertex<FAudioBuffer>(METASOUND_GET_PARAM_NAME_AND_METADATA(RightChannelVertex))
 				)
 			);
 		}
@@ -453,10 +453,10 @@ namespace Metasound
 		{
 			return FVertexInterface(
 				FInputVertexInterface(
-					TInputDataVertexModel<FAudioBuffer>(InVertexName, METASOUND_LOCTEXT("Metasound_CenterMonoOutputVertexDescription", "Center channel audio output."))
+					TInputDataVertex<FAudioBuffer>(InVertexName, FDataVertexMetadata{METASOUND_LOCTEXT("Metasound_CenterMonoOutputVertexDescription", "Center channel audio output.")})
 				),
 				FOutputVertexInterface(
-					TOutputDataVertexModel<FMonoAudioFormat>(InVertexName, METASOUND_LOCTEXT("Metasound_MonoOutputVertexDescription", "Mono Output."))
+					TOutputDataVertex<FMonoAudioFormat>(InVertexName, FDataVertexMetadata{METASOUND_LOCTEXT("Metasound_MonoOutputVertexDescription", "Mono Output.")})
 				)
 			);
 		}
@@ -586,10 +586,10 @@ namespace Metasound
 		{
 			return FVertexInterface(
 				FInputVertexInterface(
-					TInputDataVertexModel<FMonoAudioFormat>(InVertexName, METASOUND_LOCTEXT("Metasound_MonoInputVertexDescription", "Mono Input."))
+					TInputDataVertex<FMonoAudioFormat>(InVertexName, FDataVertexMetadata{METASOUND_LOCTEXT("Metasound_MonoInputVertexDescription", "Mono Input.")})
 				),
 				FOutputVertexInterface(
-					TOutputDataVertexModel<FAudioBuffer>(InVertexName, METASOUND_LOCTEXT("Metasound_CenterMonoInputVertexDescription", "Center channel audio output."))
+					TOutputDataVertex<FAudioBuffer>(InVertexName, FDataVertexMetadata{METASOUND_LOCTEXT("Metasound_CenterMonoInputVertexDescription", "Center channel audio output.")})
 				)
 			);
 		}
diff --git a/Engine/Plugins/Runtime/Metasound/Source/MetasoundStandardNodes/Private/MetasoundBPMToSecondsNode.cpp b/Engine/Plugins/Runtime/Metasound/Source/MetasoundStandardNodes/Private/MetasoundBPMToSecondsNode.cpp
--- a/Engine/Plugins/Runtime/Metasound/Source/MetasoundStandardNodes/Private/MetasoundBPMToSecondsNode.cpp
+++ b/Engine/Plugins/Runtime/Metasound/Source/MetasoundStandardNodes/Private/MetasoundBPMToSecondsNode.cpp
@@ -116,12 +116,12 @@ namespace Metasound
 
 		static const FVertexInterface Interface(
 			FInputVertexInterface(
-				TInputDataVertexModel<float>(METASOUND_GET_PARAM_NAME_AND_METADATA(InputBPM), 90.0f),
-				TInputDataVertexModel<float>(METASOUND_GET_PARAM_NAME_AND_METADATA(InputBeatMultiplier), 1.0f),
-				TInputDataVertexModel<float>(METASOUND_GET_PARAM_NAME_AND_METADATA(InputDivOfWholeNote), 4.0f)
+				TInputDataVertex<float>(METASOUND_GET_PARAM_NAME_AND_METADATA(InputBPM), 90.0f),
+				TInputDataVertex<float>(METASOUND_GET_PARAM_NAME_AND_METADATA(InputBeatMultiplier), 1.0f),
+				TInputDataVertex<float>(METASOUND_GET_PARAM_NAME_AND_METADATA(InputDivOfWholeNote), 4.0f)
 			),
 			FOutputVertexInterface(
-				TOutputDataVertexModel<FTime>(METASOUND_GET_PARAM_NAME_AND_METADATA(OutputTimeSeconds))
+				TOutputDataVertex<FTime>(METASOUND_GET_PARAM_NAME_AND_METADATA(OutputTimeSeconds))
 			)
 		);
 
diff --git a/Engine/Plugins/Runtime/Metasound/Source/MetasoundStandardNodes/Private/MetasoundBandsplitterNode.cpp b/Engine/Plugins/Runtime/Metasound/Source/MetasoundStandardNodes/Private/MetasoundBandsplitterNode.cpp
--- a/Engine/Plugins/Runtime/Metasound/Source/MetasoundStandardNodes/Private/MetasoundBandsplitterNode.cpp
+++ b/Engine/Plugins/Runtime/Metasound/Source/MetasoundStandardNodes/Private/MetasoundBandsplitterNode.cpp
@@ -179,11 +179,11 @@ namespace Metasound
 						GetAudioInputDisplayName(ChannelIndex)
 					};
 
-					InputInterface.Add(TInputDataVertexModel<FAudioBuffer>(GetAudioInputName(ChannelIndex), AudioInputMetadata));
+					InputInterface.Add(TInputDataVertex<FAudioBuffer>(GetAudioInputName(ChannelIndex), AudioInputMetadata));
 				}
 
-				InputInterface.Add(TInputDataVertexModel<FEnumBandSplitterFilterOrder>(METASOUND_GET_PARAM_NAME_AND_METADATA(InputFilterOrder)));
-				InputInterface.Add(TInputDataVertexModel<bool>(METASOUND_GET_PARAM_NAME_AND_METADATA(InputPhaseCompensate), true));
+				InputInterface.Add(TInputDataVertex<FEnumBandSplitterFilterOrder>(METASOUND_GET_PARAM_NAME_AND_METADATA(InputFilterOrder)));
+				InputInterface.Add(TInputDataVertex<bool>(METASOUND_GET_PARAM_NAME_AND_METADATA(InputPhaseCompensate), true));
 
 				// Crossover frequencies
 				for (uint32 InputIndex = 0; InputIndex < NumBands - 1; InputIndex++)
@@ -194,7 +194,7 @@ namespace Metasound
 						GetCrossoverInputDisplayName(InputIndex)
 					};
 
-					InputInterface.Add(TInputDataVertexModel<float>(GetCrossoverInputName(InputIndex), CrossoverInputMetadata, (1 + InputIndex) * 500.0f));
+					InputInterface.Add(TInputDataVertex<float>(GetCrossoverInputName(InputIndex), CrossoverInputMetadata, (1 + InputIndex) * 500.0f));
 				}
 
 				// outputs
@@ -209,7 +209,7 @@ namespace Metasound
 							GetAudioOutputDisplayName(OutputIndex, ChannelIndex)
 						};
 
-						OutputInterface.Add(TOutputDataVertexModel<FAudioBuffer>(GetAudioOutputName(OutputIndex, ChannelIndex), AudioOutputMetadata));
+						OutputInterface.Add(TOutputDataVertex<FAudioBuffer>(GetAudioOutputName(OutputIndex, ChannelIndex), AudioOutputMetadata));
 					}
 				}
 
diff --git a/Engine/Plugins/Runtime/Metasound/Source/MetasoundStandardNodes/Private/MetasoundBasicFilters.cpp b/Engine/Plugins/Runtime/Metasound/Source/MetasoundStandardNodes/Private/MetasoundBasicFilters.cpp
--- a/Engine/Plugins/Runtime/Metasound/Source/MetasoundStandardNodes/Private/MetasoundBasicFilters.cpp
+++ b/Engine/Plugins/Runtime/Metasound/Source/MetasoundStandardNodes/Private/MetasoundBasicFilters.cpp
@@ -182,12 +182,12 @@ namespace Metasound
 	{
 		static const FVertexInterface Interface(
 			FInputVertexInterface(
-				TInputDataVertexModel<FAudioBuffer>(METASOUND_GET_PARAM_NAME_AND_METADATA(ParamAudioInput)),
-				TInputDataVertexModel<float>(METASOUND_GET_PARAM_NAME_AND_METADATA(ParamCutoffFrequency), 20000.f),
-				TInputDataVertexModel<float>(METASOUND_GET_PARAM_NAME_AND_METADATA(ParamResonance), 1.f)
+				TInputDataVertex<FAudioBuffer>(METASOUND_GET_PARAM_NAME_AND_METADATA(ParamAudioInput)),
+				TInputDataVertex<float>(METASOUND_GET_PARAM_NAME_AND_METADATA(ParamCutoffFrequency), 20000.f),
+				TInputDataVertex<float>(METASOUND_GET_PARAM_NAME_AND_METADATA(ParamResonance), 1.f)
 			),
 			FOutputVertexInterface(
-				TOutputDataVertexModel<FAudioBuffer>(METASOUND_GET_PARAM_NAME_AND_METADATA(ParamAudioOutput))
+				TOutputDataVertex<FAudioBuffer>(METASOUND_GET_PARAM_NAME_AND_METADATA(ParamAudioOutput))
 			)
 		);
 
@@ -364,16 +364,16 @@ namespace Metasound
 	{
 		static const FVertexInterface Interface(
 			FInputVertexInterface(
-				TInputDataVertexModel<FAudioBuffer>(METASOUND_GET_PARAM_NAME_AND_METADATA(ParamAudioInput)),
-				TInputDataVertexModel<float>(METASOUND_GET_PARAM_NAME_AND_METADATA(ParamCutoffFrequency), 20000.f),
-				TInputDataVertexModel<float>(METASOUND_GET_PARAM_NAME_AND_METADATA(ParamResonance), 0.f),
-				TInputDataVertexModel<float>(METASOUND_GET_PARAM_NAME_AND_METADATA(ParamBandStopControl), 0.f)
+				TInputDataVertex<FAudioBuffer>(METASOUND_GET_PARAM_NAME_AND_METADATA(ParamAudioInput)),
+				TInputDataVertex<float>(METASOUND_GET_PARAM_NAME_AND_METADATA(ParamCutoffFrequency), 20000.f),
+				TInputDataVertex<float>(METASOUND_GET_PARAM_NAME_AND_METADATA(ParamResonance), 0.f),
+				TInputDataVertex<float>(METASOUND_GET_PARAM_NAME_AND_METADATA(ParamBandStopControl), 0.f)
 			),
 			FOutputVertexInterface(
-				TOutputDataVertexModel<FAudioBuffer>(METASOUND_GET_PARAM_NAME_AND_METADATA(ParamLowPassOutput)),
-				TOutputDataVertexModel<FAudioBuffer>(METASOUND_GET_PARAM_NAME_AND_METADATA(ParamHighPassOutput)),
-				TOutputDataVertexModel<FAudioBuffer>(METASOUND_GET_PARAM_NAME_AND_METADATA(ParamBandPassOutput)),
-				TOutputDataVertexModel<FAudioBuffer>(METASOUND_GET_PARAM_NAME_AND_METADATA(ParamBandStopOutput))
+				TOutputDataVertex<FAudioBuffer>(METASOUND_GET_PARAM_NAME_AND_METADATA(ParamLowPassOutput)),
+				TOutputDataVertex<FAudioBuffer>(METASOUND_GET_PARAM_NAME_AND_METADATA(ParamHighPassOutput)),
+				TOutputDataVertex<FAudioBuffer>(METASOUND_GET_PARAM_NAME_AND_METADATA(ParamBandPassOutput)),
+				TOutputDataVertex<FAudioBuffer>(METASOUND_GET_PARAM_NAME_AND_METADATA(ParamBandStopOutput))
 			)
 		);
 
@@ -536,11 +536,11 @@ namespace Metasound
 	{
 		static const FVertexInterface Interface(
 			FInputVertexInterface(
-				TInputDataVertexModel<FAudioBuffer>(METASOUND_GET_PARAM_NAME_AND_METADATA(ParamAudioInput)),
-				TInputDataVertexModel<float>(METASOUND_GET_PARAM_NAME_AND_METADATA(ParamCutoffFrequency), 20000.f)
+				TInputDataVertex<FAudioBuffer>(METASOUND_GET_PARAM_NAME_AND_METADATA(ParamAudioInput)),
+				TInputDataVertex<float>(METASOUND_GET_PARAM_NAME_AND_METADATA(ParamCutoffFrequency), 20000.f)
 			),
 			FOutputVertexInterface(
-				TOutputDataVertexModel<FAudioBuffer>(METASOUND_GET_PARAM_NAME_AND_METADATA(ParamAudioOutput))
+				TOutputDataVertex<FAudioBuffer>(METASOUND_GET_PARAM_NAME_AND_METADATA(ParamAudioOutput))
 			)
 		);
 
@@ -673,11 +673,11 @@ namespace Metasound
 	{
 		static const FVertexInterface Interface(
 			FInputVertexInterface(
-				TInputDataVertexModel<FAudioBuffer>(METASOUND_GET_PARAM_NAME_AND_METADATA(ParamAudioInput)),
-				TInputDataVertexModel<float>(METASOUND_GET_PARAM_NAME_AND_METADATA(ParamCutoffFrequency), 10.f)
+				TInputDataVertex<FAudioBuffer>(METASOUND_GET_PARAM_NAME_AND_METADATA(ParamAudioInput)),
+				TInputDataVertex<float>(METASOUND_GET_PARAM_NAME_AND_METADATA(ParamCutoffFrequency), 10.f)
 			),
 			FOutputVertexInterface(
-				TOutputDataVertexModel<FAudioBuffer>(METASOUND_GET_PARAM_NAME_AND_METADATA(ParamAudioOutput))
+				TOutputDataVertex<FAudioBuffer>(METASOUND_GET_PARAM_NAME_AND_METADATA(ParamAudioOutput))
 			)
 		);
 
@@ -835,14 +835,14 @@ namespace Metasound
 
 		static const FVertexInterface Interface(
 			FInputVertexInterface(
-				TInputDataVertexModel<FAudioBuffer>(METASOUND_GET_PARAM_NAME_AND_METADATA(ParamAudioInput)),
-				TInputDataVertexModel<float>(METASOUND_GET_PARAM_NAME_AND_METADATA(ParamCutoffFrequency), 20000.f),
-				TInputDataVertexModel<float>(METASOUND_GET_PARAM_NAME_AND_METADATA(ParamBandwidth), 1.f),
-				TInputDataVertexModel<float>(METASOUND_GET_PARAM_NAME(ParamGainDb), GainPinMetaData, 0.f),
-				TInputDataVertexModel<FEnumEBiquadFilterType>(METASOUND_GET_PARAM_NAME_AND_METADATA(ParamFilterType))
+				TInputDataVertex<FAudioBuffer>(METASOUND_GET_PARAM_NAME_AND_METADATA(ParamAudioInput)),
+				TInputDataVertex<float>(METASOUND_GET_PARAM_NAME_AND_METADATA(ParamCutoffFrequency), 20000.f),
+				TInputDataVertex<float>(METASOUND_GET_PARAM_NAME_AND_METADATA(ParamBandwidth), 1.f),
+				TInputDataVertex<float>(METASOUND_GET_PARAM_NAME(ParamGainDb), GainPinMetaData, 0.f),
+				TInputDataVertex<FEnumEBiquadFilterType>(METASOUND_GET_PARAM_NAME_AND_METADATA(ParamFilterType))
 			),
 			FOutputVertexInterface(
-				TOutputDataVertexModel<FAudioBuffer>(METASOUND_GET_PARAM_NAME_AND_METADATA(ParamAudioOutput))
+				TOutputDataVertex<FAudioBuffer>(METASOUND_GET_PARAM_NAME_AND_METADATA(ParamAudioOutput))
 			)
 		);
 
diff --git a/Engine/Plugins/Runtime/Metasound/Source/MetasoundStandardNodes/Private/MetasoundBitcrusherNode.cpp b/Engine/Plugins/Runtime/Metasound/Source/MetasoundStandardNodes/Private/MetasoundBitcrusherNode.cpp
--- a/Engine/Plugins/Runtime/Metasound/Source/MetasoundStandardNodes/Private/MetasoundBitcrusherNode.cpp
+++ b/Engine/Plugins/Runtime/Metasound/Source/MetasoundStandardNodes/Private/MetasoundBitcrusherNode.cpp
@@ -82,12 +82,12 @@ namespace Metasound
 
 			static const FVertexInterface Interface(
 				FInputVertexInterface(
-					TInputDataVertexModel<FAudioBuffer>(METASOUND_GET_PARAM_NAME_AND_METADATA(InputAudio)),
-					TInputDataVertexModel<float>(METASOUND_GET_PARAM_NAME_AND_METADATA(InputSampleRate), 8000.0f),
-					TInputDataVertexModel<float>(METASOUND_GET_PARAM_NAME_AND_METADATA(InputBitDepth), 8.0f)
+					TInputDataVertex<FAudioBuffer>(METASOUND_GET_PARAM_NAME_AND_METADATA(InputAudio)),
+					TInputDataVertex<float>(METASOUND_GET_PARAM_NAME_AND_METADATA(InputSampleRate), 8000.0f),
+					TInputDataVertex<float>(METASOUND_GET_PARAM_NAME_AND_METADATA(InputBitDepth), 8.0f)
 				),
 				FOutputVertexInterface(
-					TOutputDataVertexModel<FAudioBuffer>(METASOUND_GET_PARAM_NAME_AND_METADATA(OutputAudio))
+					TOutputDataVertex<FAudioBuffer>(METASOUND_GET_PARAM_NAME_AND_METADATA(OutputAudio))
 				)
 			);
 
diff --git a/Engine/Plugins/Runtime/Metasound/Source/MetasoundStandardNodes/Private/MetasoundClampNode.cpp b/Engine/Plugins/Runtime/Metasound/Source/MetasoundStandardNodes/Private/MetasoundClampNode.cpp
--- a/Engine/Plugins/Runtime/Metasound/Source/MetasoundStandardNodes/Private/MetasoundClampNode.cpp
+++ b/Engine/Plugins/Runtime/Metasound/Source/MetasoundStandardNodes/Private/MetasoundClampNode.cpp
@@ -157,12 +157,12 @@ namespace Metasound
 
 			static const FVertexInterface DefaultInterface(
 				FInputVertexInterface(
-					TInputDataVertexModel<ValueType>(METASOUND_GET_PARAM_NAME_AND_METADATA(InputValue)),
-					TInputDataVertexModel<ValueType>(METASOUND_GET_PARAM_NAME_AND_METADATA(InputMinValue)),
-					TInputDataVertexModel<ValueType>(METASOUND_GET_PARAM_NAME_AND_METADATA(InputMaxValue))
+					TInputDataVertex<ValueType>(METASOUND_GET_PARAM_NAME_AND_METADATA(InputValue)),
+					TInputDataVertex<ValueType>(METASOUND_GET_PARAM_NAME_AND_METADATA(InputMinValue)),
+					TInputDataVertex<ValueType>(METASOUND_GET_PARAM_NAME_AND_METADATA(InputMaxValue))
 				),
 				FOutputVertexInterface(
-					TOutputDataVertexModel<ValueType>(METASOUND_GET_PARAM_NAME_AND_METADATA(OutputValue))
+					TOutputDataVertex<ValueType>(METASOUND_GET_PARAM_NAME_AND_METADATA(OutputValue))
 				)
 			);
 
diff --git a/Engine/Plugins/Runtime/Metasound/Source/MetasoundStandardNodes/Private/MetasoundCompressorNode.cpp b/Engine/Plugins/Runtime/Metasound/Source/MetasoundStandardNodes/Private/MetasoundCompressorNode.cpp
--- a/Engine/Plugins/Runtime/Metasound/Source/MetasoundStandardNodes/Private/MetasoundCompressorNode.cpp
+++ b/Engine/Plugins/Runtime/Metasound/Source/MetasoundStandardNodes/Private/MetasoundCompressorNode.cpp
@@ -145,22 +145,22 @@ namespace Metasound
 
 			static const FVertexInterface Interface(
 				FInputVertexInterface(
-					TInputDataVertexModel<FAudioBuffer>(METASOUND_GET_PARAM_NAME_AND_METADATA(InputAudio)),
-					TInputDataVertexModel<float>(METASOUND_GET_PARAM_NAME_AND_METADATA(InputRatio), 1.5f),
-					TInputDataVertexModel<float>(METASOUND_GET_PARAM_NAME_AND_METADATA(InputThreshold), -6.0f),
-					TInputDataVertexModel<FTime>(METASOUND_GET_PARAM_NAME_AND_METADATA(InputAttackTime), 0.01f),
-					TInputDataVertexModel<FTime>(METASOUND_GET_PARAM_NAME_AND_METADATA(InputReleaseTime), 0.1f),
-					TInputDataVertexModel<FTime>(METASOUND_GET_PARAM_NAME_AND_METADATA(InputLookaheadTime), 0.01f),
-					TInputDataVertexModel<float>(METASOUND_GET_PARAM_NAME_AND_METADATA(InputKnee), 10.0f),
-					TInputDataVertexModel<FAudioBuffer>(METASOUND_GET_PARAM_NAME_AND_METADATA(InputSidechain)),
-					TInputDataVertexModel<FEnumEnvelopePeakMode>(METASOUND_GET_PARAM_NAME_AND_METADATA(InputEnvelopeMode)),
-					TInputDataVertexModel<bool>(METASOUND_GET_PARAM_NAME_AND_METADATA(InputIsAnalog), true),
-					TInputDataVertexModel<bool>(METASOUND_GET_PARAM_NAME_AND_METADATA(InputIsUpwards), false),
-					TInputDataVertexModel<float>(METASOUND_GET_PARAM_NAME_AND_METADATA(InputWetDryMix), 1.0f)
+					TInputDataVertex<FAudioBuffer>(METASOUND_GET_PARAM_NAME_AND_METADATA(InputAudio)),
+					TInputDataVertex<float>(METASOUND_GET_PARAM_NAME_AND_METADATA(InputRatio), 1.5f),
+					TInputDataVertex<float>(METASOUND_GET_PARAM_NAME_AND_METADATA(InputThreshold), -6.0f),
+					TInputDataVertex<FTime>(METASOUND_GET_PARAM_NAME_AND_METADATA(InputAttackTime), 0.01f),
+					TInputDataVertex<FTime>(METASOUND_GET_PARAM_NAME_AND_METADATA(InputReleaseTime), 0.1f),
+					TInputDataVertex<FTime>(METASOUND_GET_PARAM_NAME_AND_METADATA(InputLookaheadTime), 0.01f),
+					TInputDataVertex<float>(METASOUND_GET_PARAM_NAME_AND_METADATA(InputKnee), 10.0f),
+					TInputDataVertex<FAudioBuffer>(METASOUND_GET_PARAM_NAME_AND_METADATA(InputSidechain)),
+					TInputDataVertex<FEnumEnvelopePeakMode>(METASOUND_GET_PARAM_NAME_AND_METADATA(InputEnvelopeMode)),
+					TInputDataVertex<bool>(METASOUND_GET_PARAM_NAME_AND_METADATA(InputIsAnalog), true),
+					TInputDataVertex<bool>(METASOUND_GET_PARAM_NAME_AND_METADATA(InputIsUpwards), false),
+					TInputDataVertex<float>(METASOUND_GET_PARAM_NAME_AND_METADATA(InputWetDryMix), 1.0f)
 				),
 				FOutputVertexInterface(
-					TOutputDataVertexModel<FAudioBuffer>(METASOUND_GET_PARAM_NAME_AND_METADATA(OutputAudio)),
-					TOutputDataVertexModel<FAudioBuffer>(METASOUND_GET_PARAM_NAME_AND_METADATA(OutputEnvelope))
+					TOutputDataVertex<FAudioBuffer>(METASOUND_GET_PARAM_NAME_AND_METADATA(OutputAudio)),
+					TOutputDataVertex<FAudioBuffer>(METASOUND_GET_PARAM_NAME_AND_METADATA(OutputEnvelope))
 				)
 			);
 
diff --git a/Engine/Plugins/Runtime/Metasound/Source/MetasoundStandardNodes/Private/MetasoundConversionNodes.cpp b/Engine/Plugins/Runtime/Metasound/Source/MetasoundStandardNodes/Private/MetasoundConversionNodes.cpp
--- a/Engine/Plugins/Runtime/Metasound/Source/MetasoundStandardNodes/Private/MetasoundConversionNodes.cpp
+++ b/Engine/Plugins/Runtime/Metasound/Source/MetasoundStandardNodes/Private/MetasoundConversionNodes.cpp
@@ -157,10 +157,10 @@ namespace Metasound
 
 			static const FVertexInterface DefaultInterface(
 				FInputVertexInterface(
-					TInputDataVertexModel<FromType>(METASOUND_GET_PARAM_NAME(InputValue), InputDesc)
+					TInputDataVertex<FromType>(METASOUND_GET_PARAM_NAME(InputValue), FDataVertexMetadata{InputDesc})
 				),
 				FOutputVertexInterface(
-					TOutputDataVertexModel<ToType>(METASOUND_GET_PARAM_NAME(OutputValue), OutputDesc)
+					TOutputDataVertex<ToType>(METASOUND_GET_PARAM_NAME(OutputValue), FDataVertexMetadata{OutputDesc})
 				)
 			);
 
diff --git a/Engine/Plugins/Runtime/Metasound/Source/MetasoundStandardNodes/Private/MetasoundCrossfadeNode.cpp b/Engine/Plugins/Runtime/Metasound/Source/MetasoundStandardNodes/Private/MetasoundCrossfadeNode.cpp
--- a/Engine/Plugins/Runtime/Metasound/Source/MetasoundStandardNodes/Private/MetasoundCrossfadeNode.cpp
+++ b/Engine/Plugins/Runtime/Metasound/Source/MetasoundStandardNodes/Private/MetasoundCrossfadeNode.cpp
@@ -146,7 +146,7 @@ namespace Metasound
 			{
 				FInputVertexInterface InputInterface;
 
-				InputInterface.Add(TInputDataVertexModel<float>(METASOUND_GET_PARAM_NAME_AND_METADATA(InputCrossfadeValue)));
+				InputInterface.Add(TInputDataVertex<float>(METASOUND_GET_PARAM_NAME_AND_METADATA(InputCrossfadeValue)));
 
 				for (uint32 i = 0; i < NumInputs; ++i)
 				{
@@ -156,11 +156,11 @@ namespace Metasound
 						GetInputDisplayName(i)
 					};
 
-					InputInterface.Add(TInputDataVertexModel<ValueType>(GetInputName(i), InputMetadata));
+					InputInterface.Add(TInputDataVertex<ValueType>(GetInputName(i), InputMetadata));
 				}
 
 				FOutputVertexInterface OutputInterface;
-				OutputInterface.Add(TOutputDataVertexModel<ValueType>(METASOUND_GET_PARAM_NAME_AND_METADATA(OutputTrigger)));
+				OutputInterface.Add(TOutputDataVertex<ValueType>(METASOUND_GET_PARAM_NAME_AND_METADATA(OutputTrigger)));
 
 				return FVertexInterface(InputInterface, OutputInterface);
 			};
diff --git a/Engine/Plugins/Runtime/Metasound/Source/MetasoundStandardNodes/Private/MetasoundDebugLogNode.cpp b/Engine/Plugins/Runtime/Metasound/Source/MetasoundStandardNodes/Private/MetasoundDebugLogNode.cpp
--- a/Engine/Plugins/Runtime/Metasound/Source/MetasoundStandardNodes/Private/MetasoundDebugLogNode.cpp
+++ b/Engine/Plugins/Runtime/Metasound/Source/MetasoundStandardNodes/Private/MetasoundDebugLogNode.cpp
@@ -59,9 +59,9 @@ namespace Metasound
 				static const FVertexInterface DefaultInterface(
 					FInputVertexInterface(
 
-						TInputDataVertexModel<FTrigger>(METASOUND_GET_PARAM_NAME_AND_METADATA(InputTrigger)),
-						TInputDataVertexModel<FString>(METASOUND_GET_PARAM_NAME_AND_METADATA(InputLabel)),
-						TInputDataVertexModel<PrintLogType>(METASOUND_GET_PARAM_NAME_AND_METADATA(InputValueToLog))
+						TInputDataVertex<FTrigger>(METASOUND_GET_PARAM_NAME_AND_METADATA(InputTrigger)),
+						TInputDataVertex<FString>(METASOUND_GET_PARAM_NAME_AND_METADATA(InputLabel)),
+						TInputDataVertex<PrintLogType>(METASOUND_GET_PARAM_NAME_AND_METADATA(InputValueToLog))
 					),
 					FOutputVertexInterface(
 					)
@@ -187,4 +187,4 @@ namespace Metasound
 	METASOUND_REGISTER_NODE(FPrintLogNodeString)
 }
 
-#undef LOCTEXT_NAMESPACE
\ No newline at end of file
+#undef LOCTEXT_NAMESPACE
diff --git a/Engine/Plugins/Runtime/Metasound/Source/MetasoundStandardNodes/Private/MetasoundDecibelsToLinearGain.cpp b/Engine/Plugins/Runtime/Metasound/Source/MetasoundStandardNodes/Private/MetasoundDecibelsToLinearGain.cpp
--- a/Engine/Plugins/Runtime/Metasound/Source/MetasoundStandardNodes/Private/MetasoundDecibelsToLinearGain.cpp
+++ b/Engine/Plugins/Runtime/Metasound/Source/MetasoundStandardNodes/Private/MetasoundDecibelsToLinearGain.cpp
@@ -82,10 +82,10 @@ namespace Metasound
 
 		static const FVertexInterface Interface(
 			FInputVertexInterface(
-				TInputDataVertexModel<float>(METASOUND_GET_PARAM_NAME_AND_METADATA(InputDecibelGain), 0.0f)
+				TInputDataVertex<float>(METASOUND_GET_PARAM_NAME_AND_METADATA(InputDecibelGain), 0.0f)
 			),
 			FOutputVertexInterface(
-				TOutputDataVertexModel<float>(METASOUND_GET_PARAM_NAME_AND_METADATA(OutputLinearGain))
+				TOutputDataVertex<float>(METASOUND_GET_PARAM_NAME_AND_METADATA(OutputLinearGain))
 			)
 		);
 
diff --git a/Engine/Plugins/Runtime/Metasound/Source/MetasoundStandardNodes/Private/MetasoundDelayNode.cpp b/Engine/Plugins/Runtime/Metasound/Source/MetasoundStandardNodes/Private/MetasoundDelayNode.cpp
--- a/Engine/Plugins/Runtime/Metasound/Source/MetasoundStandardNodes/Private/MetasoundDelayNode.cpp
+++ b/Engine/Plugins/Runtime/Metasound/Source/MetasoundStandardNodes/Private/MetasoundDelayNode.cpp
@@ -177,14 +177,14 @@ namespace Metasound
 
 		static const FVertexInterface Interface(
 			FInputVertexInterface(
-				TInputDataVertexModel<FAudioBuffer>(METASOUND_GET_PARAM_NAME_AND_METADATA(InParamAudioInput)),
-				TInputDataVertexModel<FTime>(METASOUND_GET_PARAM_NAME_AND_METADATA(InParamDelayTime), 1.0f),
-				TInputDataVertexModel<float>(METASOUND_GET_PARAM_NAME_AND_METADATA(InParamDryLevel), 0.0f),
-				TInputDataVertexModel<float>(METASOUND_GET_PARAM_NAME_AND_METADATA(InParamWetLevel), 1.0f),
-				TInputDataVertexModel<float>(METASOUND_GET_PARAM_NAME_AND_METADATA(InParamFeedbackAmount), 0.0f)
+				TInputDataVertex<FAudioBuffer>(METASOUND_GET_PARAM_NAME_AND_METADATA(InParamAudioInput)),
+				TInputDataVertex<FTime>(METASOUND_GET_PARAM_NAME_AND_METADATA(InParamDelayTime), 1.0f),
+				TInputDataVertex<float>(METASOUND_GET_PARAM_NAME_AND_METADATA(InParamDryLevel), 0.0f),
+				TInputDataVertex<float>(METASOUND_GET_PARAM_NAME_AND_METADATA(InParamWetLevel), 1.0f),
+				TInputDataVertex<float>(METASOUND_GET_PARAM_NAME_AND_METADATA(InParamFeedbackAmount), 0.0f)
 			),
 			FOutputVertexInterface(
-				TOutputDataVertexModel<FAudioBuffer>(METASOUND_GET_PARAM_NAME_AND_METADATA(OutParamAudio))
+				TOutputDataVertex<FAudioBuffer>(METASOUND_GET_PARAM_NAME_AND_METADATA(OutParamAudio))
 			)
 		);
 
diff --git a/Engine/Plugins/Runtime/Metasound/Source/MetasoundStandardNodes/Private/MetasoundDiffuserNode.cpp b/Engine/Plugins/Runtime/Metasound/Source/MetasoundStandardNodes/Private/MetasoundDiffuserNode.cpp
--- a/Engine/Plugins/Runtime/Metasound/Source/MetasoundStandardNodes/Private/MetasoundDiffuserNode.cpp
+++ b/Engine/Plugins/Runtime/Metasound/Source/MetasoundStandardNodes/Private/MetasoundDiffuserNode.cpp
@@ -63,12 +63,12 @@ namespace Metasound
 
 			static const FVertexInterface Interface(
 				FInputVertexInterface(
-					TInputDataVertexModel<FAudioBuffer>(METASOUND_GET_PARAM_NAME_AND_METADATA(InputAudio)),
-					TInputDataVertexModel<int32>(METASOUND_GET_PARAM_NAME_AND_METADATA(InputDiffusionDepth), 4),
-					TInputDataVertexModel<float>(METASOUND_GET_PARAM_NAME_AND_METADATA(InputFeedbackGain), 0.707f)
+					TInputDataVertex<FAudioBuffer>(METASOUND_GET_PARAM_NAME_AND_METADATA(InputAudio)),
+					TInputDataVertex<int32>(METASOUND_GET_PARAM_NAME_AND_METADATA(InputDiffusionDepth), 4),
+					TInputDataVertex<float>(METASOUND_GET_PARAM_NAME_AND_METADATA(InputFeedbackGain), 0.707f)
 				),
 				FOutputVertexInterface(
-					TOutputDataVertexModel<FAudioBuffer>(METASOUND_GET_PARAM_NAME_AND_METADATA(OutputAudio))
+					TOutputDataVertex<FAudioBuffer>(METASOUND_GET_PARAM_NAME_AND_METADATA(OutputAudio))
 				)
 			);
 
@@ -216,4 +216,4 @@ namespace Metasound
 	METASOUND_REGISTER_NODE(FDiffuserNode)
 }
 
-#undef LOCTEXT_NAMESPACE
\ No newline at end of file
+#undef LOCTEXT_NAMESPACE
diff --git a/Engine/Plugins/Runtime/Metasound/Source/MetasoundStandardNodes/Private/MetasoundDynamicFilterNode.cpp b/Engine/Plugins/Runtime/Metasound/Source/MetasoundStandardNodes/Private/MetasoundDynamicFilterNode.cpp
--- a/Engine/Plugins/Runtime/Metasound/Source/MetasoundStandardNodes/Private/MetasoundDynamicFilterNode.cpp
+++ b/Engine/Plugins/Runtime/Metasound/Source/MetasoundStandardNodes/Private/MetasoundDynamicFilterNode.cpp
@@ -161,27 +161,27 @@ namespace Metasound
 
 			static const FVertexInterface Interface(
 				FInputVertexInterface(
-					TInputDataVertexModel<FAudioBuffer>(METASOUND_GET_PARAM_NAME_AND_METADATA(InputAudio)),
-					TInputDataVertexModel<FAudioBuffer>(METASOUND_GET_PARAM_NAME_AND_METADATA(InputSidechain)),
-					TInputDataVertexModel<FEnumEDynamicFilterType>(METASOUND_GET_PARAM_NAME_AND_METADATA(InputFilterType)),
-					TInputDataVertexModel<float>(METASOUND_GET_PARAM_NAME_AND_METADATA(InputFrequency), 1000.0f),
-					TInputDataVertexModel<float>(METASOUND_GET_PARAM_NAME_AND_METADATA(InputQ), 1.0f),
-					TInputDataVertexModel<float>(METASOUND_GET_PARAM_NAME_AND_METADATA(InputThreshold), -12.0f),
-					TInputDataVertexModel<float>(METASOUND_GET_PARAM_NAME_AND_METADATA(InputRatio), 4.0f),
-					TInputDataVertexModel<float>(METASOUND_GET_PARAM_NAME_AND_METADATA(InputKnee), 12.0f),
-					TInputDataVertexModel<float>(METASOUND_GET_PARAM_NAME_AND_METADATA(InputRange), -60.0f),
-					TInputDataVertexModel<float>(METASOUND_GET_PARAM_NAME(InputGain),
+					TInputDataVertex<FAudioBuffer>(METASOUND_GET_PARAM_NAME_AND_METADATA(InputAudio)),
+					TInputDataVertex<FAudioBuffer>(METASOUND_GET_PARAM_NAME_AND_METADATA(InputSidechain)),
+					TInputDataVertex<FEnumEDynamicFilterType>(METASOUND_GET_PARAM_NAME_AND_METADATA(InputFilterType)),
+					TInputDataVertex<float>(METASOUND_GET_PARAM_NAME_AND_METADATA(InputFrequency), 1000.0f),
+					TInputDataVertex<float>(METASOUND_GET_PARAM_NAME_AND_METADATA(InputQ), 1.0f),
+					TInputDataVertex<float>(METASOUND_GET_PARAM_NAME_AND_METADATA(InputThreshold), -12.0f),
+					TInputDataVertex<float>(METASOUND_GET_PARAM_NAME_AND_METADATA(InputRatio), 4.0f),
+					TInputDataVertex<float>(METASOUND_GET_PARAM_NAME_AND_METADATA(InputKnee), 12.0f),
+					TInputDataVertex<float>(METASOUND_GET_PARAM_NAME_AND_METADATA(InputRange), -60.0f),
+					TInputDataVertex<float>(METASOUND_GET_PARAM_NAME(InputGain),
 					{ 
 						METASOUND_GET_PARAM_TT(InputGain),
 						METASOUND_LOCTEXT("InputGainName", "Gain (dB)")
 					}, 0.0f),
-					TInputDataVertexModel<FTime>(METASOUND_GET_PARAM_NAME_AND_METADATA(InputAttackTime), 0.01f),
-					TInputDataVertexModel<FTime>(METASOUND_GET_PARAM_NAME_AND_METADATA(InputReleaseTime), 0.1f),
-					TInputDataVertexModel<FEnumEnvelopePeakMode>(METASOUND_GET_PARAM_NAME_AND_METADATA(InputEnvelopeMode), 1.0f),
-					TInputDataVertexModel<bool>(METASOUND_GET_PARAM_NAME_AND_METADATA(InputAnalogMode), true)
+					TInputDataVertex<FTime>(METASOUND_GET_PARAM_NAME_AND_METADATA(InputAttackTime), 0.01f),
+					TInputDataVertex<FTime>(METASOUND_GET_PARAM_NAME_AND_METADATA(InputReleaseTime), 0.1f),
+					TInputDataVertex<FEnumEnvelopePeakMode>(METASOUND_GET_PARAM_NAME_AND_METADATA(InputEnvelopeMode), 1.0f),
+					TInputDataVertex<bool>(METASOUND_GET_PARAM_NAME_AND_METADATA(InputAnalogMode), true)
 				),
 				FOutputVertexInterface(
-					TOutputDataVertexModel<FAudioBuffer>(METASOUND_GET_PARAM_NAME_AND_METADATA(OutputAudio))
+					TOutputDataVertex<FAudioBuffer>(METASOUND_GET_PARAM_NAME_AND_METADATA(OutputAudio))
 				)
 			);
 
diff --git a/Engine/Plugins/Runtime/Metasound/Source/MetasoundStandardNodes/Private/MetasoundEnvelopeFollowerNode.cpp b/Engine/Plugins/Runtime/Metasound/Source/MetasoundStandardNodes/Private/MetasoundEnvelopeFollowerNode.cpp
--- a/Engine/Plugins/Runtime/Metasound/Source/MetasoundStandardNodes/Private/MetasoundEnvelopeFollowerNode.cpp
+++ b/Engine/Plugins/Runtime/Metasound/Source/MetasoundStandardNodes/Private/MetasoundEnvelopeFollowerNode.cpp
@@ -178,14 +178,14 @@ namespace Metasound
 
 		static const FVertexInterface Interface(
 			FInputVertexInterface(
-				TInputDataVertexModel<FAudioBuffer>(METASOUND_GET_PARAM_NAME_AND_METADATA(InParamAudioInput)),
-				TInputDataVertexModel<FTime>(METASOUND_GET_PARAM_NAME_AND_METADATA(InParamAttackTime), 0.01f),
-				TInputDataVertexModel<FTime>(METASOUND_GET_PARAM_NAME_AND_METADATA(InParamReleaseTime), 0.1f),
-				TInputDataVertexModel<FEnumEnvelopePeakMode>(METASOUND_GET_PARAM_NAME_AND_METADATA(InParamFollowMode))
+				TInputDataVertex<FAudioBuffer>(METASOUND_GET_PARAM_NAME_AND_METADATA(InParamAudioInput)),
+				TInputDataVertex<FTime>(METASOUND_GET_PARAM_NAME_AND_METADATA(InParamAttackTime), 0.01f),
+				TInputDataVertex<FTime>(METASOUND_GET_PARAM_NAME_AND_METADATA(InParamReleaseTime), 0.1f),
+				TInputDataVertex<FEnumEnvelopePeakMode>(METASOUND_GET_PARAM_NAME_AND_METADATA(InParamFollowMode))
 			),
 			FOutputVertexInterface(
-				TOutputDataVertexModel<float>(METASOUND_GET_PARAM_NAME_AND_METADATA(OutParamEnvelope)),
-				TOutputDataVertexModel<FAudioBuffer>(METASOUND_GET_PARAM_NAME_AND_METADATA(OutputAudioEnvelope))
+				TOutputDataVertex<float>(METASOUND_GET_PARAM_NAME_AND_METADATA(OutParamEnvelope)),
+				TOutputDataVertex<FAudioBuffer>(METASOUND_GET_PARAM_NAME_AND_METADATA(OutputAudioEnvelope))
 			)
 		);
 
diff --git a/Engine/Plugins/Runtime/Metasound/Source/MetasoundStandardNodes/Private/MetasoundFlangerNode.cpp b/Engine/Plugins/Runtime/Metasound/Source/MetasoundStandardNodes/Private/MetasoundFlangerNode.cpp
--- a/Engine/Plugins/Runtime/Metasound/Source/MetasoundStandardNodes/Private/MetasoundFlangerNode.cpp
+++ b/Engine/Plugins/Runtime/Metasound/Source/MetasoundStandardNodes/Private/MetasoundFlangerNode.cpp
@@ -61,14 +61,14 @@ namespace Metasound
 
 			static const FVertexInterface Interface(
 				FInputVertexInterface(
-					TInputDataVertexModel<FAudioBuffer>(METASOUND_GET_PARAM_NAME_AND_METADATA(InputAudio)), 
-					TInputDataVertexModel<float>(METASOUND_GET_PARAM_NAME_AND_METADATA(InputModulationRate), 0.5f),
-					TInputDataVertexModel<float>(METASOUND_GET_PARAM_NAME_AND_METADATA(InputModulationDepth), 0.5f),
-					TInputDataVertexModel<float>(METASOUND_GET_PARAM_NAME_AND_METADATA(InputCenterDelay), 0.5f),
-					TInputDataVertexModel<float>(METASOUND_GET_PARAM_NAME_AND_METADATA(InputMixLevel), 0.5f)
+					TInputDataVertex<FAudioBuffer>(METASOUND_GET_PARAM_NAME_AND_METADATA(InputAudio)), 
+					TInputDataVertex<float>(METASOUND_GET_PARAM_NAME_AND_METADATA(InputModulationRate), 0.5f),
+					TInputDataVertex<float>(METASOUND_GET_PARAM_NAME_AND_METADATA(InputModulationDepth), 0.5f),
+					TInputDataVertex<float>(METASOUND_GET_PARAM_NAME_AND_METADATA(InputCenterDelay), 0.5f),
+					TInputDataVertex<float>(METASOUND_GET_PARAM_NAME_AND_METADATA(InputMixLevel), 0.5f)
 				),
 				FOutputVertexInterface(
-					TOutputDataVertexModel<FAudioBuffer>(METASOUND_GET_PARAM_NAME_AND_METADATA(OutputAudio))
+					TOutputDataVertex<FAudioBuffer>(METASOUND_GET_PARAM_NAME_AND_METADATA(OutputAudio))
 				)
 			);
 
diff --git a/Engine/Plugins/Runtime/Metasound/Source/MetasoundStandardNodes/Private/MetasoundFreqMultiplierToSemitoneNode.cpp b/Engine/Plugins/Runtime/Metasound/Source/MetasoundStandardNodes/Private/MetasoundFreqMultiplierToSemitoneNode.cpp
--- a/Engine/Plugins/Runtime/Metasound/Source/MetasoundStandardNodes/Private/MetasoundFreqMultiplierToSemitoneNode.cpp
+++ b/Engine/Plugins/Runtime/Metasound/Source/MetasoundStandardNodes/Private/MetasoundFreqMultiplierToSemitoneNode.cpp
@@ -91,10 +91,10 @@ namespace Metasound
 
 		static const FVertexInterface Interface(
 			FInputVertexInterface(
-				TInputDataVertexModel<float>(METASOUND_GET_PARAM_NAME_AND_METADATA(InputFrequencyMultiplier), 1.0f)
+				TInputDataVertex<float>(METASOUND_GET_PARAM_NAME_AND_METADATA(InputFrequencyMultiplier), 1.0f)
 			),
 			FOutputVertexInterface(
-				TOutputDataVertexModel<float>(METASOUND_GET_PARAM_NAME_AND_METADATA(OutputSemitone))
+				TOutputDataVertex<float>(METASOUND_GET_PARAM_NAME_AND_METADATA(OutputSemitone))
 			)
 		);
 
diff --git a/Engine/Plugins/Runtime/Metasound/Source/MetasoundStandardNodes/Private/MetasoundFreqToMIDINode.cpp b/Engine/Plugins/Runtime/Metasound/Source/MetasoundStandardNodes/Private/MetasoundFreqToMIDINode.cpp
--- a/Engine/Plugins/Runtime/Metasound/Source/MetasoundStandardNodes/Private/MetasoundFreqToMIDINode.cpp
+++ b/Engine/Plugins/Runtime/Metasound/Source/MetasoundStandardNodes/Private/MetasoundFreqToMIDINode.cpp
@@ -92,10 +92,10 @@ namespace Metasound
 
 		static const FVertexInterface Interface(
 			FInputVertexInterface(
-				TInputDataVertexModel<float>(METASOUND_GET_PARAM_NAME_AND_METADATA(InputFreq), 440.0f)
+				TInputDataVertex<float>(METASOUND_GET_PARAM_NAME_AND_METADATA(InputFreq), 440.0f)
 			),
 			FOutputVertexInterface(
-				TOutputDataVertexModel<float>(METASOUND_GET_PARAM_NAME_AND_METADATA(OutputMidi))
+				TOutputDataVertex<float>(METASOUND_GET_PARAM_NAME_AND_METADATA(OutputMidi))
 			)
 		);
 
diff --git a/Engine/Plugins/Runtime/Metasound/Source/MetasoundStandardNodes/Private/MetasoundITDPannerNode.cpp b/Engine/Plugins/Runtime/Metasound/Source/MetasoundStandardNodes/Private/MetasoundITDPannerNode.cpp
--- a/Engine/Plugins/Runtime/Metasound/Source/MetasoundStandardNodes/Private/MetasoundITDPannerNode.cpp
+++ b/Engine/Plugins/Runtime/Metasound/Source/MetasoundStandardNodes/Private/MetasoundITDPannerNode.cpp
@@ -228,14 +228,14 @@ namespace Metasound
 
 		static const FVertexInterface Interface(
 			FInputVertexInterface(
-				TInputDataVertexModel<FAudioBuffer>(METASOUND_GET_PARAM_NAME_AND_METADATA(InputAudio)),
-				TInputDataVertexModel<float>(METASOUND_GET_PARAM_NAME_AND_METADATA(InputPanAngle), 90.0f),
-				TInputDataVertexModel<float>(METASOUND_GET_PARAM_NAME_AND_METADATA(InputDistanceFactor), 0.0f),
-				TInputDataVertexModel<float>(METASOUND_GET_PARAM_NAME_AND_METADATA(InputHeadWidth), 34.0f)
+				TInputDataVertex<FAudioBuffer>(METASOUND_GET_PARAM_NAME_AND_METADATA(InputAudio)),
+				TInputDataVertex<float>(METASOUND_GET_PARAM_NAME_AND_METADATA(InputPanAngle), 90.0f),
+				TInputDataVertex<float>(METASOUND_GET_PARAM_NAME_AND_METADATA(InputDistanceFactor), 0.0f),
+				TInputDataVertex<float>(METASOUND_GET_PARAM_NAME_AND_METADATA(InputHeadWidth), 34.0f)
 			),
 			FOutputVertexInterface(
-				TOutputDataVertexModel<FAudioBuffer>(METASOUND_GET_PARAM_NAME_AND_METADATA(OutputAudioLeft)),
-				TOutputDataVertexModel<FAudioBuffer>(METASOUND_GET_PARAM_NAME_AND_METADATA(OutputAudioRight))
+				TOutputDataVertex<FAudioBuffer>(METASOUND_GET_PARAM_NAME_AND_METADATA(OutputAudioLeft)),
+				TOutputDataVertex<FAudioBuffer>(METASOUND_GET_PARAM_NAME_AND_METADATA(OutputAudioRight))
 			)
 		);
 
diff --git a/Engine/Plugins/Runtime/Metasound/Source/MetasoundStandardNodes/Private/MetasoundInterpToNode.cpp b/Engine/Plugins/Runtime/Metasound/Source/MetasoundStandardNodes/Private/MetasoundInterpToNode.cpp
--- a/Engine/Plugins/Runtime/Metasound/Source/MetasoundStandardNodes/Private/MetasoundInterpToNode.cpp
+++ b/Engine/Plugins/Runtime/Metasound/Source/MetasoundStandardNodes/Private/MetasoundInterpToNode.cpp
@@ -133,11 +133,11 @@ namespace Metasound
 
 		static const FVertexInterface Interface(
 			FInputVertexInterface(
-				TInputDataVertexModel<FTime>(METASOUND_GET_PARAM_NAME_AND_METADATA(InParamInterpTime), 0.1f),
-				TInputDataVertexModel<float>(METASOUND_GET_PARAM_NAME_AND_METADATA(InParamTarget), 1.0f)
+				TInputDataVertex<FTime>(METASOUND_GET_PARAM_NAME_AND_METADATA(InParamInterpTime), 0.1f),
+				TInputDataVertex<float>(METASOUND_GET_PARAM_NAME_AND_METADATA(InParamTarget), 1.0f)
 			),
 			FOutputVertexInterface(
-				TOutputDataVertexModel<float>(METASOUND_GET_PARAM_NAME_AND_METADATA(OutParamValue))
+				TOutputDataVertex<float>(METASOUND_GET_PARAM_NAME_AND_METADATA(OutParamValue))
 			)
 		);
 
diff --git a/Engine/Plugins/Runtime/Metasound/Source/MetasoundStandardNodes/Private/MetasoundLimiterNode.cpp b/Engine/Plugins/Runtime/Metasound/Source/MetasoundStandardNodes/Private/MetasoundLimiterNode.cpp
--- a/Engine/Plugins/Runtime/Metasound/Source/MetasoundStandardNodes/Private/MetasoundLimiterNode.cpp
+++ b/Engine/Plugins/Runtime/Metasound/Source/MetasoundStandardNodes/Private/MetasoundLimiterNode.cpp
@@ -123,14 +123,14 @@ namespace Metasound
 
 			static const FVertexInterface Interface(
 				FInputVertexInterface(
-					TInputDataVertexModel<FAudioBuffer>(METASOUND_GET_PARAM_NAME_AND_METADATA(InputAudio)),
-					TInputDataVertexModel<float>(METASOUND_GET_PARAM_NAME_AND_METADATA(InputInGainDb), 0.0f),
-					TInputDataVertexModel<float>(METASOUND_GET_PARAM_NAME_AND_METADATA(InputThresholdDb), 0.0f),
-					TInputDataVertexModel<FTime>(METASOUND_GET_PARAM_NAME_AND_METADATA(InputReleaseTime), 0.1f),
-					TInputDataVertexModel<FEnumKneeMode>(METASOUND_GET_PARAM_NAME_AND_METADATA(InputKneeMode))
+					TInputDataVertex<FAudioBuffer>(METASOUND_GET_PARAM_NAME_AND_METADATA(InputAudio)),
+					TInputDataVertex<float>(METASOUND_GET_PARAM_NAME_AND_METADATA(InputInGainDb), 0.0f),
+					TInputDataVertex<float>(METASOUND_GET_PARAM_NAME_AND_METADATA(InputThresholdDb), 0.0f),
+					TInputDataVertex<FTime>(METASOUND_GET_PARAM_NAME_AND_METADATA(InputReleaseTime), 0.1f),
+					TInputDataVertex<FEnumKneeMode>(METASOUND_GET_PARAM_NAME_AND_METADATA(InputKneeMode))
 				),
 				FOutputVertexInterface(
-					TOutputDataVertexModel<FAudioBuffer>(METASOUND_GET_PARAM_NAME_AND_METADATA(OutputAudio))
+					TOutputDataVertex<FAudioBuffer>(METASOUND_GET_PARAM_NAME_AND_METADATA(OutputAudio))
 				)
 			);
 
diff --git a/Engine/Plugins/Runtime/Metasound/Source/MetasoundStandardNodes/Private/MetasoundLinearGainToDecibels.cpp b/Engine/Plugins/Runtime/Metasound/Source/MetasoundStandardNodes/Private/MetasoundLinearGainToDecibels.cpp
--- a/Engine/Plugins/Runtime/Metasound/Source/MetasoundStandardNodes/Private/MetasoundLinearGainToDecibels.cpp
+++ b/Engine/Plugins/Runtime/Metasound/Source/MetasoundStandardNodes/Private/MetasoundLinearGainToDecibels.cpp
@@ -87,10 +87,10 @@ namespace Metasound
 
 		static const FVertexInterface Interface(
 			FInputVertexInterface(
-				TInputDataVertexModel<float>(METASOUND_GET_PARAM_NAME_AND_METADATA(InputLinearGain), 1.0f)
+				TInputDataVertex<float>(METASOUND_GET_PARAM_NAME_AND_METADATA(InputLinearGain), 1.0f)
 			),
 			FOutputVertexInterface(
-				TOutputDataVertexModel<float>(METASOUND_GET_PARAM_NAME_AND_METADATA(OutputDecibels))
+				TOutputDataVertex<float>(METASOUND_GET_PARAM_NAME_AND_METADATA(OutputDecibels))
 			)
 		);
 
diff --git a/Engine/Plugins/Runtime/Metasound/Source/MetasoundStandardNodes/Private/MetasoundLinearToLogFrequencyNode.cpp b/Engine/Plugins/Runtime/Metasound/Source/MetasoundStandardNodes/Private/MetasoundLinearToLogFrequencyNode.cpp
--- a/Engine/Plugins/Runtime/Metasound/Source/MetasoundStandardNodes/Private/MetasoundLinearToLogFrequencyNode.cpp
+++ b/Engine/Plugins/Runtime/Metasound/Source/MetasoundStandardNodes/Private/MetasoundLinearToLogFrequencyNode.cpp
@@ -61,14 +61,14 @@ namespace Metasound
 
 			static const FVertexInterface Interface(
 				FInputVertexInterface(
-					TInputDataVertexModel<float>(METASOUND_GET_PARAM_NAME_AND_METADATA(InputValue),     0.0f),
-					TInputDataVertexModel<float>(METASOUND_GET_PARAM_NAME_AND_METADATA(InputDomainMin), 0.0f),
-					TInputDataVertexModel<float>(METASOUND_GET_PARAM_NAME_AND_METADATA(InputDomainMax), 1.0f),
-					TInputDataVertexModel<float>(METASOUND_GET_PARAM_NAME_AND_METADATA(InputRangeMin),  20.0f),
-					TInputDataVertexModel<float>(METASOUND_GET_PARAM_NAME_AND_METADATA(InputRangeMax),  20000.0f)
+					TInputDataVertex<float>(METASOUND_GET_PARAM_NAME_AND_METADATA(InputValue),     0.0f),
+					TInputDataVertex<float>(METASOUND_GET_PARAM_NAME_AND_METADATA(InputDomainMin), 0.0f),
+					TInputDataVertex<float>(METASOUND_GET_PARAM_NAME_AND_METADATA(InputDomainMax), 1.0f),
+					TInputDataVertex<float>(METASOUND_GET_PARAM_NAME_AND_METADATA(InputRangeMin),  20.0f),
+					TInputDataVertex<float>(METASOUND_GET_PARAM_NAME_AND_METADATA(InputRangeMax),  20000.0f)
 					),
 				FOutputVertexInterface(
-					TOutputDataVertexModel<float>(METASOUND_GET_PARAM_NAME_AND_METADATA(OutputFreq))
+					TOutputDataVertex<float>(METASOUND_GET_PARAM_NAME_AND_METADATA(OutputFreq))
 				)
 			);
 
diff --git a/Engine/Plugins/Runtime/Metasound/Source/MetasoundStandardNodes/Private/MetasoundMIDIToFreqNode.cpp b/Engine/Plugins/Runtime/Metasound/Source/MetasoundStandardNodes/Private/MetasoundMIDIToFreqNode.cpp
--- a/Engine/Plugins/Runtime/Metasound/Source/MetasoundStandardNodes/Private/MetasoundMIDIToFreqNode.cpp
+++ b/Engine/Plugins/Runtime/Metasound/Source/MetasoundStandardNodes/Private/MetasoundMIDIToFreqNode.cpp
@@ -133,10 +133,10 @@ namespace Metasound
 
 		static const FVertexInterface Interface(
 			FInputVertexInterface(
-				TInputDataVertexModel<ValueType>(METASOUND_GET_PARAM_NAME_AND_METADATA(InputMidi), (ValueType)60.0f)
+				TInputDataVertex<ValueType>(METASOUND_GET_PARAM_NAME_AND_METADATA(InputMidi), (ValueType)60.0f)
 			),
 			FOutputVertexInterface(
-				TOutputDataVertexModel<float>(METASOUND_GET_PARAM_NAME_AND_METADATA(OutputFreq))
+				TOutputDataVertex<float>(METASOUND_GET_PARAM_NAME_AND_METADATA(OutputFreq))
 			)
 		);
 
diff --git a/Engine/Plugins/Runtime/Metasound/Source/MetasoundStandardNodes/Private/MetasoundMapRangeNode.cpp b/Engine/Plugins/Runtime/Metasound/Source/MetasoundStandardNodes/Private/MetasoundMapRangeNode.cpp
--- a/Engine/Plugins/Runtime/Metasound/Source/MetasoundStandardNodes/Private/MetasoundMapRangeNode.cpp
+++ b/Engine/Plugins/Runtime/Metasound/Source/MetasoundStandardNodes/Private/MetasoundMapRangeNode.cpp
@@ -49,15 +49,15 @@ namespace Metasound
 			{
 				static const FVertexInterface DefaultInterface(
 					FInputVertexInterface(
-						TInputDataVertexModel<int32>(METASOUND_GET_PARAM_NAME_AND_METADATA(InputValueName)),
-						TInputDataVertexModel<int32>(METASOUND_GET_PARAM_NAME_AND_METADATA(InputInRangeAName), 0),
-						TInputDataVertexModel<int32>(METASOUND_GET_PARAM_NAME_AND_METADATA(InputInRangeBName), 100),
-						TInputDataVertexModel<int32>(METASOUND_GET_PARAM_NAME_AND_METADATA(InputOutRangeAName), 0),
-						TInputDataVertexModel<int32>(METASOUND_GET_PARAM_NAME_AND_METADATA(InputOutRangeBName), 100),
-						TInputDataVertexModel<bool>(METASOUND_GET_PARAM_NAME_AND_METADATA(InputClampedName), true)
+						TInputDataVertex<int32>(METASOUND_GET_PARAM_NAME_AND_METADATA(InputValueName)),
+						TInputDataVertex<int32>(METASOUND_GET_PARAM_NAME_AND_METADATA(InputInRangeAName), 0),
+						TInputDataVertex<int32>(METASOUND_GET_PARAM_NAME_AND_METADATA(InputInRangeBName), 100),
+						TInputDataVertex<int32>(METASOUND_GET_PARAM_NAME_AND_METADATA(InputOutRangeAName), 0),
+						TInputDataVertex<int32>(METASOUND_GET_PARAM_NAME_AND_METADATA(InputOutRangeBName), 100),
+						TInputDataVertex<bool>(METASOUND_GET_PARAM_NAME_AND_METADATA(InputClampedName), true)
 					),
 					FOutputVertexInterface(
-						TOutputDataVertexModel<int32>(METASOUND_GET_PARAM_NAME_AND_METADATA(OutputValueName))
+						TOutputDataVertex<int32>(METASOUND_GET_PARAM_NAME_AND_METADATA(OutputValueName))
 					)
 				);
 				return DefaultInterface;
@@ -133,15 +133,15 @@ namespace Metasound
 			{
 				static const FVertexInterface DefaultInterface(
 					FInputVertexInterface(
-						TInputDataVertexModel<float>(METASOUND_GET_PARAM_NAME_AND_METADATA(InputValueName)),
-						TInputDataVertexModel<float>(METASOUND_GET_PARAM_NAME_AND_METADATA(InputInRangeAName), 0.0f),
-						TInputDataVertexModel<float>(METASOUND_GET_PARAM_NAME_AND_METADATA(InputInRangeBName), 1.0f),
-						TInputDataVertexModel<float>(METASOUND_GET_PARAM_NAME_AND_METADATA(InputOutRangeAName), 0.0f),
-						TInputDataVertexModel<float>(METASOUND_GET_PARAM_NAME_AND_METADATA(InputOutRangeBName), 1.0f),
-						TInputDataVertexModel<bool>(METASOUND_GET_PARAM_NAME_AND_METADATA(InputClampedName), true)
+						TInputDataVertex<float>(METASOUND_GET_PARAM_NAME_AND_METADATA(InputValueName)),
+						TInputDataVertex<float>(METASOUND_GET_PARAM_NAME_AND_METADATA(InputInRangeAName), 0.0f),
+						TInputDataVertex<float>(METASOUND_GET_PARAM_NAME_AND_METADATA(InputInRangeBName), 1.0f),
+						TInputDataVertex<float>(METASOUND_GET_PARAM_NAME_AND_METADATA(InputOutRangeAName), 0.0f),
+						TInputDataVertex<float>(METASOUND_GET_PARAM_NAME_AND_METADATA(InputOutRangeBName), 1.0f),
+						TInputDataVertex<bool>(METASOUND_GET_PARAM_NAME_AND_METADATA(InputClampedName), true)
 					),
 					FOutputVertexInterface(
-						TOutputDataVertexModel<float>(METASOUND_GET_PARAM_NAME_AND_METADATA(OutputValueName))
+						TOutputDataVertex<float>(METASOUND_GET_PARAM_NAME_AND_METADATA(OutputValueName))
 					)
 				);
 				return DefaultInterface;
@@ -217,15 +217,15 @@ namespace Metasound
 			{
 				static const FVertexInterface DefaultInterface(
 					FInputVertexInterface(
-						TInputDataVertexModel<FAudioBuffer>(METASOUND_GET_PARAM_NAME_AND_METADATA(InputValueName)),
-						TInputDataVertexModel<float>(METASOUND_GET_PARAM_NAME_AND_METADATA(InputInRangeAName), -1.0f),
-						TInputDataVertexModel<float>(METASOUND_GET_PARAM_NAME_AND_METADATA(InputInRangeBName), 1.0f),
-						TInputDataVertexModel<float>(METASOUND_GET_PARAM_NAME_AND_METADATA(InputOutRangeAName), -1.0f),
-						TInputDataVertexModel<float>(METASOUND_GET_PARAM_NAME_AND_METADATA(InputOutRangeBName), 1.0f),
-						TInputDataVertexModel<bool>(METASOUND_GET_PARAM_NAME_AND_METADATA(InputClampedName), true)
+						TInputDataVertex<FAudioBuffer>(METASOUND_GET_PARAM_NAME_AND_METADATA(InputValueName)),
+						TInputDataVertex<float>(METASOUND_GET_PARAM_NAME_AND_METADATA(InputInRangeAName), -1.0f),
+						TInputDataVertex<float>(METASOUND_GET_PARAM_NAME_AND_METADATA(InputInRangeBName), 1.0f),
+						TInputDataVertex<float>(METASOUND_GET_PARAM_NAME_AND_METADATA(InputOutRangeAName), -1.0f),
+						TInputDataVertex<float>(METASOUND_GET_PARAM_NAME_AND_METADATA(InputOutRangeBName), 1.0f),
+						TInputDataVertex<bool>(METASOUND_GET_PARAM_NAME_AND_METADATA(InputClampedName), true)
 					),
 					FOutputVertexInterface(
-						TOutputDataVertexModel<FAudioBuffer>(METASOUND_GET_PARAM_NAME_AND_METADATA(OutputValueName))
+						TOutputDataVertex<FAudioBuffer>(METASOUND_GET_PARAM_NAME_AND_METADATA(OutputValueName))
 					)
 				);
 
diff --git a/Engine/Plugins/Runtime/Metasound/Source/MetasoundStandardNodes/Private/MetasoundMathNodes.cpp b/Engine/Plugins/Runtime/Metasound/Source/MetasoundStandardNodes/Private/MetasoundMathNodes.cpp
--- a/Engine/Plugins/Runtime/Metasound/Source/MetasoundStandardNodes/Private/MetasoundMathNodes.cpp
+++ b/Engine/Plugins/Runtime/Metasound/Source/MetasoundStandardNodes/Private/MetasoundMathNodes.cpp
@@ -337,11 +337,11 @@ namespace Metasound
 
 			static const FVertexInterface Interface(
 				FInputVertexInterface(
-					TInputDataVertexModel<TDataClass>(METASOUND_GET_PARAM_NAME(PrimaryOperand), PrimaryOperandMetadata, static_cast<TDataClass>(0)),
-					TInputDataVertexModel<TOperandDataClass>(METASOUND_GET_PARAM_NAME(AdditionalOperands), AdditionalAddendsMetadata, static_cast<TOperandDataClass>(0))
+					TInputDataVertex<TDataClass>(METASOUND_GET_PARAM_NAME(PrimaryOperand), PrimaryOperandMetadata, static_cast<TDataClass>(0)),
+					TInputDataVertex<TOperandDataClass>(METASOUND_GET_PARAM_NAME(AdditionalOperands), AdditionalAddendsMetadata, static_cast<TOperandDataClass>(0))
 				),
 				FOutputVertexInterface(
-					TOutputDataVertexModel<TDataClass>(METASOUND_GET_PARAM_NAME(OutMath), OutputMetadata)
+					TOutputDataVertex<TDataClass>(METASOUND_GET_PARAM_NAME(OutMath), OutputMetadata)
 				)
 			);
 
@@ -406,11 +406,11 @@ namespace Metasound
 
 			static const FVertexInterface Interface(
 				FInputVertexInterface(
-					TInputDataVertexModel<TDataClass>(METASOUND_GET_PARAM_NAME(PrimaryOperand), PrimaryOperandMetadata, static_cast<TDataClass>(0)),
-					TInputDataVertexModel<TOperandDataClass>(METASOUND_GET_PARAM_NAME(AdditionalOperands), AdditionalOperandsMetadata, static_cast<TOperandDataClass>(0))
+					TInputDataVertex<TDataClass>(METASOUND_GET_PARAM_NAME(PrimaryOperand), PrimaryOperandMetadata, static_cast<TDataClass>(0)),
+					TInputDataVertex<TOperandDataClass>(METASOUND_GET_PARAM_NAME(AdditionalOperands), AdditionalOperandsMetadata, static_cast<TOperandDataClass>(0))
 				),
 				FOutputVertexInterface(
-					TOutputDataVertexModel<TDataClass>(METASOUND_GET_PARAM_NAME(OutMath), OutputMetadata)
+					TOutputDataVertex<TDataClass>(METASOUND_GET_PARAM_NAME(OutMath), OutputMetadata)
 				)
 			);
 
@@ -475,11 +475,11 @@ namespace Metasound
 
 			static const FVertexInterface Interface(
 				FInputVertexInterface(
-					TInputDataVertexModel<TDataClass>(METASOUND_GET_PARAM_NAME(PrimaryOperand), PrimaryOperandMetadata, static_cast<TDataClass>(1)),
-					TInputDataVertexModel<TOperandDataClass>(METASOUND_GET_PARAM_NAME(AdditionalOperands), AdditionalOperandsMetadata, static_cast<TOperandDataClass>(1))
+					TInputDataVertex<TDataClass>(METASOUND_GET_PARAM_NAME(PrimaryOperand), PrimaryOperandMetadata, static_cast<TDataClass>(1)),
+					TInputDataVertex<TOperandDataClass>(METASOUND_GET_PARAM_NAME(AdditionalOperands), AdditionalOperandsMetadata, static_cast<TOperandDataClass>(1))
 				),
 				FOutputVertexInterface(
-					TOutputDataVertexModel<TDataClass>(METASOUND_GET_PARAM_NAME(OutMath), OutputMetadata)
+					TOutputDataVertex<TDataClass>(METASOUND_GET_PARAM_NAME(OutMath), OutputMetadata)
 				)
 			);
 
@@ -544,11 +544,11 @@ namespace Metasound
 
 			static const FVertexInterface Interface(
 				FInputVertexInterface(
-					TInputDataVertexModel<TDataClass>(METASOUND_GET_PARAM_NAME(PrimaryOperand), PrimaryOperandMetadata, static_cast<TDataClass>(1)),
-					TInputDataVertexModel<TOperandDataClass>(METASOUND_GET_PARAM_NAME(AdditionalOperands), AdditionalOperandsMetadata, static_cast<TOperandDataClass>(1))
+					TInputDataVertex<TDataClass>(METASOUND_GET_PARAM_NAME(PrimaryOperand), PrimaryOperandMetadata, static_cast<TDataClass>(1)),
+					TInputDataVertex<TOperandDataClass>(METASOUND_GET_PARAM_NAME(AdditionalOperands), AdditionalOperandsMetadata, static_cast<TOperandDataClass>(1))
 				),
 				FOutputVertexInterface(
-					TOutputDataVertexModel<TDataClass>(METASOUND_GET_PARAM_NAME(OutMath), OutputMetadata)
+					TOutputDataVertex<TDataClass>(METASOUND_GET_PARAM_NAME(OutMath), OutputMetadata)
 				)
 			);
 
@@ -620,11 +620,11 @@ namespace Metasound
 
 			static const FVertexInterface Interface(
 				FInputVertexInterface(
-					TInputDataVertexModel<TDataClass>(METASOUND_GET_PARAM_NAME(PrimaryOperand), PrimaryOperandMetadata, static_cast<TDataClass>(1)),
-					TInputDataVertexModel<TOperandDataClass>(METASOUND_GET_PARAM_NAME(AdditionalOperands), AdditionalOperandsMetadata, static_cast<TOperandDataClass>(1))
+					TInputDataVertex<TDataClass>(METASOUND_GET_PARAM_NAME(PrimaryOperand), PrimaryOperandMetadata, static_cast<TDataClass>(1)),
+					TInputDataVertex<TOperandDataClass>(METASOUND_GET_PARAM_NAME(AdditionalOperands), AdditionalOperandsMetadata, static_cast<TOperandDataClass>(1))
 				),
 				FOutputVertexInterface(
-					TOutputDataVertexModel<TDataClass>(METASOUND_GET_PARAM_NAME(OutMath), OutputMetadata)
+					TOutputDataVertex<TDataClass>(METASOUND_GET_PARAM_NAME(OutMath), OutputMetadata)
 				)
 			);
 
@@ -698,11 +698,11 @@ namespace Metasound
 
 			static const FVertexInterface DefaultInterface(
 				FInputVertexInterface(
-					TInputDataVertexModel<float>(METASOUND_GET_PARAM_NAME(PrimaryOperand), PrimaryOperandMetadata),
-					TInputDataVertexModel<float>(METASOUND_GET_PARAM_NAME(AdditionalOperands), AdditionalOperandsMetadata)
+					TInputDataVertex<float>(METASOUND_GET_PARAM_NAME(PrimaryOperand), PrimaryOperandMetadata),
+					TInputDataVertex<float>(METASOUND_GET_PARAM_NAME(AdditionalOperands), AdditionalOperandsMetadata)
 				),
 				FOutputVertexInterface(
-					TOutputDataVertexModel<float>(METASOUND_GET_PARAM_NAME(OutMath), OutputMetadata)
+					TOutputDataVertex<float>(METASOUND_GET_PARAM_NAME(OutMath), OutputMetadata)
 				)
 			);
 
@@ -770,11 +770,11 @@ namespace Metasound
 
 			static const FVertexInterface DefaultInterface(
 				FInputVertexInterface(
-					TInputDataVertexModel<float>(METASOUND_GET_PARAM_NAME(PrimaryOperand), PrimaryOperandMetadata),
-					TInputDataVertexModel<float>(METASOUND_GET_PARAM_NAME(AdditionalOperands), AdditionalOperandsMetadata)
+					TInputDataVertex<float>(METASOUND_GET_PARAM_NAME(PrimaryOperand), PrimaryOperandMetadata),
+					TInputDataVertex<float>(METASOUND_GET_PARAM_NAME(AdditionalOperands), AdditionalOperandsMetadata)
 				),
 				FOutputVertexInterface(
-					TOutputDataVertexModel<float>(METASOUND_GET_PARAM_NAME(OutMath), OutputMetadata)
+					TOutputDataVertex<float>(METASOUND_GET_PARAM_NAME(OutMath), OutputMetadata)
 				)
 			);
 
@@ -833,11 +833,11 @@ namespace Metasound
 
 			static const FVertexInterface Interface(
 				FInputVertexInterface(
-					TInputDataVertexModel<FAudioBuffer>(METASOUND_GET_PARAM_NAME(PrimaryOperand), PrimaryOperandMetadata),
-					TInputDataVertexModel<FAudioBuffer>(METASOUND_GET_PARAM_NAME(AdditionalOperands), AdditionalOperandsMetadata)
+					TInputDataVertex<FAudioBuffer>(METASOUND_GET_PARAM_NAME(PrimaryOperand), PrimaryOperandMetadata),
+					TInputDataVertex<FAudioBuffer>(METASOUND_GET_PARAM_NAME(AdditionalOperands), AdditionalOperandsMetadata)
 				),
 				FOutputVertexInterface(
-					TOutputDataVertexModel<FAudioBuffer>(METASOUND_GET_PARAM_NAME_AND_METADATA(OutAudio))
+					TOutputDataVertex<FAudioBuffer>(METASOUND_GET_PARAM_NAME_AND_METADATA(OutAudio))
 				)
 			);
 
@@ -930,11 +930,11 @@ namespace Metasound
 
 			static const FVertexInterface Interface(
 				FInputVertexInterface(
-					TInputDataVertexModel<FTime>(METASOUND_GET_PARAM_NAME(PrimaryOperand), PrimaryOperandMetadata, 0.0f),
-					TInputDataVertexModel<FTime>(METASOUND_GET_PARAM_NAME(AdditionalOperands), AdditionalOperandsMetadata, 0.0f)
+					TInputDataVertex<FTime>(METASOUND_GET_PARAM_NAME(PrimaryOperand), PrimaryOperandMetadata, 0.0f),
+					TInputDataVertex<FTime>(METASOUND_GET_PARAM_NAME(AdditionalOperands), AdditionalOperandsMetadata, 0.0f)
 				),
 				FOutputVertexInterface(
-					TOutputDataVertexModel<FTime>(METASOUND_GET_PARAM_NAME(OutMath), OutputMetadata)
+					TOutputDataVertex<FTime>(METASOUND_GET_PARAM_NAME(OutMath), OutputMetadata)
 				)
 			);
 
@@ -988,11 +988,11 @@ namespace Metasound
 
 			static const FVertexInterface Interface(
 				FInputVertexInterface(
-					TInputDataVertexModel<FAudioBuffer>(METASOUND_GET_PARAM_NAME(PrimaryOperand), PrimaryOperandMetadata),
-					TInputDataVertexModel<float>(METASOUND_GET_PARAM_NAME(AdditionalOperands), AdditionalOperandsMetadata, 0.0f)
+					TInputDataVertex<FAudioBuffer>(METASOUND_GET_PARAM_NAME(PrimaryOperand), PrimaryOperandMetadata),
+					TInputDataVertex<float>(METASOUND_GET_PARAM_NAME(AdditionalOperands), AdditionalOperandsMetadata, 0.0f)
 				),
 				FOutputVertexInterface(
-					TOutputDataVertexModel<FAudioBuffer>(METASOUND_GET_PARAM_NAME_AND_METADATA(OutAudio))
+					TOutputDataVertex<FAudioBuffer>(METASOUND_GET_PARAM_NAME_AND_METADATA(OutAudio))
 				)
 			);
 
@@ -1050,11 +1050,11 @@ namespace Metasound
 
 			static const FVertexInterface Interface(
 				FInputVertexInterface(
-					TInputDataVertexModel<FAudioBuffer>(METASOUND_GET_PARAM_NAME(PrimaryOperand), PrimaryOperandMetadata),
-					TInputDataVertexModel<FAudioBuffer>(METASOUND_GET_PARAM_NAME(AdditionalOperands), AdditionalAddendsMetadata)
+					TInputDataVertex<FAudioBuffer>(METASOUND_GET_PARAM_NAME(PrimaryOperand), PrimaryOperandMetadata),
+					TInputDataVertex<FAudioBuffer>(METASOUND_GET_PARAM_NAME(AdditionalOperands), AdditionalAddendsMetadata)
 				),
 				FOutputVertexInterface(
-					TOutputDataVertexModel<FAudioBuffer>(METASOUND_GET_PARAM_NAME_AND_METADATA(OutAudio))
+					TOutputDataVertex<FAudioBuffer>(METASOUND_GET_PARAM_NAME_AND_METADATA(OutAudio))
 				)
 			);
 
@@ -1145,11 +1145,11 @@ namespace Metasound
 
 			static const FVertexInterface Interface(
 				FInputVertexInterface(
-					TInputDataVertexModel<FTime>(METASOUND_GET_PARAM_NAME(PrimaryOperand), PrimaryOperandMetadata, 0.0f),
-					TInputDataVertexModel<FTime>(METASOUND_GET_PARAM_NAME(AdditionalOperands), AdditionalAddendsMetadata, 0.0f)
+					TInputDataVertex<FTime>(METASOUND_GET_PARAM_NAME(PrimaryOperand), PrimaryOperandMetadata, 0.0f),
+					TInputDataVertex<FTime>(METASOUND_GET_PARAM_NAME(AdditionalOperands), AdditionalAddendsMetadata, 0.0f)
 				),
 				FOutputVertexInterface(
-					TOutputDataVertexModel<FTime>(METASOUND_GET_PARAM_NAME(OutMath), OutputMetadata)
+					TOutputDataVertex<FTime>(METASOUND_GET_PARAM_NAME(OutMath), OutputMetadata)
 				)
 			);
 
@@ -1203,11 +1203,11 @@ namespace Metasound
 
 			static const FVertexInterface Interface(
 				FInputVertexInterface(
-					TInputDataVertexModel<FAudioBuffer>(METASOUND_GET_PARAM_NAME(PrimaryOperand), PrimaryOperandMetadata),
-					TInputDataVertexModel<FAudioBuffer>(METASOUND_GET_PARAM_NAME(AdditionalOperands), AdditionalAddendsMetadata)
+					TInputDataVertex<FAudioBuffer>(METASOUND_GET_PARAM_NAME(PrimaryOperand), PrimaryOperandMetadata),
+					TInputDataVertex<FAudioBuffer>(METASOUND_GET_PARAM_NAME(AdditionalOperands), AdditionalAddendsMetadata)
 				),
 				FOutputVertexInterface(
-					TOutputDataVertexModel<FAudioBuffer>(METASOUND_GET_PARAM_NAME_AND_METADATA(OutAudio))
+					TOutputDataVertex<FAudioBuffer>(METASOUND_GET_PARAM_NAME_AND_METADATA(OutAudio))
 				)
 			);
 
@@ -1295,11 +1295,11 @@ namespace Metasound
 
 			static const FVertexInterface Interface(
 				FInputVertexInterface(
-					TInputDataVertexModel<FAudioBuffer>(METASOUND_GET_PARAM_NAME(PrimaryOperand), PrimaryOperandMetadata),
-					TInputDataVertexModel<float>(METASOUND_GET_PARAM_NAME(AdditionalOperands), AdditionalOperandsMetadata, 1.0f)
+					TInputDataVertex<FAudioBuffer>(METASOUND_GET_PARAM_NAME(PrimaryOperand), PrimaryOperandMetadata),
+					TInputDataVertex<float>(METASOUND_GET_PARAM_NAME(AdditionalOperands), AdditionalOperandsMetadata, 1.0f)
 				),
 				FOutputVertexInterface(
-					TOutputDataVertexModel<FAudioBuffer>(METASOUND_GET_PARAM_NAME_AND_METADATA(OutAudio))
+					TOutputDataVertex<FAudioBuffer>(METASOUND_GET_PARAM_NAME_AND_METADATA(OutAudio))
 				)
 			);
 
@@ -1379,11 +1379,11 @@ namespace Metasound
 
 			static const FVertexInterface Interface(
 				FInputVertexInterface(
-					TInputDataVertexModel<FTime>(METASOUND_GET_PARAM_NAME(PrimaryOperand), PrimaryOperandMetadata, 1.0f),
-					TInputDataVertexModel<float>(METASOUND_GET_PARAM_NAME(AdditionalOperands), AdditionalOperandsMetadata, 1.0f)
+					TInputDataVertex<FTime>(METASOUND_GET_PARAM_NAME(PrimaryOperand), PrimaryOperandMetadata, 1.0f),
+					TInputDataVertex<float>(METASOUND_GET_PARAM_NAME(AdditionalOperands), AdditionalOperandsMetadata, 1.0f)
 				),
 				FOutputVertexInterface(
-					TOutputDataVertexModel<FTime>(METASOUND_GET_PARAM_NAME(OutMath), OutputMetadata)
+					TOutputDataVertex<FTime>(METASOUND_GET_PARAM_NAME(OutMath), OutputMetadata)
 				)
 			);
 
@@ -1445,11 +1445,11 @@ namespace Metasound
 
 			static const FVertexInterface Interface(
 				FInputVertexInterface(
-					TInputDataVertexModel<FTime>(METASOUND_GET_PARAM_NAME(PrimaryOperand), PrimaryOperandMetadata, 1.0f),
-					TInputDataVertexModel<float>(METASOUND_GET_PARAM_NAME(AdditionalOperands), AdditionalOperandsMetadata, 1.0f)
+					TInputDataVertex<FTime>(METASOUND_GET_PARAM_NAME(PrimaryOperand), PrimaryOperandMetadata, 1.0f),
+					TInputDataVertex<float>(METASOUND_GET_PARAM_NAME(AdditionalOperands), AdditionalOperandsMetadata, 1.0f)
 				),
 				FOutputVertexInterface(
-					TOutputDataVertexModel<FTime>(METASOUND_GET_PARAM_NAME(OutMath), OutputMetadata)
+					TOutputDataVertex<FTime>(METASOUND_GET_PARAM_NAME(OutMath), OutputMetadata)
 				)
 			);
 
diff --git a/Engine/Plugins/Runtime/Metasound/Source/MetasoundStandardNodes/Private/MetasoundMaxNode.cpp b/Engine/Plugins/Runtime/Metasound/Source/MetasoundStandardNodes/Private/MetasoundMaxNode.cpp
--- a/Engine/Plugins/Runtime/Metasound/Source/MetasoundStandardNodes/Private/MetasoundMaxNode.cpp
+++ b/Engine/Plugins/Runtime/Metasound/Source/MetasoundStandardNodes/Private/MetasoundMaxNode.cpp
@@ -139,11 +139,11 @@ namespace Metasound
 
 			static const FVertexInterface DefaultInterface(
 				FInputVertexInterface(
-					TInputDataVertexModel<ValueType>(METASOUND_GET_PARAM_NAME_AND_METADATA(InputAValue)),
-					TInputDataVertexModel<ValueType>(METASOUND_GET_PARAM_NAME_AND_METADATA(InputBValue))
+					TInputDataVertex<ValueType>(METASOUND_GET_PARAM_NAME_AND_METADATA(InputAValue)),
+					TInputDataVertex<ValueType>(METASOUND_GET_PARAM_NAME_AND_METADATA(InputBValue))
 				),
 				FOutputVertexInterface(
-					TOutputDataVertexModel<ValueType>(METASOUND_GET_PARAM_NAME_AND_METADATA(OutputValue))
+					TOutputDataVertex<ValueType>(METASOUND_GET_PARAM_NAME_AND_METADATA(OutputValue))
 				)
 			);
 
diff --git a/Engine/Plugins/Runtime/Metasound/Source/MetasoundStandardNodes/Private/MetasoundMidSideNodes.cpp b/Engine/Plugins/Runtime/Metasound/Source/MetasoundStandardNodes/Private/MetasoundMidSideNodes.cpp
--- a/Engine/Plugins/Runtime/Metasound/Source/MetasoundStandardNodes/Private/MetasoundMidSideNodes.cpp
+++ b/Engine/Plugins/Runtime/Metasound/Source/MetasoundStandardNodes/Private/MetasoundMidSideNodes.cpp
@@ -87,14 +87,14 @@ namespace Metasound
 
 			static const FVertexInterface Interface(
 				FInputVertexInterface(
-					TInputDataVertexModel<FAudioBuffer>(METASOUND_GET_PARAM_NAME_AND_METADATA(InputAudioLeft)),
-					TInputDataVertexModel<FAudioBuffer>(METASOUND_GET_PARAM_NAME_AND_METADATA(InputAudioRight)),
-					TInputDataVertexModel<float>(METASOUND_GET_PARAM_NAME_AND_METADATA(InputSpreadAmount), 0.5f),
-					TInputDataVertexModel<bool>(METASOUND_GET_PARAM_NAME_AND_METADATA(InputEqualPower), false)
+					TInputDataVertex<FAudioBuffer>(METASOUND_GET_PARAM_NAME_AND_METADATA(InputAudioLeft)),
+					TInputDataVertex<FAudioBuffer>(METASOUND_GET_PARAM_NAME_AND_METADATA(InputAudioRight)),
+					TInputDataVertex<float>(METASOUND_GET_PARAM_NAME_AND_METADATA(InputSpreadAmount), 0.5f),
+					TInputDataVertex<bool>(METASOUND_GET_PARAM_NAME_AND_METADATA(InputEqualPower), false)
 				),
 				FOutputVertexInterface(
-					TOutputDataVertexModel<FAudioBuffer>(METASOUND_GET_PARAM_NAME_AND_METADATA(OutputAudioMid)),
-					TOutputDataVertexModel<FAudioBuffer>(METASOUND_GET_PARAM_NAME_AND_METADATA(OutputAudioSide))
+					TOutputDataVertex<FAudioBuffer>(METASOUND_GET_PARAM_NAME_AND_METADATA(OutputAudioMid)),
+					TOutputDataVertex<FAudioBuffer>(METASOUND_GET_PARAM_NAME_AND_METADATA(OutputAudioSide))
 				)
 			);
 
@@ -311,14 +311,14 @@ namespace Metasound
 
 			static const FVertexInterface Interface(
 				FInputVertexInterface(
-					TInputDataVertexModel<FAudioBuffer>(METASOUND_GET_PARAM_NAME_AND_METADATA(InputAudioMid)),
-					TInputDataVertexModel<FAudioBuffer>(METASOUND_GET_PARAM_NAME_AND_METADATA(InputAudioSide)),
-					TInputDataVertexModel<float>(METASOUND_GET_PARAM_NAME_AND_METADATA(InputSpreadAmount), 0.5f),
-					TInputDataVertexModel<bool>(METASOUND_GET_PARAM_NAME_AND_METADATA(InputEqualPower), false)
+					TInputDataVertex<FAudioBuffer>(METASOUND_GET_PARAM_NAME_AND_METADATA(InputAudioMid)),
+					TInputDataVertex<FAudioBuffer>(METASOUND_GET_PARAM_NAME_AND_METADATA(InputAudioSide)),
+					TInputDataVertex<float>(METASOUND_GET_PARAM_NAME_AND_METADATA(InputSpreadAmount), 0.5f),
+					TInputDataVertex<bool>(METASOUND_GET_PARAM_NAME_AND_METADATA(InputEqualPower), false)
 				),
 				FOutputVertexInterface(
-					TOutputDataVertexModel<FAudioBuffer>(METASOUND_GET_PARAM_NAME_AND_METADATA(OutputAudioLeft)),
-					TOutputDataVertexModel<FAudioBuffer>(METASOUND_GET_PARAM_NAME_AND_METADATA(OutputAudioRight))
+					TOutputDataVertex<FAudioBuffer>(METASOUND_GET_PARAM_NAME_AND_METADATA(OutputAudioLeft)),
+					TOutputDataVertex<FAudioBuffer>(METASOUND_GET_PARAM_NAME_AND_METADATA(OutputAudioRight))
 				)
 			);
 
diff --git a/Engine/Plugins/Runtime/Metasound/Source/MetasoundStandardNodes/Private/MetasoundMidiNoteQuantizerNode.cpp b/Engine/Plugins/Runtime/Metasound/Source/MetasoundStandardNodes/Private/MetasoundMidiNoteQuantizerNode.cpp
--- a/Engine/Plugins/Runtime/Metasound/Source/MetasoundStandardNodes/Private/MetasoundMidiNoteQuantizerNode.cpp
+++ b/Engine/Plugins/Runtime/Metasound/Source/MetasoundStandardNodes/Private/MetasoundMidiNoteQuantizerNode.cpp
@@ -133,12 +133,12 @@ namespace Metasound
 	{
 		static const FVertexInterface Interface(
 			FInputVertexInterface(
-				TInputDataVertexModel<float>(METASOUND_GET_PARAM_NAME_AND_METADATA(ParamNoteIn), 60.0f),
-				TInputDataVertexModel<float>(METASOUND_GET_PARAM_NAME_AND_METADATA(ParamRootNote), 0.0f),
-				TInputDataVertexModel<ScaleDegreeArrayType>(METASOUND_GET_PARAM_NAME_AND_METADATA(ParamScaleDegrees))
+				TInputDataVertex<float>(METASOUND_GET_PARAM_NAME_AND_METADATA(ParamNoteIn), 60.0f),
+				TInputDataVertex<float>(METASOUND_GET_PARAM_NAME_AND_METADATA(ParamRootNote), 0.0f),
+				TInputDataVertex<ScaleDegreeArrayType>(METASOUND_GET_PARAM_NAME_AND_METADATA(ParamScaleDegrees))
 			),
 			FOutputVertexInterface(
-				TOutputDataVertexModel<float>(METASOUND_GET_PARAM_NAME_AND_METADATA(ParamNoteOutput))
+				TOutputDataVertex<float>(METASOUND_GET_PARAM_NAME_AND_METADATA(ParamNoteOutput))
 				)
 			);
 
@@ -238,4 +238,4 @@ namespace Metasound
 
 } // namespace Metasound
 
-#undef LOCTEXT_NAMESPACE //MetasoundBasicFilterNodes
\ No newline at end of file
+#undef LOCTEXT_NAMESPACE //MetasoundBasicFilterNodes
diff --git a/Engine/Plugins/Runtime/Metasound/Source/MetasoundStandardNodes/Private/MetasoundMidiScaleToArrayNode.cpp b/Engine/Plugins/Runtime/Metasound/Source/MetasoundStandardNodes/Private/MetasoundMidiScaleToArrayNode.cpp
--- a/Engine/Plugins/Runtime/Metasound/Source/MetasoundStandardNodes/Private/MetasoundMidiScaleToArrayNode.cpp
+++ b/Engine/Plugins/Runtime/Metasound/Source/MetasoundStandardNodes/Private/MetasoundMidiScaleToArrayNode.cpp
@@ -127,11 +127,11 @@ namespace Metasound
 	{
 		static const FVertexInterface Interface(
 			FInputVertexInterface(
-				TInputDataVertexModel<FEnumEMusicalScale>(METASOUND_GET_PARAM_NAME_AND_METADATA(ParamScaleDegreesPreset)),
-				TInputDataVertexModel<bool>(METASOUND_GET_PARAM_NAME_AND_METADATA(ParamChordTonesOnly))
+				TInputDataVertex<FEnumEMusicalScale>(METASOUND_GET_PARAM_NAME_AND_METADATA(ParamScaleDegreesPreset)),
+				TInputDataVertex<bool>(METASOUND_GET_PARAM_NAME_AND_METADATA(ParamChordTonesOnly))
 			),
 			FOutputVertexInterface(
-				TOutputDataVertexModel<ScaleDegreeArrayType>(METASOUND_GET_PARAM_NAME_AND_METADATA(ParamNoteArrayOutput))
+				TOutputDataVertex<ScaleDegreeArrayType>(METASOUND_GET_PARAM_NAME_AND_METADATA(ParamNoteArrayOutput))
 				)
 			);
 
diff --git a/Engine/Plugins/Runtime/Metasound/Source/MetasoundStandardNodes/Private/MetasoundMinNode.cpp b/Engine/Plugins/Runtime/Metasound/Source/MetasoundStandardNodes/Private/MetasoundMinNode.cpp
--- a/Engine/Plugins/Runtime/Metasound/Source/MetasoundStandardNodes/Private/MetasoundMinNode.cpp
+++ b/Engine/Plugins/Runtime/Metasound/Source/MetasoundStandardNodes/Private/MetasoundMinNode.cpp
@@ -138,11 +138,11 @@ namespace Metasound
 
 			static const FVertexInterface DefaultInterface(
 				FInputVertexInterface(
-					TInputDataVertexModel<ValueType>(METASOUND_GET_PARAM_NAME_AND_METADATA(InputAValue)),
-					TInputDataVertexModel<ValueType>(METASOUND_GET_PARAM_NAME_AND_METADATA(InputBValue))
+					TInputDataVertex<ValueType>(METASOUND_GET_PARAM_NAME_AND_METADATA(InputAValue)),
+					TInputDataVertex<ValueType>(METASOUND_GET_PARAM_NAME_AND_METADATA(InputBValue))
 				),
 				FOutputVertexInterface(
-					TOutputDataVertexModel<ValueType>(METASOUND_GET_PARAM_NAME_AND_METADATA(OutputValue))
+					TOutputDataVertex<ValueType>(METASOUND_GET_PARAM_NAME_AND_METADATA(OutputValue))
 				)
 			);
 
diff --git a/Engine/Plugins/Runtime/Metasound/Source/MetasoundStandardNodes/Private/MetasoundMixerNode.cpp b/Engine/Plugins/Runtime/Metasound/Source/MetasoundStandardNodes/Private/MetasoundMixerNode.cpp
--- a/Engine/Plugins/Runtime/Metasound/Source/MetasoundStandardNodes/Private/MetasoundMixerNode.cpp
+++ b/Engine/Plugins/Runtime/Metasound/Source/MetasoundStandardNodes/Private/MetasoundMixerNode.cpp
@@ -63,7 +63,7 @@ namespace Metasound
 							GetAudioInputDescription(InputIndex, ChanIndex),
 							GetAudioInputDisplayName(InputIndex, ChanIndex)
 						};
-						InputInterface.Add(TInputDataVertexModel<FAudioBuffer>(GetAudioInputName(InputIndex, ChanIndex), AudioInputMetadata));
+						InputInterface.Add(TInputDataVertex<FAudioBuffer>(GetAudioInputName(InputIndex, ChanIndex), AudioInputMetadata));
 					}
 
 					// gain scalar
@@ -72,7 +72,7 @@ namespace Metasound
 						GetGainInputDescription(InputIndex),
 						GetGainInputDisplayName(InputIndex)
 					};
-					TInputDataVertexModel<float>GainVertexModel(GetGainInputName(InputIndex), GainPinMetaData, 1.0f);
+					TInputDataVertex<float>GainVertexModel(GetGainInputName(InputIndex), GainPinMetaData, 1.0f);
 
 					InputInterface.Add(GainVertexModel);
 				}
@@ -87,7 +87,7 @@ namespace Metasound
 						GetAudioOutputDisplayName(i)
 					};
 
-					OutputInterface.Add(TOutputDataVertexModel<FAudioBuffer>(GetAudioOutputName(i), AudioOutputMetadata));
+					OutputInterface.Add(TOutputDataVertex<FAudioBuffer>(GetAudioOutputName(i), AudioOutputMetadata));
 				}
 
 				return FVertexInterface(InputInterface, OutputInterface);
@@ -393,4 +393,4 @@ namespace Metasound
 
 } // namespace Metasound
 
-#undef LOCTEXT_NAMESPACE // "MetasoundStandardNodes_MixerNode"
\ No newline at end of file
+#undef LOCTEXT_NAMESPACE // "MetasoundStandardNodes_MixerNode"
diff --git a/Engine/Plugins/Runtime/Metasound/Source/MetasoundStandardNodes/Private/MetasoundNoiseGenerator.cpp b/Engine/Plugins/Runtime/Metasound/Source/MetasoundStandardNodes/Private/MetasoundNoiseGenerator.cpp
--- a/Engine/Plugins/Runtime/Metasound/Source/MetasoundStandardNodes/Private/MetasoundNoiseGenerator.cpp
+++ b/Engine/Plugins/Runtime/Metasound/Source/MetasoundStandardNodes/Private/MetasoundNoiseGenerator.cpp
@@ -172,11 +172,11 @@ namespace Metasound
 
 		static const FVertexInterface Interface(
 			FInputVertexInterface(
-				TInputDataVertexModel<int32>(METASOUND_GET_PARAM_NAME_AND_METADATA(InputSeed), FNoiseOperator::DefaultSeed),
-				TInputDataVertexModel<FEnumNoiseType>(METASOUND_GET_PARAM_NAME_AND_METADATA(InputType))
+				TInputDataVertex<int32>(METASOUND_GET_PARAM_NAME_AND_METADATA(InputSeed), FNoiseOperator::DefaultSeed),
+				TInputDataVertex<FEnumNoiseType>(METASOUND_GET_PARAM_NAME_AND_METADATA(InputType))
 			),
 			FOutputVertexInterface(
-				TOutputDataVertexModel<FAudioBuffer>(METASOUND_GET_PARAM_NAME_AND_METADATA(OutAudio))
+				TOutputDataVertex<FAudioBuffer>(METASOUND_GET_PARAM_NAME_AND_METADATA(OutAudio))
 			)
 		);
 		return Interface;
diff --git a/Engine/Plugins/Runtime/Metasound/Source/MetasoundStandardNodes/Private/MetasoundOscillatorNodes.cpp b/Engine/Plugins/Runtime/Metasound/Source/MetasoundStandardNodes/Private/MetasoundOscillatorNodes.cpp
--- a/Engine/Plugins/Runtime/Metasound/Source/MetasoundStandardNodes/Private/MetasoundOscillatorNodes.cpp
+++ b/Engine/Plugins/Runtime/Metasound/Source/MetasoundStandardNodes/Private/MetasoundOscillatorNodes.cpp
@@ -542,16 +542,16 @@ namespace Metasound
 			static const FVertexInterface Interface
 			{
 				FInputVertexInterface{
-					TInputDataVertexModel<bool>(METASOUND_GET_PARAM_NAME_AND_METADATA(EnabledPin), true),
-					TInputDataVertexModel<bool>(METASOUND_GET_PARAM_NAME_AND_METADATA(BiPolarPin), true),
-					TInputDataVertexModel<float>(METASOUND_GET_PARAM_NAME_AND_METADATA(OscBaseFrequencyPin), 440.f),
-					TInputDataVertexModel<FAudioBuffer>(METASOUND_GET_PARAM_NAME_AND_METADATA(FrequencyModPin)),
-					TInputDataVertexModel<FTrigger>(METASOUND_GET_PARAM_NAME_AND_METADATA(OscPhaseResetPin)),
-					TInputDataVertexModel<float>(METASOUND_GET_PARAM_NAME_AND_METADATA(PhaseOffsetPin), 0.f),
-					TInputDataVertexModel<float>(METASOUND_GET_PARAM_NAME_AND_METADATA(GlideFactorPin), 0.f)
+					TInputDataVertex<bool>(METASOUND_GET_PARAM_NAME_AND_METADATA(EnabledPin), true),
+					TInputDataVertex<bool>(METASOUND_GET_PARAM_NAME_AND_METADATA(BiPolarPin), true),
+					TInputDataVertex<float>(METASOUND_GET_PARAM_NAME_AND_METADATA(OscBaseFrequencyPin), 440.f),
+					TInputDataVertex<FAudioBuffer>(METASOUND_GET_PARAM_NAME_AND_METADATA(FrequencyModPin)),
+					TInputDataVertex<FTrigger>(METASOUND_GET_PARAM_NAME_AND_METADATA(OscPhaseResetPin)),
+					TInputDataVertex<float>(METASOUND_GET_PARAM_NAME_AND_METADATA(PhaseOffsetPin), 0.f),
+					TInputDataVertex<float>(METASOUND_GET_PARAM_NAME_AND_METADATA(GlideFactorPin), 0.f)
 				},
 				FOutputVertexInterface{
-					TOutputDataVertexModel<FAudioBuffer>(METASOUND_GET_PARAM_NAME_AND_METADATA(AudioOutPin))
+					TOutputDataVertex<FAudioBuffer>(METASOUND_GET_PARAM_NAME_AND_METADATA(AudioOutPin))
 				}
 			};
 			return Interface;
@@ -797,7 +797,7 @@ namespace Metasound
 			{
 				FVertexInterface Interface = GetCommmonVertexInterface();
 				Interface.GetInputInterface().Add(
-					TInputDataVertexModel<FEnumSineGenerationType>(METASOUND_GET_PARAM_NAME_AND_METADATA(SineType), static_cast<int32>(ESineGenerationType::Wavetable))
+					TInputDataVertex<FEnumSineGenerationType>(METASOUND_GET_PARAM_NAME_AND_METADATA(SineType), static_cast<int32>(ESineGenerationType::Wavetable))
 				);
 				return Interface;
 			};
@@ -928,7 +928,7 @@ namespace Metasound
 			{
 				FVertexInterface Interface = GetCommmonVertexInterface();
 				Interface.GetInputInterface().Add(
-					TInputDataVertexModel<FEnumSawGenerationType>(METASOUND_GET_PARAM_NAME_AND_METADATA(SawType))
+					TInputDataVertex<FEnumSawGenerationType>(METASOUND_GET_PARAM_NAME_AND_METADATA(SawType))
 				);
 				return Interface;
 			};
@@ -1168,8 +1168,8 @@ namespace Metasound
 			auto MakeInterface = []() -> FVertexInterface
 			{
 				FVertexInterface Interface = GetCommmonVertexInterface();
-				Interface.GetInputInterface().Add(TInputDataVertexModel<FEnumSquareGenerationType>(METASOUND_GET_PARAM_NAME_AND_METADATA(SquareTypePin)));
-				Interface.GetInputInterface().Add(TInputDataVertexModel<float>(METASOUND_GET_PARAM_NAME_AND_METADATA(SquarePulseWidthPin), 0.5f));
+				Interface.GetInputInterface().Add(TInputDataVertex<FEnumSquareGenerationType>(METASOUND_GET_PARAM_NAME_AND_METADATA(SquareTypePin)));
+				Interface.GetInputInterface().Add(TInputDataVertex<float>(METASOUND_GET_PARAM_NAME_AND_METADATA(SquarePulseWidthPin), 0.5f));
 				return Interface;
 			};
 			static const FVertexInterface Interface = MakeInterface();
@@ -1290,7 +1290,7 @@ namespace Metasound
 			{
 				FVertexInterface Interface = GetCommmonVertexInterface();
 				Interface.GetInputInterface().Add(
-					TInputDataVertexModel<FEnumTriangleGenerationType>(METASOUND_GET_PARAM_NAME_AND_METADATA(TriangeTypePin))
+					TInputDataVertex<FEnumTriangleGenerationType>(METASOUND_GET_PARAM_NAME_AND_METADATA(TriangeTypePin))
 				);
 				return Interface;
 			};
@@ -1358,16 +1358,16 @@ namespace Metasound
 			static const FVertexInterface Interface
 			{
 				FInputVertexInterface{
-					TInputDataVertexModel<float>(METASOUND_GET_PARAM_NAME_AND_METADATA(LfoBaseFrequencyPin), 5.f),
-					TInputDataVertexModel<FEnumLfoWaveshapeType>(METASOUND_GET_PARAM_NAME_AND_METADATA(WaveshapePin)),
-					TInputDataVertexModel<float>(METASOUND_GET_PARAM_NAME_AND_METADATA(MinOutputValuePin), -1.0f),
-					TInputDataVertexModel<float>(METASOUND_GET_PARAM_NAME_AND_METADATA(MaxOutputValuePin), 1.0f),
-					TInputDataVertexModel<FTrigger>(METASOUND_GET_PARAM_NAME_AND_METADATA(LfoPhaseResetPin)),
-					TInputDataVertexModel<float>(METASOUND_GET_PARAM_NAME_AND_METADATA(PhaseOffsetPin), 0.f),
-					TInputDataVertexModel<float>(METASOUND_GET_PARAM_NAME_AND_METADATA(LfoPulseWidthPin), 0.5f)
+					TInputDataVertex<float>(METASOUND_GET_PARAM_NAME_AND_METADATA(LfoBaseFrequencyPin), 5.f),
+					TInputDataVertex<FEnumLfoWaveshapeType>(METASOUND_GET_PARAM_NAME_AND_METADATA(WaveshapePin)),
+					TInputDataVertex<float>(METASOUND_GET_PARAM_NAME_AND_METADATA(MinOutputValuePin), -1.0f),
+					TInputDataVertex<float>(METASOUND_GET_PARAM_NAME_AND_METADATA(MaxOutputValuePin), 1.0f),
+					TInputDataVertex<FTrigger>(METASOUND_GET_PARAM_NAME_AND_METADATA(LfoPhaseResetPin)),
+					TInputDataVertex<float>(METASOUND_GET_PARAM_NAME_AND_METADATA(PhaseOffsetPin), 0.f),
+					TInputDataVertex<float>(METASOUND_GET_PARAM_NAME_AND_METADATA(LfoPulseWidthPin), 0.5f)
 				},
 				FOutputVertexInterface{
-					TOutputDataVertexModel<float>(METASOUND_GET_PARAM_NAME_AND_METADATA(LfoOutPin))
+					TOutputDataVertex<float>(METASOUND_GET_PARAM_NAME_AND_METADATA(LfoOutPin))
 				}
 			};
 			return Interface;
diff --git a/Engine/Plugins/Runtime/Metasound/Source/MetasoundStandardNodes/Private/MetasoundSampleAndHoldNode.cpp b/Engine/Plugins/Runtime/Metasound/Source/MetasoundStandardNodes/Private/MetasoundSampleAndHoldNode.cpp
--- a/Engine/Plugins/Runtime/Metasound/Source/MetasoundStandardNodes/Private/MetasoundSampleAndHoldNode.cpp
+++ b/Engine/Plugins/Runtime/Metasound/Source/MetasoundStandardNodes/Private/MetasoundSampleAndHoldNode.cpp
@@ -64,12 +64,12 @@ namespace Metasound
 
 			static const FVertexInterface Interface(
 				FInputVertexInterface(
-					TInputDataVertexModel<FTrigger>(METASOUND_GET_PARAM_NAME_AND_METADATA(InputTriggerSampleAndHold)),
-					TInputDataVertexModel<FAudioBuffer>(METASOUND_GET_PARAM_NAME_AND_METADATA(InputAudio))
+					TInputDataVertex<FTrigger>(METASOUND_GET_PARAM_NAME_AND_METADATA(InputTriggerSampleAndHold)),
+					TInputDataVertex<FAudioBuffer>(METASOUND_GET_PARAM_NAME_AND_METADATA(InputAudio))
 				),
 				FOutputVertexInterface(
-					TOutputDataVertexModel<FTrigger>(METASOUND_GET_PARAM_NAME_AND_METADATA(OutputOnSampleAndHold)),
-					TOutputDataVertexModel<FAudioBuffer>(METASOUND_GET_PARAM_NAME_AND_METADATA(OutputAudio))
+					TOutputDataVertex<FTrigger>(METASOUND_GET_PARAM_NAME_AND_METADATA(OutputOnSampleAndHold)),
+					TOutputDataVertex<FAudioBuffer>(METASOUND_GET_PARAM_NAME_AND_METADATA(OutputAudio))
 				)
 			);
 
diff --git a/Engine/Plugins/Runtime/Metasound/Source/MetasoundStandardNodes/Private/MetasoundSemitoneToFreqMultiplierNode.cpp b/Engine/Plugins/Runtime/Metasound/Source/MetasoundStandardNodes/Private/MetasoundSemitoneToFreqMultiplierNode.cpp
--- a/Engine/Plugins/Runtime/Metasound/Source/MetasoundStandardNodes/Private/MetasoundSemitoneToFreqMultiplierNode.cpp
+++ b/Engine/Plugins/Runtime/Metasound/Source/MetasoundStandardNodes/Private/MetasoundSemitoneToFreqMultiplierNode.cpp
@@ -93,10 +93,10 @@ namespace Metasound
 
 		static const FVertexInterface Interface(
 			FInputVertexInterface(
-				TInputDataVertexModel<float>(METASOUND_GET_PARAM_NAME_AND_METADATA(InputSemitone), 0.0f)
+				TInputDataVertex<float>(METASOUND_GET_PARAM_NAME_AND_METADATA(InputSemitone), 0.0f)
 			),
 			FOutputVertexInterface(
-				TOutputDataVertexModel<float>(METASOUND_GET_PARAM_NAME_AND_METADATA(OutputFrequencyMultiplier))
+				TOutputDataVertex<float>(METASOUND_GET_PARAM_NAME_AND_METADATA(OutputFrequencyMultiplier))
 			)
 		);
 
diff --git a/Engine/Plugins/Runtime/Metasound/Source/MetasoundStandardNodes/Private/MetasoundStereoDelayNode.cpp b/Engine/Plugins/Runtime/Metasound/Source/MetasoundStandardNodes/Private/MetasoundStereoDelayNode.cpp
--- a/Engine/Plugins/Runtime/Metasound/Source/MetasoundStandardNodes/Private/MetasoundStereoDelayNode.cpp
+++ b/Engine/Plugins/Runtime/Metasound/Source/MetasoundStandardNodes/Private/MetasoundStereoDelayNode.cpp
@@ -291,18 +291,18 @@ namespace Metasound
 
 		static const FVertexInterface Interface(
 			FInputVertexInterface(
-				TInputDataVertexModel<FAudioBuffer>(METASOUND_GET_PARAM_NAME_AND_METADATA(InAudioInputLeft)),
-				TInputDataVertexModel<FAudioBuffer>(METASOUND_GET_PARAM_NAME_AND_METADATA(InAudioInputRight)),
-				TInputDataVertexModel<FEnumStereoDelayMode>(METASOUND_GET_PARAM_NAME_AND_METADATA(InDelayMode)),
-				TInputDataVertexModel<FTime>(METASOUND_GET_PARAM_NAME_AND_METADATA(InDelayTime), 1.0f),
-				TInputDataVertexModel<float>(METASOUND_GET_PARAM_NAME_AND_METADATA(InDelayRatio), 0.0f),
-				TInputDataVertexModel<float>(METASOUND_GET_PARAM_NAME_AND_METADATA(InDryLevel), 0.0f),
-				TInputDataVertexModel<float>(METASOUND_GET_PARAM_NAME_AND_METADATA(InWetLevel), 1.0f),
-				TInputDataVertexModel<float>(METASOUND_GET_PARAM_NAME_AND_METADATA(InFeedbackAmount), 0.0f)
+				TInputDataVertex<FAudioBuffer>(METASOUND_GET_PARAM_NAME_AND_METADATA(InAudioInputLeft)),
+				TInputDataVertex<FAudioBuffer>(METASOUND_GET_PARAM_NAME_AND_METADATA(InAudioInputRight)),
+				TInputDataVertex<FEnumStereoDelayMode>(METASOUND_GET_PARAM_NAME_AND_METADATA(InDelayMode)),
+				TInputDataVertex<FTime>(METASOUND_GET_PARAM_NAME_AND_METADATA(InDelayTime), 1.0f),
+				TInputDataVertex<float>(METASOUND_GET_PARAM_NAME_AND_METADATA(InDelayRatio), 0.0f),
+				TInputDataVertex<float>(METASOUND_GET_PARAM_NAME_AND_METADATA(InDryLevel), 0.0f),
+				TInputDataVertex<float>(METASOUND_GET_PARAM_NAME_AND_METADATA(InWetLevel), 1.0f),
+				TInputDataVertex<float>(METASOUND_GET_PARAM_NAME_AND_METADATA(InFeedbackAmount), 0.0f)
 			),
 			FOutputVertexInterface(
-				TOutputDataVertexModel<FAudioBuffer>(METASOUND_GET_PARAM_NAME_AND_METADATA(OutAudioLeft)),
-				TOutputDataVertexModel<FAudioBuffer>(METASOUND_GET_PARAM_NAME_AND_METADATA(OutAudioRight))
+				TOutputDataVertex<FAudioBuffer>(METASOUND_GET_PARAM_NAME_AND_METADATA(OutAudioLeft)),
+				TOutputDataVertex<FAudioBuffer>(METASOUND_GET_PARAM_NAME_AND_METADATA(OutAudioRight))
 			)
 		);
 
diff --git a/Engine/Plugins/Runtime/Metasound/Source/MetasoundStandardNodes/Private/MetasoundStereoPannerNode.cpp b/Engine/Plugins/Runtime/Metasound/Source/MetasoundStandardNodes/Private/MetasoundStereoPannerNode.cpp
--- a/Engine/Plugins/Runtime/Metasound/Source/MetasoundStandardNodes/Private/MetasoundStereoPannerNode.cpp
+++ b/Engine/Plugins/Runtime/Metasound/Source/MetasoundStandardNodes/Private/MetasoundStereoPannerNode.cpp
@@ -180,13 +180,13 @@ namespace Metasound
 
 		static const FVertexInterface Interface(
 			FInputVertexInterface(
-				TInputDataVertexModel<FAudioBuffer>(METASOUND_GET_PARAM_NAME_AND_METADATA(InputAudio)),
-				TInputDataVertexModel<float>(METASOUND_GET_PARAM_NAME_AND_METADATA(InputPanAmount), 0.0f),
-				TInputDataVertexModel<FEnumPanningLaw>(METASOUND_GET_PARAM_NAME_AND_METADATA(InputPanningLaw))
+				TInputDataVertex<FAudioBuffer>(METASOUND_GET_PARAM_NAME_AND_METADATA(InputAudio)),
+				TInputDataVertex<float>(METASOUND_GET_PARAM_NAME_AND_METADATA(InputPanAmount), 0.0f),
+				TInputDataVertex<FEnumPanningLaw>(METASOUND_GET_PARAM_NAME_AND_METADATA(InputPanningLaw))
 			),
 			FOutputVertexInterface(
-				TOutputDataVertexModel<FAudioBuffer>(METASOUND_GET_PARAM_NAME_AND_METADATA(OutputAudioLeft)),
-				TOutputDataVertexModel<FAudioBuffer>(METASOUND_GET_PARAM_NAME_AND_METADATA(OutputAudioRight))
+				TOutputDataVertex<FAudioBuffer>(METASOUND_GET_PARAM_NAME_AND_METADATA(OutputAudioLeft)),
+				TOutputDataVertex<FAudioBuffer>(METASOUND_GET_PARAM_NAME_AND_METADATA(OutputAudioRight))
 			)
 		);
 
diff --git a/Engine/Plugins/Runtime/Metasound/Source/MetasoundStandardNodes/Private/MetasoundTriggerAnyNode.cpp b/Engine/Plugins/Runtime/Metasound/Source/MetasoundStandardNodes/Private/MetasoundTriggerAnyNode.cpp
--- a/Engine/Plugins/Runtime/Metasound/Source/MetasoundStandardNodes/Private/MetasoundTriggerAnyNode.cpp
+++ b/Engine/Plugins/Runtime/Metasound/Source/MetasoundStandardNodes/Private/MetasoundTriggerAnyNode.cpp
@@ -82,11 +82,11 @@ namespace Metasound
 						, GetInputTriggerDisplayName(i) // display name
 					};
 
-					InputInterface.Add(TInputDataVertexModel<FTrigger>(GetInputTriggerName(i), InputTriggerMetadata));
+					InputInterface.Add(TInputDataVertex<FTrigger>(GetInputTriggerName(i), InputTriggerMetadata));
 				}
 
 				FOutputVertexInterface OutputInterface;
-				OutputInterface.Add(TOutputDataVertexModel<FTrigger>(METASOUND_GET_PARAM_NAME_AND_METADATA(OutputTrigger)));
+				OutputInterface.Add(TOutputDataVertex<FTrigger>(METASOUND_GET_PARAM_NAME_AND_METADATA(OutputTrigger)));
 
 				return FVertexInterface(InputInterface, OutputInterface);
 			};
diff --git a/Engine/Plugins/Runtime/Metasound/Source/MetasoundStandardNodes/Private/MetasoundTriggerCoinNode.cpp b/Engine/Plugins/Runtime/Metasound/Source/MetasoundStandardNodes/Private/MetasoundTriggerCoinNode.cpp
--- a/Engine/Plugins/Runtime/Metasound/Source/MetasoundStandardNodes/Private/MetasoundTriggerCoinNode.cpp
+++ b/Engine/Plugins/Runtime/Metasound/Source/MetasoundStandardNodes/Private/MetasoundTriggerCoinNode.cpp
@@ -177,14 +177,14 @@ namespace Metasound
 
 		static const FVertexInterface Interface(
 			FInputVertexInterface(
-				TInputDataVertexModel<FTrigger>(METASOUND_GET_PARAM_NAME_AND_METADATA(InputTrigger)),
-				TInputDataVertexModel<FTrigger>(METASOUND_GET_PARAM_NAME_AND_METADATA(InputReset)),
-				TInputDataVertexModel<int32>(METASOUND_GET_PARAM_NAME_AND_METADATA(InputSeed), -1),
-				TInputDataVertexModel<float>(METASOUND_GET_PARAM_NAME_AND_METADATA(InputProbability), 1.0f)
+				TInputDataVertex<FTrigger>(METASOUND_GET_PARAM_NAME_AND_METADATA(InputTrigger)),
+				TInputDataVertex<FTrigger>(METASOUND_GET_PARAM_NAME_AND_METADATA(InputReset)),
+				TInputDataVertex<int32>(METASOUND_GET_PARAM_NAME_AND_METADATA(InputSeed), -1),
+				TInputDataVertex<float>(METASOUND_GET_PARAM_NAME_AND_METADATA(InputProbability), 1.0f)
 			),
 			FOutputVertexInterface(
-				TOutputDataVertexModel<FTrigger>(METASOUND_GET_PARAM_NAME_AND_METADATA(OutputTrueTrigger)),
-				TOutputDataVertexModel<FTrigger>(METASOUND_GET_PARAM_NAME_AND_METADATA(OutputFalseTrigger))
+				TOutputDataVertex<FTrigger>(METASOUND_GET_PARAM_NAME_AND_METADATA(OutputTrueTrigger)),
+				TOutputDataVertex<FTrigger>(METASOUND_GET_PARAM_NAME_AND_METADATA(OutputFalseTrigger))
 			)
 		);
 
diff --git a/Engine/Plugins/Runtime/Metasound/Source/MetasoundStandardNodes/Private/MetasoundTriggerControlNode.cpp b/Engine/Plugins/Runtime/Metasound/Source/MetasoundStandardNodes/Private/MetasoundTriggerControlNode.cpp
--- a/Engine/Plugins/Runtime/Metasound/Source/MetasoundStandardNodes/Private/MetasoundTriggerControlNode.cpp
+++ b/Engine/Plugins/Runtime/Metasound/Source/MetasoundStandardNodes/Private/MetasoundTriggerControlNode.cpp
@@ -182,14 +182,14 @@ namespace Metasound
 
 		static const FVertexInterface Interface(
 			FInputVertexInterface(
-				TInputDataVertexModel<FTrigger>(METASOUND_GET_PARAM_NAME_AND_METADATA(InputEnter)),
-				TInputDataVertexModel<FTrigger>(METASOUND_GET_PARAM_NAME_AND_METADATA(InputOpen)),
-				TInputDataVertexModel<FTrigger>(METASOUND_GET_PARAM_NAME_AND_METADATA(InputClose)),
-				TInputDataVertexModel<FTrigger>(METASOUND_GET_PARAM_NAME_AND_METADATA(InputToggle)),
-				TInputDataVertexModel<bool>(METASOUND_GET_PARAM_NAME_AND_METADATA(InputStartClosed), true)
+				TInputDataVertex<FTrigger>(METASOUND_GET_PARAM_NAME_AND_METADATA(InputEnter)),
+				TInputDataVertex<FTrigger>(METASOUND_GET_PARAM_NAME_AND_METADATA(InputOpen)),
+				TInputDataVertex<FTrigger>(METASOUND_GET_PARAM_NAME_AND_METADATA(InputClose)),
+				TInputDataVertex<FTrigger>(METASOUND_GET_PARAM_NAME_AND_METADATA(InputToggle)),
+				TInputDataVertex<bool>(METASOUND_GET_PARAM_NAME_AND_METADATA(InputStartClosed), true)
 			),
 			FOutputVertexInterface(
-				TOutputDataVertexModel<FTrigger>(METASOUND_GET_PARAM_NAME_AND_METADATA(OutputExit))
+				TOutputDataVertex<FTrigger>(METASOUND_GET_PARAM_NAME_AND_METADATA(OutputExit))
 			)
 		);
 
diff --git a/Engine/Plugins/Runtime/Metasound/Source/MetasoundStandardNodes/Private/MetasoundTriggerCounterNode.cpp b/Engine/Plugins/Runtime/Metasound/Source/MetasoundStandardNodes/Private/MetasoundTriggerCounterNode.cpp
--- a/Engine/Plugins/Runtime/Metasound/Source/MetasoundStandardNodes/Private/MetasoundTriggerCounterNode.cpp
+++ b/Engine/Plugins/Runtime/Metasound/Source/MetasoundStandardNodes/Private/MetasoundTriggerCounterNode.cpp
@@ -191,17 +191,17 @@ namespace Metasound
 
 		static const FVertexInterface Interface(
 			FInputVertexInterface(
-				TInputDataVertexModel<FTrigger>(METASOUND_GET_PARAM_NAME_AND_METADATA(InputInTrigger)),
-				TInputDataVertexModel<FTrigger>(METASOUND_GET_PARAM_NAME_AND_METADATA(InputReset)),
-				TInputDataVertexModel<float>(METASOUND_GET_PARAM_NAME_AND_METADATA(InputStartValue), 0.0f),
-				TInputDataVertexModel<float>(METASOUND_GET_PARAM_NAME_AND_METADATA(InputStepSize), 1.0f),
-				TInputDataVertexModel<int32>(METASOUND_GET_PARAM_NAME_AND_METADATA(InputAutoResetCount), 0)
+				TInputDataVertex<FTrigger>(METASOUND_GET_PARAM_NAME_AND_METADATA(InputInTrigger)),
+				TInputDataVertex<FTrigger>(METASOUND_GET_PARAM_NAME_AND_METADATA(InputReset)),
+				TInputDataVertex<float>(METASOUND_GET_PARAM_NAME_AND_METADATA(InputStartValue), 0.0f),
+				TInputDataVertex<float>(METASOUND_GET_PARAM_NAME_AND_METADATA(InputStepSize), 1.0f),
+				TInputDataVertex<int32>(METASOUND_GET_PARAM_NAME_AND_METADATA(InputAutoResetCount), 0)
 			),
 			FOutputVertexInterface(
-				TOutputDataVertexModel<FTrigger>(METASOUND_GET_PARAM_NAME_AND_METADATA(OutputOnTrigger)),
-				TOutputDataVertexModel<FTrigger>(METASOUND_GET_PARAM_NAME_AND_METADATA(OutputOnReset)),
-				TOutputDataVertexModel<int32>(METASOUND_GET_PARAM_NAME_AND_METADATA(OutputCount)),
-				TOutputDataVertexModel<float>(METASOUND_GET_PARAM_NAME_AND_METADATA(OutputValue))
+				TOutputDataVertex<FTrigger>(METASOUND_GET_PARAM_NAME_AND_METADATA(OutputOnTrigger)),
+				TOutputDataVertex<FTrigger>(METASOUND_GET_PARAM_NAME_AND_METADATA(OutputOnReset)),
+				TOutputDataVertex<int32>(METASOUND_GET_PARAM_NAME_AND_METADATA(OutputCount)),
+				TOutputDataVertex<float>(METASOUND_GET_PARAM_NAME_AND_METADATA(OutputValue))
 			)
 		);
 
diff --git a/Engine/Plugins/Runtime/Metasound/Source/MetasoundStandardNodes/Private/MetasoundTriggerDelayNode.cpp b/Engine/Plugins/Runtime/Metasound/Source/MetasoundStandardNodes/Private/MetasoundTriggerDelayNode.cpp
--- a/Engine/Plugins/Runtime/Metasound/Source/MetasoundStandardNodes/Private/MetasoundTriggerDelayNode.cpp
+++ b/Engine/Plugins/Runtime/Metasound/Source/MetasoundStandardNodes/Private/MetasoundTriggerDelayNode.cpp
@@ -155,12 +155,12 @@ namespace Metasound
 
 		static const FVertexInterface Interface(
 			FInputVertexInterface(
-				TInputDataVertexModel<FTrigger>(METASOUND_GET_PARAM_NAME_AND_METADATA(InputInTriggerDelay)),
-				TInputDataVertexModel<FTrigger>(METASOUND_GET_PARAM_NAME_AND_METADATA(InputResetDelay)),
-				TInputDataVertexModel<FTime>(METASOUND_GET_PARAM_NAME_AND_METADATA(InputDelayTime), 1.0f)
+				TInputDataVertex<FTrigger>(METASOUND_GET_PARAM_NAME_AND_METADATA(InputInTriggerDelay)),
+				TInputDataVertex<FTrigger>(METASOUND_GET_PARAM_NAME_AND_METADATA(InputResetDelay)),
+				TInputDataVertex<FTime>(METASOUND_GET_PARAM_NAME_AND_METADATA(InputDelayTime), 1.0f)
 			),
 			FOutputVertexInterface(
-				TOutputDataVertexModel<FTrigger>(METASOUND_GET_PARAM_NAME_AND_METADATA(OutputOnTrigger))
+				TOutputDataVertex<FTrigger>(METASOUND_GET_PARAM_NAME_AND_METADATA(OutputOnTrigger))
 		)
 		);
 
diff --git a/Engine/Plugins/Runtime/Metasound/Source/MetasoundStandardNodes/Private/MetasoundTriggerOnThresholdNode.cpp b/Engine/Plugins/Runtime/Metasound/Source/MetasoundStandardNodes/Private/MetasoundTriggerOnThresholdNode.cpp
--- a/Engine/Plugins/Runtime/Metasound/Source/MetasoundStandardNodes/Private/MetasoundTriggerOnThresholdNode.cpp
+++ b/Engine/Plugins/Runtime/Metasound/Source/MetasoundStandardNodes/Private/MetasoundTriggerOnThresholdNode.cpp
@@ -192,12 +192,12 @@ namespace Metasound
 
 		static const FVertexInterface Interface(
 			FInputVertexInterface(
-				TInputDataVertexModel<FAudioBuffer>(METASOUND_GET_PARAM_NAME_AND_METADATA(InPin)),
-				TInputDataVertexModel<float>(METASOUND_GET_PARAM_NAME_AND_METADATA(InThresholdPin), DefaultThreshold),
-				TInputDataVertexModel<FEnumBufferTriggerType>(METASOUND_GET_PARAM_NAME_AND_METADATA(InTriggerType))
+				TInputDataVertex<FAudioBuffer>(METASOUND_GET_PARAM_NAME_AND_METADATA(InPin)),
+				TInputDataVertex<float>(METASOUND_GET_PARAM_NAME_AND_METADATA(InThresholdPin), DefaultThreshold),
+				TInputDataVertex<FEnumBufferTriggerType>(METASOUND_GET_PARAM_NAME_AND_METADATA(InTriggerType))
 			),
 			FOutputVertexInterface(
-				TOutputDataVertexModel<FTrigger>(METASOUND_GET_PARAM_NAME_AND_METADATA(OutPin))
+				TOutputDataVertex<FTrigger>(METASOUND_GET_PARAM_NAME_AND_METADATA(OutPin))
 			)
 		);
 		return Interface;
diff --git a/Engine/Plugins/Runtime/Metasound/Source/MetasoundStandardNodes/Private/MetasoundTriggerOnValueChangeNode.cpp b/Engine/Plugins/Runtime/Metasound/Source/MetasoundStandardNodes/Private/MetasoundTriggerOnValueChangeNode.cpp
--- a/Engine/Plugins/Runtime/Metasound/Source/MetasoundStandardNodes/Private/MetasoundTriggerOnValueChangeNode.cpp
+++ b/Engine/Plugins/Runtime/Metasound/Source/MetasoundStandardNodes/Private/MetasoundTriggerOnValueChangeNode.cpp
@@ -132,10 +132,10 @@ namespace Metasound
 
 		static const FVertexInterface Interface(
 			FInputVertexInterface(
-				TInputDataVertexModel<ValueType>(METASOUND_GET_PARAM_NAME_AND_METADATA(InputValue), 1)
+				TInputDataVertex<ValueType>(METASOUND_GET_PARAM_NAME_AND_METADATA(InputValue), 1)
 			),
 			FOutputVertexInterface(
-				TOutputDataVertexModel<FTrigger>(METASOUND_GET_PARAM_NAME_AND_METADATA(OutputOnChange))
+				TOutputDataVertex<FTrigger>(METASOUND_GET_PARAM_NAME_AND_METADATA(OutputOnChange))
 			)
 		);
 
diff --git a/Engine/Plugins/Runtime/Metasound/Source/MetasoundStandardNodes/Private/MetasoundTriggerOnceNode.cpp b/Engine/Plugins/Runtime/Metasound/Source/MetasoundStandardNodes/Private/MetasoundTriggerOnceNode.cpp
--- a/Engine/Plugins/Runtime/Metasound/Source/MetasoundStandardNodes/Private/MetasoundTriggerOnceNode.cpp
+++ b/Engine/Plugins/Runtime/Metasound/Source/MetasoundStandardNodes/Private/MetasoundTriggerOnceNode.cpp
@@ -146,12 +146,12 @@ namespace Metasound
 
 		static const FVertexInterface Interface(
 			FInputVertexInterface(
-				TInputDataVertexModel<FTrigger>(METASOUND_GET_PARAM_NAME_AND_METADATA(InputEnter)),
-				TInputDataVertexModel<FTrigger>(METASOUND_GET_PARAM_NAME_AND_METADATA(InputReset)),
-				TInputDataVertexModel<bool>(METASOUND_GET_PARAM_NAME_AND_METADATA(InputStartClosed), false)
+				TInputDataVertex<FTrigger>(METASOUND_GET_PARAM_NAME_AND_METADATA(InputEnter)),
+				TInputDataVertex<FTrigger>(METASOUND_GET_PARAM_NAME_AND_METADATA(InputReset)),
+				TInputDataVertex<bool>(METASOUND_GET_PARAM_NAME_AND_METADATA(InputStartClosed), false)
 			),
 			FOutputVertexInterface(
-				TOutputDataVertexModel<FTrigger>(METASOUND_GET_PARAM_NAME_AND_METADATA(OutputExit))
+				TOutputDataVertex<FTrigger>(METASOUND_GET_PARAM_NAME_AND_METADATA(OutputExit))
 			)
 		);
 
diff --git a/Engine/Plugins/Runtime/Metasound/Source/MetasoundStandardNodes/Private/MetasoundTriggerPipeNode.cpp b/Engine/Plugins/Runtime/Metasound/Source/MetasoundStandardNodes/Private/MetasoundTriggerPipeNode.cpp
--- a/Engine/Plugins/Runtime/Metasound/Source/MetasoundStandardNodes/Private/MetasoundTriggerPipeNode.cpp
+++ b/Engine/Plugins/Runtime/Metasound/Source/MetasoundStandardNodes/Private/MetasoundTriggerPipeNode.cpp
@@ -155,12 +155,12 @@ namespace Metasound
 
 		static const FVertexInterface Interface(
 			FInputVertexInterface(
-				TInputDataVertexModel<FTrigger>(METASOUND_GET_PARAM_NAME_AND_METADATA(InputInTrigger)),
-				TInputDataVertexModel<FTrigger>(METASOUND_GET_PARAM_NAME_AND_METADATA(InputReset)),
-				TInputDataVertexModel<FTime>(METASOUND_GET_PARAM_NAME_AND_METADATA(InputDelayTime), 1.0f)
+				TInputDataVertex<FTrigger>(METASOUND_GET_PARAM_NAME_AND_METADATA(InputInTrigger)),
+				TInputDataVertex<FTrigger>(METASOUND_GET_PARAM_NAME_AND_METADATA(InputReset)),
+				TInputDataVertex<FTime>(METASOUND_GET_PARAM_NAME_AND_METADATA(InputDelayTime), 1.0f)
 				),
 			FOutputVertexInterface(
-				TOutputDataVertexModel<FTrigger>(METASOUND_GET_PARAM_NAME_AND_METADATA(OutputOutTrigger))
+				TOutputDataVertex<FTrigger>(METASOUND_GET_PARAM_NAME_AND_METADATA(OutputOutTrigger))
 			)
 		);
 
diff --git a/Engine/Plugins/Runtime/Metasound/Source/MetasoundStandardNodes/Private/MetasoundTriggerRepeatNode.cpp b/Engine/Plugins/Runtime/Metasound/Source/MetasoundStandardNodes/Private/MetasoundTriggerRepeatNode.cpp
--- a/Engine/Plugins/Runtime/Metasound/Source/MetasoundStandardNodes/Private/MetasoundTriggerRepeatNode.cpp
+++ b/Engine/Plugins/Runtime/Metasound/Source/MetasoundStandardNodes/Private/MetasoundTriggerRepeatNode.cpp
@@ -152,12 +152,12 @@ namespace Metasound
 
 		static const FVertexInterface Interface(
 			FInputVertexInterface(
-				TInputDataVertexModel<FTrigger>(METASOUND_GET_PARAM_NAME_AND_METADATA(InputStart)),
-				TInputDataVertexModel<FTrigger>(METASOUND_GET_PARAM_NAME_AND_METADATA(InputStop)),
-				TInputDataVertexModel<FTime>(METASOUND_GET_PARAM_NAME_AND_METADATA(InputPeriod), 0.2f)
+				TInputDataVertex<FTrigger>(METASOUND_GET_PARAM_NAME_AND_METADATA(InputStart)),
+				TInputDataVertex<FTrigger>(METASOUND_GET_PARAM_NAME_AND_METADATA(InputStop)),
+				TInputDataVertex<FTime>(METASOUND_GET_PARAM_NAME_AND_METADATA(InputPeriod), 0.2f)
 			),
 			FOutputVertexInterface(
-				TOutputDataVertexModel<FTrigger>(METASOUND_GET_PARAM_NAME_AND_METADATA(RepeatOutputOnTrigger))
+				TOutputDataVertex<FTrigger>(METASOUND_GET_PARAM_NAME_AND_METADATA(RepeatOutputOnTrigger))
 			)
 		);
 
diff --git a/Engine/Plugins/Runtime/Metasound/Source/MetasoundStandardNodes/Private/MetasoundTriggerRouteNode.cpp b/Engine/Plugins/Runtime/Metasound/Source/MetasoundStandardNodes/Private/MetasoundTriggerRouteNode.cpp
--- a/Engine/Plugins/Runtime/Metasound/Source/MetasoundStandardNodes/Private/MetasoundTriggerRouteNode.cpp
+++ b/Engine/Plugins/Runtime/Metasound/Source/MetasoundStandardNodes/Private/MetasoundTriggerRouteNode.cpp
@@ -113,13 +113,13 @@ namespace Metasound
 						, GetInputValueDisplayName(i) // display name
 					};
 
-					InputInterface.Add(TInputDataVertexModel<FTrigger>(GetInputTriggerName(i), InputTriggerMetadata));
-					InputInterface.Add(TInputDataVertexModel<ValueType>(GetInputValueName(i), InputValueMetadata));
+					InputInterface.Add(TInputDataVertex<FTrigger>(GetInputTriggerName(i), InputTriggerMetadata));
+					InputInterface.Add(TInputDataVertex<ValueType>(GetInputValueName(i), InputValueMetadata));
 				}
 
 				FOutputVertexInterface OutputInterface;
-				OutputInterface.Add(TOutputDataVertexModel<FTrigger>(METASOUND_GET_PARAM_NAME_AND_METADATA(OutputTrigger)));
-				OutputInterface.Add(TOutputDataVertexModel<ValueType>(METASOUND_GET_PARAM_NAME_AND_METADATA(OutputValue)));
+				OutputInterface.Add(TOutputDataVertex<FTrigger>(METASOUND_GET_PARAM_NAME_AND_METADATA(OutputTrigger)));
+				OutputInterface.Add(TOutputDataVertex<ValueType>(METASOUND_GET_PARAM_NAME_AND_METADATA(OutputValue)));
 				return FVertexInterface(InputInterface, OutputInterface);
 			};
 
diff --git a/Engine/Plugins/Runtime/Metasound/Source/MetasoundStandardNodes/Private/MetasoundTriggerSequenceNode.cpp b/Engine/Plugins/Runtime/Metasound/Source/MetasoundStandardNodes/Private/MetasoundTriggerSequenceNode.cpp
--- a/Engine/Plugins/Runtime/Metasound/Source/MetasoundStandardNodes/Private/MetasoundTriggerSequenceNode.cpp
+++ b/Engine/Plugins/Runtime/Metasound/Source/MetasoundStandardNodes/Private/MetasoundTriggerSequenceNode.cpp
@@ -171,14 +171,14 @@ namespace Metasound
 		auto CreateVertexInterface = []() -> FVertexInterface
 		{
 			FInputVertexInterface InputInterface;
-			InputInterface.Add(TInputDataVertexModel<FTrigger>(METASOUND_GET_PARAM_NAME_AND_METADATA(InputTrigger)));
-			InputInterface.Add(TInputDataVertexModel<FTrigger>(METASOUND_GET_PARAM_NAME_AND_METADATA(InputReset)));
-			InputInterface.Add(TInputDataVertexModel<bool>(METASOUND_GET_PARAM_NAME_AND_METADATA(InputLoop), true));
+			InputInterface.Add(TInputDataVertex<FTrigger>(METASOUND_GET_PARAM_NAME_AND_METADATA(InputTrigger)));
+			InputInterface.Add(TInputDataVertex<FTrigger>(METASOUND_GET_PARAM_NAME_AND_METADATA(InputReset)));
+			InputInterface.Add(TInputDataVertex<bool>(METASOUND_GET_PARAM_NAME_AND_METADATA(InputLoop), true));
 			
 			FOutputVertexInterface OutputInterface;
 			for (uint32 i = 0; i < NumOutputs; ++i)
 			{
-				OutputInterface.Add(TOutputDataVertexModel<FTrigger>(METASOUND_GET_PARAM_NAME_WITH_INDEX_AND_TT(OutputTrigger, i)));
+				OutputInterface.Add(TOutputDataVertex<FTrigger>(METASOUND_GET_PARAM_NAME_WITH_INDEX_AND_METADATA(OutputTrigger, i)));
 			}
 			return FVertexInterface(InputInterface, OutputInterface);
 		};
diff --git a/Engine/Plugins/Runtime/Metasound/Source/MetasoundStandardNodes/Private/MetasoundTriggerToggleNode.cpp b/Engine/Plugins/Runtime/Metasound/Source/MetasoundStandardNodes/Private/MetasoundTriggerToggleNode.cpp
--- a/Engine/Plugins/Runtime/Metasound/Source/MetasoundStandardNodes/Private/MetasoundTriggerToggleNode.cpp
+++ b/Engine/Plugins/Runtime/Metasound/Source/MetasoundStandardNodes/Private/MetasoundTriggerToggleNode.cpp
@@ -129,13 +129,13 @@ namespace Metasound
 
 		static const FVertexInterface Interface(
 			FInputVertexInterface(
-				TInputDataVertexModel<FTrigger>(METASOUND_GET_PARAM_NAME_AND_METADATA(InputOnTrigger)),
-				TInputDataVertexModel<FTrigger>(METASOUND_GET_PARAM_NAME_AND_METADATA(InputOffTrigger)),
-				TInputDataVertexModel<bool>(METASOUND_GET_PARAM_NAME_AND_METADATA(InputInit))
+				TInputDataVertex<FTrigger>(METASOUND_GET_PARAM_NAME_AND_METADATA(InputOnTrigger)),
+				TInputDataVertex<FTrigger>(METASOUND_GET_PARAM_NAME_AND_METADATA(InputOffTrigger)),
+				TInputDataVertex<bool>(METASOUND_GET_PARAM_NAME_AND_METADATA(InputInit))
 			),
 			FOutputVertexInterface(
-				TOutputDataVertexModel<FTrigger>(METASOUND_GET_PARAM_NAME_AND_METADATA(OutputTrigger)),
-				TOutputDataVertexModel<bool>(METASOUND_GET_PARAM_NAME_AND_METADATA(OutputValue))
+				TOutputDataVertex<FTrigger>(METASOUND_GET_PARAM_NAME_AND_METADATA(OutputTrigger)),
+				TOutputDataVertex<bool>(METASOUND_GET_PARAM_NAME_AND_METADATA(OutputValue))
 			)
 		);
 
diff --git a/Engine/Plugins/Runtime/Metasound/Source/MetasoundStandardNodes/Private/MetasoundWaveShaperNode.cpp b/Engine/Plugins/Runtime/Metasound/Source/MetasoundStandardNodes/Private/MetasoundWaveShaperNode.cpp
--- a/Engine/Plugins/Runtime/Metasound/Source/MetasoundStandardNodes/Private/MetasoundWaveShaperNode.cpp
+++ b/Engine/Plugins/Runtime/Metasound/Source/MetasoundStandardNodes/Private/MetasoundWaveShaperNode.cpp
@@ -137,14 +137,14 @@ namespace Metasound
 	{
 		static const FVertexInterface Interface(
 			FInputVertexInterface(
-				TInputDataVertexModel<FAudioBuffer>(METASOUND_GET_PARAM_NAME_AND_METADATA(InAudioInput)),
-				TInputDataVertexModel<float>(METASOUND_GET_PARAM_NAME_AND_METADATA(InWaveShapeAmount), 1.0f),
-				TInputDataVertexModel<float>(METASOUND_GET_PARAM_NAME_AND_METADATA(InWaveShapeBias), 0.0f),
-				TInputDataVertexModel<float>(METASOUND_GET_PARAM_NAME_AND_METADATA(InOutputGain), 1.0f),
-				TInputDataVertexModel<FEnumEWaveShaperType>(METASOUND_GET_PARAM_NAME_AND_METADATA(InType))
+				TInputDataVertex<FAudioBuffer>(METASOUND_GET_PARAM_NAME_AND_METADATA(InAudioInput)),
+				TInputDataVertex<float>(METASOUND_GET_PARAM_NAME_AND_METADATA(InWaveShapeAmount), 1.0f),
+				TInputDataVertex<float>(METASOUND_GET_PARAM_NAME_AND_METADATA(InWaveShapeBias), 0.0f),
+				TInputDataVertex<float>(METASOUND_GET_PARAM_NAME_AND_METADATA(InOutputGain), 1.0f),
+				TInputDataVertex<FEnumEWaveShaperType>(METASOUND_GET_PARAM_NAME_AND_METADATA(InType))
 			),
 			FOutputVertexInterface(
-				TOutputDataVertexModel<FAudioBuffer>(METASOUND_GET_PARAM_NAME_AND_METADATA(OutAudio))
+				TOutputDataVertex<FAudioBuffer>(METASOUND_GET_PARAM_NAME_AND_METADATA(OutAudio))
 			)
 		);
 
diff --git a/Engine/Plugins/Runtime/Metasound/Source/MetasoundStandardNodes/Private/MetasoundWaveWriterNode.cpp b/Engine/Plugins/Runtime/Metasound/Source/MetasoundStandardNodes/Private/MetasoundWaveWriterNode.cpp
--- a/Engine/Plugins/Runtime/Metasound/Source/MetasoundStandardNodes/Private/MetasoundWaveWriterNode.cpp
+++ b/Engine/Plugins/Runtime/Metasound/Source/MetasoundStandardNodes/Private/MetasoundWaveWriterNode.cpp
@@ -271,8 +271,8 @@ namespace Metasound
 
 				// inputs
 				FInputVertexInterface InputInterface(
-					TInputDataVertexModel<FString>(METASOUND_GET_PARAM_NAME_AND_METADATA(InFilenamePrefixPin), GetDefaultFileName()),
-					TInputDataVertexModel<bool>(METASOUND_GET_PARAM_NAME_AND_METADATA(InEnabledPin), true)
+					TInputDataVertex<FString>(METASOUND_GET_PARAM_NAME_AND_METADATA(InFilenamePrefixPin), GetDefaultFileName()),
+					TInputDataVertex<bool>(METASOUND_GET_PARAM_NAME_AND_METADATA(InEnabledPin), true)
 				);
 
 
@@ -285,7 +285,7 @@ namespace Metasound
 						, GetAudioInputDisplayName(InputIndex) // display name
 					};
 
-					InputInterface.Add(TInputDataVertexModel<FAudioBuffer>(GetAudioInputName(InputIndex), AudioInputMetadata));
+					InputInterface.Add(TInputDataVertex<FAudioBuffer>(GetAudioInputName(InputIndex), AudioInputMetadata));
 				}
 				FOutputVertexInterface OutputInterface;
 
diff --git a/Engine/Plugins/Runtime/Metasound/Source/MetasoundStandardNodes/Public/MetasoundRandomNode.h b/Engine/Plugins/Runtime/Metasound/Source/MetasoundStandardNodes/Public/MetasoundRandomNode.h
--- a/Engine/Plugins/Runtime/Metasound/Source/MetasoundStandardNodes/Public/MetasoundRandomNode.h
+++ b/Engine/Plugins/Runtime/Metasound/Source/MetasoundStandardNodes/Public/MetasoundRandomNode.h
@@ -263,16 +263,16 @@ namespace Metasound
 			{
 				static const FVertexInterface DefaultInterface(
 					FInputVertexInterface(
-						TInputDataVertexModel<FTrigger>(METASOUND_GET_PARAM_NAME_AND_METADATA(InputNextTrigger)),
-						TInputDataVertexModel<FTrigger>(METASOUND_GET_PARAM_NAME_AND_METADATA(InputResetTrigger)),
-						TInputDataVertexModel<int32>(METASOUND_GET_PARAM_NAME_AND_METADATA(InputSeed), DefaultSeed),
-						TInputDataVertexModel<ValueType>(METASOUND_GET_PARAM_NAME_AND_METADATA(InputMin), TRandomNodeSpecialization<ValueType>::GetDefaultMin()),
-						TInputDataVertexModel<ValueType>(METASOUND_GET_PARAM_NAME_AND_METADATA(InputMax), TRandomNodeSpecialization<ValueType>::GetDefaultMax())
+						TInputDataVertex<FTrigger>(METASOUND_GET_PARAM_NAME_AND_METADATA(InputNextTrigger)),
+						TInputDataVertex<FTrigger>(METASOUND_GET_PARAM_NAME_AND_METADATA(InputResetTrigger)),
+						TInputDataVertex<int32>(METASOUND_GET_PARAM_NAME_AND_METADATA(InputSeed), DefaultSeed),
+						TInputDataVertex<ValueType>(METASOUND_GET_PARAM_NAME_AND_METADATA(InputMin), TRandomNodeSpecialization<ValueType>::GetDefaultMin()),
+						TInputDataVertex<ValueType>(METASOUND_GET_PARAM_NAME_AND_METADATA(InputMax), TRandomNodeSpecialization<ValueType>::GetDefaultMax())
 					),
 					FOutputVertexInterface(
-						TOutputDataVertexModel<FTrigger>(METASOUND_GET_PARAM_NAME_AND_METADATA(OutputOnNextTrigger)),
-						TOutputDataVertexModel<FTrigger>(METASOUND_GET_PARAM_NAME_AND_METADATA(OutputOnResetTrigger)),
-						TOutputDataVertexModel<ValueType>(METASOUND_GET_PARAM_NAME_AND_METADATA(OutputValue))
+						TOutputDataVertex<FTrigger>(METASOUND_GET_PARAM_NAME_AND_METADATA(OutputOnNextTrigger)),
+						TOutputDataVertex<FTrigger>(METASOUND_GET_PARAM_NAME_AND_METADATA(OutputOnResetTrigger)),
+						TOutputDataVertex<ValueType>(METASOUND_GET_PARAM_NAME_AND_METADATA(OutputValue))
 					)
 				);
 
@@ -282,14 +282,14 @@ namespace Metasound
 			{
 				static const FVertexInterface DefaultInterface(
 					FInputVertexInterface(
-						TInputDataVertexModel<FTrigger>(METASOUND_GET_PARAM_NAME_AND_METADATA(InputNextTrigger)),
-						TInputDataVertexModel<FTrigger>(METASOUND_GET_PARAM_NAME_AND_METADATA(InputResetTrigger)),
-						TInputDataVertexModel<int32>(METASOUND_GET_PARAM_NAME_AND_METADATA(InputSeed), DefaultSeed)
+						TInputDataVertex<FTrigger>(METASOUND_GET_PARAM_NAME_AND_METADATA(InputNextTrigger)),
+						TInputDataVertex<FTrigger>(METASOUND_GET_PARAM_NAME_AND_METADATA(InputResetTrigger)),
+						TInputDataVertex<int32>(METASOUND_GET_PARAM_NAME_AND_METADATA(InputSeed), DefaultSeed)
 					),
 					FOutputVertexInterface(
-						TOutputDataVertexModel<FTrigger>(METASOUND_GET_PARAM_NAME_AND_METADATA(OutputOnNextTrigger)),
-						TOutputDataVertexModel<FTrigger>(METASOUND_GET_PARAM_NAME_AND_METADATA(OutputOnResetTrigger)),
-						TOutputDataVertexModel<ValueType>(METASOUND_GET_PARAM_NAME_AND_METADATA(OutputValue))
+						TOutputDataVertex<FTrigger>(METASOUND_GET_PARAM_NAME_AND_METADATA(OutputOnNextTrigger)),
+						TOutputDataVertex<FTrigger>(METASOUND_GET_PARAM_NAME_AND_METADATA(OutputOnResetTrigger)),
+						TOutputDataVertex<ValueType>(METASOUND_GET_PARAM_NAME_AND_METADATA(OutputValue))
 					)
 				);
 
diff --git a/Engine/Plugins/Runtime/Metasound/Source/MetasoundStandardNodes/Public/MetasoundTriggerAccumulatorNode.h b/Engine/Plugins/Runtime/Metasound/Source/MetasoundStandardNodes/Public/MetasoundTriggerAccumulatorNode.h
--- a/Engine/Plugins/Runtime/Metasound/Source/MetasoundStandardNodes/Public/MetasoundTriggerAccumulatorNode.h
+++ b/Engine/Plugins/Runtime/Metasound/Source/MetasoundStandardNodes/Public/MetasoundTriggerAccumulatorNode.h
@@ -41,13 +41,13 @@ namespace Metasound
 
 				for (uint32 i = 0; i < NumInputs; ++i)
 				{
-					InputInterface.Add(TInputDataVertexModel<FTrigger>(METASOUND_GET_PARAM_NAME_WITH_INDEX_AND_TT(InputTrigger, i)));
+					InputInterface.Add(TInputDataVertex<FTrigger>(METASOUND_GET_PARAM_NAME_WITH_INDEX_AND_METADATA(InputTrigger, i)));
 				}
 
-				InputInterface.Add(TInputDataVertexModel<bool>(METASOUND_GET_PARAM_NAME_AND_METADATA(InputAutoReset)));
+				InputInterface.Add(TInputDataVertex<bool>(METASOUND_GET_PARAM_NAME_AND_METADATA(InputAutoReset)));
 
 				FOutputVertexInterface OutputInterface;
-				OutputInterface.Add(TOutputDataVertexModel<FTrigger>(METASOUND_GET_PARAM_NAME_AND_METADATA(AccumulateOutputOnTrigger)));
+				OutputInterface.Add(TOutputDataVertex<FTrigger>(METASOUND_GET_PARAM_NAME_AND_METADATA(AccumulateOutputOnTrigger)));
 
 				return FVertexInterface(InputInterface, OutputInterface);
 			};
diff --git a/Engine/Plugins/Runtime/Metasound/Source/MetasoundStandardNodes/Public/MetasoundTriggerCompareNode.h b/Engine/Plugins/Runtime/Metasound/Source/MetasoundStandardNodes/Public/MetasoundTriggerCompareNode.h
--- a/Engine/Plugins/Runtime/Metasound/Source/MetasoundStandardNodes/Public/MetasoundTriggerCompareNode.h
+++ b/Engine/Plugins/Runtime/Metasound/Source/MetasoundStandardNodes/Public/MetasoundTriggerCompareNode.h
@@ -71,14 +71,14 @@ namespace Metasound
 
 			static const FVertexInterface DefaultInterface(
 				FInputVertexInterface(
-					TInputDataVertexModel<FTrigger>(METASOUND_GET_PARAM_NAME_AND_METADATA(InputCompare)),
-					TInputDataVertexModel<ValueType>(METASOUND_GET_PARAM_NAME_AND_METADATA(InputParamA)),
-					TInputDataVertexModel<ValueType>(METASOUND_GET_PARAM_NAME_AND_METADATA(InputParamB)),
-					TInputDataVertexModel<FEnumTriggerComparisonType>(METASOUND_GET_PARAM_NAME_AND_METADATA(InputCompareType))
+					TInputDataVertex<FTrigger>(METASOUND_GET_PARAM_NAME_AND_METADATA(InputCompare)),
+					TInputDataVertex<ValueType>(METASOUND_GET_PARAM_NAME_AND_METADATA(InputParamA)),
+					TInputDataVertex<ValueType>(METASOUND_GET_PARAM_NAME_AND_METADATA(InputParamB)),
+					TInputDataVertex<FEnumTriggerComparisonType>(METASOUND_GET_PARAM_NAME_AND_METADATA(InputCompareType))
 				),
 				FOutputVertexInterface(
-					TOutputDataVertexModel<FTrigger>(METASOUND_GET_PARAM_NAME_AND_METADATA(OutputOnTrue)),
-					TOutputDataVertexModel<FTrigger>(METASOUND_GET_PARAM_NAME_AND_METADATA(OutputOnFalse))
+					TOutputDataVertex<FTrigger>(METASOUND_GET_PARAM_NAME_AND_METADATA(OutputOnTrue)),
+					TOutputDataVertex<FTrigger>(METASOUND_GET_PARAM_NAME_AND_METADATA(OutputOnFalse))
 				)
 			);
 
diff --git a/Engine/Plugins/Runtime/Metasound/Source/MetasoundStandardNodes/Public/MetasoundValueNode.h b/Engine/Plugins/Runtime/Metasound/Source/MetasoundStandardNodes/Public/MetasoundValueNode.h
--- a/Engine/Plugins/Runtime/Metasound/Source/MetasoundStandardNodes/Public/MetasoundValueNode.h
+++ b/Engine/Plugins/Runtime/Metasound/Source/MetasoundStandardNodes/Public/MetasoundValueNode.h
@@ -42,15 +42,15 @@ namespace Metasound
 
 			static const FVertexInterface DefaultInterface(
 				FInputVertexInterface(
-					TInputDataVertexModel<FTrigger>(METASOUND_GET_PARAM_NAME_AND_METADATA(InputSetTrigger)),
-					TInputDataVertexModel<FTrigger>(METASOUND_GET_PARAM_NAME_AND_METADATA(InputResetTrigger)),
-					TInputDataVertexModel<ValueType>(METASOUND_GET_PARAM_NAME_AND_METADATA(InputInitValue)),
-					TInputDataVertexModel<ValueType>(METASOUND_GET_PARAM_NAME_AND_METADATA(InputTargetValue))
+					TInputDataVertex<FTrigger>(METASOUND_GET_PARAM_NAME_AND_METADATA(InputSetTrigger)),
+					TInputDataVertex<FTrigger>(METASOUND_GET_PARAM_NAME_AND_METADATA(InputResetTrigger)),
+					TInputDataVertex<ValueType>(METASOUND_GET_PARAM_NAME_AND_METADATA(InputInitValue)),
+					TInputDataVertex<ValueType>(METASOUND_GET_PARAM_NAME_AND_METADATA(InputTargetValue))
 				),
 				FOutputVertexInterface(
-					TOutputDataVertexModel<FTrigger>(METASOUND_GET_PARAM_NAME_AND_METADATA(OutputOnSet)),
-					TOutputDataVertexModel<FTrigger>(METASOUND_GET_PARAM_NAME_AND_METADATA(OutputOnReset)),
-					TOutputDataVertexModel<ValueType>(METASOUND_GET_PARAM_NAME_AND_METADATA(OutputValue))
+					TOutputDataVertex<FTrigger>(METASOUND_GET_PARAM_NAME_AND_METADATA(OutputOnSet)),
+					TOutputDataVertex<FTrigger>(METASOUND_GET_PARAM_NAME_AND_METADATA(OutputOnReset)),
+					TOutputDataVertex<ValueType>(METASOUND_GET_PARAM_NAME_AND_METADATA(OutputValue))
 				)
 			);
 
