Fix for crash in metasound value node using invalid literal

diff --git a/Engine/Plugins/Runtime/Metasound/Source/MetasoundGraphCore/Public/MetasoundVertex.h b/Engine/Plugins/Runtime/Metasound/Source/MetasoundGraphCore/Public/MetasoundVertex.h
--- a/Engine/Plugins/Runtime/Metasound/Source/MetasoundGraphCore/Public/MetasoundVertex.h
+++ b/Engine/Plugins/Runtime/Metasound/Source/MetasoundGraphCore/Public/MetasoundVertex.h
@@ -169,6 +169,7 @@ namespace Metasound
 		/** Construct an FInputDataVertex. */
 		FInputDataVertex(const FVertexName& InVertexName, const FName& InDataTypeName, const FDataVertexMetadata& InMetadata)
 			: FDataVertex(InVertexName, InDataTypeName, InMetadata)
+			, LiteralFactory(FLiteral::FNone{})
 		{
 		}
 
