Fix for renamed MetaSound wave player pin "Loop Percent"
#rnx

diff --git a/Engine/Plugins/Runtime/Metasound/Source/MetasoundEngine/Private/MetasoundWavePlayerNode.cpp b/Engine/Plugins/Runtime/Metasound/Source/MetasoundEngine/Private/MetasoundWavePlayerNode.cpp
--- a/Engine/Plugins/Runtime/Metasound/Source/MetasoundEngine/Private/MetasoundWavePlayerNode.cpp
+++ b/Engine/Plugins/Runtime/Metasound/Source/MetasoundEngine/Private/MetasoundWavePlayerNode.cpp
@@ -41,7 +41,7 @@ namespace Metasound
 		METASOUND_PARAM(OutputTriggerOnCuePoint, "On Cue Point", "Triggers when a wave cue point was hit during playback.")
 		METASOUND_PARAM(OutputCuePointID, "Cue Point ID", "The cue point ID that was triggered.")
 		METASOUND_PARAM(OutputCuePointLabel, "Cue Point Label", "The cue point label that was triggered (if there was a label parsed in the imported .wav file).")
-		METASOUND_PARAM(OutputLoopRatio, "Loop Ratio", "Returns the current playback location as a ratio of the loop (0-1) if looping is enabled.")
+		METASOUND_PARAM(OutputLoopRatio, "Loop Percent", "Returns the current playback location as a ratio of the loop (0-1) if looping is enabled.")
 		METASOUND_PARAM(OutputPlaybackLocation, "Playback Location", "Returns the absolute position of the wave playback as a ratio of wave duration (0-1).")
 		METASOUND_PARAM(OutputAudioLeft, "Out Left", "The left channel audio output. Mono wave assets will be upmixed to dual stereo.")
 		METASOUND_PARAM(OutputAudioRight, "Out Right", "The right channel audio output. Mono wave assets will be upmixed to dual stereo.")
@@ -1013,6 +1013,13 @@ namespace Metasound
 	{
 		using namespace WavePlayerVertexNames;
 
+		// Workaround to override display name of OutputLoopRatio
+		static const FDataVertexMetadata OutputLoopRatioMetadata
+		{ 
+			METASOUND_GET_PARAM_TT(OutputLoopRatio), // description 
+			METASOUND_LOCTEXT("OutputLoopRatioNotPercentDisplayName", "Loop Ratio") // display name  
+		};
+
 		return FVertexInterface(
 			FInputVertexInterface(
 				TInputDataVertex<FTrigger>(METASOUND_GET_PARAM_NAME_AND_METADATA(InputTriggerPlay)),
@@ -1032,7 +1039,7 @@ namespace Metasound
 				TOutputDataVertex<FTrigger>(METASOUND_GET_PARAM_NAME_AND_METADATA(OutputTriggerOnCuePoint)),
 				TOutputDataVertex<int32>(METASOUND_GET_PARAM_NAME_AND_METADATA(OutputCuePointID)),
 				TOutputDataVertex<FString>(METASOUND_GET_PARAM_NAME_AND_METADATA(OutputCuePointLabel)),
-				TOutputDataVertex<float>(METASOUND_GET_PARAM_NAME_AND_METADATA(OutputLoopRatio)),
+				TOutputDataVertex<float>(METASOUND_GET_PARAM_NAME(OutputLoopRatio), OutputLoopRatioMetadata),
 				TOutputDataVertex<float>(METASOUND_GET_PARAM_NAME_AND_METADATA(OutputPlaybackLocation)),
 				TOutputDataVertex<FAudioBuffer>(METASOUND_GET_PARAM_NAME_AND_METADATA(OutputAudioLeft)),
 				TOutputDataVertex<FAudioBuffer>(METASOUND_GET_PARAM_NAME_AND_METADATA(OutputAudioRight))
