Fix rename delegate not added to MetaSound members that do not have literals

diff --git a/Engine/Plugins/Runtime/Metasound/Source/MetasoundEditor/Private/MetasoundNodeDetailCustomization.cpp b/Engine/Plugins/Runtime/Metasound/Source/MetasoundEditor/Private/MetasoundNodeDetailCustomization.cpp
--- a/Engine/Plugins/Runtime/Metasound/Source/MetasoundEditor/Private/MetasoundNodeDetailCustomization.cpp
+++ b/Engine/Plugins/Runtime/Metasound/Source/MetasoundEditor/Private/MetasoundNodeDetailCustomization.cpp
@@ -830,20 +830,17 @@ namespace Metasound
 			}
 		}
 
-		void FMetasoundMemberDetailCustomization::UpdateRenameDelegate(UMetasoundEditorGraphMemberDefaultLiteral& InMemberDefaultLiteral)
+		void FMetasoundMemberDetailCustomization::UpdateRenameDelegate(UMetasoundEditorGraphMember& InMember)
 		{
-			if (UMetasoundEditorGraphMember* Member = InMemberDefaultLiteral.GetParentMember())
+			if (InMember.CanRename())
 			{
-				if (Member->CanRename())
+				if (!RenameRequestedHandle.IsValid())
 				{
-					if (!RenameRequestedHandle.IsValid())
+					InMember.OnRenameRequested.Clear();
+					RenameRequestedHandle = InMember.OnRenameRequested.AddLambda([this]()
 					{
-						Member->OnRenameRequested.Clear();
-						RenameRequestedHandle = Member->OnRenameRequested.AddLambda([this]()
-						{
-							FSlateApplication::Get().SetKeyboardFocus(NameEditableTextBox.ToSharedRef(), EFocusCause::SetDirectly);
-						});
-					}
+						FSlateApplication::Get().SetKeyboardFocus(NameEditableTextBox.ToSharedRef(), EFocusCause::SetDirectly);
+					});
 				}
 			}
 		}
@@ -875,10 +872,10 @@ namespace Metasound
 				return DefaultPropertyRows;
 			}
 
+			UpdateRenameDelegate(*GraphMember);
+
 			if (UMetasoundEditorGraphMemberDefaultLiteral* MemberDefaultLiteral = GraphMember->GetLiteral())
 			{
-				UpdateRenameDelegate(*MemberDefaultLiteral);
-
 				UClass* MemberClass = MemberDefaultLiteral->GetClass();
 				check(MemberClass);
 
diff --git a/Engine/Plugins/Runtime/Metasound/Source/MetasoundEditor/Private/MetasoundNodeDetailCustomization.h b/Engine/Plugins/Runtime/Metasound/Source/MetasoundEditor/Private/MetasoundNodeDetailCustomization.h
--- a/Engine/Plugins/Runtime/Metasound/Source/MetasoundEditor/Private/MetasoundNodeDetailCustomization.h
+++ b/Engine/Plugins/Runtime/Metasound/Source/MetasoundEditor/Private/MetasoundNodeDetailCustomization.h
@@ -313,7 +313,7 @@ namespace Metasound
 				return InDetailLayout.EditCategory("General");
 			}
 
-			void UpdateRenameDelegate(UMetasoundEditorGraphMemberDefaultLiteral& InMemberDefaultLiteral);
+			void UpdateRenameDelegate(UMetasoundEditorGraphMember& InMember);
 			void CacheMemberData(IDetailLayoutBuilder& InDetailLayout);
 			virtual void CustomizeGeneralCategory(IDetailLayoutBuilder& InDetailLayout);
 			virtual TArray<IDetailPropertyRow*> CustomizeDefaultCategory(IDetailLayoutBuilder& InDetailLayout);
