Fix for MetaSound ensure loading CitySample

diff --git a/Engine/Plugins/Runtime/Metasound/Source/MetasoundEditor/Private/MetasoundEditorModule.cpp b/Engine/Plugins/Runtime/Metasound/Source/MetasoundEditor/Private/MetasoundEditorModule.cpp
--- a/Engine/Plugins/Runtime/Metasound/Source/MetasoundEditor/Private/MetasoundEditorModule.cpp
+++ b/Engine/Plugins/Runtime/Metasound/Source/MetasoundEditor/Private/MetasoundEditorModule.cpp
@@ -709,6 +709,11 @@ namespace Metasound
 					AssetScanStatus = EAssetScanStatus::Complete;
 				}
 
+				// Metasound Engine registers USoundWave as a proxy class in the
+				// Metasound Frontend. The frontend registration must occur before
+				// the Metasound Editor registration of a USoundWave.
+				FModuleManager::LoadModuleChecked<IModuleInterface>("MetasoundEngine");
+
 				RegisterExplicitProxyClass(*USoundWave::StaticClass());
 
 				// Required to ensure logic to order nodes for presets exclusive to
