Metasound Vertex Bind
- Introduce direct binding between vertex interfaces and data references
- Introduce Bind(...) call on IOperator.
- Add FAnyDataReference container for easy management of data references.

diff --git a/Engine/Plugins/Runtime/Metasound/Source/MetasoundGraphCore/Private/MetasoundDataReferenceCollection.cpp b/Engine/Plugins/Runtime/Metasound/Source/MetasoundGraphCore/Private/MetasoundDataReferenceCollection.cpp
--- a/Engine/Plugins/Runtime/Metasound/Source/MetasoundGraphCore/Private/MetasoundDataReferenceCollection.cpp
+++ b/Engine/Plugins/Runtime/Metasound/Source/MetasoundGraphCore/Private/MetasoundDataReferenceCollection.cpp
@@ -10,11 +10,12 @@ namespace Metasound
 {
 	bool FDataReferenceCollection::ContainsDataReadReference(const FVertexName& InName, const FName& InTypeName) const
 	{
-		if (const FDataRefWrapper* Wrapper = DataReadRefMap.Find(InName))
+		if (const FAnyDataReference* DataRef = FindDataReference(InName))
 		{
-			if (const IDataReference* Ref = Wrapper->GetDataReference())
+			if (InTypeName == DataRef->GetDataTypeName())
 			{
-				return InTypeName == Ref->GetDataTypeName();
+				const EDataReferenceAccessType AccessType = DataRef->GetAccessType();
+				return (AccessType == EDataReferenceAccessType::Read) || (AccessType == EDataReferenceAccessType::Write);
 			}
 		}
 
@@ -23,59 +24,56 @@ namespace Metasound
 
 	bool FDataReferenceCollection::ContainsDataWriteReference(const FVertexName& InName, const FName& InTypeName) const
 	{
-		if (const FDataRefWrapper* Wrapper = DataWriteRefMap.Find(InName))
+		if (const FAnyDataReference* DataRef = FindDataReference(InName))
 		{
-			if (const IDataReference* Ref = Wrapper->GetDataReference())
+			if (InTypeName == DataRef->GetDataTypeName())
 			{
-				return InTypeName == Ref->GetDataTypeName();
+				const EDataReferenceAccessType AccessType = DataRef->GetAccessType();
+				return (AccessType == EDataReferenceAccessType::Write);
 			}
 		}
 
 		return false;
 	}
 
-	bool FDataReferenceCollection::AddDataReadReferenceFrom(const FVertexName& InName, const FDataReferenceCollection& OtherCollection, const FVertexName& OtherName, const FName& OtherTypeName)
+	const FAnyDataReference* FDataReferenceCollection::FindDataReference(const FVertexName& InName) const
 	{
-		if (OtherCollection.ContainsDataReadReference(OtherName, OtherTypeName))
-		{
-			DataReadRefMap.Add(InName, OtherCollection.DataReadRefMap[OtherName]);
-
-			return true;
-		}
-		else if (OtherCollection.ContainsDataWriteReference(OtherName, OtherTypeName))
-		{
-			DataReadRefMap.Add(InName, OtherCollection.DataWriteRefMap[OtherName]);
-
-			return true;
-		}
-
-		return false;
+		return DataRefMap.Find(InName);
 	}
 
-	bool FDataReferenceCollection::AddDataWriteReferenceFrom(const FVertexName& InName, const FDataReferenceCollection& OtherCollection, const FVertexName& OtherName, const FName& OtherTypeName)
+
+	bool FDataReferenceCollection::AddDataReadReferenceFrom(const FVertexName& InName, const FDataReferenceCollection& OtherCollection, const FVertexName& OtherName, const FName& OtherTypeName)
 	{
-		if (OtherCollection.ContainsDataWriteReference(OtherName, OtherTypeName))
+		if (const FAnyDataReference* OtherRef = OtherCollection.FindDataReference(OtherName))
 		{
-			DataWriteRefMap.Add(InName, OtherCollection.DataWriteRefMap[OtherName]);
-
-			return true;
+			if (OtherRef->GetDataTypeName() == OtherTypeName)
+			{
+				DataRefMap.Add(InName, *OtherRef);
+				return true;
+			}
 		}
 
 		return false;
 	}
 
-	const IDataReference* FDataReferenceCollection::GetDataReference(const FDataReferenceMap& InMap, const FVertexName& InName) const
+	bool FDataReferenceCollection::AddDataWriteReferenceFrom(const FVertexName& InName, const FDataReferenceCollection& OtherCollection, const FVertexName& OtherName, const FName& OtherTypeName)
 	{
-		if (const FDataRefWrapper* RefWrapper = InMap.Find(InName))
+		if (const FAnyDataReference* OtherRef = OtherCollection.FindDataReference(OtherName))
 		{
-			return RefWrapper->GetDataReference();
+			if (OtherRef->GetDataTypeName() == OtherTypeName)
+			{
+				if (OtherRef->GetAccessType() == EDataReferenceAccessType::Write)
+				{
+					DataRefMap.Add(InName, *OtherRef);
+					return true;
+				}
+			}
 		}
-
-		return nullptr;
+		return false;
 	}
 
-	void FDataReferenceCollection::AddDataReference(FDataReferenceMap& InMap, const FVertexName& InName, FDataRefWrapper&& InDataRef)
+	void FDataReferenceCollection::AddDataReference(const FVertexName& InName, FAnyDataReference&& InDataRef)
 	{
-		InMap.Add(InName, InDataRef);
+		DataRefMap.Add(InName, InDataRef);
 	}
 }
diff --git a/Engine/Plugins/Runtime/Metasound/Source/MetasoundGraphCore/Private/MetasoundVertexData.cpp b/Engine/Plugins/Runtime/Metasound/Source/MetasoundGraphCore/Private/MetasoundVertexData.cpp
new file mode 100644
--- /dev/null
+++ b/Engine/Plugins/Runtime/Metasound/Source/MetasoundGraphCore/Private/MetasoundVertexData.cpp
@@ -0,0 +1,209 @@
+// Copyright Epic Games, Inc. All Rights Reserved.
+
+#include "MetasoundVertexData.h"
+
+namespace Metasound
+{
+	namespace MetasoundVertexDataPrivate
+	{
+		template<typename VertexType>
+		void EmplaceBindings(TArray<TBinding<VertexType>>& InArray, const TVertexInterfaceGroup<VertexType>& InVertexInterface)
+		{
+			for (const VertexType& DataVertex : InVertexInterface)
+			{
+				InArray.Emplace(DataVertex);
+			}
+		}
+
+		template<typename BindingType>
+		BindingType* Find(TArray<BindingType>& InBindings, const FVertexName& InVertexName)
+		{
+			auto IsVertexWithName = [&InVertexName](const BindingType& InBinding)
+			{
+				return InBinding.GetVertex().VertexName == InVertexName;
+			};
+			return InBindings.FindByPredicate(IsVertexWithName);
+		}
+
+		template<typename BindingType>
+		const BindingType* Find(const TArray<BindingType>& InBindings, const FVertexName& InVertexName)
+		{
+			auto IsVertexWithName = [&InVertexName](const BindingType& InBinding)
+			{
+				return InBinding.GetVertex().VertexName == InVertexName;
+			};
+			return InBindings.FindByPredicate(IsVertexWithName);
+		}
+
+		template<typename BindingType>
+		void BindVertex(TArray<BindingType>& InBindings, const FVertexName& InVertexName, FAnyDataReference&& InDataReference)
+		{
+			if (BindingType* Binding = Find(InBindings, InVertexName))
+			{
+				if (Binding->GetVertex().DataTypeName == InDataReference.GetDataTypeName())
+				{
+					Binding->Bind(MoveTemp(InDataReference));
+				}
+				else
+				{
+					UE_LOG(LogMetaSound, Warning, TEXT("Failed bind vertex with name '%s'. Supplied data type (%s) does not match vertex data type (%s)"), *InVertexName.ToString(), *InDataReference.GetDataTypeName().ToString(), *Binding->GetVertex().DataTypeName.ToString());
+				}
+			}
+			else
+			{
+				UE_LOG(LogMetaSound, Warning, TEXT("Failed find vertex with name '%s'. Failed to bind vertex data"), *InVertexName.ToString());
+			}
+		}
+
+		template<typename BindingType>
+		void Bind(TArray<BindingType>& InBindings, const FDataReferenceCollection& InCollection)
+		{
+			for (BindingType& Binding : InBindings)
+			{
+				const FDataVertex& Vertex = Binding.GetVertex();
+
+				if (const FAnyDataReference* DataRef = InCollection.FindDataReference(Vertex.VertexName))
+				{
+					Binding.Bind(*DataRef);
+				}
+			}
+		}
+
+		template<typename BindingType>
+		bool IsVertexBound(const TArray<BindingType>& InBindings, const FVertexName& InVertexName)
+		{
+			if (const BindingType* Binding = Find(InBindings, InVertexName))
+			{
+				return Binding->IsBound();
+			}
+			return false;
+		}
+
+		template<typename BindingType>
+		bool AreAllVerticesBound(const TArray<BindingType>& InBindings)
+		{
+			return Algo::AllOf(InBindings, [](const BindingType& Binding) { return Binding.IsBound(); });
+		}
+
+		template<typename BindingType>
+		EDataReferenceAccessType GetVertexDataAccessType(const TArray<BindingType>& InBindings, const FVertexName& InVertexName)
+		{
+			if (const BindingType* Binding = Find(InBindings, InVertexName))
+			{
+				return Binding->GetAccessType();
+			}
+			return EDataReferenceAccessType::None;
+		}
+	}
+
+	FInputVertexInterfaceData::FInputVertexInterfaceData(const FInputVertexInterface& InVertexInterface)
+	{
+		MetasoundVertexDataPrivate::EmplaceBindings(Bindings, InVertexInterface);
+	}
+
+	
+	void FInputVertexInterfaceData::BindVertex(const FVertexName& InVertexName, FAnyDataReference&& InDataReference)
+	{
+		MetasoundVertexDataPrivate::BindVertex(Bindings, InVertexName, MoveTemp(InDataReference));
+	}
+
+	void FInputVertexInterfaceData::Bind(const FDataReferenceCollection& InCollection)
+	{
+		MetasoundVertexDataPrivate::Bind(Bindings, InCollection);
+	}
+
+	bool FInputVertexInterfaceData::IsVertexBound(const FVertexName& InVertexName) const
+	{
+		return MetasoundVertexDataPrivate::IsVertexBound(Bindings, InVertexName);
+	}
+
+	EDataReferenceAccessType FInputVertexInterfaceData::GetVertexDataAccessType(const FVertexName& InVertexName) const
+	{
+		return MetasoundVertexDataPrivate::GetVertexDataAccessType(Bindings, InVertexName);
+	}
+
+	bool FInputVertexInterfaceData::AreAllVerticesBound() const
+	{
+		return MetasoundVertexDataPrivate::AreAllVerticesBound(Bindings);
+	}
+
+	FInputVertexInterfaceData::FBinding* FInputVertexInterfaceData::Find(const FVertexName& InVertexName)
+	{
+		return MetasoundVertexDataPrivate::Find(Bindings, InVertexName);
+	}
+
+	const FInputVertexInterfaceData::FBinding* FInputVertexInterfaceData::Find(const FVertexName& InVertexName) const
+	{
+		return MetasoundVertexDataPrivate::Find(Bindings, InVertexName);
+	}
+
+	FInputVertexInterfaceData::FBinding* FInputVertexInterfaceData::FindChecked(const FVertexName& InVertexName)
+	{
+		FBinding* Binding = Find(InVertexName);
+		check(nullptr != Binding);
+		return Binding;
+	}
+
+	const FInputVertexInterfaceData::FBinding* FInputVertexInterfaceData::FindChecked(const FVertexName& InVertexName) const
+	{
+		const FBinding* Binding = Find(InVertexName);
+		check(nullptr != Binding);
+		return Binding;
+	}
+
+
+	FOutputVertexInterfaceData::FOutputVertexInterfaceData(const FOutputVertexInterface& InVertexInterface)
+	{
+		MetasoundVertexDataPrivate::EmplaceBindings(Bindings, InVertexInterface);
+	}
+
+	
+	void FOutputVertexInterfaceData::BindVertex(const FVertexName& InVertexName, FAnyDataReference&& InDataReference)
+	{
+		MetasoundVertexDataPrivate::BindVertex(Bindings, InVertexName, MoveTemp(InDataReference));
+	}
+
+	void FOutputVertexInterfaceData::Bind(const FDataReferenceCollection& InCollection)
+	{
+		MetasoundVertexDataPrivate::Bind(Bindings, InCollection);
+	}
+
+	bool FOutputVertexInterfaceData::IsVertexBound(const FVertexName& InVertexName) const
+	{
+		return MetasoundVertexDataPrivate::IsVertexBound(Bindings, InVertexName);
+	}
+
+	EDataReferenceAccessType FOutputVertexInterfaceData::GetVertexDataAccessType(const FVertexName& InVertexName) const
+	{
+		return MetasoundVertexDataPrivate::GetVertexDataAccessType(Bindings, InVertexName);
+	}
+
+	bool FOutputVertexInterfaceData::AreAllVerticesBound() const
+	{
+		return MetasoundVertexDataPrivate::AreAllVerticesBound(Bindings);
+	}
+
+	FOutputVertexInterfaceData::FBinding* FOutputVertexInterfaceData::Find(const FVertexName& InVertexName)
+	{
+		return MetasoundVertexDataPrivate::Find(Bindings, InVertexName);
+	}
+
+	const FOutputVertexInterfaceData::FBinding* FOutputVertexInterfaceData::Find(const FVertexName& InVertexName) const
+	{
+		return MetasoundVertexDataPrivate::Find(Bindings, InVertexName);
+	}
+
+	FOutputVertexInterfaceData::FBinding* FOutputVertexInterfaceData::FindChecked(const FVertexName& InVertexName)
+	{
+		FBinding* Binding = Find(InVertexName);
+		check(nullptr != Binding);
+		return Binding;
+	}
+
+	const FOutputVertexInterfaceData::FBinding* FOutputVertexInterfaceData::FindChecked(const FVertexName& InVertexName) const
+	{
+		const FBinding* Binding = Find(InVertexName);
+		check(nullptr != Binding);
+		return Binding;
+	}
+}
diff --git a/Engine/Plugins/Runtime/Metasound/Source/MetasoundGraphCore/Public/MetasoundDataReference.h b/Engine/Plugins/Runtime/Metasound/Source/MetasoundGraphCore/Public/MetasoundDataReference.h
--- a/Engine/Plugins/Runtime/Metasound/Source/MetasoundGraphCore/Public/MetasoundDataReference.h
+++ b/Engine/Plugins/Runtime/Metasound/Source/MetasoundGraphCore/Public/MetasoundDataReference.h
@@ -7,6 +7,7 @@
 #include "Containers/UnrealString.h"
 #include "HAL/Platform.h"
 #include "Internationalization/Text.h"
+#include "Templates/SharedPointer.h"
 #include "Templates/UniquePtr.h"
 #include "UObject/NameTypes.h"
 
@@ -369,5 +370,147 @@ namespace Metasound
 			return MakeUnique<FDataReadReference>(*this);
 		}
 	};
+
+	/** EDataReferenceAccessType describes the underlying data reference access
+	 * type for a data reference contained in a FAnyDataReference. 
+	 *
+	 * This value can be used to determine which methods are supported for accessing 
+	 * a data reference using GetDataReadReference<>() or GetDataWriteReference<>()
+	 */
+	enum class EDataReferenceAccessType : uint8
+	{
+		None, //< The data is inaccessible, or the data reference does not exist.
+		Read, //< The data is accessible through a TDataReadReference.
+		Write //< The data is accessible through a TDataWriteReference.
+	};
+
+	/** Container for any data reference. 
+	 *
+	 * This container maintains the underlying containers access type (Read or Write)
+	 * and data type. This allows for convenient storage by implementing a virtual 
+	 * copy constructor and assignment operator.
+	 */
+	class FAnyDataReference : public IDataReference
+	{
+		// Private constructor. 
+		FAnyDataReference(EDataReferenceAccessType InAccessType, const IDataReference& InDataRef)
+		: AccessType(InAccessType)
+		, DataRefPtr(InDataRef.Clone())
+		{
+			check(DataRefPtr.IsValid());
+			check(EDataReferenceAccessType::None != InAccessType);
+		}
+
+	public:
+		/** Construct with a TDataReadReference. */
+		template<typename DataType>
+		FAnyDataReference(const TDataReadReference<DataType>& InDataRef)
+		: FAnyDataReference(EDataReferenceAccessType::Read, InDataRef)
+		{
+		}
+
+		/** Construct with a TDataWriteReference. */
+		template<typename DataType>
+		FAnyDataReference(const TDataWriteReference<DataType>& InDataRef)
+		: FAnyDataReference(EDataReferenceAccessType::Write, InDataRef)
+		{
+		}
+
+		/** Copy construct with a FAnyDataReference. */
+		FAnyDataReference(const FAnyDataReference& InOther)
+		: FAnyDataReference(InOther.AccessType, *InOther.DataRefPtr)
+		{
+		}
+
+		/** Assignment. */
+		FAnyDataReference& operator=(const FAnyDataReference& InOther)
+		{
+			AccessType = InOther.AccessType;
+			DataRefPtr = InOther.DataRefPtr->Clone();
+
+			check(DataRefPtr.IsValid());
+			check(EDataReferenceAccessType::None != AccessType);
+
+			return *this;
+		}
+
+		/** Returns the access type of the underlying data reference. */
+		EDataReferenceAccessType GetAccessType() const
+		{
+			return AccessType;
+		}
+
+		/** Returns the data type name of the underlying data reference. */
+		virtual const FName& GetDataTypeName() const override
+		{
+			check(DataRefPtr.IsValid());
+			return DataRefPtr->GetDataTypeName();
+		}
+		
+		/** Returns the data type ID of the underlying data reference. */
+		virtual const void* const GetDataTypeId() const override
+		{
+			check(DataRefPtr.IsValid());
+			return DataRefPtr->GetDataTypeId();
+		}
+
+		/** Returns a clone of the underlying data reference. */
+		virtual TUniquePtr<IDataReference> Clone() const override
+		{
+			check(DataRefPtr.IsValid());
+			return DataRefPtr->Clone();
+		}
+
+		/** Get access to a TDataReadReference. 
+		 *
+		 * This method will return a valid TDataReadReference of the templated data
+		 * type. The returned object is only valid if:
+		 *     1. The template paramter DataType matches that of the underlying data reference.
+		 *     2. The underlying data reference is has either Read or Write access.
+		 *
+		 * If this method's behavior is undefined and will assert if it is called 
+		 * with a mismatched data type or unsupported access type.
+		 */
+		template<typename DataType>
+		TDataReadReference<DataType> GetDataReadReference() const
+		{
+			check(DataRefPtr.IsValid());
+			check(IsDataReferenceOfType<DataType>(*DataRefPtr));
+
+			if (EDataReferenceAccessType::Read == AccessType)
+			{
+				return *static_cast<const TDataReadReference<DataType>*>(DataRefPtr.Get());
+			}
+			else 
+			{
+				check(EDataReferenceAccessType::Write == AccessType);
+				return *static_cast<const TDataWriteReference<DataType>*>(DataRefPtr.Get());
+			}
+		}
+
+		/** Get access to a TDataWriteReference. 
+		 *
+		 * This method will return a valid TDataWriteReference of the templated data
+		 * type. The returned object is only valid if:
+		 *     1. The template paramter DataType matches that of the underlying data reference.
+		 *     2. The underlying data reference is has Write access.
+		 *
+		 * If this method's behavior is undefined and will assert if it is called 
+		 * with a mismatched data type or unsupported access type.
+		 */
+		template<typename DataType>
+		TDataWriteReference<DataType> GetDataWriteReference() const
+		{
+			check(DataRefPtr.IsValid());
+			check(IsDataReferenceOfType<DataType>(*DataRefPtr));
+			check(EDataReferenceAccessType::Write == AccessType);
+			return *static_cast<const TDataWriteReference<DataType>*>(DataRefPtr.Get());
+		}
+
+	private:
+
+		EDataReferenceAccessType AccessType = EDataReferenceAccessType::None;
+		TUniquePtr<IDataReference> DataRefPtr;
+	};
 }
 
diff --git a/Engine/Plugins/Runtime/Metasound/Source/MetasoundGraphCore/Public/MetasoundDataReferenceCollection.h b/Engine/Plugins/Runtime/Metasound/Source/MetasoundGraphCore/Public/MetasoundDataReferenceCollection.h
--- a/Engine/Plugins/Runtime/Metasound/Source/MetasoundGraphCore/Public/MetasoundDataReferenceCollection.h
+++ b/Engine/Plugins/Runtime/Metasound/Source/MetasoundGraphCore/Public/MetasoundDataReferenceCollection.h
@@ -2,7 +2,6 @@
 
 #pragma once
 
-#include "CoreMinimal.h"
 #include "MetasoundDataFactory.h"
 #include "MetasoundDataReference.h"
 #include "MetasoundOperatorSettings.h"
@@ -10,64 +9,12 @@
 
 namespace Metasound
 {
+	template<typename VertexType>
+	class TVertexInterfaceGroupData;
+
 	class METASOUNDGRAPHCORE_API FDataReferenceCollection
 	{
-		// Convenience class to handle copying and moving of data refs
-		// without knowing data refs derived type.  
-		//
-		// Uses Clone() function as a virtual copy constructor.
-		class FDataRefWrapper 
-		{
-			public:
-				// Anything that implements a IDataReference interface
-				template<typename DataRefType>
-				FDataRefWrapper(const DataRefType& InDataReference)
-				:	DataRefPtr(InDataReference.Clone())
-				{
-				}
-
-				// Override copy constructor to use underlying Clone method.
-				FDataRefWrapper(const FDataRefWrapper& Other)
-				:	DataRefPtr(Other.CloneDataReference())
-				{
-				}
-
-				// Override equal operator to use underlying Clone method.
-				FDataRefWrapper& operator=(const FDataRefWrapper& Other)
-				{
-					DataRefPtr = Other.CloneDataReference();
-					return *this;
-				}
-
-				bool IsValid() const
-				{
-					return DataRefPtr.IsValid();
-				}
-
-				// Get underlying data ref
-				const IDataReference* GetDataReference() const
-				{
-					return DataRefPtr.Get();
-				}
-
-				// Clone underlying data ref
-				TUniquePtr<IDataReference> CloneDataReference() const
-				{
-					if (DataRefPtr.IsValid())
-					{
-						return DataRefPtr->Clone();
-					}
-				
-					return TUniquePtr<IDataReference>();
-				}
-
-			private:
-
-				TUniquePtr<IDataReference> DataRefPtr;
-		};
-
-
-		using FDataReferenceMap = TMap<FVertexName, FDataRefWrapper>;
+		using FDataReferenceMap = TMap<FVertexName, FAnyDataReference>;
 
 		public:
 
@@ -79,9 +26,9 @@ namespace Metasound
 			template<typename DataType>
 			void AddDataReadReference(const FVertexName& InName, const TDataReadReference<DataType>& InDataReference)
 			{
-				FDataRefWrapper Wrapper(InDataReference);
+				FAnyDataReference Wrapper(InDataReference);
 
-				AddDataReference(DataReadRefMap, InName, MoveTemp(Wrapper));
+				AddDataReference(InName, MoveTemp(Wrapper));
 			}
 
 			/** Add a readable data reference to the collection.
@@ -105,9 +52,9 @@ namespace Metasound
 			template<typename DataType>
 			void AddDataWriteReference(const FVertexName& InName, const TDataWriteReference<DataType>& InDataReference)
 			{
-				FDataRefWrapper Wrapper(InDataReference);
+				FAnyDataReference Wrapper(InDataReference);
 
-				AddDataReference(DataWriteRefMap, InName, MoveTemp(Wrapper));
+				AddDataReference(InName, MoveTemp(Wrapper));
 			}
 
 			/** Add a readable data reference from another collection.
@@ -150,6 +97,11 @@ namespace Metasound
 			 */
 			bool ContainsDataWriteReference(const FVertexName& InName, const FName& InTypeName) const;
 
+			/** Returns pointer to data reference if it exists for a given vertex name.
+			 *
+			 * If no data reference exists for the given vertex name, a nullptr is returned. 
+			 */
+			const FAnyDataReference* FindDataReference(const FVertexName& InName) const;
 
 			/** Query whether a readable data reference is within the collection. 
 			 *
@@ -160,7 +112,16 @@ namespace Metasound
 			template<typename DataType>
 			bool ContainsDataReadReference(const FVertexName& InName) const
 			{
-				return Contains<DataType>(DataReadRefMap, InName);
+				if (const FAnyDataReference* DataRef = FindDataReference(InName))
+				{
+					if (IsDataReferenceOfType<DataType>(*DataRef))
+					{
+						const EDataReferenceAccessType AccessType = DataRef->GetAccessType();
+						return (AccessType == EDataReferenceAccessType::Read) || (AccessType == EDataReferenceAccessType::Write);
+					}
+				}
+
+				return false;
 			}
 
 			/** Query whether a writable data reference is within the collection. 
@@ -172,7 +133,16 @@ namespace Metasound
 			template<typename DataType>
 			bool ContainsDataWriteReference(const FVertexName& InName) const
 			{
-				return Contains<DataType>(DataWriteRefMap, InName);
+				if (const FAnyDataReference* DataRef = FindDataReference(InName))
+				{
+					if (IsDataReferenceOfType<DataType>(*DataRef))
+					{
+						const EDataReferenceAccessType AccessType = DataRef->GetAccessType();
+						return AccessType == EDataReferenceAccessType::Write;
+					}
+				}
+
+				return false;
 			}
 
 			/** Returns a readable data ref from the collection.
@@ -184,15 +154,10 @@ namespace Metasound
 			template<typename DataType>
 			TDataReadReference<DataType> GetDataReadReference(const FVertexName& InName) const
 			{
-				typedef TDataReadReference<DataType> FDataRefType;
-
-				const IDataReference* DataRefPtr = GetDataReference(DataReadRefMap, InName);
-
+				const FAnyDataReference* DataRefPtr = FindDataReference(InName);
 				check(nullptr != DataRefPtr);
+				return DataRefPtr->GetDataReadReference<DataType>();
 
-				check(IsDataReferenceOfType<DataType>(*DataRefPtr));
-
-				return FDataRefType(*static_cast<const FDataRefType*>(DataRefPtr));
 			}
 
 			/** Returns a readable data ref from the collection or construct one
@@ -258,15 +223,10 @@ namespace Metasound
 			template<typename DataType>
 			TDataWriteReference<DataType> GetDataWriteReference(const FVertexName& InName) const
 			{
-				typedef TDataWriteReference<DataType> FDataRefType;
-
-				const IDataReference* DataRefPtr = GetDataReference(DataWriteRefMap, InName);
-
+				const FAnyDataReference* DataRefPtr = FindDataReference(InName);
 				check(nullptr != DataRefPtr);
 
-				check(IsDataReferenceOfType<DataType>(*DataRefPtr));
-
-				return FDataRefType(*static_cast<const FDataRefType*>(DataRefPtr));
+				return DataRefPtr->GetDataWriteReference<DataType>();
 			}
 
 			/** Returns a writable data ref from the collection or construct
@@ -292,29 +252,12 @@ namespace Metasound
 				}
 			}
 
-		private:
-
-			template<typename DataType>
-			bool Contains(const FDataReferenceMap& InMap, const FVertexName& InName) const
-			{
-				if (InMap.Contains(InName))
-				{
-					const IDataReference* DataRefPtr = InMap[InName].GetDataReference();
-
-					if (nullptr != DataRefPtr)
-					{
-						return IsDataReferenceOfType<DataType>(*DataRefPtr);
-					}
-				}
 
-				return false;
-			}
+		private:
 
-			const IDataReference* GetDataReference(const FDataReferenceMap& InMap, const FVertexName& InName) const;
 
-			void AddDataReference(FDataReferenceMap& InMap, const FVertexName& InName, FDataRefWrapper&& InDataReference);
+			void AddDataReference(const FVertexName& InName, FAnyDataReference&& InDataReference);
 
-			FDataReferenceMap DataReadRefMap;
-			FDataReferenceMap DataWriteRefMap;
+			FDataReferenceMap DataRefMap;
 	};
 }
diff --git a/Engine/Plugins/Runtime/Metasound/Source/MetasoundGraphCore/Public/MetasoundOperatorInterface.h b/Engine/Plugins/Runtime/Metasound/Source/MetasoundGraphCore/Public/MetasoundOperatorInterface.h
--- a/Engine/Plugins/Runtime/Metasound/Source/MetasoundGraphCore/Public/MetasoundOperatorInterface.h
+++ b/Engine/Plugins/Runtime/Metasound/Source/MetasoundGraphCore/Public/MetasoundOperatorInterface.h
@@ -2,10 +2,8 @@
 
 #pragma once
 
-#include "CoreMinimal.h"
-
-#include "MetasoundDataReference.h"
 #include "MetasoundDataReferenceCollection.h"
+#include "MetasoundVertexData.h"
 
 namespace Metasound
 {
@@ -42,6 +40,18 @@ namespace Metasound
 		 */
 		virtual FDataReferenceCollection GetOutputs() const = 0;
 
+		/** Bind data references to the interface of this operator.
+		 *
+		 * Operators should bind all data references to the vertices of the provided
+		 * InVertexData. This method is meant to replace GetInputs() and GetOutputs()
+		 * because it provides better error checking against mismatched vertex names.
+		 */
+		virtual void Bind(FVertexInterfaceData& InVertexData) const
+		{
+			InVertexData.GetInputs().Bind(GetInputs());
+			InVertexData.GetOutputs().Bind(GetOutputs());
+		}
+
 		/** Return the execution function to call during graph execution.
 		 *
 		 * The IOperator* argument to the FExecutionFunction will be the same IOperator instance
diff --git a/Engine/Plugins/Runtime/Metasound/Source/MetasoundGraphCore/Public/MetasoundVertexData.h b/Engine/Plugins/Runtime/Metasound/Source/MetasoundGraphCore/Public/MetasoundVertexData.h
new file mode 100644
--- /dev/null
+++ b/Engine/Plugins/Runtime/Metasound/Source/MetasoundGraphCore/Public/MetasoundVertexData.h
@@ -0,0 +1,490 @@
+// Copyright Epic Games, Inc. All Rights Reserved.
+
+#pragma once
+
+#include "Algo/AllOf.h"
+#include "Containers/Array.h"
+#include "MetasoundDataFactory.h"
+#include "MetasoundDataReference.h"
+#include "MetasoundDataReferenceCollection.h"
+#include "MetasoundLog.h"
+#include "MetasoundOperatorSettings.h"
+#include "MetasoundVertex.h"
+
+#include <type_traits>
+
+namespace Metasound
+{
+	namespace MetasoundVertexDataPrivate
+	{
+		template<typename V>
+		using EnableIfIsInputDataVertex = std::enable_if_t<std::is_same_v<FInputDataVertex, V>>;
+
+		// Binds a vertex to a data reference.
+		template<typename VertexType>
+		class TBinding
+		{
+
+		public:
+			TBinding(VertexType&& InVertex)
+			: Vertex(MoveTemp(InVertex))
+			{
+			}
+
+			TBinding(const VertexType& InVertex)
+			: Vertex(InVertex)
+			{
+			}
+
+			template<typename DataReferenceType>
+			void Bind(const DataReferenceType& InDataReference)
+			{
+				check(Vertex.DataTypeName == InDataReference.GetDataTypeName());
+				Data.Emplace(InDataReference);
+			}
+
+			void Bind(FAnyDataReference&& InAnyDataReference)
+			{
+				check(Vertex.DataTypeName == InAnyDataReference.GetDataTypeName());
+				Data.Emplace(MoveTemp(InAnyDataReference));
+			}
+			
+			const VertexType& GetVertex() const
+			{
+				return Vertex;
+			}
+
+			bool IsBound() const
+			{
+				return Data.IsSet();
+			}
+
+			EDataReferenceAccessType GetAccessType() const
+			{
+				if (Data.IsSet())
+				{
+					Data->GetAccessType();
+				}
+				return EDataReferenceAccessType::None;
+			}
+
+			// Get data read reference assuming data is bound and read or write accessible.
+			template<typename DataType>
+			TDataReadReference<DataType> GetDataReadReference() const
+			{
+				check(Data.IsSet());
+				return Data->GetDataReadReference<DataType>();
+			}
+
+			// Get the bound data read reference if it exists. Otherwise create and 
+			// return a data read reference by constructing one using the Vertex's 
+			// default literal.
+			template<
+				typename DataType,
+				typename V = VertexType,
+				typename = EnableIfIsInputDataVertex<V>
+			> 
+			TDataReadReference<DataType> GetOrCreateDefaultDataReadReference(const FOperatorSettings& InSettings) const
+			{
+				if (Data.IsSet())
+				{
+					return Data->GetDataReadReference<DataType>();
+				}
+				else
+				{
+					return TDataReadReferenceLiteralFactory<DataType>::CreateExplicitArgs(InSettings, Vertex.GetDefaultLiteral());
+				}
+			}
+
+			// Get the bound data read reference if it exists. Otherwise create and 
+			// return a data read reference by constructing one the supplied constructor 
+			// arguments.
+			template<typename DataType, typename ... ConstructorArgTypes>
+			TDataReadReference<DataType> GetOrConstructDataReadReference(ConstructorArgTypes&&... ConstructorArgs) const
+			{
+				if (Data.IsSet())
+				{
+					return Data->GetDataReadReference<DataType>();
+				}
+
+				return TDataReadReference<DataType>::CreateNew(Forward<ConstructorArgTypes>(ConstructorArgs)...);
+			}
+
+			// Get data write reference assuming data is bound and write accessible.
+			template<typename DataType> 
+			TDataWriteReference<DataType> GetDataWriteReference() const
+			{
+				check(Data.IsSet());
+				return Data->GetDataWriteReference<DataType>();
+			}
+
+			// Get the bound data write reference if it exists. Otherwise create and 
+			// return a data write reference by constructing one using the Vertex's 
+			// default literal.
+			template<
+				typename DataType,
+				typename V = VertexType,
+				typename = EnableIfIsInputDataVertex<V>
+			> 
+			TDataWriteReference<DataType> GetOrCreateDefaultDataWriteReference(const FOperatorSettings& InSettings) const
+			{
+				if (Data.IsSet())
+				{
+					return Data->GetDataWriteReference<DataType>();
+				}
+				return TDataWriteReferenceLiteralFactory<DataType>::CreateExplicitArgs(InSettings, Vertex.GetDefaultLiteral());
+			}
+
+			// Get the bound data write reference if it exists. Otherwise create and 
+			// return a data write reference by constructing one the supplied constructor 
+			// arguments.
+			template<typename DataType, typename ... ConstructorArgTypes>
+			TDataWriteReference<DataType> GetOrConstructDataWriteReference(ConstructorArgTypes&&... ConstructorArgs) const
+			{
+				if (Data.IsSet())
+				{
+					return Data->GetDataWriteReference<DataType>();
+				}
+				return TDataWriteReference<DataType>::CreateNew(Forward<ConstructorArgTypes>(ConstructorArgs)...);
+			}
+
+		private:
+
+			VertexType Vertex;
+			TOptional<FAnyDataReference> Data;
+		};
+	}
+
+	/** An input vertex interface with optionally bound data references. */
+	class METASOUNDGRAPHCORE_API FInputVertexInterfaceData
+	{
+		using FBinding = MetasoundVertexDataPrivate::TBinding<FInputDataVertex>;
+
+	public:
+
+		using FRangedForIteratorType = typename TArray<FBinding>::RangedForIteratorType;
+		using FRangedForConstIteratorType = typename TArray<FBinding>::RangedForConstIteratorType;
+
+		/** Construct with an FInputVertexInterface. */
+		FInputVertexInterfaceData(const FInputVertexInterface& InVertexInterface);
+
+		/** Bind a data reference to a vertex. */
+		template<typename DataReferenceType>
+		void BindVertex(const FVertexName& InVertexName, const DataReferenceType& InDataReference)
+		{
+			BindVertex(InVertexName, FAnyDataReference{InDataReference});
+		}
+
+		/** Bind a data reference to a vertex. */
+		void BindVertex(const FVertexName& InVertexName, FAnyDataReference&& InDataReference);
+
+		/** Bind a data references to vertices with matching vertex names. */
+		void Bind(const FDataReferenceCollection& InCollection);
+
+		/** Returns true if a vertex with the given vertex name exists and is bound
+		 * to a data reference. */
+		bool IsVertexBound(const FVertexName& InVertexName) const;
+
+		/** Returns the access type of a bound vertex. If the vertex does not exist
+		 * or if it is unbound, this will return EDataReferenceAccessType::None
+		 */
+		EDataReferenceAccessType GetVertexDataAccessType(const FVertexName& InVertexName) const;
+
+		/** Returns true if all vertices in the FInputVertexInterface are bound to 
+		 * data references. */
+		bool AreAllVerticesBound() const;
+
+		FRangedForIteratorType begin()
+		{
+			return Bindings.begin();
+		}
+
+		FRangedForIteratorType end()
+		{
+			return Bindings.end();
+		}
+
+		FRangedForConstIteratorType begin() const
+		{
+			return Bindings.begin();
+		}
+
+		FRangedForConstIteratorType end() const
+		{
+			return Bindings.end();
+		}
+
+
+		/** Get data read reference assuming data is bound and read or write accessible. */
+		template<typename DataType> 
+		TDataReadReference<DataType> GetDataReadReference(const FVertexName& InVertexName) const
+		{
+			const FBinding* Binding = FindChecked(InVertexName);
+			return Binding->GetDataReadReference<DataType>();
+		}
+
+		/** Get the bound data read reference if it exists. Otherwise create and 
+		 * return a data read reference by constructing one using the Vertex's 
+		 * default literal.
+		 */
+		template<typename DataType> 
+		TDataReadReference<DataType> GetOrCreateDefaultDataReadReference(const FVertexName& InVertexName, const FOperatorSettings& InSettings) const
+		{
+			if (const FBinding* Binding = Find(InVertexName))
+			{
+				return Binding->GetOrCreateDefaultDataReadReference<DataType>(InSettings);
+			}
+			else
+			{
+				UE_LOG(LogMetaSound, Warning, TEXT("Failed find vertex with name '%s'. Cannot check for existing bound data."), *InVertexName.ToString());
+			}
+
+			return TDataReadReferenceLiteralFactory<DataType>::CreateAny(InSettings, FLiteral::CreateInvalid());
+		}
+
+		/** Get the bound data read reference if it exists. Otherwise create and 
+		 * return a data read reference by constructing one the supplied constructor 
+		 * arguments.
+		 */
+		template<typename DataType, typename ... ConstructorArgTypes>
+		TDataReadReference<DataType> GetOrConstructDataReadReference(const FVertexName& InVertexName, ConstructorArgTypes&&... ConstructorArgs) const
+		{
+			if (const FBinding* Binding = Find(InVertexName))
+			{
+				return Binding->GetOrConstructDataReadReference<DataType>(Forward<ConstructorArgTypes>(ConstructorArgs)...);
+			}
+			else
+			{
+				UE_LOG(LogMetaSound, Warning, TEXT("Failed find vertex with name '%s'. Cannot check for existing bound data."), *InVertexName.ToString());
+			}
+
+			return TDataReadReference<DataType>::CreateNew(Forward<ConstructorArgTypes>(ConstructorArgs)...);
+		}
+
+
+		/** Get data write reference assuming data is bound and write accessible. */
+		template<typename DataType> 
+		TDataWriteReference<DataType> GetDataWriteReference(const FVertexName& InVertexName) const
+		{
+			const FBinding* Binding = FindChecked(InVertexName);
+			return Binding->GetDataWriteReference<DataType>();
+		}
+
+		/** Get the bound data write reference if it exists. Otherwise create and 
+		 * return a data write reference by constructing one using the Vertex's 
+		 * default literal.
+		 */
+		template<typename DataType> 
+		TDataWriteReference<DataType> GetOrCreateDefaultDataWriteReference(const FVertexName& InVertexName, const FOperatorSettings& InSettings) const
+		{
+			if (const FBinding* Binding = Find(InVertexName))
+			{
+				return Binding->GetOrCreateDefaultDataWriteReference<DataType>(InSettings);
+			}
+			else
+			{
+				UE_LOG(LogMetaSound, Warning, TEXT("Failed find vertex with name '%s'. Cannot check for existing bound data."), *InVertexName.ToString());
+			}
+
+			return TDataWriteReferenceLiteralFactory<DataType>::CreateAny(InSettings, FLiteral::CreateInvalid());
+		}
+
+		/** Get the bound data write reference if it exists. Otherwise create and 
+		 * return a data write reference by constructing one the supplied constructor 
+		 * arguments.
+		 */
+		template<typename DataType, typename ... ConstructorArgTypes>
+		TDataWriteReference<DataType> GetOrConstructDataWriteReference(const FVertexName& InVertexName, ConstructorArgTypes&&... ConstructorArgs) const
+		{
+			if (const FBinding* Binding = Find(InVertexName))
+			{
+				return Binding->GetOrConstructDataWriteReference<DataType>(Forward<ConstructorArgTypes>(ConstructorArgs)...);
+			}
+			else
+			{
+				UE_LOG(LogMetaSound, Warning, TEXT("Failed find vertex with name '%s'. Cannot check for existing bound data."), *InVertexName.ToString());
+			}
+
+			return TDataWriteReference<DataType>::CreateNew(Forward<ConstructorArgTypes>(ConstructorArgs)...);
+		}
+	
+	private:
+
+		FBinding* Find(const FVertexName& InVertexName);
+		const FBinding* Find(const FVertexName& InVertexName) const;
+
+		FBinding* FindChecked(const FVertexName& InVertexName);
+		const FBinding* FindChecked(const FVertexName& InVertexName) const;
+
+		TArray<FBinding> Bindings;
+	};
+
+	/** An output vertex interface with optionally bound data references. */
+	class METASOUNDGRAPHCORE_API FOutputVertexInterfaceData
+	{
+		using FBinding = MetasoundVertexDataPrivate::TBinding<FOutputDataVertex>;
+
+	public:
+
+		using FRangedForIteratorType = typename TArray<FBinding>::RangedForIteratorType;
+		using FRangedForConstIteratorType = typename TArray<FBinding>::RangedForConstIteratorType;
+
+		/* Construct using an FOutputVertexInterface. */
+		FOutputVertexInterfaceData(const FOutputVertexInterface& InVertexInterface);
+
+		/** Bind a data reference to a vertex. */
+		template<typename DataReferenceType>
+		void BindVertex(const FVertexName& InVertexName, const DataReferenceType& InDataReference)
+		{
+			BindVertex(InVertexName, FAnyDataReference{InDataReference});
+		}
+
+		/** Bind a data reference to a vertex. */
+		void BindVertex(const FVertexName& InVertexName, FAnyDataReference&& InDataReference);
+
+		/** Bind a data references to vertices with matching vertex names. */
+		void Bind(const FDataReferenceCollection& InCollection);
+
+		/** Returns true if a vertex with the given vertex name exists and is bound
+		 * to a data reference. */
+		bool IsVertexBound(const FVertexName& InVertexName) const;
+
+		/** Returns the access type of a bound vertex. If the vertex does not exist
+		 * or if it is unbound, this will return EDataReferenceAccessType::None
+		 */
+		EDataReferenceAccessType GetVertexDataAccessType(const FVertexName& InVertexName) const;
+
+		/** Returns true if all vertices in the FInputVertexInterface are bound to 
+		 * data references. */
+		bool AreAllVerticesBound() const;
+
+		FRangedForIteratorType begin()
+		{
+			return Bindings.begin();
+		}
+
+		FRangedForIteratorType end()
+		{
+			return Bindings.end();
+		}
+
+		FRangedForConstIteratorType begin() const
+		{
+			return Bindings.begin();
+		}
+
+		FRangedForConstIteratorType end() const
+		{
+			return Bindings.end();
+		}
+
+
+		/** Get data read reference assuming data is bound and read or write accessible. */
+		template<typename DataType> 
+		TDataReadReference<DataType> GetDataReadReference(const FVertexName& InVertexName) const
+		{
+			const FBinding* Binding = FindChecked(InVertexName);
+			return Binding->GetDataReadReference<DataType>();
+		}
+
+		/** Get the bound data read reference if it exists. Otherwise create and 
+		 * return a data read reference by constructing one the supplied constructor 
+		 * arguments.
+		 */
+		template<typename DataType, typename ... ConstructorArgTypes>
+		TDataReadReference<DataType> GetOrConstructDataReadReference(const FVertexName& InVertexName, ConstructorArgTypes&&... ConstructorArgs) const
+		{
+			if (const FBinding* Binding = Find(InVertexName))
+			{
+				return Binding->GetOrConstructDataReadReference<DataType>(Forward<ConstructorArgTypes>(ConstructorArgs)...);
+			}
+			else
+			{
+				UE_LOG(LogMetaSound, Warning, TEXT("Failed find vertex with name '%s'. Cannot check for existing bound data."), *InVertexName.ToString());
+			}
+
+			return TDataReadReference<DataType>::CreateNew(Forward<ConstructorArgTypes>(ConstructorArgs)...);
+		}
+
+
+		/** Get data write reference assuming data is bound and write accessible. */
+		template<typename DataType> 
+		TDataWriteReference<DataType> GetDataWriteReference(const FVertexName& InVertexName) const
+		{
+			const FBinding* Binding = FindChecked(InVertexName);
+			return Binding->GetDataWriteReference<DataType>();
+		}
+
+		/** Get the bound data write reference if it exists. Otherwise create and 
+		 * return a data write reference by constructing one the supplied constructor 
+		 * arguments.
+		 */
+		template<typename DataType, typename ... ConstructorArgTypes>
+		TDataWriteReference<DataType> GetOrConstructDataWriteReference(const FVertexName& InVertexName, ConstructorArgTypes&&... ConstructorArgs) const
+		{
+			if (const FBinding* Binding = Find(InVertexName))
+			{
+				return Binding->GetOrConstructDataWriteReference<DataType>(Forward<ConstructorArgTypes>(ConstructorArgs)...);
+			}
+			else
+			{
+				UE_LOG(LogMetaSound, Warning, TEXT("Failed find vertex with name '%s'. Cannot check for existing bound data."), *InVertexName.ToString());
+			}
+
+			return TDataWriteReference<DataType>::CreateNew(Forward<ConstructorArgTypes>(ConstructorArgs)...);
+		}
+	
+	private:
+
+		FBinding* Find(const FVertexName& InVertexName);
+		const FBinding* Find(const FVertexName& InVertexName) const;
+
+		FBinding* FindChecked(const FVertexName& InVertexName);
+		const FBinding* FindChecked(const FVertexName& InVertexName) const;
+
+		TArray<FBinding> Bindings;
+	};
+
+
+	/** A vertex interface with optionally bound data. */
+	class FVertexInterfaceData
+	{
+	public:
+
+		/** Construct using an FVertexInterface. */
+		FVertexInterfaceData(const FVertexInterface& InVertexInterface)
+		: InputVertexInterfaceData(InVertexInterface.GetInputInterface())
+		, OutputVertexInterfaceData(InVertexInterface.GetOutputInterface())
+		{
+		}
+
+		/** Get input vertex interface data. */
+		FInputVertexInterfaceData& GetInputs()
+		{
+			return InputVertexInterfaceData;
+		}
+
+		/** Get input vertex interface data. */
+		const FInputVertexInterfaceData& GetInputs() const
+		{
+			return InputVertexInterfaceData;
+		}
+
+		/** Get output vertex interface data. */
+		FOutputVertexInterfaceData& GetOutputs()
+		{
+			return OutputVertexInterfaceData;
+		}
+
+		/** Get output vertex interface data. */
+		const FOutputVertexInterfaceData& GetOutputs() const
+		{
+			return OutputVertexInterfaceData;
+		}
+
+	private:
+
+		FInputVertexInterfaceData InputVertexInterfaceData;
+		FOutputVertexInterfaceData OutputVertexInterfaceData;
+	};
+}
