Before attempting to load metasound assets on startup if their class info can't be found, check that the package mount point actually exists. This stops us from trying to load assets that are listed in the asset registry, but belong to plugins that we don't have loaded.

diff --git a/Engine/Plugins/Runtime/Metasound/Source/MetasoundEngine/Private/MetasoundAssetSubsystem.cpp b/Engine/Plugins/Runtime/Metasound/Source/MetasoundEngine/Private/MetasoundAssetSubsystem.cpp
--- a/Engine/Plugins/Runtime/Metasound/Source/MetasoundEngine/Private/MetasoundAssetSubsystem.cpp
+++ b/Engine/Plugins/Runtime/Metasound/Source/MetasoundEngine/Private/MetasoundAssetSubsystem.cpp
@@ -208,15 +208,18 @@ Metasound::Frontend::FNodeRegistryKey UMetaSoundAssetSubsystem::AddOrUpdateAsset
 		UObject* Object = nullptr;
 
 		FSoftObjectPath Path(InAssetData.ObjectPath);
-		if (InAssetData.IsAssetLoaded())
+		if (!FPackageName::GetPackageMountPoint(InAssetData.ObjectPath.ToString()).IsNone())
 		{
-			Object = Path.ResolveObject();
-			UE_LOG(LogMetaSound, Verbose, TEXT("Adding loaded asset '%s' to MetaSoundAsset registry."), *Object->GetName());
-		}
-		else
-		{
-			Object = Path.TryLoad();
-			UE_LOG(LogMetaSound, Verbose, TEXT("Loaded asset '%s' and adding to MetaSoundAsset registry."), *Object->GetName());
+			if (InAssetData.IsAssetLoaded())
+			{
+				Object = Path.ResolveObject();
+				UE_LOG(LogMetaSound, Verbose, TEXT("Adding loaded asset '%s' to MetaSoundAsset registry."), *Object->GetName());
+			}
+			else
+			{
+				Object = Path.TryLoad();
+				UE_LOG(LogMetaSound, Verbose, TEXT("Loaded asset '%s' and adding to MetaSoundAsset registry."), *Object->GetName());
+			}
 		}
 
 		if (Object)
