Fix crash when attempting to synchronize a MetaSound that has yet to be initialized/missing a graph reference.

diff --git a/Engine/Plugins/Runtime/Metasound/Source/MetasoundEditor/Private/MetasoundEditorGraphBuilder.cpp b/Engine/Plugins/Runtime/Metasound/Source/MetasoundEditor/Private/MetasoundEditorGraphBuilder.cpp
--- a/Engine/Plugins/Runtime/Metasound/Source/MetasoundEditor/Private/MetasoundEditorGraphBuilder.cpp
+++ b/Engine/Plugins/Runtime/Metasound/Source/MetasoundEditor/Private/MetasoundEditorGraphBuilder.cpp
@@ -10,6 +10,7 @@
 #include "EdGraph/EdGraphPin.h"
 #include "Framework/Notifications/NotificationManager.h"
 #include "GraphEditor.h"
+#include "Kismet/KismetSystemLibrary.h"
 #include "Metasound.h"
 #include "MetasoundAssetBase.h"
 #include "MetasoundEditor.h"
@@ -18,6 +19,7 @@
 #include "MetasoundEditorGraphInputNode.h"
 #include "MetasoundEditorGraphMemberDefaults.h"
 #include "MetasoundEditorGraphNode.h"
+#include "MetasoundEditorGraphSchema.h"
 #include "MetasoundEditorGraphValidation.h"
 #include "MetasoundEditorModule.h"
 #include "MetasoundEditorSettings.h"
@@ -1741,6 +1743,32 @@ namespace Metasound
 
 		bool FGraphBuilder::SynchronizeGraph(UObject& InMetaSound, bool bForceRefreshNodes)
 		{
+			FMetasoundAssetBase* MetaSoundAsset = IMetasoundUObjectRegistry::Get().GetObjectAsAssetBase(&InMetaSound);
+			check(MetaSoundAsset);
+
+			// If no graph is set, MetaSound has been created outside of asset factory, so initialize it here.
+			// TODO: Move factory initialization and this code to single builder function (in header so cannot move
+			// until 5.1+).
+			if(!MetaSoundAsset->GetGraph())
+			{
+				FString Author = UKismetSystemLibrary::GetPlatformUserName();
+				if (const UMetasoundEditorSettings* EditorSettings = GetDefault<UMetasoundEditorSettings>())
+				{
+					if (!EditorSettings->DefaultAuthor.IsEmpty())
+					{
+						Author = EditorSettings->DefaultAuthor;
+					}
+				}
+
+				FGraphBuilder::InitMetaSound(InMetaSound, Author);
+
+				// Initial graph generation is not something to be managed by the transaction
+				// stack, so don't track dirty state until after initial setup if necessary.
+				UMetasoundEditorGraph* Graph = NewObject<UMetasoundEditorGraph>(&InMetaSound, FName(), RF_Transactional);
+				Graph->Schema = UMetasoundEditorGraphSchema::StaticClass();
+				MetaSoundAsset->SetGraph(Graph);
+			}
+
 			bool bEditorGraphModified = SynchronizeGraphMembers(InMetaSound);
 			bEditorGraphModified |= SynchronizeNodeMembers(InMetaSound);
 			bEditorGraphModified |= SynchronizeNodes(InMetaSound);
@@ -1754,8 +1782,6 @@ namespace Metasound
 			bForceRefreshNodes |= bEditorGraphModified;
 			const bool bIsValid = FGraphBuilder::ValidateGraph(InMetaSound, bForceRefreshNodes);
 
-			FMetasoundAssetBase* MetaSoundAsset = IMetasoundUObjectRegistry::Get().GetObjectAsAssetBase(&InMetaSound);
-			check(MetaSoundAsset);
 			MetaSoundAsset->ResetSynchronizationState();
 
 			return bIsValid;
