Add type colors + icons to Metasound member list and stub in editor side for constructor pins
- Fix color for Time:Array type

diff --git a/Engine/Plugins/Runtime/Metasound/Source/MetasoundEditor/Private/MetasoundEditor.cpp b/Engine/Plugins/Runtime/Metasound/Source/MetasoundEditor/Private/MetasoundEditor.cpp
--- a/Engine/Plugins/Runtime/Metasound/Source/MetasoundEditor/Private/MetasoundEditor.cpp
+++ b/Engine/Plugins/Runtime/Metasound/Source/MetasoundEditor/Private/MetasoundEditor.cpp
@@ -487,19 +487,66 @@ namespace Metasound
 			{
 				TSharedRef<SWidget> TextWidget = SGraphPaletteItem::CreateTextSlotWidget(InCreateData, bIsReadOnly);
 
-				TSharedRef<SHorizontalBox> LayoutWidget = SNew(SHorizontalBox)
-				+ SHorizontalBox::Slot()
+				bool bIsConstructorPin = false;
+
+				const FSlateBrush* IconBrush = nullptr;
+				const FVector2D IconSize16 = FVector2D(16.0f, 16.0f);
+				FSlateColor IconColor = FSlateColor::UseForeground();
+
+				const bool bIsInterfaceMember = InterfaceVersion.IsValid();
+				const FSlateBrush* InterfaceIconBrush = bIsInterfaceMember ? FEditorStyle::GetBrush("Icons.Lock") : FStyleDefaults::GetNoBrush();
+
+				if (TSharedPtr<FMetasoundGraphMemberSchemaAction> GraphMemberAction = StaticCastSharedPtr<FMetasoundGraphMemberSchemaAction>(InCreateData->Action))
+				{
+
+					UMetasoundEditorGraphMember* GraphMember = GraphMemberAction->GetGraphMember();
+					if (const UMetasoundEditorGraphInput* Input = Cast<UMetasoundEditorGraphInput>(GraphMember))
+					{
+						bIsConstructorPin = Input->ConstructorPin;
+					}
+					FName DataTypeName = GraphMember->GetDataType(); 
+
+					IMetasoundEditorModule& EditorModule = FModuleManager::GetModuleChecked<IMetasoundEditorModule>("MetaSoundEditor");
+					if (const FEditorDataType* EditorDataType = EditorModule.FindDataType(DataTypeName))
+					{
+						FEdGraphPinType PinType = EditorDataType->PinType;
+						if (const UMetasoundEditorGraphSchema* Schema = GetDefault<UMetasoundEditorGraphSchema>())
+						{
+							IconColor = Schema->GetPinTypeColor(PinType);
+						}
+
+						IconBrush = EditorDataType->GetIconBrush(bIsConstructorPin);
+					}
+				}
+
+				TSharedRef<SHorizontalBox> LayoutWidget = SNew(SHorizontalBox);
+				LayoutWidget->AddSlot()
 				.AutoWidth()
 				.HAlign(HAlign_Left)
 				.VAlign(VAlign_Center)
 				[
 					SNew(SImage)
-					.Image(FEditorStyle::GetBrush(InterfaceVersion.IsValid() ? "Icons.Lock" : "Icons.BulletPoint"))
-					.ToolTipText(InterfaceVersion.IsValid() ? FText::Format(LOCTEXT("InterfaceMemberToolTipFormat", "Cannot Add/Remove: Member of interface '{0}'"), FText::FromName(InterfaceVersion.Name)): FText())
-					.ColorAndOpacity(FSlateColor::UseForeground())
-					.DesiredSizeOverride(FVector2D(16, 16))
-				]
-				+ SHorizontalBox::Slot()
+					.Image(IconBrush)
+					.ColorAndOpacity(IconColor)
+					.DesiredSizeOverride(IconSize16)
+				];
+
+				if (bIsInterfaceMember)
+				{
+					LayoutWidget->AddSlot()
+					.AutoWidth()
+					.HAlign(HAlign_Left)
+					.VAlign(VAlign_Center)
+					[
+						SNew(SImage)
+						.Image(InterfaceIconBrush)
+						.ToolTipText(bIsInterfaceMember ? FText::Format(LOCTEXT("InterfaceMemberToolTipFormat", "Cannot Add/Remove: Member of interface '{0}'"), FText::FromName(InterfaceVersion.Name)) : FText())
+						.ColorAndOpacity(FSlateColor::UseForeground())
+						.DesiredSizeOverride(IconSize16)
+					];
+				}
+				
+				LayoutWidget->AddSlot()
 				.AutoWidth()
 				.HAlign(HAlign_Left)
 				.VAlign(VAlign_Center)
diff --git a/Engine/Plugins/Runtime/Metasound/Source/MetasoundEditor/Private/MetasoundEditorModule.cpp b/Engine/Plugins/Runtime/Metasound/Source/MetasoundEditor/Private/MetasoundEditorModule.cpp
--- a/Engine/Plugins/Runtime/Metasound/Source/MetasoundEditor/Private/MetasoundEditorModule.cpp
+++ b/Engine/Plugins/Runtime/Metasound/Source/MetasoundEditor/Private/MetasoundEditorModule.cpp
@@ -67,6 +67,25 @@ namespace Metasound
 	namespace Editor
 	{
 		static const FName AssetToolName { "AssetTools" };
+		
+		const FSlateBrush* FEditorDataType::GetIconBrush(const bool bIsConstructorType) const
+		{
+			if (const ISlateStyle* MetasoundStyle = FSlateStyleRegistry::FindSlateStyle("MetaSoundStyle"))
+			{
+				if (PinType.IsArray())
+				{
+					return bIsConstructorType ? MetasoundStyle->GetBrush(TEXT("MetasoundEditor.Graph.ConstructorPinArray")) : MetasoundStyle->GetBrush(TEXT("MetasoundEditor.Graph.ArrayPin"));
+				}
+				else
+				{
+					return bIsConstructorType ? MetasoundStyle->GetBrush(TEXT("MetasoundEditor.Graph.ConstructorPin")) : FEditorStyle::GetBrush("Icons.BulletPoint");
+				}
+			}
+			else
+			{
+				return PinType.IsArray() ? FEditorStyle::GetBrush("Graph.ArrayPin.Connected") : FEditorStyle::GetBrush("Icons.BulletPoint");
+			}
+		}
 
 		template <typename T>
 		void AddAssetAction(IAssetTools& AssetTools, TArray<TSharedPtr<FAssetTypeActions_Base>>& AssetArray)
@@ -135,6 +154,12 @@ namespace Metasound
 					Set("MetasoundEditor.Graph.Node.Math.Power", new FSlateImageBrush(RootToContentDir(TEXT("/Graph/node_math_power_40x.png")), Icon40x40));
 					Set("MetasoundEditor.Graph.Node.Math.Logarithm", new FSlateImageBrush(RootToContentDir(TEXT("/Graph/node_math_logarithm_40x.png")), Icon40x40));
 					Set("MetasoundEditor.Graph.Node.Conversion", new FSlateImageBrush(RootToContentDir(TEXT("/Graph/node_conversion_40x.png")), Icon40x40));
+					
+					Set("MetasoundEditor.Graph.ConstructorPinArray", new FSlateImageBrush(RootToContentDir(TEXT("/Icons/array_pin_rotated.png")), Icon16));
+					Set("MetasoundEditor.Graph.ConstructorPinArrayDisconnected", new FSlateImageBrush(RootToContentDir(TEXT("/Icons/array_pin_rotated_disconnected.png")), Icon16));
+					Set("MetasoundEditor.Graph.ArrayPin", new FSlateImageBrush(RootToContentDir(TEXT("/Icons/array_pin.png")), Icon16));
+					Set("MetasoundEditor.Graph.ConstructorPin", new FSlateImageBrush(RootToContentDir(TEXT("/Icons/square_pin_rotated.png")), Icon16));
+					Set("MetasoundEditor.Graph.ConstructorPinDisconnected", new FSlateImageBrush(RootToContentDir(TEXT("/Icons/square_pin_rotated_disconnected.png")), Icon16));
 
 					// Analyzers
 					Set("MetasoundEditor.Analyzers.BackgroundColor", FLinearColor(0.0075f, 0.0075f, 0.0075, 1.0f));
@@ -460,7 +485,8 @@ namespace Metasound
 									//}
 
 									// Differentiate stronger numeric types associated with audio
-									if (DataTypeName == GetMetasoundDataTypeName<FTime>())
+									if (DataTypeName == GetMetasoundDataTypeName<FTime>() || 
+										DataTypeName == CreateArrayTypeNameFromElementTypeName(GetMetasoundDataTypeName<FTime>()))
 									{
 										PinSubCategory = FGraphBuilder::PinSubCategoryTime;
 									}
diff --git a/Engine/Plugins/Runtime/Metasound/Source/MetasoundEditor/Private/MetasoundEditorSettings.cpp b/Engine/Plugins/Runtime/Metasound/Source/MetasoundEditor/Private/MetasoundEditorSettings.cpp
--- a/Engine/Plugins/Runtime/Metasound/Source/MetasoundEditor/Private/MetasoundEditorSettings.cpp
+++ b/Engine/Plugins/Runtime/Metasound/Source/MetasoundEditor/Private/MetasoundEditorSettings.cpp
@@ -22,7 +22,7 @@ UMetasoundEditorSettings::UMetasoundEditorSettings(const FObjectInitializer& Obj
 	TimePinTypeColor = FLinearColor(0.3f, 1.0f, 1.0f, 1.0f);					// cyan
 	TriggerPinTypeColor = FLinearColor(1.0f, 1.0f, 1.0f, 1.0f);					// white
 
-	NativeNodeTitleColor = FLinearColor(0.4f, 0.85f, 0.35f, 1.0f);				// i forgot
+	NativeNodeTitleColor = FLinearColor(0.4f, 0.85f, 0.35f, 1.0f);				// pale green
 	AssetReferenceNodeTitleColor = FLinearColor(0.047f, 0.686f, 0.988f);		// sky blue
 	InputNodeTitleColor = FLinearColor(0.168f, 1.0f, 0.7294f);					// sea foam
 	OutputNodeTitleColor = FLinearColor(1.0f, 0.878f, 0.1686f);					// yellow
diff --git a/Engine/Plugins/Runtime/Metasound/Source/MetasoundEditor/Private/MetasoundNodeDetailCustomization.cpp b/Engine/Plugins/Runtime/Metasound/Source/MetasoundEditor/Private/MetasoundNodeDetailCustomization.cpp
--- a/Engine/Plugins/Runtime/Metasound/Source/MetasoundEditor/Private/MetasoundNodeDetailCustomization.cpp
+++ b/Engine/Plugins/Runtime/Metasound/Source/MetasoundEditor/Private/MetasoundNodeDetailCustomization.cpp
@@ -55,8 +55,11 @@ namespace Metasound
 
 			static const FText OverrideInputDefaultText = LOCTEXT("OverridePresetInputDefault", "Override Inherited Default");
 			static const FText OverrideInputDefaultTooltip = LOCTEXT("OverridePresetInputTooltip",
-				"Enables overriding the input's inherited default value otherwise provided by the referenced graph."
-				"Setting to true disables auto-updating the input's default value if modified on the referenced asset.");
+				"Enables overriding the input's inherited default value otherwise provided by the referenced graph. Setting to true disables auto-updating the input's default value if modified on the referenced asset.");
+
+			static const FText ConstructorPinText = LOCTEXT("ConstructorPinText", "Is Constructor Pin");
+			static const FText ConstructorPinTooltip = LOCTEXT("ConstructorPinTooltip",
+				"Whether this input is a constructor pin (value is only read on construction, and is not updated at runtime).");
 
 			void GetDataTypeFromElementPropertyHandle(TSharedPtr<IPropertyHandle> ElementPropertyHandle, Frontend::FDataTypeRegistryInfo& OutDataTypeInfo)
 			{
@@ -726,7 +729,7 @@ namespace Metasound
 			[
 				SNew(STextBlock)
 				.Text(InRowName)
-				.Font(IDetailLayoutBuilder::GetDetailFont())
+				.Font(IDetailLayoutBuilder::GetDetailFontBold())
 			]
 			.ValueContent()
 			[
@@ -1255,6 +1258,19 @@ namespace Metasound
 				}
 			}
 		}
+		void FMetasoundInputDetailCustomization::CustomizeGeneralCategory(IDetailLayoutBuilder& InDetailLayout)
+		{
+			FMetasoundVertexDetailCustomization::CustomizeGeneralCategory(InDetailLayout);
+
+			if (UMetasoundEditorGraphInput* Input = Cast<UMetasoundEditorGraphInput>(GraphMember.Get()))
+			{
+				IDetailCategoryBuilder& CategoryBuilder = FMetasoundMemberDetailCustomization::GetGeneralCategoryBuilder(InDetailLayout);
+				// TODO: Remove ConstructorPin bool from UMetasoundEditorGraphInput and add constructor pin bool property from frontend here 
+				/*IDetailPropertyRow* Row = CategoryBuilder.AddExternalObjectProperty(TArray<UObject*>({ Input }), GET_MEMBER_NAME_CHECKED(UMetasoundEditorGraphInput, ConstructorPin));
+				auto PropertyEnabled = TAttribute<bool>::CreateLambda([this] { return !GetDocumentHandle()->GetRootGraphClass().PresetOptions.bIsPreset; });
+				Row->EditCondition(PropertyEnabled, { });*/
+			}
+		}
 
 		void FMetasoundInputDetailCustomization::CustomizeDetails(IDetailLayoutBuilder& InDetailLayout)
 		{
diff --git a/Engine/Plugins/Runtime/Metasound/Source/MetasoundEditor/Private/MetasoundNodeDetailCustomization.h b/Engine/Plugins/Runtime/Metasound/Source/MetasoundEditor/Private/MetasoundNodeDetailCustomization.h
--- a/Engine/Plugins/Runtime/Metasound/Source/MetasoundEditor/Private/MetasoundNodeDetailCustomization.h
+++ b/Engine/Plugins/Runtime/Metasound/Source/MetasoundEditor/Private/MetasoundNodeDetailCustomization.h
@@ -361,10 +361,13 @@ namespace Metasound
 		{
 		public:
 			virtual ~FMetasoundInputDetailCustomization() = default;
-
+			
 			virtual void CustomizeDetails(IDetailLayoutBuilder& InDetailLayout) override;
 			virtual bool IsDefaultEditable() const override;
 
+		protected: 
+			virtual void CustomizeGeneralCategory(IDetailLayoutBuilder& InDetailLayout) override;
+
 		private:
 			bool GetInputInheritsDefault() const;
 			void SetInputInheritsDefault();
diff --git a/Engine/Plugins/Runtime/Metasound/Source/MetasoundEditor/Private/SMetasoundGraphPin.h b/Engine/Plugins/Runtime/Metasound/Source/MetasoundEditor/Private/SMetasoundGraphPin.h
--- a/Engine/Plugins/Runtime/Metasound/Source/MetasoundEditor/Private/SMetasoundGraphPin.h
+++ b/Engine/Plugins/Runtime/Metasound/Source/MetasoundEditor/Private/SMetasoundGraphPin.h
@@ -203,6 +203,43 @@ namespace Metasound
 						]
 					];
 			}
+
+			virtual const FSlateBrush* GetPinIcon() const
+			{
+				const bool bIsConnected = ParentPinType::IsConnected();
+				bool bIsConstructorPin = false;
+				if (const UEdGraphPin* Pin = ParentPinType::GetPinObj())
+				{
+					if (const UMetasoundEditorGraphNode* Node = Cast<UMetasoundEditorGraphNode>(Pin->GetOwningNode()))
+					{
+						if (const UMetasoundEditorGraphMemberNode* MemberNode = Cast<UMetasoundEditorGraphMemberNode>(Node))
+						{
+							if (const UMetasoundEditorGraphInput* Input = Cast<UMetasoundEditorGraphInput>(MemberNode->GetMember()))
+							{
+								bIsConstructorPin = Input->ConstructorPin; 
+							}
+						}
+					}
+				}
+
+				if (bIsConstructorPin)
+				{
+					if (const ISlateStyle* MetasoundStyle = FSlateStyleRegistry::FindSlateStyle("MetaSoundStyle"))
+					{
+						if (ParentPinType::IsArray())
+						{
+							return bIsConnected ? MetasoundStyle->GetBrush(TEXT("MetasoundEditor.Graph.ConstructorPinArray")) : 
+								MetasoundStyle->GetBrush(TEXT("MetasoundEditor.Graph.ConstructorPinArrayDisconnected"));
+						}
+						else
+						{
+							return bIsConnected ? MetasoundStyle->GetBrush(TEXT("MetasoundEditor.Graph.ConstructorPin")) : 
+								MetasoundStyle->GetBrush(TEXT("MetasoundEditor.Graph.ConstructorPinDisconnected"));
+						}
+					}
+				}
+				return SGraphPin::GetPinIcon();
+			}
 		};
 
 		class SMetasoundGraphPin : public TMetasoundGraphPin<SGraphPin>
diff --git a/Engine/Plugins/Runtime/Metasound/Source/MetasoundEditor/Public/MetasoundEditorGraph.h b/Engine/Plugins/Runtime/Metasound/Source/MetasoundEditor/Public/MetasoundEditorGraph.h
--- a/Engine/Plugins/Runtime/Metasound/Source/MetasoundEditor/Public/MetasoundEditorGraph.h
+++ b/Engine/Plugins/Runtime/Metasound/Source/MetasoundEditor/Public/MetasoundEditorGraph.h
@@ -250,6 +250,10 @@ class METASOUNDEDITOR_API UMetasoundEditorGraphInput : public UMetasoundEditorGr
 	GENERATED_BODY()
 
 public:
+	/** Whether this input is a constructor pin (value is only read on construction, and is not updated at runtime). */
+	UPROPERTY(Transient, EditAnywhere, Category = General)
+	bool ConstructorPin = false;
+
 	virtual int32 GetSortOrderIndex() const override;
 	virtual void SetSortOrderIndex(int32 InSortOrderIndex) override;
 
diff --git a/Engine/Plugins/Runtime/Metasound/Source/MetasoundEditor/Public/MetasoundEditorModule.h b/Engine/Plugins/Runtime/Metasound/Source/MetasoundEditor/Public/MetasoundEditorModule.h
--- a/Engine/Plugins/Runtime/Metasound/Source/MetasoundEditor/Public/MetasoundEditorModule.h
+++ b/Engine/Plugins/Runtime/Metasound/Source/MetasoundEditor/Public/MetasoundEditorModule.h
@@ -9,6 +9,7 @@
 #include "MetasoundFrontendRegistries.h"
 #include "Modules/ModuleInterface.h"
 #include "PropertyHandle.h"
+#include "Styling/SlateStyleRegistry.h"
 #include "Templates/Function.h"
 
 // Forward Declarations
@@ -55,6 +56,9 @@ namespace Metasound
 				, RegistryInfo(InRegistryInfo)
 			{
 			}
+
+			// Get the corresponding icon brush for this type 
+			const FSlateBrush* GetIconBrush(const bool bIsConstructorType) const;
 		};
 
 		class METASOUNDEDITOR_API FMetasoundDefaultLiteralCustomizationBase
