Metasound Frontend Search Engine Memory Improvements
- Split out search engine into Core functionality and EditorOnly functionality
- Removed some ISearchEngine integration in the INodeController
- Fixed up callsites that were relying on EditorOnly ISearchEngine functionality on non-editor builds.

diff --git a/Engine/Plugins/Runtime/Metasound/Source/MetasoundEditor/Private/MetasoundEditor.cpp b/Engine/Plugins/Runtime/Metasound/Source/MetasoundEditor/Private/MetasoundEditor.cpp
--- a/Engine/Plugins/Runtime/Metasound/Source/MetasoundEditor/Private/MetasoundEditor.cpp
+++ b/Engine/Plugins/Runtime/Metasound/Source/MetasoundEditor/Private/MetasoundEditor.cpp
@@ -2607,11 +2607,14 @@ namespace Metasound
 			{
 				if (UMetasoundEditorGraphExternalNode* ExternalNode = Cast<UMetasoundEditorGraphExternalNode>(Object))
 				{
-					FMetasoundFrontendVersionNumber HighestVersion = ExternalNode->FindHighestVersionInRegistry();
-					Metasound::Frontend::FConstNodeHandle NodeHandle = ExternalNode->GetConstNodeHandle();
+					Metasound::Frontend::FNodeHandle NodeHandle = ExternalNode->GetNodeHandle();
 					const FMetasoundFrontendClassMetadata& Metadata = NodeHandle->GetClassMetadata();
+
+					// Check for new version
+					FMetasoundFrontendVersionNumber HighestVersion = ExternalNode->FindHighestVersionInRegistry();
 					const bool bHasNewVersion = HighestVersion.IsValid() && HighestVersion > Metadata.GetVersion();
 
+					// Check for non-native classes
 					const FNodeRegistryKey RegistryKey = NodeRegistryKey::CreateKey(Metadata);
 					const bool bIsClassNative = FMetasoundFrontendRegistryContainer::Get()->IsNodeNative(RegistryKey);
 
@@ -2622,8 +2625,7 @@ namespace Metasound
 						constexpr TArray<INodeController::FVertexNameAndType>* DisconnectedInputs = nullptr;
 						constexpr TArray<INodeController::FVertexNameAndType>* DisconnectedOutputs = nullptr;
 
-						FNodeHandle ExistingNode = ExternalNode->GetNodeHandle();
-						FNodeHandle NewNode = ExistingNode->ReplaceWithVersion(HighestVersion, DisconnectedInputs, DisconnectedOutputs);
+						FNodeHandle NewNode = NodeHandle->ReplaceWithVersion(HighestVersion, DisconnectedInputs, DisconnectedOutputs);
 						bReplacedNodes = true;
 					}
 				}
diff --git a/Engine/Plugins/Runtime/Metasound/Source/MetasoundEditor/Private/MetasoundEditorGraphNode.cpp b/Engine/Plugins/Runtime/Metasound/Source/MetasoundEditor/Private/MetasoundEditorGraphNode.cpp
--- a/Engine/Plugins/Runtime/Metasound/Source/MetasoundEditor/Private/MetasoundEditorGraphNode.cpp
+++ b/Engine/Plugins/Runtime/Metasound/Source/MetasoundEditor/Private/MetasoundEditorGraphNode.cpp
@@ -712,12 +712,19 @@ void UMetasoundEditorGraphExternalNode::ReconstructNode()
 
 FMetasoundFrontendVersionNumber UMetasoundEditorGraphExternalNode::FindHighestVersionInRegistry() const
 {
-	return GetConstNodeHandle()->FindHighestVersionInRegistry();
-}
+	using namespace Metasound::Frontend;
 
-FMetasoundFrontendVersionNumber UMetasoundEditorGraphExternalNode::FindHighestMinorVersionInRegistry() const
-{
-	return GetConstNodeHandle()->FindHighestMinorVersionInRegistry();
+	FMetasoundFrontendClass HighestVersionClass;
+	FMetasoundFrontendVersionNumber HighestVersionNumber = FMetasoundFrontendVersionNumber::GetInvalid();
+
+	Metasound::Frontend::FConstNodeHandle NodeHandle = GetConstNodeHandle();
+	const FMetasoundFrontendClassMetadata& Metadata = NodeHandle->GetClassMetadata();
+	if (ISearchEngine::Get().FindClassWithHighestVersion(Metadata.GetClassName(), HighestVersionClass))
+	{
+		HighestVersionNumber = HighestVersionClass.Metadata.GetVersion();
+	}
+
+	return HighestVersionNumber;
 }
 
 bool UMetasoundEditorGraphExternalNode::CanAutoUpdate() const
diff --git a/Engine/Plugins/Runtime/Metasound/Source/MetasoundEditor/Public/MetasoundEditorGraphNode.h b/Engine/Plugins/Runtime/Metasound/Source/MetasoundEditor/Public/MetasoundEditorGraphNode.h
--- a/Engine/Plugins/Runtime/Metasound/Source/MetasoundEditor/Public/MetasoundEditorGraphNode.h
+++ b/Engine/Plugins/Runtime/Metasound/Source/MetasoundEditor/Public/MetasoundEditorGraphNode.h
@@ -219,7 +219,6 @@ public:
 	virtual void CacheTitle() override;
 
 	FMetasoundFrontendVersionNumber FindHighestVersionInRegistry() const;
-	FMetasoundFrontendVersionNumber FindHighestMinorVersionInRegistry() const;
 	bool CanAutoUpdate() const;
 
 	// Validates node and returns whether or not the node is valid.
diff --git a/Engine/Plugins/Runtime/Metasound/Source/MetasoundEngineTest/Private/EngineTestMetasoundSource.cpp b/Engine/Plugins/Runtime/Metasound/Source/MetasoundEngineTest/Private/EngineTestMetasoundSource.cpp
--- a/Engine/Plugins/Runtime/Metasound/Source/MetasoundEngineTest/Private/EngineTestMetasoundSource.cpp
+++ b/Engine/Plugins/Runtime/Metasound/Source/MetasoundEngineTest/Private/EngineTestMetasoundSource.cpp
@@ -55,7 +55,6 @@ namespace EngineTestMetasoundSourcePrivate
 	Metasound::Frontend::FNodeHandle AddNode(Metasound::Frontend::IGraphController& InGraph, const Metasound::FNodeClassName& InClassName)
 	{
 		FMetasoundFrontendClass NodeClass;
-		ensure(Metasound::Frontend::ISearchEngine::Get().FindClassWithHighestVersion(InClassName, NodeClass));
 		Metasound::Frontend::FNodeHandle Node = InGraph.AddNode(NodeClass.Metadata);
 		check(Node->IsValid());
 		return Node;
diff --git a/Engine/Plugins/Runtime/Metasound/Source/MetasoundFrontend/Private/MetasoundAssetBase.cpp b/Engine/Plugins/Runtime/Metasound/Source/MetasoundFrontend/Private/MetasoundAssetBase.cpp
--- a/Engine/Plugins/Runtime/Metasound/Source/MetasoundFrontend/Private/MetasoundAssetBase.cpp
+++ b/Engine/Plugins/Runtime/Metasound/Source/MetasoundFrontend/Private/MetasoundAssetBase.cpp
@@ -218,7 +218,7 @@ void FMetasoundAssetBase::RegisterGraphWithFrontend(Metasound::Frontend::FMetaSo
 		RegistryKey = FMetasoundFrontendRegistryContainer::Get()->RegisterNode(MakeUnique<FNodeRegistryEntry>(AssetName, *Doc, AssetClassInfo.AssetPath));
 	}
 
-	if (NodeRegistryKey::IsValid(RegistryKey))
+	if (NodeRegistryKey::IsValid(RegistryKey) && FMetasoundFrontendRegistryContainer::Get()->IsNodeRegistered(RegistryKey))
 	{
 #if WITH_EDITORONLY_DATA
 		// Refresh Asset Registry Info if successfully registered with Frontend
@@ -227,7 +227,6 @@ void FMetasoundAssetBase::RegisterGraphWithFrontend(Metasound::Frontend::FMetaSo
 		AssetClassInfo.AssetClassID = FGuid(DocumentClassMetadata.GetClassName().Name.ToString());
 		FNodeClassName ClassName = DocumentClassMetadata.GetClassName().ToNodeClassName();
 		FMetasoundFrontendClass GraphClass;
-		ensure(ISearchEngine::Get().FindClassWithMajorVersion(ClassName, DocumentClassMetadata.GetVersion().Major, GraphClass));
 
 		AssetClassInfo.Version = DocumentClassMetadata.GetVersion();
 
diff --git a/Engine/Plugins/Runtime/Metasound/Source/MetasoundFrontend/Private/MetasoundFrontendInvalidController.h b/Engine/Plugins/Runtime/Metasound/Source/MetasoundFrontend/Private/MetasoundFrontendInvalidController.h
--- a/Engine/Plugins/Runtime/Metasound/Source/MetasoundFrontend/Private/MetasoundFrontendInvalidController.h
+++ b/Engine/Plugins/Runtime/Metasound/Source/MetasoundFrontend/Private/MetasoundFrontendInvalidController.h
@@ -248,8 +248,6 @@ namespace Metasound
 
 			virtual bool DiffAgainstRegistryInterface(FClassInterfaceUpdates& OutInterfaceUpdates, bool bInUseHighestMinorVersion) const override { return false; }
 			virtual bool CanAutoUpdate(FClassInterfaceUpdates& OutInterfaceUpdates) const override { OutInterfaceUpdates = { }; return false; }
-			virtual FMetasoundFrontendVersionNumber FindHighestVersionInRegistry() const override { return FMetasoundFrontendVersionNumber::GetInvalid(); }
-			virtual FMetasoundFrontendVersionNumber FindHighestMinorVersionInRegistry() const override { return FMetasoundFrontendVersionNumber::GetInvalid(); }
 
 			virtual TSharedRef<IGraphController> AsGraph() override;
 			virtual TSharedRef<const IGraphController> AsGraph() const override;
diff --git a/Engine/Plugins/Runtime/Metasound/Source/MetasoundFrontend/Private/MetasoundFrontendNodeController.cpp b/Engine/Plugins/Runtime/Metasound/Source/MetasoundFrontend/Private/MetasoundFrontendNodeController.cpp
--- a/Engine/Plugins/Runtime/Metasound/Source/MetasoundFrontend/Private/MetasoundFrontendNodeController.cpp
+++ b/Engine/Plugins/Runtime/Metasound/Source/MetasoundFrontend/Private/MetasoundFrontendNodeController.cpp
@@ -690,56 +690,23 @@ namespace Metasound
 			return IGraphController::GetInvalidHandle();
 		}
 
-		FMetasoundFrontendVersionNumber FBaseNodeController::FindHighestMinorVersionInRegistry() const
-		{
-			const FMetasoundFrontendClassMetadata& Metadata = GetClassMetadata();
-			const FMetasoundFrontendVersionNumber& CurrentVersion = Metadata.GetVersion();
-			Metasound::FNodeClassName NodeClassName = Metadata.GetClassName().ToNodeClassName();
-
-			FMetasoundFrontendClass ClassWithMajorVersion;
-			if (ISearchEngine::Get().FindClassWithMajorVersion(NodeClassName, CurrentVersion.Major, ClassWithMajorVersion))
-			{
-				if (ClassWithMajorVersion.Metadata.GetVersion().Minor >= CurrentVersion.Minor)
-				{
-					return ClassWithMajorVersion.Metadata.GetVersion();
-				}
-			}
-
-			return FMetasoundFrontendVersionNumber::GetInvalid();
-		}
-
-		FMetasoundFrontendVersionNumber FBaseNodeController::FindHighestVersionInRegistry() const
-		{
-			const FMetasoundFrontendClassMetadata& Metadata = GetClassMetadata();
-			const FMetasoundFrontendVersionNumber& CurrentVersion = Metadata.GetVersion();
-			Metasound::FNodeClassName NodeClassName = Metadata.GetClassName().ToNodeClassName();
-
-			FMetasoundFrontendClass ClassWithHighestVersion;
-			if (ISearchEngine::Get().FindClassWithHighestVersion(NodeClassName, ClassWithHighestVersion))
-			{
-				if (ClassWithHighestVersion.Metadata.GetVersion().Major >= CurrentVersion.Major)
-				{
-					return ClassWithHighestVersion.Metadata.GetVersion();
-				}
-			}
-
-			return FMetasoundFrontendVersionNumber::GetInvalid();
-
-		}
-
 		FNodeHandle FBaseNodeController::ReplaceWithVersion(const FMetasoundFrontendVersionNumber& InNewVersion, TArray<FVertexNameAndType>* OutDisconnectedInputs, TArray<FVertexNameAndType>* OutDisconnectedOutputs)
 		{
 			const FMetasoundFrontendClassMetadata Metadata = GetClassMetadata();
-			const TArray<FMetasoundFrontendClass> Versions = ISearchEngine::Get().FindClassesWithName(Metadata.GetClassName().ToNodeClassName(), false /* bInSortByVersion */);
 
-			auto IsClassOfNewVersion = [InNewVersion](const FMetasoundFrontendClass& RegisteredClass)
-			{
-				return RegisteredClass.Metadata.GetVersion() == InNewVersion;
-			};
+			// Lookup new version in node registry
+			FNodeRegistryKey NewVersionRegistryKey = NodeRegistryKey::CreateKey(Metadata.GetType(), Metadata.GetClassName().ToString(), InNewVersion.Major, InNewVersion.Minor);
+			FMetasoundFrontendRegistryContainer* Registry = FMetasoundFrontendRegistryContainer::Get();
+			checkf(nullptr != Registry, TEXT("The metasound node registry should always be available if the metasound plugin is loaded"));
 
-			const FMetasoundFrontendClass* RegisteredClass = Versions.FindByPredicate(IsClassOfNewVersion);
-			if (!ensure(RegisteredClass))
+			FMetasoundFrontendClass NewMetasoundClass;
+			const bool bFoundNewClass = Registry->FindFrontendClassFromRegistered(NewVersionRegistryKey, NewMetasoundClass);
+			if (!bFoundNewClass)
 			{
+				FString ClassNameString = Metadata.GetClassName().ToString();
+				FString NewVersionString = InNewVersion.ToString();
+				UE_LOG(LogMetaSound, Error, TEXT("Failed to change class version from %s to %s for class %s. %s %s is not registered."), *Metadata.GetVersion().ToString(), *NewVersionString, *ClassNameString, *ClassNameString, *NewVersionString);
+
 				return this->AsShared();
 			}
 
@@ -818,13 +785,13 @@ namespace Metasound
 			// Make sure classes are up-to-date with registered versions of class.
 			// Note that this may break other nodes in the graph that have stale
 			// class API, but that's on the caller to fix-up or report invalid state.
-			const FNodeRegistryKey RegistryKey = NodeRegistryKey::CreateKey(RegisteredClass->Metadata);
+			const FNodeRegistryKey RegistryKey = NodeRegistryKey::CreateKey(NewMetasoundClass.Metadata);
 			FDocumentHandle Document = GetOwningGraph()->GetOwningDocument();
 
 			constexpr bool bRefreshFromRegistry = true;
 			ensureAlways(Document->FindOrAddClass(RegistryKey, bRefreshFromRegistry).Get() != nullptr);
 
-			FNodeHandle ReplacementNode = GetOwningGraph()->AddNode(RegisteredClass->Metadata, ReplacedNodeGuid);
+			FNodeHandle ReplacementNode = GetOwningGraph()->AddNode(NewMetasoundClass.Metadata, ReplacedNodeGuid);
 			if (!ensureAlways(ReplacementNode->IsValid()))
 			{
 				return this->AsShared();
@@ -918,7 +885,7 @@ namespace Metasound
 
 			if (bInUseHighestMinorVersion)
 			{
-				if (!ISearchEngine::Get().FindClassWithMajorVersion(NodeClassName, NodeClassMetadata.GetVersion().Major, OutInterfaceUpdates.RegistryClass))
+				if (!ISearchEngine::Get().FindClassWithHighestMinorVersion(NodeClassName, NodeClassMetadata.GetVersion().Major, OutInterfaceUpdates.RegistryClass))
 				{
 					Algo::Transform(NodeClassInterface.Inputs, OutInterfaceUpdates.RemovedInputs, [&](const FMetasoundFrontendClassInput& Input) { return &Input; });
 					Algo::Transform(NodeClassInterface.Outputs, OutInterfaceUpdates.RemovedOutputs, [&](const FMetasoundFrontendClassOutput& Output) { return &Output; });
@@ -927,20 +894,20 @@ namespace Metasound
 			}
 			else
 			{
-				constexpr bool bSortByVersion = true;
-				const TArray<FMetasoundFrontendClass> Classes = ISearchEngine::Get().FindClassesWithName(NodeClassName, bSortByVersion);
-				const FMetasoundFrontendClass* ExactClass = Classes.FindByPredicate([CurrentVersion = &NodeClassMetadata.GetVersion()](const FMetasoundFrontendClass& AvailableClass)
-				{
-					return AvailableClass.Metadata.GetVersion() == *CurrentVersion;
-				});
+				// Find class with same metadata in the node registry.
+				FMetasoundFrontendRegistryContainer* Registry = FMetasoundFrontendRegistryContainer::Get();
+				checkf(nullptr != Registry, TEXT("The metasound node registry should always be available if the metasound plugin is loaded"));
+				FMetasoundFrontendClass RegisteredClass;
+				bool bFoundRegisteredClass = Registry->FindFrontendClassFromRegistered(NodeRegistryKey::CreateKey(GetClassMetadata()), OutInterfaceUpdates.RegistryClass);
 
-				if (!ExactClass)
+				if (!bFoundRegisteredClass)
 				{
+					// If the class was not found, mark all inputs and outputs as removed.
+					UE_LOG(LogMetaSound, Warning, TEXT("Could not find registered version of interface. %s %s is not registered."), *NodeClassMetadata.GetClassName().ToString(), *NodeClassMetadata.GetVersion().ToString());
 					Algo::Transform(NodeClassInterface.Inputs, OutInterfaceUpdates.RemovedInputs, [&](const FMetasoundFrontendClassInput& Input) { return &Input; });
 					Algo::Transform(NodeClassInterface.Outputs, OutInterfaceUpdates.RemovedOutputs, [&](const FMetasoundFrontendClassOutput& Output) { return &Output; });
 					return false;
 				}
-				OutInterfaceUpdates.RegistryClass = *ExactClass;
 			}
 
 			Algo::Transform(OutInterfaceUpdates.RegistryClass.Interface.Inputs, OutInterfaceUpdates.AddedInputs, [&](const FMetasoundFrontendClassInput& Input) { return &Input; });
@@ -1005,10 +972,7 @@ namespace Metasound
 			}
 
 			FMetasoundFrontendClass RegistryClass;
-			if (!ISearchEngine::Get().FindClassWithMajorVersion(
-				NodeClassMetadata.GetClassName().ToNodeClassName(),
-				NodeClassMetadata.GetVersion().Major,
-				RegistryClass))
+			if (!ISearchEngine::Get().FindClassWithHighestMinorVersion( NodeClassMetadata.GetClassName().ToNodeClassName(), NodeClassMetadata.GetVersion().Major, RegistryClass))
 			{
 				return false;
 			}
diff --git a/Engine/Plugins/Runtime/Metasound/Source/MetasoundFrontend/Private/MetasoundFrontendNodeController.h b/Engine/Plugins/Runtime/Metasound/Source/MetasoundFrontend/Private/MetasoundFrontendNodeController.h
--- a/Engine/Plugins/Runtime/Metasound/Source/MetasoundFrontend/Private/MetasoundFrontendNodeController.h
+++ b/Engine/Plugins/Runtime/Metasound/Source/MetasoundFrontend/Private/MetasoundFrontendNodeController.h
@@ -60,8 +60,6 @@ namespace Metasound
 
 			virtual bool CanAutoUpdate(FClassInterfaceUpdates& OutInterfaceUpdates) const override;
 			virtual FNodeHandle ReplaceWithVersion(const FMetasoundFrontendVersionNumber& InNewVersion, TArray<FVertexNameAndType>* OutDisconnectedInputs, TArray<FVertexNameAndType>* OutDisconnectedOutputs) override;
-			virtual FMetasoundFrontendVersionNumber FindHighestVersionInRegistry() const override;
-			virtual FMetasoundFrontendVersionNumber FindHighestMinorVersionInRegistry() const override;
 
 			virtual const FVertexName& GetNodeName() const override;
 
diff --git a/Engine/Plugins/Runtime/Metasound/Source/MetasoundFrontend/Private/MetasoundFrontendSearchEngine.cpp b/Engine/Plugins/Runtime/Metasound/Source/MetasoundFrontend/Private/MetasoundFrontendSearchEngine.cpp
--- a/Engine/Plugins/Runtime/Metasound/Source/MetasoundFrontend/Private/MetasoundFrontendSearchEngine.cpp
+++ b/Engine/Plugins/Runtime/Metasound/Source/MetasoundFrontend/Private/MetasoundFrontendSearchEngine.cpp
@@ -1,10 +1,10 @@
 // Copyright Epic Games, Inc. All Rights Reserved.
 
 #include "MetasoundFrontendSearchEngine.h"
+#include "MetasoundFrontendSearchEngineCore.h"
+#include "MetasoundFrontendSearchEngineEditorOnly.h"
 
 #include "Algo/MaxElement.h"
-#include "Algo/Transform.h"
-#include "HAL/IConsoleManager.h"
 #include "MetasoundFrontendArchetypeRegistry.h"
 #include "MetasoundFrontendQuery.h"
 #include "MetasoundFrontendQuerySteps.h"
@@ -12,880 +12,19 @@
 #include "MetasoundTrace.h"
 #include "MetasoundAssetManager.h"
 
-static int32 ClearMetaSoundFrontendSearchEngineCacheCVar = 0;
-FAutoConsoleVariableRef CVarClearMetaSoundFrontendSearchEngineCache(
-	TEXT("au.Debug.MetaSounds.SearchEngine.ClearCache"),
-	ClearMetaSoundFrontendSearchEngineCacheCVar,
-	TEXT("If true, flags the MetaSound registry SearchEngine cache to be cleared on subsequent query.\n")
-	TEXT("0: Do not clear, !0: Flag to clear"),
-	ECVF_Default);
-
-
 namespace Metasound
 {
 	namespace Frontend
 	{
-		namespace SearchEngineQuerySteps
-		{
-			class FMapClassToFullClassNameAndMajorVersion : public IFrontendQueryMapStep
-			{
-			public:
-				static FFrontendQueryKey GetKey(const FMetasoundFrontendClassName& InClassName, int32 InMajorVersion)
-				{
-					return FFrontendQueryKey(FString::Format(TEXT("{0}_v{1}"), {InClassName.ToString(), InMajorVersion}));
-				}
-
-				virtual FFrontendQueryKey Map(const FFrontendQueryEntry& InEntry) const override
-				{
-					if (ensure(InEntry.Value.IsType<FMetasoundFrontendClass>()))
-					{
-						const FMetasoundFrontendClassMetadata& Metadata = InEntry.Value.Get<FMetasoundFrontendClass>().Metadata;
-						return FMapClassToFullClassNameAndMajorVersion::GetKey(Metadata.GetClassName(), Metadata.GetVersion().Major);
-					}
-					return FFrontendQueryKey();
-				}
-			};
-
-			class FInterfaceRegistryTransactionSource : public IFrontendQueryStreamStep
-			{
-			public:
-				FInterfaceRegistryTransactionSource()
-				: CurrentTransactionID(GetOriginRegistryTransactionID())
-				{
-				}
-
-				virtual void Stream(TArray<FFrontendQueryValue>& OutEntries) override
-				{
-					auto AddValue = [&OutEntries](const FInterfaceRegistryTransaction& InTransaction)
-					{
-						OutEntries.Emplace(TInPlaceType<FInterfaceRegistryTransaction>(), InTransaction);
-					};
-					
-					IInterfaceRegistry::Get().ForEachRegistryTransactionSince(CurrentTransactionID, &CurrentTransactionID, AddValue);
-				}
-
-			private:
-				FRegistryTransactionID CurrentTransactionID;
-			};
-
-			class FMapInterfaceRegistryTransactionsToInterfaceRegistryKeys : public IFrontendQueryMapStep
-			{
-			public:
-				virtual FFrontendQueryKey Map(const FFrontendQueryEntry& InEntry) const override
-				{
-					FInterfaceRegistryKey RegistryKey;
-					if (ensure(InEntry.Value.IsType<FInterfaceRegistryTransaction>()))
-					{
-						RegistryKey = InEntry.Value.Get<FInterfaceRegistryTransaction>().GetInterfaceRegistryKey();
-					}
-					return FFrontendQueryKey{RegistryKey};
-				}
-			};
-
-			class FMapInterfaceRegistryTransactionToInterfaceName : public IFrontendQueryMapStep
-			{
-			public:
-				virtual FFrontendQueryKey Map(const FFrontendQueryEntry& InEntry) const override
-				{
-					FInterfaceRegistryKey RegistryKey;
-					if (ensure(InEntry.Value.IsType<FInterfaceRegistryTransaction>()))
-					{
-						return FFrontendQueryKey{InEntry.Value.Get<FInterfaceRegistryTransaction>().GetInterfaceVersion().Name};
-					}
-					return FFrontendQueryKey();
-				}
-			};
-
-			class FReduceInterfaceRegistryTransactionsToCurrentStatus : public IFrontendQueryReduceStep
-			{
-			public:
-				virtual void Reduce(const FFrontendQueryKey& InKey, FFrontendQueryPartition& InOutEntries) const override
-				{
-					// Get most recent transaction
-					const FFrontendQueryEntry* FinalEntry = Algo::MaxElementBy(InOutEntries, GetTransactionTimestamp);
-
-					// Check if most recent transaction is valid and not an unregister transaction.
-					if (IsValidTransactionOfType(FInterfaceRegistryTransaction::ETransactionType::InterfaceRegistration, FinalEntry))
-					{
-						FFrontendQueryEntry Entry = *FinalEntry;
-						InOutEntries.Reset();
-						InOutEntries.Add(Entry);
-					}
-					else
-					{
-						InOutEntries.Reset();
-					}
-				}
-			private:
-				static FInterfaceRegistryTransaction::FTimeType GetTransactionTimestamp(const FFrontendQueryEntry& InEntry)
-				{
-					if (ensure(InEntry.Value.IsType<FInterfaceRegistryTransaction>()))
-					{
-						return InEntry.Value.Get<FInterfaceRegistryTransaction>().GetTimestamp();
-					}
-					return 0;
-				}
-
-				static bool IsValidTransactionOfType(FInterfaceRegistryTransaction::ETransactionType InType, const FFrontendQueryEntry* InEntry)
-				{
-					if (nullptr != InEntry)
-					{
-						if (InEntry->Value.IsType<FInterfaceRegistryTransaction>())
-						{
-							return InEntry->Value.Get<FInterfaceRegistryTransaction>().GetTransactionType() == InType;
-						}
-					}
-					return false;
-				}
-
-			};
-
-
-			class FTransformInterfaceRegistryTransactionToInterface : public IFrontendQueryTransformStep
-			{
-			public:
-				virtual void Transform(FFrontendQueryEntry::FValue& InValue) const override
-				{
-					FMetasoundFrontendInterface Interface;
-					if (ensure(InValue.IsType<FInterfaceRegistryTransaction>()))
-					{
-						FInterfaceRegistryKey RegistryKey = InValue.Get<FInterfaceRegistryTransaction>().GetInterfaceRegistryKey();
-						IInterfaceRegistry::Get().FindInterface(RegistryKey, Interface);
-					}
-					InValue.Set<FMetasoundFrontendInterface>(MoveTemp(Interface));
-				}
-			};
-
-			class FMapInterfaceToInterfaceNameAndMajorVersion : public IFrontendQueryMapStep
-			{
-			public:
-				static FFrontendQueryKey GetKey(const FString& InName, int32 InMajorVersion)
-				{
-					return FFrontendQueryKey{FString::Format(TEXT("{0}_{1}"), { InName, InMajorVersion })};
-				}
-
-				virtual FFrontendQueryKey Map(const FFrontendQueryEntry& InEntry) const override
-				{
-					FString Key;
-					if (ensure(InEntry.Value.IsType<FMetasoundFrontendInterface>()))
-					{
-						const FMetasoundFrontendInterface& Interface = InEntry.Value.Get<FMetasoundFrontendInterface>();
-
-						return FMapInterfaceToInterfaceNameAndMajorVersion::GetKey(Interface.Version.Name.ToString(), Interface.Version.Number.Major);
-					}
-					return FFrontendQueryKey{Key};
-				}
-			};
-
-			class FReduceInterfacesToHighestVersion : public IFrontendQueryReduceStep
-			{
-			public:
-				virtual void Reduce(const FFrontendQueryKey& InKey, FFrontendQueryPartition& InOutEntries) const override
-				{
-					FFrontendQueryEntry* HighestVersionEntry = nullptr;
-					FMetasoundFrontendVersionNumber HighestVersion = FMetasoundFrontendVersionNumber::GetInvalid();
-
-					for (FFrontendQueryEntry& Entry : InOutEntries)
-					{
-						if (ensure(Entry.Value.IsType<FMetasoundFrontendInterface>()))
-						{
-							const FMetasoundFrontendVersionNumber& VersionNumber = Entry.Value.Get<FMetasoundFrontendInterface>().Version.Number;
-							if (VersionNumber > HighestVersion)
-							{
-								HighestVersionEntry = &Entry;
-								HighestVersion = VersionNumber;
-							}
-						}
-					}
-
-					if (HighestVersionEntry)
-					{
-						FFrontendQueryEntry Entry = *HighestVersionEntry;
-						InOutEntries.Reset();
-						InOutEntries.Add(Entry);
-					}
-					else
-					{
-						InOutEntries.Reset();
-					}
-				}
-			};
-
-			class FRemoveDeprecatedClasses : public IFrontendQueryFilterStep
-			{
-			public:
-				virtual bool Filter(const FFrontendQueryEntry& InEntry) const override
-				{
-					if (InEntry.Value.IsType<FMetasoundFrontendClass>())
-					{
-						return !InEntry.Value.Get<FMetasoundFrontendClass>().Metadata.GetIsDeprecated();
-					}
-					return false;
-				}
-			};
-
-			class FRemoveInterfacesWhichAreNotDefault : public IFrontendQueryFilterStep
-			{
-				virtual bool Filter(const FFrontendQueryEntry& InEntry) const override
-				{
-					if (InEntry.Value.IsType<FMetasoundFrontendInterface>())
-					{
-						FInterfaceRegistryKey RegistryKey = GetInterfaceRegistryKey(InEntry.Value.Get<FMetasoundFrontendInterface>());
-						const IInterfaceRegistryEntry* RegEntry = IInterfaceRegistry::Get().FindInterfaceRegistryEntry(RegistryKey);
-
-						if (nullptr != RegEntry)
-						{
-							return RegEntry->IsDefault();
-						}
-					}
-					return false;
-				}
-			};
-
-			class FMapInterfaceToInterfaceName : public IFrontendQueryMapStep
-			{
-			public:
-				virtual FFrontendQueryKey Map(const FFrontendQueryEntry& InEntry) const override
-				{
-					FName Key;
-					if (ensure(InEntry.Value.IsType<FMetasoundFrontendInterface>()))
-					{
-						Key = InEntry.Value.Get<FMetasoundFrontendInterface>().Version.Name;
-					}
-					return FFrontendQueryKey{Key};
-				}
-			};
-
-			class FTransformInterfaceRegistryTransactionToInterfaceVersion : public IFrontendQueryTransformStep
-			{
-			public:
-				virtual void Transform(FFrontendQueryEntry::FValue& InValue) const override
-				{
-					FMetasoundFrontendVersion Version;
-					if (ensure(InValue.IsType<FInterfaceRegistryTransaction>()))
-					{
-						Version = InValue.Get<FInterfaceRegistryTransaction>().GetInterfaceVersion();
-					}
-					InValue.Set<FMetasoundFrontendVersion>(Version);
-				}
-			};
-		}
-
-		// Base interface for a search engine query. This allows the query to be
-		// primed to improve execution performance in critical call-stacks.
-		struct ISearchEngineQuery
-		{
-			// Prime the query to improve execution performance during subsequent 
-			// queries. 
-			virtual void Prime() = 0;
-			virtual ~ISearchEngineQuery() = default;
-		};
-
-		// TSearchEngineQuery provides a base for all queries used in the search
-		// engine. It requires a `QueryPolicy` as a template parameter. 
-		//
-		// A `QueryPolicy` is a struct containing a static function to create the 
-		// query, a static function to build a result from the query selection, 
-		// and query result type.
-		//
-		// Example:
-		//
-		// struct FQueryPolicyExample
-		// {
-		//		using ResultType = FMetasoundFrontendClass;
-		//
-		//		static FFrontendQuery CreateQuery()
-		//		{
-		//			using namespace SearchEngineQuerySteps;
-		//			FFrontendQuery Query;
-		//			Query.AddStep<FNodeClassRegistrationEvents>()
-		//				.AddStep<FMapRegistrationEventsToNodeRegistryKeys>()
-		//				.AddStep<FReduceRegistrationEventsToCurrentStatus>()
-		//				.AddStep<FTransformRegistrationEventsToClasses>();
-		//			return Query;
-		//		}
-		//
-		//		static ResultType BuildResult(const FFrontendQueryPartition& InPartition)
-		//		{
-		//			FMetasoundFrontendClass MetaSoundClass;
-		//			if (ensureMsgf(InPartition.Num() > 0, TEXT("Cannot retrieve class from empty partition")))
-		//			{
-		//				MetaSoundClass = InPartition.CreateConstIterator()->Value.Get<FMetasoundFrontendClass>();
-		//			}
-		//			return MetaSoundClass;
-		//		}
-		// };
-		template<typename QueryPolicy>
-		struct TSearchEngineQuery : public ISearchEngineQuery
-		{
-			using ResultType = typename QueryPolicy::ResultType;
-
-			TSearchEngineQuery()
-			: Query(QueryPolicy::CreateQuery())
-			{
-			}
-
-			virtual void Prime() override
-			{
-				METASOUND_LLM_SCOPE;
-				FScopeLock Lock(&QueryCriticalSection);
-				Update();
-			}
-
-			// Updates the query to the most recent value and returns the
-			// result.
-			ResultType UpdateAndFindResult(const FFrontendQueryKey& InKey)
-			{
-				METASOUND_LLM_SCOPE;
-				FScopeLock Lock(&QueryCriticalSection);
-
-				Update();
-				if (const ResultType* Result = ResultCache.Find(InKey))
-				{
-					return *Result;
-				}
-				return ResultType();
-			}
-
-			// Updates the query to the most recent value and assigns the value 
-			// to OutResult. Returns true on success, false on failure. 
-			bool UpdateAndFindResult(const FFrontendQueryKey& InKey, ResultType& OutResult)
-			{
-				METASOUND_LLM_SCOPE;
-				FScopeLock Lock(&QueryCriticalSection);
-
-				Update();
-				if (const ResultType* Result = ResultCache.Find(InKey))
-				{
-					OutResult = *Result;
-					return true;
-				}
-				return false;
-			}
-
-		private:
-			void Update()
-			{
-				METASOUND_TRACE_CPUPROFILER_EVENT_SCOPE(MetaSound::SearchEngineQuery::Update);
-
-				TSet<FFrontendQueryKey> UpdatedKeys;
-				const FFrontendQuerySelection& Selection = Query.Update(UpdatedKeys);
-
-				if (UpdatedKeys.Num() > 0)
-				{
-					METASOUND_TRACE_CPUPROFILER_EVENT_SCOPE(MetaSound::SearchEngineQuery::BuildResult);
-
-					for (const FFrontendQueryKey& Key : UpdatedKeys)
-					{
-						const FFrontendQueryPartition* Partition = Query.GetSelection().Find(Key);
-						if (nullptr != Partition)
-						{
-							ResultCache.Add(Key, QueryPolicy::BuildResult(*Partition));
-						}
-						else
-						{
-							ResultCache.Remove(Key);
-						}
-					}
-				}
-			}
-
-			FCriticalSection QueryCriticalSection;
-			FFrontendQuery Query;
-			TMap<FFrontendQueryKey, ResultType> ResultCache;
-		};
-
-		TArray<FMetasoundFrontendClass> BuildArrayOfClassesFromPartition(const FFrontendQueryPartition& InPartition)
-		{
-			TArray<FMetasoundFrontendClass> Result;
-			for (const FFrontendQueryEntry& Entry : InPartition)
-			{
-				check(Entry.Value.IsType<FMetasoundFrontendClass>());
-				Result.Add(Entry.Value.Get<FMetasoundFrontendClass>());
-			}
-			return Result;
-		}
-
-		FMetasoundFrontendClass BuildSingleClassFromPartition(const FFrontendQueryPartition& InPartition)
-		{
-			FMetasoundFrontendClass MetaSoundClass;
-			if (ensureMsgf(InPartition.Num() > 0, TEXT("Cannot retrieve class from empty partition")))
-			{
-				MetaSoundClass = InPartition.CreateConstIterator()->Value.Get<FMetasoundFrontendClass>();
-			}
-			return MetaSoundClass;
-		}
-
-		TArray<FMetasoundFrontendInterface> BuildArrayOfInterfacesFromPartition(const FFrontendQueryPartition& InPartition)
-		{
-			TArray<FMetasoundFrontendInterface> Result;
-			for (const FFrontendQueryEntry& Entry : InPartition)
-			{
-				check(Entry.Value.IsType<FMetasoundFrontendInterface>());
-				Result.Add(Entry.Value.Get<FMetasoundFrontendInterface>());
-			}
-			return Result;
-		}
-
-		FMetasoundFrontendInterface BuildSingleInterfaceFromPartition(const FFrontendQueryPartition& InPartition)
-		{
-			FMetasoundFrontendInterface MetaSoundInterface;
-			if (ensureMsgf(InPartition.Num() > 0, TEXT("Cannot retrieve class from empty partition")))
-			{
-				MetaSoundInterface = InPartition.CreateConstIterator()->Value.Get<FMetasoundFrontendInterface>();
-			}
-			return MetaSoundInterface;
-		}
-
-		TArray<FMetasoundFrontendVersion> BuildArrayOfVersionsFromPartition(const FFrontendQueryPartition& InPartition)
-		{
-			TArray<FMetasoundFrontendVersion> Result;
-			for (const FFrontendQueryEntry& Entry : InPartition)
-			{
-				check(Entry.Value.IsType<FMetasoundFrontendVersion>());
-				Result.Add(Entry.Value.Get<FMetasoundFrontendVersion>());
-			}
-			return Result;
-		}
-
-		// Policy for finding all registered metasound classes. 
-		struct FFindAllClassesQueryPolicy
-		{
-			using ResultType = TArray<FMetasoundFrontendClass>;
-
-			static FFrontendQuery CreateQuery()
-			{
-				using namespace SearchEngineQuerySteps;
-				FFrontendQuery Query;
-				Query.AddStep<FNodeClassRegistrationEvents>()
-					.AddStep<FMapRegistrationEventsToNodeRegistryKeys>()
-					.AddStep<FReduceRegistrationEventsToCurrentStatus>()
-					.AddStep<FTransformRegistrationEventsToClasses>()
-					.AddStep<FRemoveDeprecatedClasses>()
-					.AddStep<FMapToNull>();
-				return Query;
-			}
-
-			static ResultType BuildResult(const FFrontendQueryPartition& InPartition)
-			{
-				return BuildArrayOfClassesFromPartition(InPartition);
-			}
-		};
-
-		// Policy for finding all registered metasound classes, including deprecated classes. 
-		struct FFindAllClassesIncludingAllVersionsQueryPolicy
-		{
-			using ResultType = TArray<FMetasoundFrontendClass>;
-
-			static FFrontendQuery CreateQuery()
-			{
-				using namespace SearchEngineQuerySteps;
-				FFrontendQuery Query;
-				Query.AddStep<FNodeClassRegistrationEvents>()
-					.AddStep<FMapRegistrationEventsToNodeRegistryKeys>()
-					.AddStep<FReduceRegistrationEventsToCurrentStatus>()
-					.AddStep<FTransformRegistrationEventsToClasses>()
-					.AddStep<FMapToNull>();
-				return Query;
-			}
-
-			static ResultType BuildResult(const FFrontendQueryPartition& InPartition)
-			{
-				return BuildArrayOfClassesFromPartition(InPartition);
-			}
-		};
-
-		// Policy for finding all registered metasound classes sorted by version 
-		// and indexed by name.
-		struct FFindClassesWithNameSortedQueryPolicy
-		{
-			using ResultType = TArray<FMetasoundFrontendClass>;
-
-			static FFrontendQuery CreateQuery()
-			{
-				using namespace SearchEngineQuerySteps;
-				FFrontendQuery Query;
-				Query.AddStep<FNodeClassRegistrationEvents>()
-					.AddStep<FMapRegistrationEventsToNodeRegistryKeys>()
-					.AddStep<FReduceRegistrationEventsToCurrentStatus>()
-					.AddStep<FTransformRegistrationEventsToClasses>()
-					.AddStep<FMapClassesToClassName>()
-					.AddStep<FSortClassesByVersion>();
-				return Query;
-			}
-
-			static ResultType BuildResult(const FFrontendQueryPartition& InPartition)
-			{
-				return BuildArrayOfClassesFromPartition(InPartition);
-			}
-		};
-
-		// Policy for finding all registered metasound classes indexed by name.
-		struct FFindClassesWithNameUnsortedQueryPolicy
-		{
-			using ResultType = TArray<FMetasoundFrontendClass>;
-
-			static FFrontendQuery CreateQuery()
-			{
-				using namespace SearchEngineQuerySteps;
-				FFrontendQuery Query;
-				Query.AddStep<FNodeClassRegistrationEvents>()
-					.AddStep<FMapRegistrationEventsToNodeRegistryKeys>()
-					.AddStep<FReduceRegistrationEventsToCurrentStatus>()
-					.AddStep<FTransformRegistrationEventsToClasses>()
-					.AddStep<FMapClassesToClassName>();
-				return Query;
-			}
-
-			static ResultType BuildResult(const FFrontendQueryPartition& InPartition)
-			{
-				return BuildArrayOfClassesFromPartition(InPartition);
-			}
-		};
-
-		// Policy for finding highest versioned metasound class by name.
-		struct FFindClassWithHighestVersionQueryPolicy
-		{
-			using ResultType = FMetasoundFrontendClass;
-
-			static FFrontendQuery CreateQuery()
-			{
-				using namespace SearchEngineQuerySteps;
-				FFrontendQuery Query;
-				Query.AddStep<FNodeClassRegistrationEvents>()
-					.AddStep<FMapRegistrationEventsToNodeRegistryKeys>()
-					.AddStep<FReduceRegistrationEventsToCurrentStatus>()
-					.AddStep<FTransformRegistrationEventsToClasses>()
-					.AddStep<FRemoveDeprecatedClasses>()
-					.AddStep<FMapToFullClassName>()
-					.AddStep<FReduceClassesToHighestVersion>();
-				return Query;
-			}
-
-			static ResultType BuildResult(const FFrontendQueryPartition& InPartition)
-			{
-				return BuildSingleClassFromPartition(InPartition);
-			}
-		};
-
-		// Policy for finding highest versioned metasound class by name and major version.
-		struct FFindClassWithMajorVersionQueryPolicy
-		{
-			using ResultType = FMetasoundFrontendClass;
-
-			static FFrontendQuery CreateQuery()
-			{
-				using namespace SearchEngineQuerySteps;
-				FFrontendQuery Query;
-				Query.AddStep<FNodeClassRegistrationEvents>()
-					.AddStep<FMapRegistrationEventsToNodeRegistryKeys>()
-					.AddStep<FReduceRegistrationEventsToCurrentStatus>()
-					.AddStep<FTransformRegistrationEventsToClasses>()
-					.AddStep<FRemoveDeprecatedClasses>()
-					.AddStep<FMapClassToFullClassNameAndMajorVersion>()
-					.AddStep<FReduceClassesToHighestVersion>();
-				return Query;
-			}
-
-			static ResultType BuildResult(const FFrontendQueryPartition& InPartition)
-			{
-				return BuildSingleClassFromPartition(InPartition);
-			}
-		};
-
-		// Policy for finding all registered interfaces. 
-		struct FFindAllInterfacesQueryPolicy
-		{
-			using ResultType = TArray<FMetasoundFrontendInterface>;
-
-			static FFrontendQuery CreateQuery()
-			{
-				using namespace SearchEngineQuerySteps;
-				FFrontendQuery Query;
-				Query.AddStep<FInterfaceRegistryTransactionSource>()
-					.AddStep<FMapInterfaceRegistryTransactionsToInterfaceRegistryKeys>()
-					.AddStep<FReduceInterfaceRegistryTransactionsToCurrentStatus>()
-					.AddStep<FTransformInterfaceRegistryTransactionToInterface>()
-					.AddStep<FMapInterfaceToInterfaceNameAndMajorVersion>()
-					.AddStep<FReduceInterfacesToHighestVersion>()
-					.AddStep<FMapToNull>();
-				return Query;
-			}
-
-			static ResultType BuildResult(const FFrontendQueryPartition& InPartition)
-			{
-				return BuildArrayOfInterfacesFromPartition(InPartition);
-			}
-		};
-
-		// Policy for finding all registered default interfaces. 
-		struct FFindAllDefaultInterfacesQueryPolicy
-		{
-			using ResultType = TArray<FMetasoundFrontendInterface>;
-
-			static FFrontendQuery CreateQuery()
-			{
-				using namespace SearchEngineQuerySteps;
-				FFrontendQuery Query;
-				Query.AddStep<FInterfaceRegistryTransactionSource>()
-					.AddStep<FMapInterfaceRegistryTransactionsToInterfaceRegistryKeys>()
-					.AddStep<FReduceInterfaceRegistryTransactionsToCurrentStatus>()
-					.AddStep<FTransformInterfaceRegistryTransactionToInterface>()
-					.AddStep<FMapInterfaceToInterfaceNameAndMajorVersion>()
-					.AddStep<FReduceInterfacesToHighestVersion>()
-					.AddStep<FRemoveInterfacesWhichAreNotDefault>()
-					.AddStep<FMapToNull>();
-				return Query;
-			}
-
-			static ResultType BuildResult(const FFrontendQueryPartition& InPartition)
-			{
-				return BuildArrayOfInterfacesFromPartition(InPartition);
-			}
-		};
-
-		// Policy for finding all registered interfaces (including deprecated). 
-		struct FFindAllInterfacesIncludingAllVersionsQueryPolicy
-		{
-			using ResultType = TArray<FMetasoundFrontendInterface>;
-
-			static FFrontendQuery CreateQuery()
-			{
-				using namespace SearchEngineQuerySteps;
-				FFrontendQuery Query;
-				Query.AddStep<FInterfaceRegistryTransactionSource>()
-					.AddStep<FMapInterfaceRegistryTransactionsToInterfaceRegistryKeys>()
-					.AddStep<FReduceInterfaceRegistryTransactionsToCurrentStatus>()
-					.AddStep<FTransformInterfaceRegistryTransactionToInterface>()
-					.AddStep<FMapToNull>();
-				return Query;
-			}
-
-			static TArray<FMetasoundFrontendInterface> BuildResult(const FFrontendQueryPartition& InPartition)
-			{
-				return BuildArrayOfInterfacesFromPartition(InPartition);
-			}
-		};
-
-		// Policy for finding all registered interface versions (name & version number)
-		// by name.
-		struct FFindAllRegisteredInterfacesWithNameQueryPolicy
-		{
-			using ResultType = TArray<FMetasoundFrontendVersion>;
-
-			static FFrontendQuery CreateQuery()
-			{
-				using namespace SearchEngineQuerySteps;
-				FFrontendQuery Query;
-				Query.AddStep<FInterfaceRegistryTransactionSource>()
-					.AddStep<FMapInterfaceRegistryTransactionsToInterfaceRegistryKeys>()
-					.AddStep<FReduceInterfaceRegistryTransactionsToCurrentStatus>()
-					.AddStep<FMapInterfaceRegistryTransactionToInterfaceName>()
-					.AddStep<FTransformInterfaceRegistryTransactionToInterfaceVersion>();
-				return Query;
-			}
-
-			static TArray<FMetasoundFrontendVersion> BuildResult(const FFrontendQueryPartition& InPartition)
-			{
-				return BuildArrayOfVersionsFromPartition(InPartition);
-			}
-		};
-
-		// Policy for finding the highest version registered interfaces by name.
-		struct FFindInterfaceWithHighestVersionQueryPolicy
-		{
-			using ResultType = FMetasoundFrontendInterface;
-
-			static FFrontendQuery CreateQuery()
-			{
-				using namespace SearchEngineQuerySteps;
-				FFrontendQuery Query;
-				Query.AddStep<FInterfaceRegistryTransactionSource>()
-					.AddStep<FMapInterfaceRegistryTransactionsToInterfaceRegistryKeys>()
-					.AddStep<FReduceInterfaceRegistryTransactionsToCurrentStatus>()
-					.AddStep<FTransformInterfaceRegistryTransactionToInterface>()
-					.AddStep<FMapInterfaceToInterfaceName>()
-					.AddStep<FReduceInterfacesToHighestVersion>();
-				return Query;
-			}
-
-			static FMetasoundFrontendInterface BuildResult(const FFrontendQueryPartition& InPartition)
-			{
-				return BuildSingleInterfaceFromPartition(InPartition);
-			}
-		};
-
-		class FSearchEngine : public ISearchEngine
-		{
-			FSearchEngine(const FSearchEngine&) = delete;
-			FSearchEngine(FSearchEngine&&) = delete;
-			FSearchEngine& operator=(const FSearchEngine&) = delete;
-			FSearchEngine& operator=(FSearchEngine&&) = delete;
-
-		public:
-			FSearchEngine() = default;
-			virtual ~FSearchEngine() = default;
-
-			virtual void Prime() override;
-			virtual TArray<FMetasoundFrontendClass> FindAllClasses(bool bInIncludeAllVersions) override;
-			virtual TArray<FMetasoundFrontendClass> FindClassesWithName(const FMetasoundFrontendClassName& InName, bool bInSortByVersion) override;
-			virtual bool FindClassWithHighestVersion(const FMetasoundFrontendClassName& InName, FMetasoundFrontendClass& OutClass) override;
-			virtual bool FindClassWithMajorVersion(const FMetasoundFrontendClassName& InName, int32 InMajorVersion, FMetasoundFrontendClass& OutClass) override;
-
-			virtual TArray<FMetasoundFrontendInterface> FindAllInterfaces(bool bInIncludeAllVersions) override;
-			virtual TArray<FMetasoundFrontendInterface> FindUClassDefaultInterfaces(FName InUClassName) override;
-			virtual TArray<FMetasoundFrontendVersion> FindAllRegisteredInterfacesWithName(FName InInterfaceName) override;
-			virtual bool FindInterfaceWithHighestVersion(FName InInterfaceName, FMetasoundFrontendInterface& OutInterface) override;
-
-		private:
-
-			TSearchEngineQuery<FFindAllClassesQueryPolicy> FindAllClassesQuery;
-			TSearchEngineQuery<FFindAllClassesIncludingAllVersionsQueryPolicy> FindAllClassesIncludingAllVersionsQuery;
-			TSearchEngineQuery<FFindClassesWithNameUnsortedQueryPolicy> FindClassesWithNameUnsortedQuery;
-			TSearchEngineQuery<FFindClassesWithNameSortedQueryPolicy> FindClassesWithNameSortedQuery;
-			TSearchEngineQuery<FFindClassWithHighestVersionQueryPolicy> FindClassWithHighestVersionQuery;
-			TSearchEngineQuery<FFindClassWithMajorVersionQueryPolicy> FindClassWithMajorVersionQuery;
-			TSearchEngineQuery<FFindAllInterfacesQueryPolicy> FindAllInterfacesQuery;
-			TSearchEngineQuery<FFindAllDefaultInterfacesQueryPolicy> FFindAllDefaultInterfacesQuery;
-			TSearchEngineQuery<FFindAllInterfacesIncludingAllVersionsQueryPolicy> FindAllInterfacesIncludingAllVersionsQuery;
-			TSearchEngineQuery<FFindAllRegisteredInterfacesWithNameQueryPolicy> FindAllRegisteredInterfacesWithNameQuery;
-			TSearchEngineQuery<FFindInterfaceWithHighestVersionQueryPolicy> FindInterfaceWithHighestVersionQuery;
-		};
-
 		ISearchEngine& ISearchEngine::Get()
 		{
-			static FSearchEngine SearchEngine;
+#if WITH_EDITORONLY_DATA
+			static FSearchEngineEditorOnly SearchEngine;
+#else
+			static FSearchEngineCore SearchEngine;
+#endif
 			return SearchEngine;
 		}
-
-		void FSearchEngine::Prime()
-		{
-			METASOUND_LLM_SCOPE;
-			METASOUND_TRACE_CPUPROFILER_EVENT_SCOPE(metasound::FSearchEngine::Prime);
-
-			FindAllClassesQuery.Prime();
-			FindAllClassesIncludingAllVersionsQuery.Prime();
-			FindClassesWithNameUnsortedQuery.Prime();
-			FindClassesWithNameSortedQuery.Prime();
-			FindClassWithHighestVersionQuery.Prime();
-			FindClassWithMajorVersionQuery.Prime();
-			FindAllInterfacesIncludingAllVersionsQuery.Prime();
-			FindAllInterfacesQuery.Prime();
-			FindAllRegisteredInterfacesWithNameQuery.Prime();
-			FindInterfaceWithHighestVersionQuery.Prime();
-		}
-
-		TArray<FMetasoundFrontendClass> FSearchEngine::FindAllClasses(bool bInIncludeAllVersions)
-		{
-			METASOUND_TRACE_CPUPROFILER_EVENT_SCOPE(metasound::FSearchEngine::FindAllClasses);
-
-			FFrontendQueryKey NullKey;
-			if (bInIncludeAllVersions)
-			{
-				return FindAllClassesIncludingAllVersionsQuery.UpdateAndFindResult(NullKey);
-			}
-			else
-			{
-				return FindAllClassesQuery.UpdateAndFindResult(NullKey);
-			}
-		}
-
-		TArray<FMetasoundFrontendClass> FSearchEngine::FindClassesWithName(const FMetasoundFrontendClassName& InName, bool bInSortByVersion)
-		{
-			METASOUND_TRACE_CPUPROFILER_EVENT_SCOPE(metasound::FSearchEngine::FindClassesWithName);
-
-			const FFrontendQueryKey Key(InName.GetFullName());
-			if (bInSortByVersion)
-			{
-				return FindClassesWithNameSortedQuery.UpdateAndFindResult(Key);
-			}
-			else 
-			{
-				return FindClassesWithNameUnsortedQuery.UpdateAndFindResult(Key);
-			}
-		}
-
-		bool FSearchEngine::FindClassWithHighestVersion(const FMetasoundFrontendClassName& InName, FMetasoundFrontendClass& OutClass)
-		{
-			METASOUND_TRACE_CPUPROFILER_EVENT_SCOPE(metasound::FSearchEngine::FindClassWithHighestVersion);
-
-			const FFrontendQueryKey Key{InName.GetFullName()};
-			return FindClassWithHighestVersionQuery.UpdateAndFindResult(Key, OutClass);
-		}
-
-		bool FSearchEngine::FindClassWithMajorVersion(const FMetasoundFrontendClassName& InName, int32 InMajorVersion, FMetasoundFrontendClass& OutClass)
-		{
-			METASOUND_TRACE_CPUPROFILER_EVENT_SCOPE(metasound::FSearchEngine::FindClassWithMajorVersion);
-			using namespace SearchEngineQuerySteps;
-
-			const FFrontendQueryKey Key = FMapClassToFullClassNameAndMajorVersion::GetKey(InName, InMajorVersion);
-			return FindClassWithMajorVersionQuery.UpdateAndFindResult(Key, OutClass);
-		}
-
-		TArray<FMetasoundFrontendInterface> FSearchEngine::FindAllInterfaces(bool bInIncludeAllVersions)
-		{
-			METASOUND_TRACE_CPUPROFILER_EVENT_SCOPE(metasound::FSearchEngine::FindAllInterfaces);
-
-			const FFrontendQueryKey NullKey;
-			if (bInIncludeAllVersions)
-			{
-				return FindAllInterfacesIncludingAllVersionsQuery.UpdateAndFindResult(NullKey);
-			}
-			else
-			{
-				return FindAllInterfacesQuery.UpdateAndFindResult(NullKey);
-			}
-		}
-
-		TArray<FMetasoundFrontendInterface> FSearchEngine::FindUClassDefaultInterfaces(FName InUClassName)
-		{
-			METASOUND_TRACE_CPUPROFILER_EVENT_SCOPE(metasound::FSearchEngine::FindUClassDefaultInterfaces);
-
-			const FFrontendQueryKey NullKey;
-			TArray<FMetasoundFrontendInterface> DefaultInterfaces = FFindAllDefaultInterfacesQuery.UpdateAndFindResult(NullKey);
-
-			// All default interfaces are filtered here because a UClass name cannot effectively 
-			// be used as a "Key" for this query.
-			//
-			// We cannot know all the UClasses associated with each default interface given only a UClassName
-			// because we need to inspect UClass inheritance to handle derived UClasses
-			// which should be compatible with their parent's default interfaces. 
-			//
-			// This module does not have access to the UClass inheritance structure.
-			// Instead, this information is accessed in the IInterfaceRegistryEntry
-			// concrete implementation. 
-			auto DoesNotSupportUClass = [&InUClassName](const FMetasoundFrontendInterface& InInterface)
-			{
-				FInterfaceRegistryKey RegistryKey = GetInterfaceRegistryKey(InInterface);
-				const IInterfaceRegistryEntry* RegEntry = IInterfaceRegistry::Get().FindInterfaceRegistryEntry(RegistryKey);
-				if (ensure(RegEntry))
-				{
-					return !RegEntry->UClassIsSupported(InUClassName);
-				}
-
-				return true;
-			};
-
-			DefaultInterfaces.RemoveAllSwap(DoesNotSupportUClass);
-
-			return DefaultInterfaces;
-		}
-
-		TArray<FMetasoundFrontendVersion> FSearchEngine::FindAllRegisteredInterfacesWithName(FName InInterfaceName)
-		{
-			METASOUND_TRACE_CPUPROFILER_EVENT_SCOPE(metasound::FSearchEngine::FindAllRegisteredInterfacesWithName);
-
-			const FFrontendQueryKey Key{InInterfaceName};
-			return FindAllRegisteredInterfacesWithNameQuery.UpdateAndFindResult(Key);
-		}
-
-		bool FSearchEngine::FindInterfaceWithHighestVersion(FName InInterfaceName, FMetasoundFrontendInterface& OutInterface)
-		{
-			METASOUND_TRACE_CPUPROFILER_EVENT_SCOPE(metasound::FSearchEngine::FindInterfaceWithHighestVersion);
-
-			const FFrontendQueryKey Key{InInterfaceName};
-			return FindInterfaceWithHighestVersionQuery.UpdateAndFindResult(Key, OutInterface);
-		}
 	}
 }
 
diff --git a/Engine/Plugins/Runtime/Metasound/Source/MetasoundFrontend/Private/MetasoundFrontendSearchEngineCore.cpp b/Engine/Plugins/Runtime/Metasound/Source/MetasoundFrontend/Private/MetasoundFrontendSearchEngineCore.cpp
new file mode 100644
--- /dev/null
+++ b/Engine/Plugins/Runtime/Metasound/Source/MetasoundFrontend/Private/MetasoundFrontendSearchEngineCore.cpp
@@ -0,0 +1,414 @@
+// Copyright Epic Games, Inc. All Rights Reserved.
+
+#include "MetasoundFrontendSearchEngineCore.h"
+
+#include "Algo/MaxElement.h"
+#include "Containers/Array.h"
+#include "Containers/UnrealString.h"
+#include "MetasoundFrontendArchetypeRegistry.h"
+#include "MetasoundFrontendQuery.h"
+#include "MetasoundFrontendQuerySteps.h"
+#include "MetasoundFrontendRegistryTransaction.h"
+#include "MetasoundTrace.h"
+
+namespace Metasound
+{
+	namespace Frontend
+	{
+		namespace SearchEngineQuerySteps
+		{
+			FFrontendQueryKey FMapClassToFullClassNameAndMajorVersion::GetKey(const FMetasoundFrontendClassName& InClassName, int32 InMajorVersion)
+			{
+				return FFrontendQueryKey(FString::Format(TEXT("{0}_v{1}"), {InClassName.ToString(), InMajorVersion}));
+			}
+
+			FFrontendQueryKey FMapClassToFullClassNameAndMajorVersion::Map(const FFrontendQueryEntry& InEntry) const
+			{
+				if (ensure(InEntry.Value.IsType<FMetasoundFrontendClass>()))
+				{
+					const FMetasoundFrontendClassMetadata& Metadata = InEntry.Value.Get<FMetasoundFrontendClass>().Metadata;
+					return FMapClassToFullClassNameAndMajorVersion::GetKey(Metadata.GetClassName(), Metadata.GetVersion().Major);
+				}
+				return FFrontendQueryKey();
+			}
+
+			FInterfaceRegistryTransactionSource::FInterfaceRegistryTransactionSource()
+			: CurrentTransactionID(GetOriginRegistryTransactionID())
+			{
+			}
+
+			void FInterfaceRegistryTransactionSource::Stream(TArray<FFrontendQueryValue>& OutEntries)
+			{
+				auto AddValue = [&OutEntries](const FInterfaceRegistryTransaction& InTransaction)
+				{
+					OutEntries.Emplace(TInPlaceType<FInterfaceRegistryTransaction>(), InTransaction);
+				};
+				
+				IInterfaceRegistry::Get().ForEachRegistryTransactionSince(CurrentTransactionID, &CurrentTransactionID, AddValue);
+			}
+
+			FFrontendQueryKey FMapInterfaceRegistryTransactionsToInterfaceRegistryKeys::Map(const FFrontendQueryEntry& InEntry) const
+			{
+				FInterfaceRegistryKey RegistryKey;
+				if (ensure(InEntry.Value.IsType<FInterfaceRegistryTransaction>()))
+				{
+					RegistryKey = InEntry.Value.Get<FInterfaceRegistryTransaction>().GetInterfaceRegistryKey();
+				}
+				return FFrontendQueryKey{RegistryKey};
+			}
+
+			FFrontendQueryKey FMapInterfaceRegistryTransactionToInterfaceName::Map(const FFrontendQueryEntry& InEntry) const
+			{
+				FInterfaceRegistryKey RegistryKey;
+				if (ensure(InEntry.Value.IsType<FInterfaceRegistryTransaction>()))
+				{
+					return FFrontendQueryKey{InEntry.Value.Get<FInterfaceRegistryTransaction>().GetInterfaceVersion().Name};
+				}
+				return FFrontendQueryKey();
+			}
+
+			void FReduceInterfaceRegistryTransactionsToCurrentStatus::Reduce(const FFrontendQueryKey& InKey, FFrontendQueryPartition& InOutEntries) const
+			{
+				// Get most recent transaction
+				const FFrontendQueryEntry* FinalEntry = Algo::MaxElementBy(InOutEntries, GetTransactionTimestamp);
+
+				// Check if most recent transaction is valid and not an unregister transaction.
+				if (IsValidTransactionOfType(FInterfaceRegistryTransaction::ETransactionType::InterfaceRegistration, FinalEntry))
+				{
+					FFrontendQueryEntry Entry = *FinalEntry;
+					InOutEntries.Reset();
+					InOutEntries.Add(Entry);
+				}
+				else
+				{
+					InOutEntries.Reset();
+				}
+			}
+
+			FInterfaceRegistryTransaction::FTimeType FReduceInterfaceRegistryTransactionsToCurrentStatus::GetTransactionTimestamp(const FFrontendQueryEntry& InEntry)
+			{
+				if (ensure(InEntry.Value.IsType<FInterfaceRegistryTransaction>()))
+				{
+					return InEntry.Value.Get<FInterfaceRegistryTransaction>().GetTimestamp();
+				}
+				return 0;
+			}
+
+			bool FReduceInterfaceRegistryTransactionsToCurrentStatus::IsValidTransactionOfType(FInterfaceRegistryTransaction::ETransactionType InType, const FFrontendQueryEntry* InEntry)
+			{
+				if (nullptr != InEntry)
+				{
+					if (InEntry->Value.IsType<FInterfaceRegistryTransaction>())
+					{
+						return InEntry->Value.Get<FInterfaceRegistryTransaction>().GetTransactionType() == InType;
+					}
+				}
+				return false;
+			}
+
+			void FTransformInterfaceRegistryTransactionToInterface::Transform(FFrontendQueryEntry::FValue& InValue) const
+			{
+				FMetasoundFrontendInterface Interface;
+				if (ensure(InValue.IsType<FInterfaceRegistryTransaction>()))
+				{
+					FInterfaceRegistryKey RegistryKey = InValue.Get<FInterfaceRegistryTransaction>().GetInterfaceRegistryKey();
+					IInterfaceRegistry::Get().FindInterface(RegistryKey, Interface);
+				}
+				InValue.Set<FMetasoundFrontendInterface>(MoveTemp(Interface));
+			}
+
+			FFrontendQueryKey FMapInterfaceToInterfaceNameAndMajorVersion::GetKey(const FString& InName, int32 InMajorVersion)
+			{
+				return FFrontendQueryKey{FString::Format(TEXT("{0}_{1}"), { InName, InMajorVersion })};
+			}
+
+			FFrontendQueryKey FMapInterfaceToInterfaceNameAndMajorVersion::Map(const FFrontendQueryEntry& InEntry) const
+			{
+				FString Key;
+				if (ensure(InEntry.Value.IsType<FMetasoundFrontendInterface>()))
+				{
+					const FMetasoundFrontendInterface& Interface = InEntry.Value.Get<FMetasoundFrontendInterface>();
+
+					return FMapInterfaceToInterfaceNameAndMajorVersion::GetKey(Interface.Version.Name.ToString(), Interface.Version.Number.Major);
+				}
+				return FFrontendQueryKey{Key};
+			}
+
+			void FReduceInterfacesToHighestVersion::Reduce(const FFrontendQueryKey& InKey, FFrontendQueryPartition& InOutEntries) const
+			{
+				FFrontendQueryEntry* HighestVersionEntry = nullptr;
+				FMetasoundFrontendVersionNumber HighestVersion = FMetasoundFrontendVersionNumber::GetInvalid();
+
+				for (FFrontendQueryEntry& Entry : InOutEntries)
+				{
+					if (ensure(Entry.Value.IsType<FMetasoundFrontendInterface>()))
+					{
+						const FMetasoundFrontendVersionNumber& VersionNumber = Entry.Value.Get<FMetasoundFrontendInterface>().Version.Number;
+						if (VersionNumber > HighestVersion)
+						{
+							HighestVersionEntry = &Entry;
+							HighestVersion = VersionNumber;
+						}
+					}
+				}
+
+				if (HighestVersionEntry)
+				{
+					FFrontendQueryEntry Entry = *HighestVersionEntry;
+					InOutEntries.Reset();
+					InOutEntries.Add(Entry);
+				}
+				else
+				{
+					InOutEntries.Reset();
+				}
+			}
+
+			bool FRemoveDeprecatedClasses::Filter(const FFrontendQueryEntry& InEntry) const
+			{
+				if (InEntry.Value.IsType<FMetasoundFrontendClass>())
+				{
+					return !InEntry.Value.Get<FMetasoundFrontendClass>().Metadata.GetIsDeprecated();
+				}
+				return false;
+			}
+
+			bool FRemoveInterfacesWhichAreNotDefault::Filter(const FFrontendQueryEntry& InEntry) const
+			{
+				if (InEntry.Value.IsType<FMetasoundFrontendInterface>())
+				{
+					FInterfaceRegistryKey RegistryKey = GetInterfaceRegistryKey(InEntry.Value.Get<FMetasoundFrontendInterface>());
+					const IInterfaceRegistryEntry* RegEntry = IInterfaceRegistry::Get().FindInterfaceRegistryEntry(RegistryKey);
+
+					if (nullptr != RegEntry)
+					{
+						return RegEntry->IsDefault();
+					}
+				}
+				return false;
+			}
+
+			FFrontendQueryKey FMapInterfaceToInterfaceName::Map(const FFrontendQueryEntry& InEntry) const
+			{
+				FName Key;
+				if (ensure(InEntry.Value.IsType<FMetasoundFrontendInterface>()))
+				{
+					Key = InEntry.Value.Get<FMetasoundFrontendInterface>().Version.Name;
+				}
+				return FFrontendQueryKey{Key};
+			}
+
+			void FTransformInterfaceRegistryTransactionToInterfaceVersion::Transform(FFrontendQueryEntry::FValue& InValue) const
+			{
+				FMetasoundFrontendVersion Version;
+				if (ensure(InValue.IsType<FInterfaceRegistryTransaction>()))
+				{
+					Version = InValue.Get<FInterfaceRegistryTransaction>().GetInterfaceVersion();
+				}
+				InValue.Set<FMetasoundFrontendVersion>(Version);
+			}
+
+			TArray<FMetasoundFrontendClass> BuildArrayOfClassesFromPartition(const FFrontendQueryPartition& InPartition)
+			{
+				TArray<FMetasoundFrontendClass> Result;
+				for (const FFrontendQueryEntry& Entry : InPartition)
+				{
+					check(Entry.Value.IsType<FMetasoundFrontendClass>());
+					Result.Add(Entry.Value.Get<FMetasoundFrontendClass>());
+				}
+				return Result;
+			}
+
+			FMetasoundFrontendClass BuildSingleClassFromPartition(const FFrontendQueryPartition& InPartition)
+			{
+				FMetasoundFrontendClass MetaSoundClass;
+				if (ensureMsgf(InPartition.Num() > 0, TEXT("Cannot retrieve class from empty partition")))
+				{
+					MetaSoundClass = InPartition.CreateConstIterator()->Value.Get<FMetasoundFrontendClass>();
+				}
+				return MetaSoundClass;
+			}
+
+			TArray<FMetasoundFrontendInterface> BuildArrayOfInterfacesFromPartition(const FFrontendQueryPartition& InPartition)
+			{
+				TArray<FMetasoundFrontendInterface> Result;
+				for (const FFrontendQueryEntry& Entry : InPartition)
+				{
+					check(Entry.Value.IsType<FMetasoundFrontendInterface>());
+					Result.Add(Entry.Value.Get<FMetasoundFrontendInterface>());
+				}
+				return Result;
+			}
+
+			FMetasoundFrontendInterface BuildSingleInterfaceFromPartition(const FFrontendQueryPartition& InPartition)
+			{
+				FMetasoundFrontendInterface MetaSoundInterface;
+				if (ensureMsgf(InPartition.Num() > 0, TEXT("Cannot retrieve class from empty partition")))
+				{
+					MetaSoundInterface = InPartition.CreateConstIterator()->Value.Get<FMetasoundFrontendInterface>();
+				}
+				return MetaSoundInterface;
+			}
+
+			TArray<FMetasoundFrontendVersion> BuildArrayOfVersionsFromPartition(const FFrontendQueryPartition& InPartition)
+			{
+				TArray<FMetasoundFrontendVersion> Result;
+				for (const FFrontendQueryEntry& Entry : InPartition)
+				{
+					check(Entry.Value.IsType<FMetasoundFrontendVersion>());
+					Result.Add(Entry.Value.Get<FMetasoundFrontendVersion>());
+				}
+				return Result;
+			}
+		}
+
+
+		FFrontendQuery FFindClassWithHighestMinorVersionQueryPolicy::CreateQuery()
+		{
+			using namespace SearchEngineQuerySteps;
+			FFrontendQuery Query;
+			Query.AddStep<FNodeClassRegistrationEvents>()
+				.AddStep<FMapRegistrationEventsToNodeRegistryKeys>()
+				.AddStep<FReduceRegistrationEventsToCurrentStatus>()
+				.AddStep<FTransformRegistrationEventsToClasses>()
+				.AddStep<FRemoveDeprecatedClasses>()
+				.AddStep<FMapClassToFullClassNameAndMajorVersion>()
+				.AddStep<FReduceClassesToHighestVersion>();
+			return Query;
+		}
+
+		FMetasoundFrontendClass FFindClassWithHighestMinorVersionQueryPolicy::BuildResult(const FFrontendQueryPartition& InPartition)
+		{
+			using namespace SearchEngineQuerySteps;
+			return BuildSingleClassFromPartition(InPartition);
+		}
+
+		FFrontendQuery FFindAllDefaultInterfacesQueryPolicy::CreateQuery()
+		{
+			using namespace SearchEngineQuerySteps;
+			FFrontendQuery Query;
+			Query.AddStep<FInterfaceRegistryTransactionSource>()
+				.AddStep<FMapInterfaceRegistryTransactionsToInterfaceRegistryKeys>()
+				.AddStep<FReduceInterfaceRegistryTransactionsToCurrentStatus>()
+				.AddStep<FTransformInterfaceRegistryTransactionToInterface>()
+				.AddStep<FMapInterfaceToInterfaceNameAndMajorVersion>()
+				.AddStep<FReduceInterfacesToHighestVersion>()
+				.AddStep<FRemoveInterfacesWhichAreNotDefault>()
+				.AddStep<FMapToNull>();
+			return Query;
+		}
+
+		TArray<FMetasoundFrontendInterface> FFindAllDefaultInterfacesQueryPolicy::BuildResult(const FFrontendQueryPartition& InPartition)
+		{
+			using namespace SearchEngineQuerySteps;
+			return BuildArrayOfInterfacesFromPartition(InPartition);
+		}
+
+		FFrontendQuery FFindAllRegisteredInterfacesWithNameQueryPolicy::CreateQuery()
+		{
+			using namespace SearchEngineQuerySteps;
+			FFrontendQuery Query;
+			Query.AddStep<FInterfaceRegistryTransactionSource>()
+				.AddStep<FMapInterfaceRegistryTransactionsToInterfaceRegistryKeys>()
+				.AddStep<FReduceInterfaceRegistryTransactionsToCurrentStatus>()
+				.AddStep<FMapInterfaceRegistryTransactionToInterfaceName>()
+				.AddStep<FTransformInterfaceRegistryTransactionToInterfaceVersion>();
+			return Query;
+		}
+
+		TArray<FMetasoundFrontendVersion> FFindAllRegisteredInterfacesWithNameQueryPolicy::BuildResult(const FFrontendQueryPartition& InPartition)
+		{
+			using namespace SearchEngineQuerySteps;
+			return BuildArrayOfVersionsFromPartition(InPartition);
+		}
+
+		FFrontendQuery FFindInterfaceWithHighestVersionQueryPolicy::CreateQuery()
+		{
+			using namespace SearchEngineQuerySteps;
+			FFrontendQuery Query;
+			Query.AddStep<FInterfaceRegistryTransactionSource>()
+				.AddStep<FMapInterfaceRegistryTransactionsToInterfaceRegistryKeys>()
+				.AddStep<FReduceInterfaceRegistryTransactionsToCurrentStatus>()
+				.AddStep<FTransformInterfaceRegistryTransactionToInterface>()
+				.AddStep<FMapInterfaceToInterfaceName>()
+				.AddStep<FReduceInterfacesToHighestVersion>();
+			return Query;
+		}
+
+		FMetasoundFrontendInterface FFindInterfaceWithHighestVersionQueryPolicy::BuildResult(const FFrontendQueryPartition& InPartition)
+		{
+			using namespace SearchEngineQuerySteps;
+			return BuildSingleInterfaceFromPartition(InPartition);
+		}
+
+		void FSearchEngineCore::Prime()
+		{
+			METASOUND_LLM_SCOPE;
+			METASOUND_TRACE_CPUPROFILER_EVENT_SCOPE(metasound::FSearchEngineCore::Prime);
+
+			FindClassWithHighestMinorVersionQuery.Prime();
+			FindAllDefaultInterfacesQuery.Prime();
+			FindAllRegisteredInterfacesWithNameQuery.Prime();
+			FindInterfaceWithHighestVersionQuery.Prime();
+		}
+
+		bool FSearchEngineCore::FindClassWithHighestMinorVersion(const FMetasoundFrontendClassName& InName, int32 InMajorVersion, FMetasoundFrontendClass& OutClass)
+		{
+			METASOUND_TRACE_CPUPROFILER_EVENT_SCOPE(metasound::FSearchEngineCore::FindClassWithHighestMinorVersion);
+			using namespace SearchEngineQuerySteps;
+
+			const FFrontendQueryKey Key = FMapClassToFullClassNameAndMajorVersion::GetKey(InName, InMajorVersion);
+			return FindClassWithHighestMinorVersionQuery.UpdateAndFindResult(Key, OutClass);
+		}
+
+		TArray<FMetasoundFrontendInterface> FSearchEngineCore::FindUClassDefaultInterfaces(FName InUClassName)
+		{
+			METASOUND_TRACE_CPUPROFILER_EVENT_SCOPE(metasound::FSearchEngineCore::FindUClassDefaultInterfaces);
+
+			const FFrontendQueryKey NullKey;
+			TArray<FMetasoundFrontendInterface> DefaultInterfaces = FindAllDefaultInterfacesQuery.UpdateAndFindResult(NullKey);
+
+			// All default interfaces are filtered here because a UClass name cannot effectively 
+			// be used as a "Key" for this query.
+			//
+			// We cannot know all the UClasses associated with each default interface given only a UClassName
+			// because we need to inspect UClass inheritance to handle derived UClasses
+			// which should be compatible with their parent's default interfaces. 
+			//
+			// This module does not have access to the UClass inheritance structure.
+			// Instead, this information is accessed in the IInterfaceRegistryEntry
+			// concrete implementation. 
+			auto DoesNotSupportUClass = [&InUClassName](const FMetasoundFrontendInterface& InInterface)
+			{
+				FInterfaceRegistryKey RegistryKey = GetInterfaceRegistryKey(InInterface);
+				const IInterfaceRegistryEntry* RegEntry = IInterfaceRegistry::Get().FindInterfaceRegistryEntry(RegistryKey);
+				if (ensure(RegEntry))
+				{
+					return !RegEntry->UClassIsSupported(InUClassName);
+				}
+
+				return true;
+			};
+
+			DefaultInterfaces.RemoveAllSwap(DoesNotSupportUClass);
+
+			return DefaultInterfaces;
+		}
+
+		TArray<FMetasoundFrontendVersion> FSearchEngineCore::FindAllRegisteredInterfacesWithName(FName InInterfaceName)
+		{
+			METASOUND_TRACE_CPUPROFILER_EVENT_SCOPE(metasound::FSearchEngineCore::FindAllRegisteredInterfacesWithName);
+
+			const FFrontendQueryKey Key{InInterfaceName};
+			return FindAllRegisteredInterfacesWithNameQuery.UpdateAndFindResult(Key);
+		}
+
+		bool FSearchEngineCore::FindInterfaceWithHighestVersion(FName InInterfaceName, FMetasoundFrontendInterface& OutInterface)
+		{
+			METASOUND_TRACE_CPUPROFILER_EVENT_SCOPE(metasound::FSearchEngineCore::FindInterfaceWithHighestVersion);
+
+			const FFrontendQueryKey Key{InInterfaceName};
+			return FindInterfaceWithHighestVersionQuery.UpdateAndFindResult(Key, OutInterface);
+		}
+	}
+}
diff --git a/Engine/Plugins/Runtime/Metasound/Source/MetasoundFrontend/Private/MetasoundFrontendSearchEngineCore.h b/Engine/Plugins/Runtime/Metasound/Source/MetasoundFrontend/Private/MetasoundFrontendSearchEngineCore.h
new file mode 100644
--- /dev/null
+++ b/Engine/Plugins/Runtime/Metasound/Source/MetasoundFrontend/Private/MetasoundFrontendSearchEngineCore.h
@@ -0,0 +1,307 @@
+// Copyright Epic Games, Inc. All Rights Reserved.
+
+#pragma once
+
+#include "Containers/Array.h"
+#include "HAL/CriticalSection.h"
+#include "MetasoundFrontendArchetypeRegistry.h"
+#include "MetasoundFrontendQuery.h"
+#include "MetasoundFrontendQuerySteps.h"
+#include "MetasoundFrontendRegistries.h"
+#include "MetasoundFrontendRegistryTransaction.h"
+#include "MetasoundFrontendSearchEngine.h"
+#include "MetasoundTrace.h"
+
+namespace Metasound
+{
+	namespace Frontend
+	{
+		namespace SearchEngineQuerySteps
+		{
+			/* A collection of reusable steps to be used with the Metasound frontend query system. */
+			 
+			class FMapClassToFullClassNameAndMajorVersion : public IFrontendQueryMapStep
+			{
+			public:
+				static FFrontendQueryKey GetKey(const FMetasoundFrontendClassName& InClassName, int32 InMajorVersion);
+				virtual FFrontendQueryKey Map(const FFrontendQueryEntry& InEntry) const override;
+			};
+
+			class FInterfaceRegistryTransactionSource : public IFrontendQueryStreamStep
+			{
+			public:
+				FInterfaceRegistryTransactionSource();
+				virtual void Stream(TArray<FFrontendQueryValue>& OutEntries) override;
+
+			private:
+				FRegistryTransactionID CurrentTransactionID;
+			};
+
+			class FMapInterfaceRegistryTransactionsToInterfaceRegistryKeys : public IFrontendQueryMapStep
+			{
+			public:
+				virtual FFrontendQueryKey Map(const FFrontendQueryEntry& InEntry) const override;
+			};
+
+			class FMapInterfaceRegistryTransactionToInterfaceName : public IFrontendQueryMapStep
+			{
+			public:
+				virtual FFrontendQueryKey Map(const FFrontendQueryEntry& InEntry) const override;
+			};
+
+			class FReduceInterfaceRegistryTransactionsToCurrentStatus : public IFrontendQueryReduceStep
+			{
+			public:
+				virtual void Reduce(const FFrontendQueryKey& InKey, FFrontendQueryPartition& InOutEntries) const override;
+			private:
+				static FInterfaceRegistryTransaction::FTimeType GetTransactionTimestamp(const FFrontendQueryEntry& InEntry);
+				static bool IsValidTransactionOfType(FInterfaceRegistryTransaction::ETransactionType InType, const FFrontendQueryEntry* InEntry);
+			};
+
+			class FTransformInterfaceRegistryTransactionToInterface : public IFrontendQueryTransformStep
+			{
+			public:
+				virtual void Transform(FFrontendQueryEntry::FValue& InValue) const override;
+			};
+
+			class FMapInterfaceToInterfaceNameAndMajorVersion : public IFrontendQueryMapStep
+			{
+			public:
+				static FFrontendQueryKey GetKey(const FString& InName, int32 InMajorVersion);
+				virtual FFrontendQueryKey Map(const FFrontendQueryEntry& InEntry) const override;
+			};
+
+			class FReduceInterfacesToHighestVersion : public IFrontendQueryReduceStep
+			{
+			public:
+				virtual void Reduce(const FFrontendQueryKey& InKey, FFrontendQueryPartition& InOutEntries) const override;
+			};
+
+			class FRemoveDeprecatedClasses : public IFrontendQueryFilterStep
+			{
+			public:
+				virtual bool Filter(const FFrontendQueryEntry& InEntry) const override;
+			};
+
+			class FRemoveInterfacesWhichAreNotDefault : public IFrontendQueryFilterStep
+			{
+				virtual bool Filter(const FFrontendQueryEntry& InEntry) const override;
+			};
+
+			class FMapInterfaceToInterfaceName : public IFrontendQueryMapStep
+			{
+			public:
+				virtual FFrontendQueryKey Map(const FFrontendQueryEntry& InEntry) const override;
+			};
+
+			class FTransformInterfaceRegistryTransactionToInterfaceVersion : public IFrontendQueryTransformStep
+			{
+			public:
+				virtual void Transform(FFrontendQueryEntry::FValue& InValue) const override;
+			};
+
+			// Convert an array of partitions into an array of metasound frontend classes.
+			TArray<FMetasoundFrontendClass> BuildArrayOfClassesFromPartition(const FFrontendQueryPartition& InPartition);
+
+			FMetasoundFrontendClass BuildSingleClassFromPartition(const FFrontendQueryPartition& InPartition);
+
+			// Convert an array of partitions into an array of metasound frontend interfaces.
+			TArray<FMetasoundFrontendInterface> BuildArrayOfInterfacesFromPartition(const FFrontendQueryPartition& InPartition);
+
+			FMetasoundFrontendInterface BuildSingleInterfaceFromPartition(const FFrontendQueryPartition& InPartition);
+
+			// Convert an array of partitions into an array of metasound frontend versions.
+			TArray<FMetasoundFrontendVersion> BuildArrayOfVersionsFromPartition(const FFrontendQueryPartition& InPartition);
+		}
+
+		// Base interface for a search engine query. This allows the query to be
+		// primed to improve execution performance in critical call-stacks.
+		struct ISearchEngineQuery
+		{
+			// Prime the query to improve execution performance during subsequent 
+			// queries. 
+			virtual void Prime() = 0;
+			virtual ~ISearchEngineQuery() = default;
+		};
+
+		// TSearchEngineQuery provides a base for all queries used in the search
+		// engine. It requires a `QueryPolicy` as a template parameter. 
+		//
+		// A `QueryPolicy` is a struct containing a static function to create the 
+		// query, a static function to build a result from the query selection, 
+		// and query result type.
+		//
+		// Example:
+		//
+		// struct FQueryPolicyExample
+		// {
+		//		using ResultType = FMetasoundFrontendClass;
+		//
+		//		static FFrontendQuery CreateQuery()
+		//		{
+		//			using namespace SearchEngineQuerySteps;
+		//			FFrontendQuery Query;
+		//			Query.AddStep<FNodeClassRegistrationEvents>()
+		//				.AddStep<FMapRegistrationEventsToNodeRegistryKeys>()
+		//				.AddStep<FReduceRegistrationEventsToCurrentStatus>()
+		//				.AddStep<FTransformRegistrationEventsToClasses>();
+		//			return Query;
+		//		}
+		//
+		//		static ResultType BuildResult(const FFrontendQueryPartition& InPartition)
+		//		{
+		//			FMetasoundFrontendClass MetaSoundClass;
+		//			if (ensureMsgf(InPartition.Num() > 0, TEXT("Cannot retrieve class from empty partition")))
+		//			{
+		//				MetaSoundClass = InPartition.CreateConstIterator()->Value.Get<FMetasoundFrontendClass>();
+		//			}
+		//			return MetaSoundClass;
+		//		}
+		// };
+		template<typename QueryPolicy>
+		struct TSearchEngineQuery : public ISearchEngineQuery
+		{
+			using ResultType = typename QueryPolicy::ResultType;
+
+			TSearchEngineQuery()
+			: Query(QueryPolicy::CreateQuery())
+			{
+			}
+
+			virtual void Prime() override
+			{
+				METASOUND_LLM_SCOPE;
+				FScopeLock Lock(&QueryCriticalSection);
+				Update();
+			}
+
+			// Updates the query to the most recent value and returns the
+			// result.
+			ResultType UpdateAndFindResult(const FFrontendQueryKey& InKey)
+			{
+				METASOUND_LLM_SCOPE;
+				FScopeLock Lock(&QueryCriticalSection);
+
+				Update();
+				if (const ResultType* Result = ResultCache.Find(InKey))
+				{
+					return *Result;
+				}
+				return ResultType();
+			}
+
+			// Updates the query to the most recent value and assigns the value 
+			// to OutResult. Returns true on success, false on failure. 
+			bool UpdateAndFindResult(const FFrontendQueryKey& InKey, ResultType& OutResult)
+			{
+				METASOUND_LLM_SCOPE;
+				FScopeLock Lock(&QueryCriticalSection);
+
+				Update();
+				if (const ResultType* Result = ResultCache.Find(InKey))
+				{
+					OutResult = *Result;
+					return true;
+				}
+				return false;
+			}
+
+		private:
+			void Update()
+			{
+				METASOUND_TRACE_CPUPROFILER_EVENT_SCOPE(MetaSound::SearchEngineQuery::Update);
+
+				TSet<FFrontendQueryKey> UpdatedKeys;
+				const FFrontendQuerySelection& Selection = Query.Update(UpdatedKeys);
+
+				if (UpdatedKeys.Num() > 0)
+				{
+					METASOUND_TRACE_CPUPROFILER_EVENT_SCOPE(MetaSound::SearchEngineQuery::BuildResult);
+
+					for (const FFrontendQueryKey& Key : UpdatedKeys)
+					{
+						const FFrontendQueryPartition* Partition = Query.GetSelection().Find(Key);
+						if (nullptr != Partition)
+						{
+							ResultCache.Add(Key, QueryPolicy::BuildResult(*Partition));
+						}
+						else
+						{
+							ResultCache.Remove(Key);
+						}
+					}
+				}
+			}
+
+			FCriticalSection QueryCriticalSection;
+			FFrontendQuery Query;
+			TMap<FFrontendQueryKey, ResultType> ResultCache;
+		};
+
+
+		// Policy for finding highest versioned metasound class by name and major version.
+		struct FFindClassWithHighestMinorVersionQueryPolicy
+		{
+			using ResultType = FMetasoundFrontendClass;
+
+			static FFrontendQuery CreateQuery();
+			static ResultType BuildResult(const FFrontendQueryPartition& InPartition);
+		};
+
+		// Policy for finding all registered default interfaces. 
+		struct FFindAllDefaultInterfacesQueryPolicy
+		{
+			using ResultType = TArray<FMetasoundFrontendInterface>;
+
+			static FFrontendQuery CreateQuery();
+			static ResultType BuildResult(const FFrontendQueryPartition& InPartition);
+		};
+
+		// Policy for finding all registered interface versions (name & version number)
+		// by name.
+		struct FFindAllRegisteredInterfacesWithNameQueryPolicy
+		{
+			using ResultType = TArray<FMetasoundFrontendVersion>;
+
+			static FFrontendQuery CreateQuery();
+			static TArray<FMetasoundFrontendVersion> BuildResult(const FFrontendQueryPartition& InPartition);
+		};
+
+		// Policy for finding the highest version registered interfaces by name.
+		struct FFindInterfaceWithHighestVersionQueryPolicy
+		{
+			using ResultType = FMetasoundFrontendInterface;
+
+			static FFrontendQuery CreateQuery();
+			static FMetasoundFrontendInterface BuildResult(const FFrontendQueryPartition& InPartition);
+		};
+
+		/** Supports essential search engine functionality needed for runtime. */
+		class FSearchEngineCore : public ISearchEngine
+		{
+			FSearchEngineCore(const FSearchEngineCore&) = delete;
+			FSearchEngineCore(FSearchEngineCore&&) = delete;
+			FSearchEngineCore& operator=(const FSearchEngineCore&) = delete;
+			FSearchEngineCore& operator=(FSearchEngineCore&&) = delete;
+
+		public:
+			FSearchEngineCore() = default;
+			virtual ~FSearchEngineCore() = default;
+
+			virtual void Prime() override;
+			virtual bool FindClassWithHighestMinorVersion(const FMetasoundFrontendClassName& InName, int32 InMajorVersion, FMetasoundFrontendClass& OutClass) override;
+
+			virtual TArray<FMetasoundFrontendInterface> FindUClassDefaultInterfaces(FName InUClassName) override;
+			virtual TArray<FMetasoundFrontendVersion> FindAllRegisteredInterfacesWithName(FName InInterfaceName) override;
+			virtual bool FindInterfaceWithHighestVersion(FName InInterfaceName, FMetasoundFrontendInterface& OutInterface) override;
+
+		private:
+
+			TSearchEngineQuery<FFindClassWithHighestMinorVersionQueryPolicy> FindClassWithHighestMinorVersionQuery;
+			TSearchEngineQuery<FFindAllDefaultInterfacesQueryPolicy> FindAllDefaultInterfacesQuery;
+			TSearchEngineQuery<FFindAllRegisteredInterfacesWithNameQueryPolicy> FindAllRegisteredInterfacesWithNameQuery;
+			TSearchEngineQuery<FFindInterfaceWithHighestVersionQueryPolicy> FindInterfaceWithHighestVersionQuery;
+		};
+	}
+}
+
diff --git a/Engine/Plugins/Runtime/Metasound/Source/MetasoundFrontend/Private/MetasoundFrontendSearchEngineEditorOnly.cpp b/Engine/Plugins/Runtime/Metasound/Source/MetasoundFrontend/Private/MetasoundFrontendSearchEngineEditorOnly.cpp
new file mode 100644
--- /dev/null
+++ b/Engine/Plugins/Runtime/Metasound/Source/MetasoundFrontend/Private/MetasoundFrontendSearchEngineEditorOnly.cpp
@@ -0,0 +1,225 @@
+// Copyright Epic Games, Inc. All Rights Reserved.
+
+#if WITH_EDITORONLY_DATA
+
+#include "MetasoundFrontendSearchEngineEditorOnly.h"
+
+#include "Containers/Array.h"
+#include "MetasoundFrontendArchetypeRegistry.h"
+#include "MetasoundFrontendDocument.h"
+#include "MetasoundFrontendQuery.h"
+#include "MetasoundFrontendQuerySteps.h"
+#include "MetasoundFrontendRegistryTransaction.h"
+#include "MetasoundFrontendSearchEngineCore.h"
+#include "MetasoundTrace.h"
+
+namespace Metasound
+{
+	namespace Frontend
+	{
+		FFrontendQuery FFindAllClassesQueryPolicy::CreateQuery()
+		{
+			using namespace SearchEngineQuerySteps;
+			FFrontendQuery Query;
+			Query.AddStep<FNodeClassRegistrationEvents>()
+				.AddStep<FMapRegistrationEventsToNodeRegistryKeys>()
+				.AddStep<FReduceRegistrationEventsToCurrentStatus>()
+				.AddStep<FTransformRegistrationEventsToClasses>()
+				.AddStep<FRemoveDeprecatedClasses>()
+				.AddStep<FMapToNull>();
+			return Query;
+		}
+
+		TArray<FMetasoundFrontendClass> FFindAllClassesQueryPolicy::BuildResult(const FFrontendQueryPartition& InPartition)
+		{
+			using namespace SearchEngineQuerySteps;
+			return BuildArrayOfClassesFromPartition(InPartition);
+		}
+
+		FFrontendQuery FFindAllClassesIncludingAllVersionsQueryPolicy::CreateQuery()
+		{
+			using namespace SearchEngineQuerySteps;
+			FFrontendQuery Query;
+			Query.AddStep<FNodeClassRegistrationEvents>()
+				.AddStep<FMapRegistrationEventsToNodeRegistryKeys>()
+				.AddStep<FReduceRegistrationEventsToCurrentStatus>()
+				.AddStep<FTransformRegistrationEventsToClasses>()
+				.AddStep<FMapToNull>();
+			return Query;
+		}
+
+		TArray<FMetasoundFrontendClass> FFindAllClassesIncludingAllVersionsQueryPolicy::BuildResult(const FFrontendQueryPartition& InPartition)
+		{
+			using namespace SearchEngineQuerySteps;
+			return BuildArrayOfClassesFromPartition(InPartition);
+		}
+
+		FFrontendQuery FFindClassesWithNameSortedQueryPolicy::CreateQuery()
+		{
+			using namespace SearchEngineQuerySteps;
+			FFrontendQuery Query;
+			Query.AddStep<FNodeClassRegistrationEvents>()
+				.AddStep<FMapRegistrationEventsToNodeRegistryKeys>()
+				.AddStep<FReduceRegistrationEventsToCurrentStatus>()
+				.AddStep<FTransformRegistrationEventsToClasses>()
+				.AddStep<FMapClassesToClassName>()
+				.AddStep<FSortClassesByVersion>();
+			return Query;
+		}
+
+		TArray<FMetasoundFrontendClass> FFindClassesWithNameSortedQueryPolicy::BuildResult(const FFrontendQueryPartition& InPartition)
+		{
+			using namespace SearchEngineQuerySteps;
+			return BuildArrayOfClassesFromPartition(InPartition);
+		}
+
+		FFrontendQuery FFindClassesWithNameUnsortedQueryPolicy::CreateQuery()
+		{
+			using namespace SearchEngineQuerySteps;
+			FFrontendQuery Query;
+			Query.AddStep<FNodeClassRegistrationEvents>()
+				.AddStep<FMapRegistrationEventsToNodeRegistryKeys>()
+				.AddStep<FReduceRegistrationEventsToCurrentStatus>()
+				.AddStep<FTransformRegistrationEventsToClasses>()
+				.AddStep<FMapClassesToClassName>();
+			return Query;
+		}
+
+		TArray<FMetasoundFrontendClass> FFindClassesWithNameUnsortedQueryPolicy::BuildResult(const FFrontendQueryPartition& InPartition)
+		{
+			using namespace SearchEngineQuerySteps;
+			return BuildArrayOfClassesFromPartition(InPartition);
+		}
+
+		FFrontendQuery FFindClassWithHighestVersionQueryPolicy::CreateQuery()
+		{
+			using namespace SearchEngineQuerySteps;
+			FFrontendQuery Query;
+			Query.AddStep<FNodeClassRegistrationEvents>()
+				.AddStep<FMapRegistrationEventsToNodeRegistryKeys>()
+				.AddStep<FReduceRegistrationEventsToCurrentStatus>()
+				.AddStep<FTransformRegistrationEventsToClasses>()
+				.AddStep<FRemoveDeprecatedClasses>()
+				.AddStep<FMapToFullClassName>()
+				.AddStep<FReduceClassesToHighestVersion>();
+			return Query;
+		}
+
+		FMetasoundFrontendClass FFindClassWithHighestVersionQueryPolicy::BuildResult(const FFrontendQueryPartition& InPartition)
+		{
+			using namespace SearchEngineQuerySteps;
+			return BuildSingleClassFromPartition(InPartition);
+		}
+
+		FFrontendQuery FFindAllInterfacesQueryPolicy::CreateQuery()
+		{
+			using namespace SearchEngineQuerySteps;
+			FFrontendQuery Query;
+			Query.AddStep<FInterfaceRegistryTransactionSource>()
+				.AddStep<FMapInterfaceRegistryTransactionsToInterfaceRegistryKeys>()
+				.AddStep<FReduceInterfaceRegistryTransactionsToCurrentStatus>()
+				.AddStep<FTransformInterfaceRegistryTransactionToInterface>()
+				.AddStep<FMapInterfaceToInterfaceNameAndMajorVersion>()
+				.AddStep<FReduceInterfacesToHighestVersion>()
+				.AddStep<FMapToNull>();
+			return Query;
+		}
+
+		TArray<FMetasoundFrontendInterface> FFindAllInterfacesQueryPolicy::BuildResult(const FFrontendQueryPartition& InPartition)
+		{
+			using namespace SearchEngineQuerySteps;
+			return BuildArrayOfInterfacesFromPartition(InPartition);
+		}
+
+		FFrontendQuery FFindAllInterfacesIncludingAllVersionsQueryPolicy::CreateQuery()
+		{
+			using namespace SearchEngineQuerySteps;
+			FFrontendQuery Query;
+			Query.AddStep<FInterfaceRegistryTransactionSource>()
+				.AddStep<FMapInterfaceRegistryTransactionsToInterfaceRegistryKeys>()
+				.AddStep<FReduceInterfaceRegistryTransactionsToCurrentStatus>()
+				.AddStep<FTransformInterfaceRegistryTransactionToInterface>()
+				.AddStep<FMapToNull>();
+			return Query;
+		}
+
+		TArray<FMetasoundFrontendInterface> FFindAllInterfacesIncludingAllVersionsQueryPolicy::BuildResult(const FFrontendQueryPartition& InPartition)
+		{
+			using namespace SearchEngineQuerySteps;
+			return BuildArrayOfInterfacesFromPartition(InPartition);
+		}
+
+		void FSearchEngineEditorOnly::Prime()
+		{
+			Super::Prime();
+			{
+				METASOUND_LLM_SCOPE;
+				METASOUND_TRACE_CPUPROFILER_EVENT_SCOPE(metasound::FSearchEngineEditorOnly::Prime);
+
+				FindAllClassesQuery.Prime();
+				FindAllClassesIncludingAllVersionsQuery.Prime();
+				FindClassesWithNameUnsortedQuery.Prime();
+				FindClassesWithNameSortedQuery.Prime();
+				FindClassWithHighestVersionQuery.Prime();
+				FindAllInterfacesIncludingAllVersionsQuery.Prime();
+				FindAllInterfacesQuery.Prime();
+			}
+		}
+
+		TArray<FMetasoundFrontendClass> FSearchEngineEditorOnly::FindAllClasses(bool bInIncludeAllVersions)
+		{
+			METASOUND_TRACE_CPUPROFILER_EVENT_SCOPE(metasound::FSearchEngineEditorOnly::FindAllClasses);
+
+			FFrontendQueryKey NullKey;
+			if (bInIncludeAllVersions)
+			{
+				return FindAllClassesIncludingAllVersionsQuery.UpdateAndFindResult(NullKey);
+			}
+			else
+			{
+				return FindAllClassesQuery.UpdateAndFindResult(NullKey);
+			}
+		}
+
+		TArray<FMetasoundFrontendClass> FSearchEngineEditorOnly::FindClassesWithName(const FMetasoundFrontendClassName& InName, bool bInSortByVersion)
+		{
+			METASOUND_TRACE_CPUPROFILER_EVENT_SCOPE(metasound::FSearchEngineEditorOnly::FindClassesWithName);
+
+			const FFrontendQueryKey Key(InName.GetFullName());
+			if (bInSortByVersion)
+			{
+				return FindClassesWithNameSortedQuery.UpdateAndFindResult(Key);
+			}
+			else 
+			{
+				return FindClassesWithNameUnsortedQuery.UpdateAndFindResult(Key);
+			}
+		}
+
+		bool FSearchEngineEditorOnly::FindClassWithHighestVersion(const FMetasoundFrontendClassName& InName, FMetasoundFrontendClass& OutClass)
+		{
+			METASOUND_TRACE_CPUPROFILER_EVENT_SCOPE(metasound::FSearchEngineEditorOnly::FindClassWithHighestVersion);
+
+			const FFrontendQueryKey Key{InName.GetFullName()};
+			return FindClassWithHighestVersionQuery.UpdateAndFindResult(Key, OutClass);
+		}
+
+		TArray<FMetasoundFrontendInterface> FSearchEngineEditorOnly::FindAllInterfaces(bool bInIncludeAllVersions)
+		{
+			METASOUND_TRACE_CPUPROFILER_EVENT_SCOPE(metasound::FSearchEngineEditorOnly::FindAllInterfaces);
+
+			const FFrontendQueryKey NullKey;
+			if (bInIncludeAllVersions)
+			{
+				return FindAllInterfacesIncludingAllVersionsQuery.UpdateAndFindResult(NullKey);
+			}
+			else
+			{
+				return FindAllInterfacesQuery.UpdateAndFindResult(NullKey);
+			}
+		}
+
+	}
+}
+
+#endif
diff --git a/Engine/Plugins/Runtime/Metasound/Source/MetasoundFrontend/Private/MetasoundFrontendSearchEngineEditorOnly.h b/Engine/Plugins/Runtime/Metasound/Source/MetasoundFrontend/Private/MetasoundFrontendSearchEngineEditorOnly.h
new file mode 100644
--- /dev/null
+++ b/Engine/Plugins/Runtime/Metasound/Source/MetasoundFrontend/Private/MetasoundFrontendSearchEngineEditorOnly.h
@@ -0,0 +1,116 @@
+// Copyright Epic Games, Inc. All Rights Reserved.
+
+#pragma once
+
+#if WITH_EDITORONLY_DATA
+
+#include "Containers/Array.h"
+#include "MetasoundFrontendDocument.h"
+#include "MetasoundFrontendQuery.h"
+#include "MetasoundFrontendSearchEngine.h"
+#include "MetasoundFrontendSearchEngineCore.h"
+
+namespace Metasound
+{
+	namespace Frontend
+	{
+		// Policy for finding all registered metasound classes. 
+		struct FFindAllClassesQueryPolicy
+		{
+			using ResultType = TArray<FMetasoundFrontendClass>;
+
+			static FFrontendQuery CreateQuery();
+			static ResultType BuildResult(const FFrontendQueryPartition& InPartition);
+		};
+
+		// Policy for finding all registered metasound classes, including deprecated classes. 
+		struct FFindAllClassesIncludingAllVersionsQueryPolicy
+		{
+			using ResultType = TArray<FMetasoundFrontendClass>;
+
+			static FFrontendQuery CreateQuery();
+			static ResultType BuildResult(const FFrontendQueryPartition& InPartition);
+		};
+
+		// Policy for finding all registered metasound classes sorted by version 
+		// and indexed by name.
+		struct FFindClassesWithNameSortedQueryPolicy
+		{
+			using ResultType = TArray<FMetasoundFrontendClass>;
+
+			static FFrontendQuery CreateQuery();
+			static ResultType BuildResult(const FFrontendQueryPartition& InPartition);
+		};
+
+		// Policy for finding all registered metasound classes indexed by name.
+		struct FFindClassesWithNameUnsortedQueryPolicy
+		{
+			using ResultType = TArray<FMetasoundFrontendClass>;
+
+			static FFrontendQuery CreateQuery();
+			static ResultType BuildResult(const FFrontendQueryPartition& InPartition);
+		};
+
+		// Policy for finding highest versioned metasound class by name.
+		struct FFindClassWithHighestVersionQueryPolicy
+		{
+			using ResultType = FMetasoundFrontendClass;
+
+			static FFrontendQuery CreateQuery();
+			static ResultType BuildResult(const FFrontendQueryPartition& InPartition);
+		};
+
+		// Policy for finding all registered interfaces. 
+		struct FFindAllInterfacesQueryPolicy
+		{
+			using ResultType = TArray<FMetasoundFrontendInterface>;
+
+			static FFrontendQuery CreateQuery();
+			static ResultType BuildResult(const FFrontendQueryPartition& InPartition);
+		};
+
+		// Policy for finding all registered interfaces (including deprecated). 
+		struct FFindAllInterfacesIncludingAllVersionsQueryPolicy
+		{
+			using ResultType = TArray<FMetasoundFrontendInterface>;
+
+			static FFrontendQuery CreateQuery();
+			static TArray<FMetasoundFrontendInterface> BuildResult(const FFrontendQueryPartition& InPartition);
+		};
+
+		// To minimize runtime memory costs during gampeplay, some search engine
+		// queries are only exposed when the editor only data is available. This
+		// class supports the editor only functionality of the search engine.
+		class FSearchEngineEditorOnly : public FSearchEngineCore
+		{
+			FSearchEngineEditorOnly(const FSearchEngineEditorOnly&) = delete;
+			FSearchEngineEditorOnly(FSearchEngineEditorOnly&&) = delete;
+			FSearchEngineEditorOnly& operator=(const FSearchEngineEditorOnly&) = delete;
+			FSearchEngineEditorOnly& operator=(FSearchEngineEditorOnly&&) = delete;
+
+			using Super = FSearchEngineCore;
+
+		public:
+			FSearchEngineEditorOnly() = default;
+			virtual ~FSearchEngineEditorOnly() = default;
+
+			virtual void Prime() override;
+			virtual TArray<FMetasoundFrontendClass> FindAllClasses(bool bInIncludeAllVersions) override;
+			virtual TArray<FMetasoundFrontendClass> FindClassesWithName(const FMetasoundFrontendClassName& InName, bool bInSortByVersion) override;
+			virtual bool FindClassWithHighestVersion(const FMetasoundFrontendClassName& InName, FMetasoundFrontendClass& OutClass) override;
+			virtual TArray<FMetasoundFrontendInterface> FindAllInterfaces(bool bInIncludeAllVersions) override;
+
+		private:
+
+			TSearchEngineQuery<FFindAllClassesQueryPolicy> FindAllClassesQuery;
+			TSearchEngineQuery<FFindAllClassesIncludingAllVersionsQueryPolicy> FindAllClassesIncludingAllVersionsQuery;
+			TSearchEngineQuery<FFindClassesWithNameUnsortedQueryPolicy> FindClassesWithNameUnsortedQuery;
+			TSearchEngineQuery<FFindClassesWithNameSortedQueryPolicy> FindClassesWithNameSortedQuery;
+			TSearchEngineQuery<FFindClassWithHighestVersionQueryPolicy> FindClassWithHighestVersionQuery;
+			TSearchEngineQuery<FFindAllInterfacesQueryPolicy> FindAllInterfacesQuery;
+			TSearchEngineQuery<FFindAllInterfacesIncludingAllVersionsQueryPolicy> FindAllInterfacesIncludingAllVersionsQuery;
+		};
+	}
+}
+
+#endif
diff --git a/Engine/Plugins/Runtime/Metasound/Source/MetasoundFrontend/Private/MetasoundFrontendTransform.cpp b/Engine/Plugins/Runtime/Metasound/Source/MetasoundFrontend/Private/MetasoundFrontendTransform.cpp
--- a/Engine/Plugins/Runtime/Metasound/Source/MetasoundFrontend/Private/MetasoundFrontendTransform.cpp
+++ b/Engine/Plugins/Runtime/Metasound/Source/MetasoundFrontend/Private/MetasoundFrontendTransform.cpp
@@ -569,16 +569,21 @@ namespace Metasound
 					return;
 				}
 
-				FMetasoundFrontendVersionNumber UpdateVersion = NodeHandle->FindHighestMinorVersionInRegistry();
-				if (UpdateVersion.IsValid() && UpdateVersion > ClassMetadata.GetVersion())
+				// Check if a updated minor version exists.
+				FMetasoundFrontendClass ClassWithHighestMinorVersion;
+				bool bFoundClassInSearchEngine = Frontend::ISearchEngine::Get().FindClassWithHighestMinorVersion(ClassMetadata.GetClassName(), ClassMetadata.GetVersion().Major, ClassWithHighestMinorVersion);
+
+				if (bFoundClassInSearchEngine && (ClassWithHighestMinorVersion.Metadata.GetVersion() > ClassMetadata.GetVersion()))
 				{
+					const FMetasoundFrontendVersionNumber UpdateVersion = ClassWithHighestMinorVersion.Metadata.GetVersion();
+
 					UE_LOG(LogMetaSound, Display, TEXT("Auto-Updating '%s' node class '%s': Newer version '%s' found."), *DebugAssetPath, *ClassMetadata.GetClassName().ToString(), *UpdateVersion.ToString());
 
 					NodesToUpdate.Add(TPair<FNodeHandle, FMetasoundFrontendVersionNumber>(NodeHandle, UpdateVersion));
 				}
 				else if (InterfaceUpdates.ContainsChanges())
 				{
-					UpdateVersion = ClassMetadata.GetVersion();
+					const FMetasoundFrontendVersionNumber UpdateVersion = ClassMetadata.GetVersion();
 					UE_LOG(LogMetaSound, Display, TEXT("Auto-Updating '%s' node class '%s (%s)': Interface change detected."), *DebugAssetPath, *ClassMetadata.GetClassName().ToString(), *UpdateVersion.ToString());
 
 					NodesToUpdate.Add(TPair<FNodeHandle, FMetasoundFrontendVersionNumber>(NodeHandle, UpdateVersion));
@@ -589,7 +594,7 @@ namespace Metasound
 #if WITH_EDITORONLY_DATA
 				else
 				{
-					NodesToUpdate.Add(TPair<FNodeHandle, FMetasoundFrontendVersionNumber>(NodeHandle, UpdateVersion));
+					NodesToUpdate.Add(TPair<FNodeHandle, FMetasoundFrontendVersionNumber>(NodeHandle, ClassMetadata.GetVersion()));
 				}
 #endif // WITH_EDITORONLY_DATA
 			}, EMetasoundFrontendClassType::External);
@@ -1113,14 +1118,33 @@ namespace Metasound
 				// with an "Invalid" interface version.
 				if (ensure(Interfaces.IsEmpty()))
 				{
-					constexpr bool bIncludeAllVersions = true;
-					TArray<FMetasoundFrontendInterface> AllInterfaces = ISearchEngine::Get().FindAllInterfaces(bIncludeAllVersions);
+					// At the time when version 1.4 of the document was introduced, 
+					// these were the only available interfaces. 
+					static const FMetasoundFrontendVersion PreexistingInterfaceVersions[] = {
+						FMetasoundFrontendVersion{"MetaSound", {1, 0}},
+						FMetasoundFrontendVersion{"MonoSource", {1, 0}},
+						FMetasoundFrontendVersion{"StereoSource", {1, 0}},
+						FMetasoundFrontendVersion{"MonoSource", {1, 1}},
+						FMetasoundFrontendVersion{"StereoSource", {1, 1}}
+					};
+					static const int32 NumPreexistingInterfaceVersions = sizeof(PreexistingInterfaceVersions) / sizeof(PreexistingInterfaceVersions[0]);
+
+					TArray<FMetasoundFrontendInterface> CandidateInterfaces;
+					IInterfaceRegistry& InterfaceRegistry = IInterfaceRegistry::Get();
+					for (int32 i = 0; i < NumPreexistingInterfaceVersions; i++)
+					{
+						FMetasoundFrontendInterface Interface;
+						if (InterfaceRegistry.FindInterface(GetInterfaceRegistryKey(PreexistingInterfaceVersions[i]), Interface))
+						{
+							CandidateInterfaces.Add(Interface);
+						}
+					}
 
 					const FMetasoundFrontendGraphClass& RootGraph = InDocument->GetRootGraphClass();
 					const TArray<FMetasoundFrontendClass>& Dependencies = InDocument->GetDependencies();
 					const TArray<FMetasoundFrontendGraphClass>& Subgraphs = InDocument->GetSubgraphs();
 
-					if (const FMetasoundFrontendInterface* Interface = FindMostSimilarInterfaceSupportingEnvironment(RootGraph, Dependencies, Subgraphs, AllInterfaces))
+					if (const FMetasoundFrontendInterface* Interface = FindMostSimilarInterfaceSupportingEnvironment(RootGraph, Dependencies, Subgraphs, CandidateInterfaces))
 					{
 						UE_LOG(LogMetaSound, Display, TEXT("Assigned interface [InterfaceVersion:%s] to document [RootGraphClassName:%s]"),
 							*Interface->Version.ToString(), *RootGraph.Metadata.GetClassName().ToString());
diff --git a/Engine/Plugins/Runtime/Metasound/Source/MetasoundFrontend/Public/MetasoundFrontendController.h b/Engine/Plugins/Runtime/Metasound/Source/MetasoundFrontend/Public/MetasoundFrontendController.h
--- a/Engine/Plugins/Runtime/Metasound/Source/MetasoundFrontend/Public/MetasoundFrontendController.h
+++ b/Engine/Plugins/Runtime/Metasound/Source/MetasoundFrontend/Public/MetasoundFrontendController.h
@@ -472,12 +472,6 @@ namespace Metasound
 			/** Returns interface version if node is a required member of a given interface, otherwise returns invalid version. */
 			virtual const FMetasoundFrontendVersion& GetInterfaceVersion() const = 0;
 
-			/** Returns the highest minor version number available from the class registry that shares this node's name & major version. */
-			virtual FMetasoundFrontendVersionNumber FindHighestVersionInRegistry() const = 0;
-
-			/** Returns the highest version number available from the class registry that shares this node's name. */
-			virtual FMetasoundFrontendVersionNumber FindHighestMinorVersionInRegistry() const = 0;
-
 			/**
 			  * Replaces this node with a new node of the provided version number, and attempts to 
 			  * rebuild edges where possible with matching vertex names that share the same DataType.
@@ -485,6 +479,7 @@ namespace Metasound
 			  */
 			virtual FNodeHandle ReplaceWithVersion(const FMetasoundFrontendVersionNumber& InNewVersion, TArray<FVertexNameAndType>* OutDisconnectedInputs, TArray<FVertexNameAndType>* OutDisconnectedOutputs) = 0;
 
+
 			/** Returns an input with the given id.
 			 *
 			 * If the input does not exist, an invalid handle is returned.
diff --git a/Engine/Plugins/Runtime/Metasound/Source/MetasoundFrontend/Public/MetasoundFrontendSearchEngine.h b/Engine/Plugins/Runtime/Metasound/Source/MetasoundFrontend/Public/MetasoundFrontendSearchEngine.h
--- a/Engine/Plugins/Runtime/Metasound/Source/MetasoundFrontend/Public/MetasoundFrontendSearchEngine.h
+++ b/Engine/Plugins/Runtime/Metasound/Source/MetasoundFrontend/Public/MetasoundFrontendSearchEngine.h
@@ -20,9 +20,21 @@ namespace Metasound
 				/** Return an instance of a search engine. */
 				static ISearchEngine& Get();
 
+				virtual ~ISearchEngine() = default;
+
 				/** Updates internal state to speed up queries. */
 				virtual void Prime() = 0;
 
+				/** Find the class with the given ClassName & Major Version. Returns false if not found, true if found. */
+				virtual bool FindClassWithHighestMinorVersion(const FMetasoundFrontendClassName& InName, int32 InMajorVersion, FMetasoundFrontendClass& OutClass) = 0;
+
+				virtual TArray<FMetasoundFrontendVersion> FindAllRegisteredInterfacesWithName(FName InInterfaceName) = 0;
+
+				virtual bool FindInterfaceWithHighestVersion(FName InInterfaceName, FMetasoundFrontendInterface& OutInterface) = 0;
+
+				virtual TArray<FMetasoundFrontendInterface> FindUClassDefaultInterfaces(FName InUClassName) = 0;
+
+#if WITH_EDITORONLY_DATA
 				/** Find all FMetasoundFrontendClasses.
 				  * (Optional) Include all versions (i.e. deprecated classes and versions of classes that are not the highest major version).
 				  */
@@ -36,18 +48,9 @@ namespace Metasound
 				/** Find the highest version of a class with the given ClassName. Returns false if not found, true if found. */
 				virtual bool FindClassWithHighestVersion(const FMetasoundFrontendClassName& InName, FMetasoundFrontendClass& OutClass) = 0;
 
-				/** Find the class with the given ClassName & Major Version. Returns false if not found, true if found. */
-				virtual bool FindClassWithMajorVersion(const FMetasoundFrontendClassName& InName, int32 InMajorVersion, FMetasoundFrontendClass& OutClass) = 0;
-
 				virtual TArray<FMetasoundFrontendInterface> FindAllInterfaces(bool bInIncludeAllVersions = false) = 0;
+#endif
 
-				virtual TArray<FMetasoundFrontendInterface> FindUClassDefaultInterfaces(FName InUClassName) = 0;
-
-				virtual TArray<FMetasoundFrontendVersion> FindAllRegisteredInterfacesWithName(FName InInterfaceName) = 0;
-
-				virtual bool FindInterfaceWithHighestVersion(FName InInterfaceName, FMetasoundFrontendInterface& OutInterface) = 0;
-
-				virtual ~ISearchEngine() = default;
 
 			protected:
 				ISearchEngine() = default;
