Fix for MetaSounds that spawn out of range being culled when Virtualization Setting is 'Restart'

diff --git a/Engine/Plugins/Runtime/Metasound/Source/MetasoundEngine/Private/MetasoundSource.cpp b/Engine/Plugins/Runtime/Metasound/Source/MetasoundEngine/Private/MetasoundSource.cpp
--- a/Engine/Plugins/Runtime/Metasound/Source/MetasoundEngine/Private/MetasoundSource.cpp
+++ b/Engine/Plugins/Runtime/Metasound/Source/MetasoundEngine/Private/MetasoundSource.cpp
@@ -498,8 +498,12 @@ bool UMetaSoundSource::SupportsSubtitles() const
 
 float UMetaSoundSource::GetDuration() const
 {
-	// eh? this is kind of a weird field anyways.
-	return Super::GetDuration();
+	// This is an unfortunate function required by logic in determining what sounds can be potentially
+	// culled (in this case prematurally). MetaSound OneShots are stopped either by internally logic that
+	// triggers OnFinished, or if an external system requests the sound to be stopped. Setting the duration
+	// as a "close to" maximum length without being considered looping avoids the MetaSound from being
+	// culled inappropriately.
+	return IsOneShot() ? INDEFINITELY_LOOPING_DURATION - 1.0f : INDEFINITELY_LOOPING_DURATION;
 }
 
 bool UMetaSoundSource::ImplementsParameterInterface(Audio::FParameterInterfacePtr InInterface) const
