Non unity fix in SMetasoundGraphPin

#rnx

diff --git a/Engine/Plugins/Runtime/Metasound/Source/MetasoundEditor/Private/SMetasoundGraphPin.h b/Engine/Plugins/Runtime/Metasound/Source/MetasoundEditor/Private/SMetasoundGraphPin.h
--- a/Engine/Plugins/Runtime/Metasound/Source/MetasoundEditor/Private/SMetasoundGraphPin.h
+++ b/Engine/Plugins/Runtime/Metasound/Source/MetasoundEditor/Private/SMetasoundGraphPin.h
@@ -13,6 +13,7 @@
 #include "MetasoundEditorGraphNode.h"
 #include "MetasoundFrontendController.h"
 #include "SGraphPin.h"
+#include "Styling/SlateStyleRegistry.h"
 #include "Widgets/Input/SButton.h"
 #include "Widgets/SBoxPanel.h"
 #include "Widgets/SNullWidget.h"
