MetaSound Frontend Memory Improvements
- Reduces baseline metasound plugin memory usage from 6.0 -> 3.5 mb
- Use a TSharedRef for storing node and data type registry entries to avoid storing multiple copies.
- Rework the runtime search engine query "FindHighestMinorVersion" to use less memory.

diff --git a/Engine/Plugins/Runtime/Metasound/Source/MetasoundFrontend/Private/MetasoundFrontendDataTypeRegistry.cpp b/Engine/Plugins/Runtime/Metasound/Source/MetasoundFrontend/Private/MetasoundFrontendDataTypeRegistry.cpp
--- a/Engine/Plugins/Runtime/Metasound/Source/MetasoundFrontend/Private/MetasoundFrontendDataTypeRegistry.cpp
+++ b/Engine/Plugins/Runtime/Metasound/Source/MetasoundFrontend/Private/MetasoundFrontendDataTypeRegistry.cpp
@@ -140,11 +140,6 @@ namespace Metasound
 					return nullptr;
 				}
 
-				virtual const FMetasoundFrontendClass& GetFrontendClass() const override
-				{
-					return FrontendClass;
-				}
-
 				virtual bool IsNative() const override
 				{
 					return true;
@@ -152,16 +147,14 @@ namespace Metasound
 
 			protected:
 
-				void SetFrontendClass(const FMetasoundFrontendClass& InClass)
+				void UpdateNodeClassInfo(const FMetasoundFrontendClass& InClass)
 				{
-					FrontendClass = InClass;
-					ClassInfo = FNodeClassInfo(FrontendClass.Metadata);
+					ClassInfo = FNodeClassInfo(InClass.Metadata);
 				}
 
 			private:
 				
 				FNodeClassInfo ClassInfo;
-				FMetasoundFrontendClass FrontendClass;
 			};
 
 			// Node registry entry for input nodes created from a data type registry entry.
@@ -170,38 +163,32 @@ namespace Metasound
 			public:
 				FInputNodeRegistryEntry() = delete;
 
-				FInputNodeRegistryEntry(TUniquePtr<IDataTypeRegistryEntry>&& InDataTypeEntry)
-				: DataTypeEntry(MoveTemp(InDataTypeEntry))
+				FInputNodeRegistryEntry(const TSharedRef<IDataTypeRegistryEntry, ESPMode::ThreadSafe>& InDataTypeEntry)
+				: DataTypeEntry(InDataTypeEntry)
 				{
-					if (DataTypeEntry.IsValid())
-					{
-						SetFrontendClass(DataTypeEntry->GetFrontendInputClass());
-					}
+					UpdateNodeClassInfo(DataTypeEntry->GetFrontendInputClass());
 				}
 
 				virtual ~FInputNodeRegistryEntry() = default;
 
+				virtual const FMetasoundFrontendClass& GetFrontendClass() const override
+				{
+					return DataTypeEntry->GetFrontendInputClass();
+				}
+
 				virtual TUniquePtr<INode> CreateNode(FDefaultNamedVertexWithLiteralNodeConstructorParams&& InParams) const override
 				{
-					if (DataTypeEntry.IsValid())
-					{
-						return DataTypeEntry->CreateInputNode(MoveTemp(InParams));
-					}
-					return nullptr;
+					return DataTypeEntry->CreateInputNode(MoveTemp(InParams));
 				}
 
 				virtual TUniquePtr<INodeRegistryEntry> Clone() const override
 				{
-					if (DataTypeEntry.IsValid())
-					{
-						return MakeUnique<FInputNodeRegistryEntry>(DataTypeEntry->Clone());
-					}
-					return MakeUnique<FInputNodeRegistryEntry>(TUniquePtr<IDataTypeRegistryEntry>());
+					return MakeUnique<FInputNodeRegistryEntry>(DataTypeEntry);
 				}
 
 			private:
 				
-				TUniquePtr<IDataTypeRegistryEntry> DataTypeEntry;
+				TSharedRef<IDataTypeRegistryEntry, ESPMode::ThreadSafe> DataTypeEntry;
 			};
 
 			// Node registry entry for output nodes created from a data type registry entry.
@@ -210,38 +197,32 @@ namespace Metasound
 			public:
 				FOutputNodeRegistryEntry() = delete;
 
-				FOutputNodeRegistryEntry(TUniquePtr<IDataTypeRegistryEntry>&& InDataTypeEntry)
-				: DataTypeEntry(MoveTemp(InDataTypeEntry))
+				FOutputNodeRegistryEntry(const TSharedRef<IDataTypeRegistryEntry, ESPMode::ThreadSafe>& InDataTypeEntry)
+				: DataTypeEntry(InDataTypeEntry)
 				{
-					if (DataTypeEntry.IsValid())
-					{
-						SetFrontendClass(DataTypeEntry->GetFrontendOutputClass());
-					}
+					UpdateNodeClassInfo(DataTypeEntry->GetFrontendOutputClass());
 				}
 
 				virtual ~FOutputNodeRegistryEntry() = default;
 
+				virtual const FMetasoundFrontendClass& GetFrontendClass() const override
+				{
+					return DataTypeEntry->GetFrontendOutputClass();
+				}
+
 				virtual TUniquePtr<INode> CreateNode(FDefaultNamedVertexNodeConstructorParams&& InParams) const override
 				{
-					if (DataTypeEntry.IsValid())
-					{
-						return DataTypeEntry->CreateOutputNode(MoveTemp(InParams));
-					}
-					return nullptr;
+					return DataTypeEntry->CreateOutputNode(MoveTemp(InParams));
 				}
 
 				virtual TUniquePtr<INodeRegistryEntry> Clone() const override
 				{
-					if (DataTypeEntry.IsValid())
-					{
-						return MakeUnique<FOutputNodeRegistryEntry>(DataTypeEntry->Clone());
-					}
-					return MakeUnique<FOutputNodeRegistryEntry>(TUniquePtr<IDataTypeRegistryEntry>());
+					return MakeUnique<FOutputNodeRegistryEntry>(DataTypeEntry);
 				}
 
 			private:
 				
-				TUniquePtr<IDataTypeRegistryEntry> DataTypeEntry;
+				TSharedRef<IDataTypeRegistryEntry, ESPMode::ThreadSafe> DataTypeEntry;
 			};
 
 			// Node registry entry for literal nodes created from a data type registry entry.
@@ -250,38 +231,32 @@ namespace Metasound
 			public:
 				FLiteralNodeRegistryEntry() = delete;
 
-				FLiteralNodeRegistryEntry(TUniquePtr<IDataTypeRegistryEntry>&& InDataTypeEntry)
-				: DataTypeEntry(MoveTemp(InDataTypeEntry))
+				FLiteralNodeRegistryEntry(const TSharedRef<IDataTypeRegistryEntry, ESPMode::ThreadSafe>& InDataTypeEntry)
+				: DataTypeEntry(InDataTypeEntry)
 				{
-					if (DataTypeEntry.IsValid())
-					{
-						SetFrontendClass(DataTypeEntry->GetFrontendLiteralClass());
-					}
+					UpdateNodeClassInfo(DataTypeEntry->GetFrontendLiteralClass());
 				}
 
 				virtual ~FLiteralNodeRegistryEntry() = default;
 
+				virtual const FMetasoundFrontendClass& GetFrontendClass() const override
+				{
+					return DataTypeEntry->GetFrontendLiteralClass();
+				}
+
 				virtual TUniquePtr<INode> CreateNode(FDefaultLiteralNodeConstructorParams&& InParams) const override
 				{
-					if (DataTypeEntry.IsValid())
-					{
-						return DataTypeEntry->CreateLiteralNode(MoveTemp(InParams));
-					}
-					return nullptr;
+					return DataTypeEntry->CreateLiteralNode(MoveTemp(InParams));
 				}
 
 				virtual TUniquePtr<INodeRegistryEntry> Clone() const override
 				{
-					if (DataTypeEntry.IsValid())
-					{
-						return MakeUnique<FLiteralNodeRegistryEntry>(DataTypeEntry->Clone());
-					}
-					return MakeUnique<FLiteralNodeRegistryEntry>(TUniquePtr<IDataTypeRegistryEntry>());
+					return MakeUnique<FLiteralNodeRegistryEntry>(DataTypeEntry);
 				}
 
 			private:
 				
-				TUniquePtr<IDataTypeRegistryEntry> DataTypeEntry;
+				TSharedRef<IDataTypeRegistryEntry, ESPMode::ThreadSafe> DataTypeEntry;
 			};
 
 
@@ -291,38 +266,32 @@ namespace Metasound
 			public:
 				FVariableNodeRegistryEntry() = delete;
 
-				FVariableNodeRegistryEntry(TUniquePtr<IDataTypeRegistryEntry>&& InDataTypeEntry)
-				: DataTypeEntry(MoveTemp(InDataTypeEntry))
+				FVariableNodeRegistryEntry(const TSharedRef<IDataTypeRegistryEntry, ESPMode::ThreadSafe>& InDataTypeEntry)
+				: DataTypeEntry(InDataTypeEntry)
 				{
-					if (DataTypeEntry.IsValid())
-					{
-						SetFrontendClass(DataTypeEntry->GetFrontendVariableClass());
-					}
+					UpdateNodeClassInfo(DataTypeEntry->GetFrontendVariableClass());
 				}
 
 				virtual ~FVariableNodeRegistryEntry() = default;
 
+				virtual const FMetasoundFrontendClass& GetFrontendClass() const override
+				{
+					return DataTypeEntry->GetFrontendVariableClass();
+				}
+
 				virtual TUniquePtr<INode> CreateNode(FDefaultLiteralNodeConstructorParams&& InParams) const override
 				{
-					if (DataTypeEntry.IsValid())
-					{
-						return DataTypeEntry->CreateVariableNode(MoveTemp(InParams));
-					}
-					return nullptr;
+					return DataTypeEntry->CreateVariableNode(MoveTemp(InParams));
 				}
 
 				virtual TUniquePtr<INodeRegistryEntry> Clone() const override
 				{
-					if (DataTypeEntry.IsValid())
-					{
-						return MakeUnique<FVariableNodeRegistryEntry>(DataTypeEntry->Clone());
-					}
-					return MakeUnique<FVariableNodeRegistryEntry>(TUniquePtr<IDataTypeRegistryEntry>());
+					return MakeUnique<FVariableNodeRegistryEntry>(DataTypeEntry);
 				}
 
 			private:
 				
-				TUniquePtr<IDataTypeRegistryEntry> DataTypeEntry;
+				TSharedRef<IDataTypeRegistryEntry, ESPMode::ThreadSafe> DataTypeEntry;
 			};
 
 			// Node registry entry for set variable nodes created from a data type registry entry.
@@ -331,38 +300,32 @@ namespace Metasound
 			public:
 				FVariableMutatorNodeRegistryEntry() = delete;
 
-				FVariableMutatorNodeRegistryEntry(TUniquePtr<IDataTypeRegistryEntry>&& InDataTypeEntry)
-				: DataTypeEntry(MoveTemp(InDataTypeEntry))
+				FVariableMutatorNodeRegistryEntry(const TSharedRef<IDataTypeRegistryEntry, ESPMode::ThreadSafe>& InDataTypeEntry)
+				: DataTypeEntry(InDataTypeEntry)
 				{
-					if (DataTypeEntry.IsValid())
-					{
-						SetFrontendClass(DataTypeEntry->GetFrontendVariableMutatorClass());
-					}
+					UpdateNodeClassInfo(DataTypeEntry->GetFrontendVariableMutatorClass());
 				}
 
 				virtual ~FVariableMutatorNodeRegistryEntry() = default;
 
+				virtual const FMetasoundFrontendClass& GetFrontendClass() const override
+				{
+					return DataTypeEntry->GetFrontendVariableMutatorClass();
+				}
+
 				virtual TUniquePtr<INode> CreateNode(const FNodeInitData& InParams) const override
 				{
-					if (DataTypeEntry.IsValid())
-					{
-						return DataTypeEntry->CreateVariableMutatorNode(InParams);
-					}
-					return nullptr;
+					return DataTypeEntry->CreateVariableMutatorNode(InParams);
 				}
 
 				virtual TUniquePtr<INodeRegistryEntry> Clone() const override
 				{
-					if (DataTypeEntry.IsValid())
-					{
-						return MakeUnique<FVariableMutatorNodeRegistryEntry>(DataTypeEntry->Clone());
-					}
-					return MakeUnique<FVariableMutatorNodeRegistryEntry>(TUniquePtr<IDataTypeRegistryEntry>());
+					return MakeUnique<FVariableMutatorNodeRegistryEntry>(DataTypeEntry);
 				}
 
 			private:
 				
-				TUniquePtr<IDataTypeRegistryEntry> DataTypeEntry;
+				TSharedRef<IDataTypeRegistryEntry, ESPMode::ThreadSafe> DataTypeEntry;
 			};
 
 			// Node registry entry for get variable nodes created from a data type registry entry.
@@ -371,38 +334,32 @@ namespace Metasound
 			public:
 				FVariableAccessorNodeRegistryEntry() = delete;
 
-				FVariableAccessorNodeRegistryEntry(TUniquePtr<IDataTypeRegistryEntry>&& InDataTypeEntry)
-				: DataTypeEntry(MoveTemp(InDataTypeEntry))
+				FVariableAccessorNodeRegistryEntry(const TSharedRef<IDataTypeRegistryEntry, ESPMode::ThreadSafe>& InDataTypeEntry)
+				: DataTypeEntry(InDataTypeEntry)
 				{
-					if (DataTypeEntry.IsValid())
-					{
-						SetFrontendClass(DataTypeEntry->GetFrontendVariableAccessorClass());
-					}
+					UpdateNodeClassInfo(DataTypeEntry->GetFrontendVariableAccessorClass());
 				}
 
 				virtual ~FVariableAccessorNodeRegistryEntry() = default;
 
+				virtual const FMetasoundFrontendClass& GetFrontendClass() const override
+				{
+					return DataTypeEntry->GetFrontendVariableAccessorClass();
+				}
+
 				virtual TUniquePtr<INode> CreateNode(const FNodeInitData& InParams) const override
 				{
-					if (DataTypeEntry.IsValid())
-					{
-						return DataTypeEntry->CreateVariableAccessorNode(InParams);
-					}
-					return nullptr;
+					return DataTypeEntry->CreateVariableAccessorNode(InParams);
 				}
 
 				virtual TUniquePtr<INodeRegistryEntry> Clone() const override
 				{
-					if (DataTypeEntry.IsValid())
-					{
-						return MakeUnique<FVariableAccessorNodeRegistryEntry>(DataTypeEntry->Clone());
-					}
-					return MakeUnique<FVariableAccessorNodeRegistryEntry>(TUniquePtr<IDataTypeRegistryEntry>());
+					return MakeUnique<FVariableAccessorNodeRegistryEntry>(DataTypeEntry);
 				}
 
 			private:
 				
-				TUniquePtr<IDataTypeRegistryEntry> DataTypeEntry;
+				TSharedRef<IDataTypeRegistryEntry, ESPMode::ThreadSafe> DataTypeEntry;
 			};
 
 			// Node registry entry for get delayed variable nodes created from a data type registry entry.
@@ -411,38 +368,32 @@ namespace Metasound
 			public:
 				FVariableDeferredAccessorNodeRegistryEntry() = delete;
 
-				FVariableDeferredAccessorNodeRegistryEntry(TUniquePtr<IDataTypeRegistryEntry>&& InDataTypeEntry)
-				: DataTypeEntry(MoveTemp(InDataTypeEntry))
+				FVariableDeferredAccessorNodeRegistryEntry(const TSharedRef<IDataTypeRegistryEntry, ESPMode::ThreadSafe>& InDataTypeEntry)
+				: DataTypeEntry(InDataTypeEntry)
 				{
-					if (DataTypeEntry.IsValid())
-					{
-						SetFrontendClass(DataTypeEntry->GetFrontendVariableDeferredAccessorClass());
-					}
+					UpdateNodeClassInfo(DataTypeEntry->GetFrontendVariableDeferredAccessorClass());
 				}
 
 				virtual ~FVariableDeferredAccessorNodeRegistryEntry() = default;
 
+				virtual const FMetasoundFrontendClass& GetFrontendClass() const override
+				{
+					return DataTypeEntry->GetFrontendVariableDeferredAccessorClass();
+				}
+
 				virtual TUniquePtr<INode> CreateNode(const FNodeInitData& InParams) const override
 				{
-					if (DataTypeEntry.IsValid())
-					{
-						return DataTypeEntry->CreateVariableDeferredAccessorNode(InParams);
-					}
-					return nullptr;
+					return DataTypeEntry->CreateVariableDeferredAccessorNode(InParams);
 				}
 
 				virtual TUniquePtr<INodeRegistryEntry> Clone() const override
 				{
-					if (DataTypeEntry.IsValid())
-					{
-						return MakeUnique<FVariableDeferredAccessorNodeRegistryEntry>(DataTypeEntry->Clone());
-					}
-					return MakeUnique<FVariableDeferredAccessorNodeRegistryEntry>(TUniquePtr<IDataTypeRegistryEntry>());
+					return MakeUnique<FVariableDeferredAccessorNodeRegistryEntry>(DataTypeEntry);
 				}
 
 			private:
 				
-				TUniquePtr<IDataTypeRegistryEntry> DataTypeEntry;
+				TSharedRef<IDataTypeRegistryEntry, ESPMode::ThreadSafe> DataTypeEntry;
 			};
 
 			class FDataTypeRegistry : public IDataTypeRegistry
@@ -509,7 +460,7 @@ namespace Metasound
 
 				const IDataTypeRegistryEntry* FindDataTypeEntry(const FName& InDataTypeName) const;
 
-				TMap<FName, TUniquePtr<IDataTypeRegistryEntry>> RegisteredDataTypes;
+				TMap<FName, TSharedRef<IDataTypeRegistryEntry, ESPMode::ThreadSafe>> RegisteredDataTypes;
 
 				// UObject type names to DataTypeNames
 				TMap<const UClass*, FName> RegisteredObjectClasses;
@@ -521,7 +472,9 @@ namespace Metasound
 
 				if (InEntry.IsValid())
 				{
-					const FName Name = InEntry->GetDataTypeInfo().DataTypeName;
+					TSharedRef<IDataTypeRegistryEntry, ESPMode::ThreadSafe> Entry(InEntry.Release());
+
+					const FName Name = Entry->GetDataTypeInfo().DataTypeName;
 
 					if (!ensureAlwaysMsgf(!RegisteredDataTypes.Contains(Name),
 						TEXT("Name collision when trying to register Metasound Data Type [Name:%s]. DataType must have "
@@ -535,19 +488,19 @@ namespace Metasound
 					FMetasoundFrontendRegistryContainer* NodeRegistry = FMetasoundFrontendRegistryContainer::Get();
 					if (ensure(nullptr != NodeRegistry))
 					{
-						if (InEntry->GetDataTypeInfo().bIsParsable)
+						if (Entry->GetDataTypeInfo().bIsParsable)
 						{
-							NodeRegistry->RegisterNode(MakeUnique<FInputNodeRegistryEntry>(InEntry->Clone()));
-							NodeRegistry->RegisterNode(MakeUnique<FOutputNodeRegistryEntry>(InEntry->Clone()));
-							NodeRegistry->RegisterNode(MakeUnique<FLiteralNodeRegistryEntry>(InEntry->Clone()));
-							NodeRegistry->RegisterNode(MakeUnique<FVariableNodeRegistryEntry>(InEntry->Clone()));
-							NodeRegistry->RegisterNode(MakeUnique<FVariableMutatorNodeRegistryEntry>(InEntry->Clone()));
-							NodeRegistry->RegisterNode(MakeUnique<FVariableAccessorNodeRegistryEntry>(InEntry->Clone()));
-							NodeRegistry->RegisterNode(MakeUnique<FVariableDeferredAccessorNodeRegistryEntry>(InEntry->Clone()));
+							NodeRegistry->RegisterNode(MakeUnique<FInputNodeRegistryEntry>(Entry));
+							NodeRegistry->RegisterNode(MakeUnique<FOutputNodeRegistryEntry>(Entry));
+							NodeRegistry->RegisterNode(MakeUnique<FLiteralNodeRegistryEntry>(Entry));
+							NodeRegistry->RegisterNode(MakeUnique<FVariableNodeRegistryEntry>(Entry));
+							NodeRegistry->RegisterNode(MakeUnique<FVariableMutatorNodeRegistryEntry>(Entry));
+							NodeRegistry->RegisterNode(MakeUnique<FVariableAccessorNodeRegistryEntry>(Entry));
+							NodeRegistry->RegisterNode(MakeUnique<FVariableDeferredAccessorNodeRegistryEntry>(Entry));
 						}
 					}
 
-					const FDataTypeRegistryInfo& RegistryInfo = InEntry->GetDataTypeInfo();
+					const FDataTypeRegistryInfo& RegistryInfo = Entry->GetDataTypeInfo();
 					if (!RegistryInfo.IsArrayType())
 					{
 						if (const UClass* Class = RegistryInfo.ProxyGeneratorClass)
@@ -556,7 +509,7 @@ namespace Metasound
 						}
 					}
 
-					RegisteredDataTypes.Add(Name, MoveTemp(InEntry));
+					RegisteredDataTypes.Add(Name, Entry);
 
 					UE_LOG(LogMetaSound, Verbose, TEXT("Registered Metasound Datatype [Name:%s]."), *Name.ToString());
 					return true;
@@ -599,7 +552,7 @@ namespace Metasound
 
 			void FDataTypeRegistry::IterateDataTypeInfo(TFunctionRef<void(const FDataTypeRegistryInfo&)> InFunction) const
 			{
-				for (const TPair<FName, TUniquePtr<IDataTypeRegistryEntry>>& Entry : RegisteredDataTypes)
+				for (const TPair<FName, TSharedRef<IDataTypeRegistryEntry, ESPMode::ThreadSafe>>& Entry : RegisteredDataTypes)
 				{
 					InFunction(Entry.Value->GetDataTypeInfo());
 				}
@@ -956,14 +909,11 @@ namespace Metasound
 
 			const IDataTypeRegistryEntry* FDataTypeRegistry::FindDataTypeEntry(const FName& InDataTypeName) const
 			{
-				const TUniquePtr<IDataTypeRegistryEntry>* Entry = RegisteredDataTypes.Find(InDataTypeName);
+				const TSharedRef<IDataTypeRegistryEntry, ESPMode::ThreadSafe>* Entry = RegisteredDataTypes.Find(InDataTypeName);
 
 				if (ensureMsgf(nullptr != Entry, TEXT("Data type not registered [Name:%s]"), *InDataTypeName.ToString()))
 				{
-					if (Entry->IsValid())
-					{
-						return Entry->Get();
-					}
+					return &Entry->Get();
 				}
 
 				return nullptr;
diff --git a/Engine/Plugins/Runtime/Metasound/Source/MetasoundFrontend/Private/MetasoundFrontendQuery.cpp b/Engine/Plugins/Runtime/Metasound/Source/MetasoundFrontend/Private/MetasoundFrontendQuery.cpp
--- a/Engine/Plugins/Runtime/Metasound/Source/MetasoundFrontend/Private/MetasoundFrontendQuery.cpp
+++ b/Engine/Plugins/Runtime/Metasound/Source/MetasoundFrontend/Private/MetasoundFrontendQuery.cpp
@@ -12,6 +12,14 @@ namespace Metasound
 {
 	namespace FrontendQueryPrivate
 	{
+		void CompactSelection(FFrontendQuerySelection& InSelection)
+		{
+			for (auto& Pair : InSelection)
+			{
+				Pair.Value.Compact();
+			}
+		}
+
 		// Wrapper for step defined by function
 		struct FStreamFunctionFrontendQueryStep: IFrontendQueryStreamStep
 		{
@@ -944,6 +952,8 @@ namespace Metasound
 
 	void FFrontendQuery::UpdateInternal(TSet<FFrontendQueryKey>& OutUpdatedKeys)
 	{
+		using namespace FrontendQueryPrivate;
+
 		FIncremental Incremental;
 
 		const int32 LastStepIndex = Steps.Num() - 1;
@@ -977,15 +987,19 @@ namespace Metasound
 			if (bIsLastStep)
 			{
 				StepInfo.Step->Merge(Incremental, *Result);
+				CompactSelection(*Result);
 				OutUpdatedKeys = Incremental.ActiveKeys.Union(Incremental.ActiveRemovalKeys);
 			}
 			else if (StepInfo.bMergeAndCacheOutput)
 			{
 				StepInfo.Step->Merge(Incremental, StepInfo.OutputCache);
+				
 				// If entries were removed during a merge, the subsequent step 
 				// needs to re-evaluate the entire partition. The entire partition 
 				// is added to the active set for each removed key.
 				AppendPartitions(Incremental.ActiveRemovalKeys, StepInfo.OutputCache, Incremental.ActiveKeys, Incremental.ActiveSelection);
+
+				CompactSelection(StepInfo.OutputCache);
 			}
 		}
 	}
diff --git a/Engine/Plugins/Runtime/Metasound/Source/MetasoundFrontend/Private/MetasoundFrontendRegistries.cpp b/Engine/Plugins/Runtime/Metasound/Source/MetasoundFrontend/Private/MetasoundFrontendRegistries.cpp
--- a/Engine/Plugins/Runtime/Metasound/Source/MetasoundFrontend/Private/MetasoundFrontendRegistries.cpp
+++ b/Engine/Plugins/Runtime/Metasound/Source/MetasoundFrontend/Private/MetasoundFrontendRegistries.cpp
@@ -146,7 +146,7 @@ namespace Metasound
 				FCriticalSection LazyInitCommandCritSection;
 
 				// Registry in which we keep all information about nodes implemented in C++.
-				TMap<FNodeRegistryKey, TUniquePtr<INodeRegistryEntry>> RegisteredNodes;
+				TMap<FNodeRegistryKey, TSharedRef<INodeRegistryEntry, ESPMode::ThreadSafe>> RegisteredNodes;
 
 				// Registry in which we keep lists of possible nodes to use to convert between two datatypes
 				TMap<FConverterNodeRegistryKey, FConverterNodeRegistryValue> ConverterNodeRegistry;
@@ -263,9 +263,11 @@ namespace Metasound
 
 				if (InEntry.IsValid())
 				{
+					TSharedRef<INodeRegistryEntry, ESPMode::ThreadSafe> Entry(InEntry.Release());
+
 					FNodeRegistryTransaction::FTimeType Timestamp = FPlatformTime::Cycles64();
 
-					Key = NodeRegistryKey::CreateKey(InEntry->GetClassInfo());
+					Key = NodeRegistryKey::CreateKey(Entry->GetClassInfo());
 
 					// check to see if an identical node was already registered, and log
 					ensureAlwaysMsgf(
@@ -277,10 +279,10 @@ namespace Metasound
 
 					// Store update to newly registered node in history so nodes
 					// can be queried by transaction ID
-					RegistryTransactionHistory.Add(FNodeRegistryTransaction(FNodeRegistryTransaction::ETransactionType::NodeRegistration, Key, InEntry->GetClassInfo(), Timestamp));
+					RegistryTransactionHistory.Add(FNodeRegistryTransaction(FNodeRegistryTransaction::ETransactionType::NodeRegistration, Key, Entry->GetClassInfo(), Timestamp));
 
 					// Store registry elements in map so nodes can be queried using registry key.
-					RegisteredNodes.Add(Key, MoveTemp(InEntry));
+					RegisteredNodes.Add(Key, Entry);
 				}
 
 				return Key;
@@ -399,7 +401,7 @@ namespace Metasound
 
 			void FRegistryContainerImpl::IterateRegistry(Metasound::FIterateMetasoundFrontendClassFunction InIterFunc, EMetasoundFrontendClassType InClassType) const
 			{
-				auto WrappedFunc = [&](const TPair<FNodeRegistryKey, TUniquePtr<INodeRegistryEntry>>& Pair)
+				auto WrappedFunc = [&](const TPair<FNodeRegistryKey, TSharedRef<INodeRegistryEntry, ESPMode::ThreadSafe>>& Pair)
 				{
 					InIterFunc(Pair.Value->GetFrontendClass());
 				};
@@ -412,7 +414,7 @@ namespace Metasound
 				else
 				{
 					// Only call function on classes of certain type.
-					auto IsMatchingClassType = [&](const TPair<FNodeRegistryKey, TUniquePtr<INodeRegistryEntry>>& Pair)
+					auto IsMatchingClassType = [&](const TPair<FNodeRegistryKey, TSharedRef<INodeRegistryEntry, ESPMode::ThreadSafe>>& Pair)
 					{
 						return Pair.Value->GetClassInfo().Type == InClassType;
 					};
@@ -423,12 +425,9 @@ namespace Metasound
 
 			const INodeRegistryEntry* FRegistryContainerImpl::FindNodeEntry(const FNodeRegistryKey& InKey) const
 			{
-				if (const TUniquePtr<INodeRegistryEntry>* Entry = RegisteredNodes.Find(InKey))
+				if (const TSharedRef<INodeRegistryEntry, ESPMode::ThreadSafe>* Entry = RegisteredNodes.Find(InKey))
 				{
-					if (Entry->IsValid())
-					{
-						return Entry->Get();
-					}
+					return &Entry->Get();
 				}
 
 				return nullptr;
diff --git a/Engine/Plugins/Runtime/Metasound/Source/MetasoundFrontend/Private/MetasoundFrontendSearchEngineCore.cpp b/Engine/Plugins/Runtime/Metasound/Source/MetasoundFrontend/Private/MetasoundFrontendSearchEngineCore.cpp
--- a/Engine/Plugins/Runtime/Metasound/Source/MetasoundFrontend/Private/MetasoundFrontendSearchEngineCore.cpp
+++ b/Engine/Plugins/Runtime/Metasound/Source/MetasoundFrontend/Private/MetasoundFrontendSearchEngineCore.cpp
@@ -17,7 +17,7 @@ namespace Metasound
 	{
 		namespace SearchEngineQuerySteps
 		{
-			FFrontendQueryKey FMapClassToFullClassNameAndMajorVersion::GetKey(const FMetasoundFrontendClassName& InClassName, int32 InMajorVersion)
+			FFrontendQueryKey CreateKeyFromFullClassNameAndMajorVersion(const FMetasoundFrontendClassName& InClassName, int32 InMajorVersion)
 			{
 				return FFrontendQueryKey(FString::Format(TEXT("{0}_v{1}"), {InClassName.ToString(), InMajorVersion}));
 			}
@@ -27,7 +27,7 @@ namespace Metasound
 				if (ensure(InEntry.Value.IsType<FMetasoundFrontendClass>()))
 				{
 					const FMetasoundFrontendClassMetadata& Metadata = InEntry.Value.Get<FMetasoundFrontendClass>().Metadata;
-					return FMapClassToFullClassNameAndMajorVersion::GetKey(Metadata.GetClassName(), Metadata.GetVersion().Major);
+					return CreateKeyFromFullClassNameAndMajorVersion(Metadata.GetClassName(), Metadata.GetVersion().Major);
 				}
 				return FFrontendQueryKey();
 			}
@@ -188,6 +188,23 @@ namespace Metasound
 				return false;
 			}
 
+			class FMapNodeRegistrationEventsToClassAndMajorVersion : public IFrontendQueryMapStep
+			{
+			public:
+				virtual FFrontendQueryKey Map(const FFrontendQueryEntry& InEntry) const override
+				{
+					if (InEntry.Value.IsType<FNodeRegistryTransaction>())
+					{
+						const FNodeClassInfo& Info = InEntry.Value.Get<FNodeRegistryTransaction>().GetNodeClassInfo();
+						return CreateKeyFromFullClassNameAndMajorVersion(Info.ClassName, Info.Version.Major);
+					}
+					else
+					{
+						return FFrontendQueryKey();
+					}
+				}
+			};
+
 			FFrontendQueryKey FMapInterfaceToInterfaceName::Map(const FFrontendQueryEntry& InEntry) const
 			{
 				FName Key;
@@ -260,6 +277,8 @@ namespace Metasound
 				}
 				return Result;
 			}
+
+
 		}
 
 
@@ -270,17 +289,34 @@ namespace Metasound
 			Query.AddStep<FNodeClassRegistrationEvents>()
 				.AddStep<FMapRegistrationEventsToNodeRegistryKeys>()
 				.AddStep<FReduceRegistrationEventsToCurrentStatus>()
-				.AddStep<FTransformRegistrationEventsToClasses>()
-				.AddStep<FRemoveDeprecatedClasses>()
-				.AddStep<FMapClassToFullClassNameAndMajorVersion>()
-				.AddStep<FReduceClassesToHighestVersion>();
+				.AddStep<FMapNodeRegistrationEventsToClassAndMajorVersion>();
 			return Query;
 		}
 
-		FMetasoundFrontendClass FFindClassWithHighestMinorVersionQueryPolicy::BuildResult(const FFrontendQueryPartition& InPartition)
+		FNodeRegistryKey FFindClassWithHighestMinorVersionQueryPolicy::BuildResult(const FFrontendQueryPartition& InPartition)
 		{
 			using namespace SearchEngineQuerySteps;
-			return BuildSingleClassFromPartition(InPartition);
+
+			FMetasoundFrontendVersionNumber HighestVersionNumber{-1, -1};
+			FNodeRegistryKey Key = NodeRegistryKey::GetInvalid();
+			for (const FFrontendQueryEntry& InEntry : InPartition)
+			{
+				if (InEntry.Value.IsType<FNodeRegistryTransaction>())
+				{
+					const FNodeRegistryTransaction& Transaction = InEntry.Value.Get<FNodeRegistryTransaction>();
+					if (FNodeRegistryTransaction::ETransactionType::NodeRegistration == Transaction.GetTransactionType())
+					{
+						const FNodeClassInfo& NodeClassInfo = Transaction.GetNodeClassInfo();
+						if (NodeClassInfo.Version > HighestVersionNumber)
+						{
+							HighestVersionNumber = NodeClassInfo.Version;
+							Key = Transaction.GetNodeRegistryKey();
+						}
+					}
+				}
+			}
+
+			return Key;
 		}
 
 		FFrontendQuery FFindAllDefaultInterfacesQueryPolicy::CreateQuery()
@@ -357,8 +393,16 @@ namespace Metasound
 			METASOUND_TRACE_CPUPROFILER_EVENT_SCOPE(metasound::FSearchEngineCore::FindClassWithHighestMinorVersion);
 			using namespace SearchEngineQuerySteps;
 
-			const FFrontendQueryKey Key = FMapClassToFullClassNameAndMajorVersion::GetKey(InName, InMajorVersion);
-			return FindClassWithHighestMinorVersionQuery.UpdateAndFindResult(Key, OutClass);
+			const FFrontendQueryKey QueryKey = CreateKeyFromFullClassNameAndMajorVersion(InName, InMajorVersion);
+			FNodeRegistryKey NodeRegistryKey; 
+			if (FindClassWithHighestMinorVersionQuery.UpdateAndFindResult(QueryKey, NodeRegistryKey))
+			{
+				if (FMetasoundFrontendRegistryContainer* NodeRegistry = FMetasoundFrontendRegistryContainer::Get())
+				{
+					return NodeRegistry->FindFrontendClassFromRegistered(NodeRegistryKey, OutClass);
+				}
+			}
+			return false;
 		}
 
 		TArray<FMetasoundFrontendInterface> FSearchEngineCore::FindUClassDefaultInterfaces(FName InUClassName)
diff --git a/Engine/Plugins/Runtime/Metasound/Source/MetasoundFrontend/Private/MetasoundFrontendSearchEngineCore.h b/Engine/Plugins/Runtime/Metasound/Source/MetasoundFrontend/Private/MetasoundFrontendSearchEngineCore.h
--- a/Engine/Plugins/Runtime/Metasound/Source/MetasoundFrontend/Private/MetasoundFrontendSearchEngineCore.h
+++ b/Engine/Plugins/Runtime/Metasound/Source/MetasoundFrontend/Private/MetasoundFrontendSearchEngineCore.h
@@ -19,11 +19,10 @@ namespace Metasound
 		namespace SearchEngineQuerySteps
 		{
 			/* A collection of reusable steps to be used with the Metasound frontend query system. */
-			 
+
 			class FMapClassToFullClassNameAndMajorVersion : public IFrontendQueryMapStep
 			{
 			public:
-				static FFrontendQueryKey GetKey(const FMetasoundFrontendClassName& InClassName, int32 InMajorVersion);
 				virtual FFrontendQueryKey Map(const FFrontendQueryEntry& InEntry) const override;
 			};
 
@@ -242,7 +241,7 @@ namespace Metasound
 		// Policy for finding highest versioned metasound class by name and major version.
 		struct FFindClassWithHighestMinorVersionQueryPolicy
 		{
-			using ResultType = FMetasoundFrontendClass;
+			using ResultType = FNodeRegistryKey;
 
 			static FFrontendQuery CreateQuery();
 			static ResultType BuildResult(const FFrontendQueryPartition& InPartition);
