Fix for metasound engine test failing
#rnx

diff --git a/Engine/Plugins/Runtime/Metasound/Source/MetasoundEngineTest/Private/EngineTestMetasoundSource.cpp b/Engine/Plugins/Runtime/Metasound/Source/MetasoundEngineTest/Private/EngineTestMetasoundSource.cpp
--- a/Engine/Plugins/Runtime/Metasound/Source/MetasoundEngineTest/Private/EngineTestMetasoundSource.cpp
+++ b/Engine/Plugins/Runtime/Metasound/Source/MetasoundEngineTest/Private/EngineTestMetasoundSource.cpp
@@ -52,11 +52,15 @@ namespace EngineTestMetasoundSourcePrivate
 		return OutPath;
 	}
 
-	Metasound::Frontend::FNodeHandle AddNode(Metasound::Frontend::IGraphController& InGraph, const Metasound::FNodeClassName& InClassName)
+	Metasound::Frontend::FNodeHandle AddNode(Metasound::Frontend::IGraphController& InGraph, const Metasound::FNodeClassName& InClassName, int32 InMajorVersion)
 	{
+		Metasound::Frontend::FNodeHandle Node = Metasound::Frontend::INodeController::GetInvalidHandle();
 		FMetasoundFrontendClass NodeClass;
-		Metasound::Frontend::FNodeHandle Node = InGraph.AddNode(NodeClass.Metadata);
-		check(Node->IsValid());
+		if (ensure(Metasound::Frontend::ISearchEngine::Get().FindClassWithHighestMinorVersion(InClassName, InMajorVersion, NodeClass)))
+		{
+			Node = InGraph.AddNode(NodeClass.Metadata);
+			check(Node->IsValid());
+		}
 		return Node;
 	}
 
@@ -107,7 +111,7 @@ namespace EngineTestMetasoundSourcePrivate
 		check(AudioOutputNode->IsValid());
 
 		// osc node
-		FNodeHandle OscNode = AddNode(*RootGraph, { "UE", "Sine", "Audio" });
+		FNodeHandle OscNode = AddNode(*RootGraph, { "UE", "Sine", "Audio" }, 1);
 
 		// Make connections:
 
diff --git a/Engine/Plugins/Runtime/Metasound/Source/MetasoundEngineTest/Private/MetasoundEngineTestModule.cpp b/Engine/Plugins/Runtime/Metasound/Source/MetasoundEngineTest/Private/MetasoundEngineTestModule.cpp
--- a/Engine/Plugins/Runtime/Metasound/Source/MetasoundEngineTest/Private/MetasoundEngineTestModule.cpp
+++ b/Engine/Plugins/Runtime/Metasound/Source/MetasoundEngineTest/Private/MetasoundEngineTestModule.cpp
@@ -10,6 +10,8 @@ namespace Metasound
 	public:
 		virtual void StartupModule() override
 		{
+			FModuleManager::Get().LoadModuleChecked("MetasoundFrontend");
+			FModuleManager::Get().LoadModuleChecked("MetasoundStandardNodes");
 			FModuleManager::Get().LoadModuleChecked("MetasoundEngine");
 		}
 	};
