Trigger route fix

diff --git a/Engine/Plugins/Runtime/Metasound/Source/MetasoundStandardNodes/Private/MetasoundTriggerRouteNode.cpp b/Engine/Plugins/Runtime/Metasound/Source/MetasoundStandardNodes/Private/MetasoundTriggerRouteNode.cpp
--- a/Engine/Plugins/Runtime/Metasound/Source/MetasoundStandardNodes/Private/MetasoundTriggerRouteNode.cpp
+++ b/Engine/Plugins/Runtime/Metasound/Source/MetasoundStandardNodes/Private/MetasoundTriggerRouteNode.cpp
@@ -215,12 +215,11 @@ namespace Metasound
 					[this, i](int32 StartFrame, int32 EndFrame)
 					{
 						CurrentIndex = i;
+						*OutputValue = *InputValues[CurrentIndex];
 						OutputTrigger->TriggerFrame(StartFrame);
 					}
-					);
-			}
-
-			*OutputValue = *InputValues[CurrentIndex];
+				);
+			}			
 		}
 
 	private:
