Clamping feedback pin to avoid explosions in audio

diff --git a/Engine/Plugins/Runtime/Metasound/Source/MetasoundStandardNodes/Private/MetasoundDiffuserNode.cpp b/Engine/Plugins/Runtime/Metasound/Source/MetasoundStandardNodes/Private/MetasoundDiffuserNode.cpp
--- a/Engine/Plugins/Runtime/Metasound/Source/MetasoundStandardNodes/Private/MetasoundDiffuserNode.cpp
+++ b/Engine/Plugins/Runtime/Metasound/Source/MetasoundStandardNodes/Private/MetasoundDiffuserNode.cpp
@@ -145,9 +145,11 @@ namespace Metasound
 
 			FMemory::Memcpy(AudioOutput->GetData(), AudioInput->GetData(), NumSamples * sizeof(float));
 
+			const float FeedbackClamped = FMath::Clamp(*Feedback, 0.0f, 1.0f - KINDA_SMALL_NUMBER);
+
 			for (Audio::FLongDelayAPF* Delay : Delays)
 			{
-				Delay->SetG(*Feedback);
+				Delay->SetG(FeedbackClamped);
 				Delay->ProcessAudio(*AudioOutput);
 			}
 		}
