[Metasound] Move various members to WITH_EDITORONLY_DATA to save on non-editor memory.

diff --git a/Engine/Plugins/Runtime/Metasound/Source/MetasoundFrontend/Private/MetasoundFrontend.cpp b/Engine/Plugins/Runtime/Metasound/Source/MetasoundFrontend/Private/MetasoundFrontend.cpp
--- a/Engine/Plugins/Runtime/Metasound/Source/MetasoundFrontend/Private/MetasoundFrontend.cpp
+++ b/Engine/Plugins/Runtime/Metasound/Source/MetasoundFrontend/Private/MetasoundFrontend.cpp
@@ -31,7 +31,9 @@ namespace Metasound
 
 			ClassDescription.Metadata = FMetasoundFrontendClassMetadata::GenerateClassMetadata(InNodeMetadata, ClassType);
 			ClassDescription.Interface = FMetasoundFrontendClassInterface::GenerateClassInterface(InNodeMetadata.DefaultInterface);
+#if WITH_EDITORONLY_DATA
 			ClassDescription.Style = FMetasoundFrontendClassStyle::GenerateClassStyle(InNodeMetadata.DisplayStyle);
+#endif // WITH_EDITORONLY_DATA
 
 			return ClassDescription;
 		}
diff --git a/Engine/Plugins/Runtime/Metasound/Source/MetasoundFrontend/Private/MetasoundFrontendDocument.cpp b/Engine/Plugins/Runtime/Metasound/Source/MetasoundFrontend/Private/MetasoundFrontendDocument.cpp
--- a/Engine/Plugins/Runtime/Metasound/Source/MetasoundFrontend/Private/MetasoundFrontendDocument.cpp
+++ b/Engine/Plugins/Runtime/Metasound/Source/MetasoundFrontend/Private/MetasoundFrontendDocument.cpp
@@ -206,9 +206,10 @@ FMetasoundFrontendClassInterface FMetasoundFrontendClassInterface::GenerateClass
 			ClassInput.TypeName = InputVertex.DataTypeName;
 			ClassInput.VertexID = FGuid::NewGuid();
 
-			const FDataVertexMetadata& VertexMetadata = InputVertex.Metadata;
 
 #if WITH_EDITOR
+			const FDataVertexMetadata& VertexMetadata = InputVertex.Metadata;
+
 			ClassInput.Metadata.SetSerializeText(false);
 			ClassInput.Metadata.SetDisplayName(VertexMetadata.DisplayName);
 			ClassInput.Metadata.SetDescription(VertexMetadata.Description);
@@ -259,9 +260,9 @@ FMetasoundFrontendClassInterface FMetasoundFrontendClassInterface::GenerateClass
 			ClassOutput.TypeName = OutputVertex.DataTypeName;
 			ClassOutput.VertexID = FGuid::NewGuid();
 
+#if WITH_EDITOR
 			const FDataVertexMetadata& VertexMetadata = OutputVertex.Metadata;
 
-#if WITH_EDITOR
 			ClassOutput.Metadata.SetSerializeText(false);
 			ClassOutput.Metadata.SetDisplayName(VertexMetadata.DisplayName);
 			ClassOutput.Metadata.SetDescription(VertexMetadata.Description);
@@ -460,9 +461,11 @@ bool FMetasoundFrontendClass::CacheGraphDependencyMetadataFromRegistry(FMetasoun
 }
 #endif // WITH_EDITOR
 
+#if WITH_EDITORONLY_DATA
 FMetasoundFrontendClassStyle FMetasoundFrontendClassStyle::GenerateClassStyle(const Metasound::FNodeDisplayStyle& InNodeDisplayStyle)
 {
 	FMetasoundFrontendClassStyle Style;
+
 	Style.Display.bShowName = InNodeDisplayStyle.bShowName;
 	Style.Display.bShowInputNames = InNodeDisplayStyle.bShowInputNames;
 	Style.Display.bShowOutputNames = InNodeDisplayStyle.bShowOutputNames;
@@ -470,6 +473,7 @@ FMetasoundFrontendClassStyle FMetasoundFrontendClassStyle::GenerateClassStyle(co
 
 	return Style;
 }
+#endif // WITH_EDITORONLY_DATA
 
 FMetasoundFrontendClassMetadata FMetasoundFrontendClassMetadata::GenerateClassMetadata(const Metasound::FNodeClassMetadata& InNodeClassMetadata, EMetasoundFrontendClassType InType)
 {
diff --git a/Engine/Plugins/Runtime/Metasound/Source/MetasoundFrontend/Private/MetasoundFrontendGraph.cpp b/Engine/Plugins/Runtime/Metasound/Source/MetasoundFrontend/Private/MetasoundFrontendGraph.cpp
--- a/Engine/Plugins/Runtime/Metasound/Source/MetasoundFrontend/Private/MetasoundFrontendGraph.cpp
+++ b/Engine/Plugins/Runtime/Metasound/Source/MetasoundFrontend/Private/MetasoundFrontendGraph.cpp
@@ -183,7 +183,7 @@ namespace Metasound
 			}
 			else
 			{
-				UE_LOG(LogMetaSound, Error, TEXT("Cannot create variable node [NodeID:%s]. [Variable:%s] cannot be constructed with the provided literal type."), *InNode.GetID().ToString(), *FrontendVariable->DisplayName.ToString());
+				UE_LOG(LogMetaSound, Error, TEXT("Cannot create variable node [NodeID:%s]. [Variable:%s] cannot be constructed with the provided literal type."), *InNode.GetID().ToString(), *FrontendVariable->Name.ToString());
 			}
 		}
 		else
diff --git a/Engine/Plugins/Runtime/Metasound/Source/MetasoundFrontend/Private/MetasoundFrontendGraphController.cpp b/Engine/Plugins/Runtime/Metasound/Source/MetasoundFrontend/Private/MetasoundFrontendGraphController.cpp
--- a/Engine/Plugins/Runtime/Metasound/Source/MetasoundFrontend/Private/MetasoundFrontendGraphController.cpp
+++ b/Engine/Plugins/Runtime/Metasound/Source/MetasoundFrontend/Private/MetasoundFrontendGraphController.cpp
@@ -279,7 +279,9 @@ namespace Metasound
 					FGuid VariableID = FGuid::NewGuid();
 
 					FMetasoundFrontendVariable Variable;
+#if WITH_EDITORONLY_DATA
 					Variable.DisplayName = Info.DataTypeDisplayText;
+#endif // WITH_EDITORONLY_DATA
 					Variable.TypeName = Info.DataTypeName;
 					Variable.Literal.SetFromLiteral(Registry.CreateDefaultLiteral(InDataType));
 					Variable.ID = VariableID;
diff --git a/Engine/Plugins/Runtime/Metasound/Source/MetasoundFrontend/Private/MetasoundFrontendInjectReceiveNodes.cpp b/Engine/Plugins/Runtime/Metasound/Source/MetasoundFrontend/Private/MetasoundFrontendInjectReceiveNodes.cpp
--- a/Engine/Plugins/Runtime/Metasound/Source/MetasoundFrontend/Private/MetasoundFrontendInjectReceiveNodes.cpp
+++ b/Engine/Plugins/Runtime/Metasound/Source/MetasoundFrontend/Private/MetasoundFrontendInjectReceiveNodes.cpp
@@ -142,7 +142,7 @@ namespace Metasound
 
 			if (!ReceiveNode.IsValid())
 			{
-				UE_LOG(LogMetaSound, Warning, TEXT("Failed to create receive node while injecting receive node for graph input [VertexKey:%s, VertexDescription:%s, DataTypeName:%s]"), *VertexKey.ToString(), *InputDestination.Vertex.Metadata.Description.ToString(), *DataType.ToString());
+				UE_LOG(LogMetaSound, Warning, TEXT("Failed to create receive node while injecting receive node for graph input [VertexName:%s, DataTypeName:%s]"), *VertexKey.ToString(), *DataType.ToString());
 				return false;
 			}
 
@@ -152,7 +152,7 @@ namespace Metasound
 
 			if (!AddressNode.IsValid())
 			{
-				UE_LOG(LogMetaSound, Warning, TEXT("Failed to create address node while injecting receive node for graph input [VertexKey:%s, VertexDescription:%s, DataTypeName:%s]"), *VertexKey.ToString(), *InputDestination.Vertex.Metadata.Description.ToString(), *DataType.ToString());
+				UE_LOG(LogMetaSound, Warning, TEXT("Failed to create address node while injecting receive node for graph input [VertexName:%s, DataTypeName:%s]"), *VertexKey.ToString(), *DataType.ToString());
 				return false;
 			}
 
diff --git a/Engine/Plugins/Runtime/Metasound/Source/MetasoundFrontend/Private/MetasoundFrontendRegistries.cpp b/Engine/Plugins/Runtime/Metasound/Source/MetasoundFrontend/Private/MetasoundFrontendRegistries.cpp
--- a/Engine/Plugins/Runtime/Metasound/Source/MetasoundFrontend/Private/MetasoundFrontendRegistries.cpp
+++ b/Engine/Plugins/Runtime/Metasound/Source/MetasoundFrontend/Private/MetasoundFrontendRegistries.cpp
@@ -279,7 +279,7 @@ namespace Metasound
 
 					// Store update to newly registered node in history so nodes
 					// can be queried by transaction ID
-					RegistryTransactionHistory.Add(FNodeRegistryTransaction(FNodeRegistryTransaction::ETransactionType::NodeRegistration, Key, Entry->GetClassInfo(), Timestamp));
+					RegistryTransactionHistory.Add(FNodeRegistryTransaction(FNodeRegistryTransaction::ETransactionType::NodeRegistration, Entry->GetClassInfo(), Timestamp));
 
 					// Store registry elements in map so nodes can be queried using registry key.
 					RegisteredNodes.Add(Key, Entry);
@@ -298,7 +298,7 @@ namespace Metasound
 					{
 						FNodeRegistryTransaction::FTimeType Timestamp = FPlatformTime::Cycles64();
 
-						RegistryTransactionHistory.Add(FNodeRegistryTransaction(FNodeRegistryTransaction::ETransactionType::NodeUnregistration, InKey, Entry->GetClassInfo(), Timestamp));
+						RegistryTransactionHistory.Add(FNodeRegistryTransaction(FNodeRegistryTransaction::ETransactionType::NodeUnregistration, Entry->GetClassInfo(), Timestamp));
 
 						RegisteredNodes.Remove(InKey);
 						return true;
@@ -434,14 +434,19 @@ namespace Metasound
 			}
 		} // namespace MetasoundFrontendRegistriesPrivate
 
-		FNodeRegistryTransaction::FNodeRegistryTransaction(ETransactionType InType, const FNodeRegistryKey& InKey, const FNodeClassInfo& InNodeClassInfo, FNodeRegistryTransaction::FTimeType InTimestamp)
+
+		FNodeRegistryTransaction::FNodeRegistryTransaction(ETransactionType InType, const FNodeClassInfo& InNodeClassInfo, FNodeRegistryTransaction::FTimeType InTimestamp)
 		: Type(InType)
-		, Key(InKey)
 		, NodeClassInfo(InNodeClassInfo)
 		, Timestamp(InTimestamp)
 		{
 		}
 
+		FNodeRegistryTransaction::FNodeRegistryTransaction(ETransactionType InType, const FNodeRegistryKey& InKey, const FNodeClassInfo& InNodeClassInfo, FNodeRegistryTransaction::FTimeType InTimestamp)
+		: FNodeRegistryTransaction(InType, InNodeClassInfo, InTimestamp)
+		{
+		}
+
 		FNodeRegistryTransaction::ETransactionType FNodeRegistryTransaction::GetTransactionType() const
 		{
 			return Type;
@@ -452,9 +457,9 @@ namespace Metasound
 			return NodeClassInfo;
 		}
 
-		const FNodeRegistryKey& FNodeRegistryTransaction::GetNodeRegistryKey() const
+		FNodeRegistryKey FNodeRegistryTransaction::GetNodeRegistryKey() const
 		{
-			return Key;
+			return NodeRegistryKey::CreateKey(NodeClassInfo);
 		}
 
 		FNodeRegistryTransaction::FTimeType FNodeRegistryTransaction::GetTimestamp() const
@@ -549,7 +554,7 @@ namespace Metasound
 			, Version(InClass.Metadata.GetVersion())
 		{
 			ensure(!AssetPath.IsNone());
-
+#if WITH_EDITORONLY_DATA
 			for (const FMetasoundFrontendClassInput& Input : InClass.Interface.Inputs)
 			{
 				InputTypes.Add(Input.TypeName);
@@ -559,6 +564,7 @@ namespace Metasound
 			{
 				OutputTypes.Add(Output.TypeName);
 			}
+#endif // WITH_EDITORONLY_DATA
 		}
 	} // namespace Frontend
 } // namespace Metasound
@@ -567,6 +573,7 @@ FMetasoundFrontendRegistryContainer* FMetasoundFrontendRegistryContainer::LazySi
 
 FMetasoundFrontendRegistryContainer* FMetasoundFrontendRegistryContainer::Get()
 {
+	METASOUND_LLM_SCOPE;
 	if (!LazySingleton)
 	{
 		LazySingleton = new Metasound::Frontend::MetasoundFrontendRegistryPrivate::FRegistryContainerImpl();
diff --git a/Engine/Plugins/Runtime/Metasound/Source/MetasoundFrontend/Public/MetasoundFrontendDocument.h b/Engine/Plugins/Runtime/Metasound/Source/MetasoundFrontend/Public/MetasoundFrontendDocument.h
--- a/Engine/Plugins/Runtime/Metasound/Source/MetasoundFrontend/Public/MetasoundFrontendDocument.h
+++ b/Engine/Plugins/Runtime/Metasound/Source/MetasoundFrontend/Public/MetasoundFrontendDocument.h
@@ -291,6 +291,7 @@ struct FMetasoundFrontendVariable
 	UPROPERTY(VisibleAnywhere, Category = CustomView)
 	FName Name;
 
+#if WITH_EDITORONLY_DATA
 	// Variable display name
 	UPROPERTY()
 	FText DisplayName;
@@ -299,6 +300,8 @@ struct FMetasoundFrontendVariable
 	UPROPERTY()
 	FText Description;
 
+#endif // WITH_EDITORONLY_DATA
+
 	// Variable data type name
 	UPROPERTY()
 	FName TypeName;
@@ -365,6 +368,7 @@ struct METASOUNDFRONTEND_API FMetasoundFrontendNodeStyleDisplay
 {
 	GENERATED_BODY()
 
+#if WITH_EDITORONLY_DATA
 	// DEPRECATED in Document Model v1.1: Visibility state of node
 	UPROPERTY()
 	EMetasoundFrontendNodeStyleDisplayVisibility Visibility = EMetasoundFrontendNodeStyleDisplayVisibility::Visible;
@@ -373,6 +377,7 @@ struct METASOUNDFRONTEND_API FMetasoundFrontendNodeStyleDisplay
 	// more than one place on the graph (Only functionally relevant for nodes that cannot contain inputs.)
 	UPROPERTY()
 	TMap<FGuid, FVector2D> Locations;
+#endif // WITH_EDITORONLY_DATA
 };
 
 USTRUCT()
@@ -380,6 +385,7 @@ struct METASOUNDFRONTEND_API FMetasoundFrontendNodeStyle
 {
 	GENERATED_BODY()
 
+#if WITH_EDITORONLY_DATA
 	// Display style of a node
 	UPROPERTY()
 	FMetasoundFrontendNodeStyleDisplay Display;
@@ -391,8 +397,10 @@ struct METASOUNDFRONTEND_API FMetasoundFrontendNodeStyle
 
 	UPROPERTY()
 	bool bIsPrivate = false;
+#endif // WITH_EDITORONLY_DATA
 };
 
+
 // An FMetasoundFrontendNode represents a single instance of a FMetasoundFrontendClass
 USTRUCT()
 struct METASOUNDFRONTEND_API FMetasoundFrontendNode
@@ -426,9 +434,11 @@ public:
 	UPROPERTY()
 	TArray<FMetasoundFrontendVertexLiteral> InputLiterals;
 
+#if WITH_EDITORONLY_DATA
 	// Style info related to a node.
 	UPROPERTY()
 	FMetasoundFrontendNodeStyle Style;
+#endif // WITH_EDITORONLY_DATA
 
 	const FGuid& GetID() const
 	{
@@ -465,6 +475,7 @@ struct FMetasoundFrontendEdge
 	FGuid ToVertexID = Metasound::FrontendInvalidID;
 };
 
+
 // Display style for an edge.
 UENUM()
 enum class EMetasoundFrontendStyleEdgeDisplay : uint8
@@ -480,8 +491,10 @@ struct FMetasoundFrontendStyleEdge
 {
 	GENERATED_BODY()
 
+#if WITH_EDITORONLY_DATA
 	UPROPERTY()
 	EMetasoundFrontendStyleEdgeDisplay Display = EMetasoundFrontendStyleEdgeDisplay::Default;
+#endif // WITH_EDITORONLY_DATA
 };
 
 // Styling for a class of edges dependent upon edge data type.
@@ -490,6 +503,7 @@ struct FMetasoundFrontendStyleEdgeClass
 {
 	GENERATED_BODY()
 
+#if WITH_EDITORONLY_DATA
 	// Datatype of edge to apply style to
 	UPROPERTY()
 	FName TypeName;
@@ -497,6 +511,7 @@ struct FMetasoundFrontendStyleEdgeClass
 	// Style information for edge.
 	UPROPERTY()
 	FMetasoundFrontendStyleEdge Style;
+#endif // WITH_EDITORONLY_DATA
 };
 
 // Styling for a class
@@ -505,6 +520,7 @@ struct FMetasoundFrontendGraphStyle
 {
 	GENERATED_BODY()
 
+#if WITH_EDITORONLY_DATA
 	// Whether or not the graph is editable by a user
 	UPROPERTY()
 	bool bIsGraphEditable = true;
@@ -512,8 +528,10 @@ struct FMetasoundFrontendGraphStyle
 	// Edge styles for graph.
 	UPROPERTY()
 	TArray<FMetasoundFrontendStyleEdgeClass> EdgeStyles;
+#endif // WITH_EDITORONLY_DATA
 };
 
+
 USTRUCT()
 struct FMetasoundFrontendGraph
 {
@@ -531,9 +549,13 @@ struct FMetasoundFrontendGraph
 	UPROPERTY()
 	TArray<FMetasoundFrontendVariable> Variables;
 
+#if WITH_EDITORONLY_DATA
+
 	// Style of graph display.
 	UPROPERTY()
 	FMetasoundFrontendGraphStyle Style;
+
+#endif // WITH_EDITORONLY_DATA
 };
 
 // Metadata associated with a vertex.
@@ -668,12 +690,14 @@ struct METASOUNDFRONTEND_API FMetasoundFrontendClassVertex : public FMetasoundFr
 	static bool IsFunctionalEquivalent(const FMetasoundFrontendClassVertex& InLHS, const FMetasoundFrontendClassVertex& InRHS);
 };
 
+
 // Information regarding how to display a node class
 USTRUCT()
 struct METASOUNDFRONTEND_API FMetasoundFrontendClassStyleDisplay
 {
 	GENERATED_BODY()
 
+#if WITH_EDITORONLY_DATA
 	FMetasoundFrontendClassStyleDisplay() = default;
 
 	FMetasoundFrontendClassStyleDisplay(const Metasound::FNodeDisplayStyle& InDisplayStyle)
@@ -695,8 +719,10 @@ struct METASOUNDFRONTEND_API FMetasoundFrontendClassStyleDisplay
 
 	UPROPERTY()
 	bool bShowOutputNames = true;
+#endif // WITH_EDITORONLY_DATA
 };
 
+
 // Contains info for input vertex of a Metasound class.
 USTRUCT() 
 struct METASOUNDFRONTEND_API FMetasoundFrontendClassInput : public FMetasoundFrontendClassVertex
@@ -828,6 +854,7 @@ struct METASOUNDFRONTEND_API FMetasoundFrontendClassInterface
 	GENERATED_BODY()
 
 private:
+#if WITH_EDITORONLY_DATA
 
 	// Style info for inputs.
 	UPROPERTY()
@@ -836,6 +863,7 @@ private:
 	// Style info for outputs.
 	UPROPERTY()
 	FMetasoundFrontendInterfaceStyle OutputStyle;
+#endif // WITH_EDITORONLY_DATA
 
 public:
 
@@ -860,6 +888,7 @@ private:
 	FGuid ChangeID;
 
 public:
+#if WITH_EDITORONLY_DATA
 	const FMetasoundFrontendInterfaceStyle& GetInputStyle() const
 	{
 		return InputStyle;
@@ -882,7 +911,6 @@ public:
 		ChangeID = FGuid::NewGuid();
 	}
 
-#if WITH_EDITORONLY_DATA
 	void AddRequiredInputToStyle(const FName& InInputName, const FText& InRequiredText)
 	{
 		InputStyle.RequiredMembers.Add(InInputName, InRequiredText);
@@ -1216,18 +1244,20 @@ public:
 	}
 };
 
+
 USTRUCT()
 struct FMetasoundFrontendClassStyle
 {
 	GENERATED_BODY()
 
+#if WITH_EDITORONLY_DATA
+
 	UPROPERTY()
 	FMetasoundFrontendClassStyleDisplay Display;
 
 	// Generates class style from core node class metadata.
 	static FMetasoundFrontendClassStyle GenerateClassStyle(const Metasound::FNodeDisplayStyle& InNodeDisplayStyle);
 
-#if WITH_EDITORONLY_DATA
 	// Editor only ID that allows for pumping view to reflect changes to class.
 	void UpdateChangeID()
 	{
@@ -1261,9 +1291,13 @@ struct METASOUNDFRONTEND_API FMetasoundFrontendClass
 	UPROPERTY(EditAnywhere, Category = CustomView)
 	FMetasoundFrontendClassInterface Interface;
 
+#if WITH_EDITORONLY_DATA
+
 	UPROPERTY()
 	FMetasoundFrontendClassStyle Style;
 
+#endif // WITH_EDITORONLY_DATA
+
 #if WITH_EDITOR
 	/*
 	 * Caches transient style, class & vertex Metadata found in the registry
diff --git a/Engine/Plugins/Runtime/Metasound/Source/MetasoundFrontend/Public/MetasoundFrontendRegistries.h b/Engine/Plugins/Runtime/Metasound/Source/MetasoundFrontend/Public/MetasoundFrontendRegistries.h
--- a/Engine/Plugins/Runtime/Metasound/Source/MetasoundFrontend/Public/MetasoundFrontendRegistries.h
+++ b/Engine/Plugins/Runtime/Metasound/Source/MetasoundFrontend/Public/MetasoundFrontendRegistries.h
@@ -71,11 +71,13 @@ namespace Metasound
 			// Version of the registered class
 			FMetasoundFrontendVersionNumber Version;
 
+#if WITH_EDITORONLY_DATA
 			// Types of class inputs
 			TSet<FName> InputTypes;
 
 			// Types of class outputs
 			TSet<FName> OutputTypes;
+#endif // WITH_EDITORONLY_DATA
 
 			FNodeClassInfo() = default;
 
@@ -212,17 +214,18 @@ namespace Metasound
 				Invalid
 			};
 
+			UE_DEPRECATED("5.1", "This constructor is deprecated. Use a different constructor for FNodeREegistryTransaction")
 			FNodeRegistryTransaction(ETransactionType InType, const FNodeRegistryKey& InKey, const FNodeClassInfo& InNodeClassInfo, FTimeType InTimestamp);
+			FNodeRegistryTransaction(ETransactionType InType, const FNodeClassInfo& InNodeClassInfo, FTimeType InTimestamp);
 
 			ETransactionType GetTransactionType() const;
 			const FNodeClassInfo& GetNodeClassInfo() const;
-			const FNodeRegistryKey& GetNodeRegistryKey() const;
+			FNodeRegistryKey GetNodeRegistryKey() const;
 			FTimeType GetTimestamp() const;
 
 		private:
 
 			ETransactionType Type;
-			FNodeRegistryKey Key;
 			FNodeClassInfo NodeClassInfo;
 			FTimeType Timestamp;
 		};
diff --git a/Engine/Plugins/Runtime/Metasound/Source/MetasoundGraphCore/Public/MetasoundVertex.h b/Engine/Plugins/Runtime/Metasound/Source/MetasoundGraphCore/Public/MetasoundVertex.h
--- a/Engine/Plugins/Runtime/Metasound/Source/MetasoundGraphCore/Public/MetasoundVertex.h
+++ b/Engine/Plugins/Runtime/Metasound/Source/MetasoundGraphCore/Public/MetasoundVertex.h
@@ -43,10 +43,13 @@ namespace Metasound
 		FDataVertex(const FVertexName& InVertexName, const FName& InDataTypeName, const FDataVertexMetadata& InMetadata)
 		:	VertexName(InVertexName)
 		,	DataTypeName(InDataTypeName)
+#if WITH_EDITORONLY_DATA
 		,	Metadata(InMetadata)
+#endif // WITH_EDITORONLY_DATA
 		{
 		}
 
+
 		virtual ~FDataVertex() = default;
 
 		/** Name of vertex. */
@@ -55,8 +58,10 @@ namespace Metasound
 		/** Type name of data. */
 		FName DataTypeName;
 
+#if WITH_EDITORONLY_DATA
 		/** Metadata associated with vertex. */
 		FDataVertexMetadata Metadata;
+#endif // WITH_EDITORONLY_DATA
 
 		UE_DEPRECATED(5.1, "GetVertexName() is deprecated. Use VertexName")
 		const FVertexName& GetVertexName() const
@@ -70,11 +75,13 @@ namespace Metasound
 			return DataTypeName;
 		}
 
+#if WITH_EDITORONLY_DATA
 		UE_DEPRECATED(5.1, "GetMetadata() is deprecated. Use Metadata")
 		const FDataVertexMetadata& GetMetadata() const
 		{
 			return Metadata;
 		}
+#endif // WITH_EDITORONLY_DATA
 	};
 
 	/** FInputDataVertex */
@@ -274,7 +281,9 @@ namespace Metasound
 		 */
 		FEnvironmentVertex(const FVertexName& InVertexName, const FText& InDescription)
 		:	VertexName(InVertexName)
+#if WITH_EDITORONLY_DATA
 		,	Description(InDescription)
+#endif // WITH_EDITORONLY_DATA
 		{
 		}
 
@@ -283,8 +292,10 @@ namespace Metasound
 		/** Name of vertex. */
 		FVertexName VertexName;
 
+#if WITH_EDITORONLY_DATA
 		/** Description of the vertex. */
 		FText Description;
+#endif // WITH_EDITORONLY_DATA
 
 		/** Name of vertex. */
 		UE_DEPRECATED(5.1, "GetVertexName() is deprecated. Use VertexName")
@@ -293,12 +304,14 @@ namespace Metasound
 			return VertexName;
 		}
 
+#if WITH_EDITORONLY_DATA
 		/** Description of the vertex. */
 		UE_DEPRECATED(5.1, "GetDescription() is deprecated. Use Description")
 		const FText& GetDescription() const
 		{
 			return Description;
 		}
+#endif // WITH_EDITORONLY_DATA
 
 		friend bool METASOUNDGRAPHCORE_API operator==(const FEnvironmentVertex& InLHS, const FEnvironmentVertex& InRHS);
 		friend bool METASOUNDGRAPHCORE_API operator!=(const FEnvironmentVertex& InLHS, const FEnvironmentVertex& InRHS);
@@ -382,6 +395,10 @@ namespace Metasound
 			template<typename... VertexTypes>
 			TVertexInterfaceGroup(VertexTypes&&... InVertexs)
 			{
+				// Reserve array to hold exact number of vertices to avoid
+				// over allocation.
+				Vertices.Reserve(sizeof...(VertexTypes));
+
 				CopyInputs(Forward<VertexTypes>(InVertexs)...);
 			}
 
