Remove LLM tag from FMetaSoundFrontendRegistryContainer creation to unblock Xbox

diff --git a/Engine/Plugins/Runtime/Metasound/Source/MetasoundFrontend/Private/MetasoundFrontendRegistries.cpp b/Engine/Plugins/Runtime/Metasound/Source/MetasoundFrontend/Private/MetasoundFrontendRegistries.cpp
--- a/Engine/Plugins/Runtime/Metasound/Source/MetasoundFrontend/Private/MetasoundFrontendRegistries.cpp
+++ b/Engine/Plugins/Runtime/Metasound/Source/MetasoundFrontend/Private/MetasoundFrontendRegistries.cpp
@@ -573,7 +573,6 @@ FMetasoundFrontendRegistryContainer* FMetasoundFrontendRegistryContainer::LazySi
 
 FMetasoundFrontendRegistryContainer* FMetasoundFrontendRegistryContainer::Get()
 {
-	METASOUND_LLM_SCOPE;
 	if (!LazySingleton)
 	{
 		LazySingleton = new Metasound::Frontend::MetasoundFrontendRegistryPrivate::FRegistryContainerImpl();
