[Metasound] Swapping transaction history for transaction stream

diff --git a/Engine/Plugins/Runtime/Metasound/Source/MetasoundFrontend/Private/MetasoundFrontendArchetypeRegistry.cpp b/Engine/Plugins/Runtime/Metasound/Source/MetasoundFrontend/Private/MetasoundFrontendArchetypeRegistry.cpp
--- a/Engine/Plugins/Runtime/Metasound/Source/MetasoundFrontend/Private/MetasoundFrontendArchetypeRegistry.cpp
+++ b/Engine/Plugins/Runtime/Metasound/Source/MetasoundFrontend/Private/MetasoundFrontendArchetypeRegistry.cpp
@@ -1,6 +1,7 @@
 // Copyright Epic Games, Inc. All Rights Reserved.
 
 #include "MetasoundFrontendArchetypeRegistry.h"
+#include "MetasoundFrontendInterfaceRegistryPrivate.h"
 
 #include "HAL/PlatformTime.h"
 #include "MetasoundFrontendRegistryTransaction.h"
@@ -10,71 +11,68 @@ namespace Metasound
 {
 	namespace Frontend
 	{
-		namespace MetasoundFrontendArchetypeRegistryPrivate
+		FInterfaceRegistry::FInterfaceRegistry()
+		: TransactionBuffer(MakeShared<TTransactionBuffer<FInterfaceRegistryTransaction>>())
 		{
-			class FInterfaceRegistry : public IInterfaceRegistry
-			{
-			public:
-				virtual ~FInterfaceRegistry() = default;
+		}
 
-				virtual void RegisterInterface(TUniquePtr<IInterfaceRegistryEntry>&& InEntry) override
-				{
+		void FInterfaceRegistry::RegisterInterface(TUniquePtr<IInterfaceRegistryEntry>&& InEntry)
+		{
 					METASOUND_LLM_SCOPE;
 
-					FInterfaceRegistryTransaction::FTimeType TransactionTime = FPlatformTime::Cycles64();
-					if (InEntry.IsValid())
-					{
-						FInterfaceRegistryKey Key = GetInterfaceRegistryKey(InEntry->GetInterface());
-						if (IsValidInterfaceRegistryKey(Key))
-						{
-							if (const IInterfaceRegistryEntry* Entry = FindInterfaceRegistryEntry(Key))
-							{
-								UE_LOG(LogMetaSound, Warning, TEXT("Registration of interface overwriting previously registered interface [RegistryKey: %s]"), *Key);
-								
-								FInterfaceRegistryTransaction Transaction{FInterfaceRegistryTransaction::ETransactionType::InterfaceUnregistration, Key, Entry->GetInterface().Version, TransactionTime};
-								Transactions.Add(MoveTemp(Transaction));
-							}
-							
-							FInterfaceRegistryTransaction Transaction{FInterfaceRegistryTransaction::ETransactionType::InterfaceRegistration, Key, InEntry->GetInterface().Version, TransactionTime};
-							Transactions.Add(MoveTemp(Transaction));
-
-							Entries.Add(Key, MoveTemp(InEntry)).Get();
-						}
-					}
-				}
-
-				virtual const IInterfaceRegistryEntry* FindInterfaceRegistryEntry(const FInterfaceRegistryKey& InKey) const override
+			FInterfaceRegistryTransaction::FTimeType TransactionTime = FPlatformTime::Cycles64();
+			if (InEntry.IsValid())
+			{
+				FInterfaceRegistryKey Key = GetInterfaceRegistryKey(InEntry->GetInterface());
+				if (IsValidInterfaceRegistryKey(Key))
 				{
-					if (const TUniquePtr<IInterfaceRegistryEntry>* Entry = Entries.Find(InKey))
+					if (const IInterfaceRegistryEntry* Entry = FindInterfaceRegistryEntry(Key))
 					{
-						return Entry->Get();
+						UE_LOG(LogMetaSound, Warning, TEXT("Registration of interface overwriting previously registered interface [RegistryKey: %s]"), *Key);
+						
+						FInterfaceRegistryTransaction Transaction{FInterfaceRegistryTransaction::ETransactionType::InterfaceUnregistration, Key, Entry->GetInterface().Version, TransactionTime};
+						TransactionBuffer->AddTransaction(MoveTemp(Transaction));
 					}
-					return nullptr;
+					
+					FInterfaceRegistryTransaction Transaction{FInterfaceRegistryTransaction::ETransactionType::InterfaceRegistration, Key, InEntry->GetInterface().Version, TransactionTime};
+					TransactionBuffer->AddTransaction(MoveTemp(Transaction));
+
+					Entries.Add(Key, MoveTemp(InEntry)).Get();
 				}
+			}
+		}
 
-				virtual bool FindInterface(const FInterfaceRegistryKey& InKey, FMetasoundFrontendInterface& OutInterface) const override
-				{
-					if (const IInterfaceRegistryEntry* Entry = FindInterfaceRegistryEntry(InKey))
-					{
-						OutInterface = Entry->GetInterface();
-						return true;
-					}
+		const IInterfaceRegistryEntry* FInterfaceRegistry::FindInterfaceRegistryEntry(const FInterfaceRegistryKey& InKey) const
+		{
+			if (const TUniquePtr<IInterfaceRegistryEntry>* Entry = Entries.Find(InKey))
+			{
+				return Entry->Get();
+			}
+			return nullptr;
+		}
 
-					return false;
-				}
+		bool FInterfaceRegistry::FindInterface(const FInterfaceRegistryKey& InKey, FMetasoundFrontendInterface& OutInterface) const
+		{
+			if (const IInterfaceRegistryEntry* Entry = FindInterfaceRegistryEntry(InKey))
+			{
+				OutInterface = Entry->GetInterface();
+				return true;
+			}
 
-				virtual void ForEachRegistryTransactionSince(FRegistryTransactionID InSince, FRegistryTransactionID* OutCurrentRegistryTransactionID, TFunctionRef<void(const FInterfaceRegistryTransaction&)> InFunc) const override
-				{
-					Transactions.ForEachTransactionSince(InSince, OutCurrentRegistryTransactionID, InFunc);
-				}
+			return false;
+		}
 
-			private:
+		TUniquePtr<FInterfaceTransactionStream> FInterfaceRegistry::CreateTransactionStream()
+		{
+			return MakeUnique<FInterfaceTransactionStream>(TransactionBuffer);
+		}
 
-				TMap<FInterfaceRegistryKey, TUniquePtr<IInterfaceRegistryEntry>> Entries;
-				TRegistryTransactionHistory<FInterfaceRegistryTransaction> Transactions;
-			};
+		void FInterfaceRegistry::ForEachRegistryTransactionSince(FRegistryTransactionID InSince, FRegistryTransactionID* OutCurrentRegistryTransactionID, TFunctionRef<void(const FInterfaceRegistryTransaction&)> InFunc) const
+		{
+			UE_LOG(LogMetaSound, Error, TEXT("IInterfaceRegistry::ForEachRegistryTransactionSince is no longer supported and should not be called"));
 		}
 
+
 		bool IsValidInterfaceRegistryKey(const FInterfaceRegistryKey& InKey)
 		{
 			return !InKey.IsEmpty();
@@ -119,10 +117,15 @@ namespace Metasound
 			return Timestamp;
 		}
 
-		IInterfaceRegistry& IInterfaceRegistry::Get()
+		FInterfaceRegistry& FInterfaceRegistry::Get()
 		{
-			static MetasoundFrontendArchetypeRegistryPrivate::FInterfaceRegistry Registry;
+			static FInterfaceRegistry Registry;
 			return Registry;
 		}
+
+		IInterfaceRegistry& IInterfaceRegistry::Get()
+		{
+			return FInterfaceRegistry::Get();
+		}
 	}
 }
diff --git a/Engine/Plugins/Runtime/Metasound/Source/MetasoundFrontend/Private/MetasoundFrontendInterfaceRegistryPrivate.h b/Engine/Plugins/Runtime/Metasound/Source/MetasoundFrontend/Private/MetasoundFrontendInterfaceRegistryPrivate.h
new file mode 100644
--- /dev/null
+++ b/Engine/Plugins/Runtime/Metasound/Source/MetasoundFrontend/Private/MetasoundFrontendInterfaceRegistryPrivate.h
@@ -0,0 +1,44 @@
+// Copyright Epic Games, Inc. All Rights Reserved.
+
+#pragma once
+
+#include "MetasoundFrontendArchetypeRegistry.h"
+#include "MetasoundFrontendController.h"
+#include "MetasoundFrontendDocument.h"
+#include "MetasoundFrontendRegistryTransaction.h"
+#include "UObject/NoExportTypes.h"
+
+namespace Metasound
+{
+	namespace Frontend
+	{
+		using FInterfaceTransactionStream = TTransactionStream<FInterfaceRegistryTransaction>;
+
+		class FInterfaceRegistry : public IInterfaceRegistry
+		{
+		public:
+
+			static FInterfaceRegistry& Get();
+
+			FInterfaceRegistry();
+
+			virtual ~FInterfaceRegistry() = default;
+
+			virtual void RegisterInterface(TUniquePtr<IInterfaceRegistryEntry>&& InEntry) override;
+
+			virtual const IInterfaceRegistryEntry* FindInterfaceRegistryEntry(const FInterfaceRegistryKey& InKey) const override;
+
+			virtual bool FindInterface(const FInterfaceRegistryKey& InKey, FMetasoundFrontendInterface& OutInterface) const override;
+
+			TUniquePtr<FInterfaceTransactionStream> CreateTransactionStream();
+
+			virtual void ForEachRegistryTransactionSince(FRegistryTransactionID InSince, FRegistryTransactionID* OutCurrentRegistryTransactionID, TFunctionRef<void(const FInterfaceRegistryTransaction&)> InFunc) const override;
+		private:
+			using FInterfaceTransactionBuffer = TTransactionBuffer<FInterfaceRegistryTransaction>;
+
+			TMap<FInterfaceRegistryKey, TUniquePtr<IInterfaceRegistryEntry>> Entries;
+			TSharedRef<FInterfaceTransactionBuffer> TransactionBuffer;
+		};
+	}
+}
+
diff --git a/Engine/Plugins/Runtime/Metasound/Source/MetasoundFrontend/Private/MetasoundFrontendNodeRegistryPrivate.h b/Engine/Plugins/Runtime/Metasound/Source/MetasoundFrontend/Private/MetasoundFrontendNodeRegistryPrivate.h
new file mode 100644
--- /dev/null
+++ b/Engine/Plugins/Runtime/Metasound/Source/MetasoundFrontend/Private/MetasoundFrontendNodeRegistryPrivate.h
@@ -0,0 +1,111 @@
+// Copyright Epic Games, Inc. All Rights Reserved.
+#pragma once
+
+#include "MetasoundFrontendRegistries.h"
+
+#include "Containers/Array.h"
+#include "Containers/Map.h"
+#include "HAL/CriticalSection.h"
+#include "MetasoundFrontendRegistryTransaction.h"
+
+namespace Metasound
+{
+	namespace Frontend
+	{
+		using FNodeRegistryTransactionBuffer = TTransactionBuffer<FNodeRegistryTransaction>;
+		using FNodeRegistryTransactionStream = TTransactionStream<FNodeRegistryTransaction>; 
+
+		// Registry container private implementation.
+		class FRegistryContainerImpl : public FMetasoundFrontendRegistryContainer
+		{
+
+		public:
+			using FConverterNodeRegistryKey = ::Metasound::Frontend::FConverterNodeRegistryKey;
+			using FConverterNodeRegistryValue = ::Metasound::Frontend::FConverterNodeRegistryValue;
+			using FConverterNodeInfo = ::Metasound::Frontend::FConverterNodeInfo;
+
+			using FNodeRegistryKey = Metasound::Frontend::FNodeRegistryKey;
+
+			using FDataTypeRegistryInfo = Metasound::Frontend::FDataTypeRegistryInfo;
+			using FNodeClassMetadata = Metasound::FNodeClassMetadata;
+			using IEnumDataTypeInterface = Metasound::Frontend::IEnumDataTypeInterface;
+
+			FRegistryContainerImpl();
+
+			FRegistryContainerImpl(const FRegistryContainerImpl&) = delete;
+			FRegistryContainerImpl& operator=(const FRegistryContainerImpl&) = delete;
+
+			static FRegistryContainerImpl& Get();
+			static void Shutdown();
+
+			virtual ~FRegistryContainerImpl() = default;
+
+			// Add a function to the init command array.
+			virtual bool EnqueueInitCommand(TUniqueFunction<void()>&& InFunc) override;
+
+			// This is called on module startup. This invokes any registration commands enqueued by our registration macros.
+			virtual void RegisterPendingNodes() override;
+
+			/** Register external node with the frontend.
+			 *
+			 * @param InCreateNode - Function for creating node from FNodeInitData.
+			 * @param InCreateDescription - Function for creating a FMetasoundFrontendClass.
+			 *
+			 * @return True on success.
+			 */
+			virtual FNodeRegistryKey RegisterNode(TUniquePtr<Metasound::Frontend::INodeRegistryEntry>&&) override;
+			virtual void ForEachNodeRegistryTransactionSince(Metasound::Frontend::FRegistryTransactionID InSince, Metasound::Frontend::FRegistryTransactionID* OutCurrentRegistryTransactionID, TFunctionRef<void(const Metasound::Frontend::FNodeRegistryTransaction&)> InFunc) const override;
+			virtual bool UnregisterNode(const FNodeRegistryKey& InKey) override;
+			virtual bool IsNodeRegistered(const FNodeRegistryKey& InKey) const override;
+			virtual bool IsNodeNative(const FNodeRegistryKey& InKey) const override;
+
+			virtual bool RegisterConversionNode(const FConverterNodeRegistryKey& InNodeKey, const FConverterNodeInfo& InNodeInfo) override;
+
+			virtual void IterateRegistry(Metasound::FIterateMetasoundFrontendClassFunction InIterFunc, EMetasoundFrontendClassType InClassType = EMetasoundFrontendClassType::Invalid) const override;
+
+			// Find Frontend Document data.
+			virtual bool FindFrontendClassFromRegistered(const FNodeRegistryKey& InKey, FMetasoundFrontendClass& OutClass) override;
+			virtual bool FindNodeClassInfoFromRegistered(const Metasound::Frontend::FNodeRegistryKey& InKey, FNodeClassInfo& OutInfo) override;
+			virtual bool FindInputNodeRegistryKeyForDataType(const FName& InDataTypeName, FNodeRegistryKey& OutKey) override;
+			virtual bool FindVariableNodeRegistryKeyForDataType(const FName& InDataTypeName, FNodeRegistryKey& OutKey) override;
+			virtual bool FindOutputNodeRegistryKeyForDataType(const FName& InDataTypeName, FNodeRegistryKey& OutKey) override;
+
+			// Create a new instance of a C++ implemented node from the registry.
+			virtual TUniquePtr<Metasound::INode> CreateNode(const FNodeRegistryKey& InKey, const Metasound::FNodeInitData& InInitData) const override;
+			virtual TUniquePtr<INode> CreateNode(const FNodeRegistryKey& InKey, FDefaultLiteralNodeConstructorParams&&) const override;
+			virtual TUniquePtr<INode> CreateNode(const FNodeRegistryKey& InKey, FDefaultNamedVertexNodeConstructorParams&&) const override;
+			virtual TUniquePtr<INode> CreateNode(const FNodeRegistryKey& InKey, FDefaultNamedVertexWithLiteralNodeConstructorParams&&) const override;
+
+			// Returns a list of possible nodes to use to convert from FromDataType to ToDataType.
+			// Returns an empty array if none are available.
+			virtual TArray<FConverterNodeInfo> GetPossibleConverterNodes(const FName& FromDataType, const FName& ToDataType) override;
+
+			// Create a transaction stream for any newly transactions
+			TUniquePtr<FNodeRegistryTransactionStream> CreateTransactionStream();
+
+		private:
+			static FRegistryContainerImpl* LazySingleton;
+
+			const INodeRegistryEntry* FindNodeEntry(const FNodeRegistryKey& InKey) const;
+
+
+			// This buffer is used to enqueue nodes and datatypes to register when the module has been initialized,
+			// in order to avoid bad behavior with ensures, logs, etc. on static initialization.
+			// The bad news is that TInlineAllocator is the safest allocator to use on static init.
+			// The good news is that none of these lambdas typically have captures, so this should have low memory overhead.
+			static constexpr int32 MaxNumNodesAndDatatypesToInitialize = 2048;
+			TArray<TUniqueFunction<void()>, TInlineAllocator<MaxNumNodesAndDatatypesToInitialize>> LazyInitCommands;
+			
+			FCriticalSection LazyInitCommandCritSection;
+
+			// Registry in which we keep all information about nodes implemented in C++.
+			TMap<FNodeRegistryKey, TSharedRef<INodeRegistryEntry, ESPMode::ThreadSafe>> RegisteredNodes;
+
+			// Registry in which we keep lists of possible nodes to use to convert between two datatypes
+			TMap<FConverterNodeRegistryKey, FConverterNodeRegistryValue> ConverterNodeRegistry;
+
+			TSharedRef<FNodeRegistryTransactionBuffer> TransactionBuffer;
+		};
+	}
+}
+
diff --git a/Engine/Plugins/Runtime/Metasound/Source/MetasoundFrontend/Private/MetasoundFrontendQuerySteps.cpp b/Engine/Plugins/Runtime/Metasound/Source/MetasoundFrontend/Private/MetasoundFrontendQuerySteps.cpp
--- a/Engine/Plugins/Runtime/Metasound/Source/MetasoundFrontend/Private/MetasoundFrontendQuerySteps.cpp
+++ b/Engine/Plugins/Runtime/Metasound/Source/MetasoundFrontend/Private/MetasoundFrontendQuerySteps.cpp
@@ -5,30 +5,48 @@
 #include "Algo/MaxElement.h"
 #include "MetasoundFrontend.h"
 #include "MetasoundFrontendDocument.h"
+#include "MetasoundFrontendNodeRegistryPrivate.h"
 #include "MetasoundFrontendQuery.h"
 #include "MetasoundFrontendRegistries.h"
 #include "MetasoundFrontendRegistryTransaction.h"
 
 namespace Metasound
 {
+	class FNodeClassRegistrationEventsPimpl : public IFrontendQueryStreamStep
+	{
+	public:
+		FNodeClassRegistrationEventsPimpl()
+		{
+			TransactionStream = Frontend::FRegistryContainerImpl::Get().CreateTransactionStream();
+		}
+
+		virtual void Stream(TArray<FFrontendQueryValue>& OutValues) override
+		{
+			using namespace Frontend;
+
+			auto AddEntry = [&OutValues](const FNodeRegistryTransaction& InTransaction)
+			{
+				OutValues.Emplace(TInPlaceType<FNodeRegistryTransaction>(), InTransaction);
+			};
+
+			if (TransactionStream.IsValid())
+			{
+				TransactionStream->Stream(AddEntry);
+			}
+		}
+
+	private:
+		TUniquePtr<Frontend::FNodeRegistryTransactionStream> TransactionStream;
+	};
+
 	FNodeClassRegistrationEvents::FNodeClassRegistrationEvents()
-	: CurrentTransactionID(Frontend::GetOriginRegistryTransactionID())
+	: Pimpl(MakePimpl<FNodeClassRegistrationEventsPimpl>())
 	{
 	}
 
 	void FNodeClassRegistrationEvents::Stream(TArray<FFrontendQueryValue>& OutValues)
 	{
-		using namespace Frontend;
-
-		auto AddEntry = [&OutValues](const FNodeRegistryTransaction& InTransaction)
-		{
-			OutValues.Emplace(TInPlaceType<FNodeRegistryTransaction>(), InTransaction);
-		};
-		
-		if (FMetasoundFrontendRegistryContainer* Registry = FMetasoundFrontendRegistryContainer::Get())
-		{
-			Registry->ForEachNodeRegistryTransactionSince(CurrentTransactionID, &CurrentTransactionID, AddEntry);
-		}
+		Pimpl->Stream(OutValues);
 	}
 
 	FFrontendQueryKey FMapRegistrationEventsToNodeRegistryKeys::Map(const FFrontendQueryEntry& InEntry) const 
diff --git a/Engine/Plugins/Runtime/Metasound/Source/MetasoundFrontend/Private/MetasoundFrontendRegistries.cpp b/Engine/Plugins/Runtime/Metasound/Source/MetasoundFrontend/Private/MetasoundFrontendRegistries.cpp
--- a/Engine/Plugins/Runtime/Metasound/Source/MetasoundFrontend/Private/MetasoundFrontendRegistries.cpp
+++ b/Engine/Plugins/Runtime/Metasound/Source/MetasoundFrontend/Private/MetasoundFrontendRegistries.cpp
@@ -6,6 +6,7 @@
 #include "CoreMinimal.h"
 #include "HAL/PlatformTime.h"
 #include "MetasoundFrontendDataTypeRegistry.h"
+#include "MetasoundFrontendNodeRegistryPrivate.h"
 #include "MetasoundFrontendRegistryTransaction.h"
 #include "MetasoundFrontendSearchEngine.h"
 #include "MetasoundLog.h"
@@ -67,372 +68,317 @@ namespace Metasound
 						return InvalidType;
 				}
 			}
+		}
 
-			// Registry container private implementation.
-			class FRegistryContainerImpl : public FMetasoundFrontendRegistryContainer
-			{
-
-			public:
-				using FConverterNodeRegistryKey = ::Metasound::Frontend::FConverterNodeRegistryKey;
-				using FConverterNodeRegistryValue = ::Metasound::Frontend::FConverterNodeRegistryValue;
-				using FConverterNodeInfo = ::Metasound::Frontend::FConverterNodeInfo;
-
-				using FNodeRegistryKey = Metasound::Frontend::FNodeRegistryKey;
-
-				using FDataTypeRegistryInfo = Metasound::Frontend::FDataTypeRegistryInfo;
-				using FNodeClassMetadata = Metasound::FNodeClassMetadata;
-				using IEnumDataTypeInterface = Metasound::Frontend::IEnumDataTypeInterface;
-
-				FRegistryContainerImpl() = default;
-
-				FRegistryContainerImpl(const FRegistryContainerImpl&) = delete;
-				FRegistryContainerImpl& operator=(const FRegistryContainerImpl&) = delete;
-
-				virtual ~FRegistryContainerImpl() = default;
-
-				// Add a function to the init command array.
-				bool EnqueueInitCommand(TUniqueFunction<void()>&& InFunc) override;
+		FRegistryContainerImpl* FRegistryContainerImpl::LazySingleton = nullptr;
 
-				// This is called on module startup. This invokes any registration commands enqueued by our registration macros.
-				void RegisterPendingNodes() override;
+		FRegistryContainerImpl& FRegistryContainerImpl::Get()
+		{
+			if (!LazySingleton)
+			{
+				LazySingleton = new Metasound::Frontend::FRegistryContainerImpl();
+			}
 
-				/** Register external node with the frontend.
-				 *
-				 * @param InCreateNode - Function for creating node from FNodeInitData.
-				 * @param InCreateDescription - Function for creating a FMetasoundFrontendClass.
-				 *
-				 * @return True on success.
-				 */
-				virtual FNodeRegistryKey RegisterNode(TUniquePtr<Metasound::Frontend::INodeRegistryEntry>&&) override;
-				virtual bool UnregisterNode(const FNodeRegistryKey& InKey) override;
-				virtual bool IsNodeRegistered(const FNodeRegistryKey& InKey) const override;
-				virtual bool IsNodeNative(const FNodeRegistryKey& InKey) const override;
+			return *LazySingleton;
+		}
 
-				bool RegisterConversionNode(const FConverterNodeRegistryKey& InNodeKey, const FConverterNodeInfo& InNodeInfo) override;
+		void FRegistryContainerImpl::Shutdown()
+		{
+			if (nullptr != LazySingleton)
+			{
+				delete LazySingleton;
+				LazySingleton = nullptr;
+			}
+		}
 
-				virtual void ForEachNodeRegistryTransactionSince(FRegistryTransactionID InSince, FRegistryTransactionID* OutCurrentRegistryTransactionID, TFunctionRef<void(const FNodeRegistryTransaction&)> InFunc) const override;
+		FRegistryContainerImpl::FRegistryContainerImpl()
+		: TransactionBuffer(MakeShared<FNodeRegistryTransactionBuffer>())
+		{
+		}
 
-				void IterateRegistry(Metasound::FIterateMetasoundFrontendClassFunction InIterFunc, EMetasoundFrontendClassType InClassType = EMetasoundFrontendClassType::Invalid) const override;
+		void FRegistryContainerImpl::RegisterPendingNodes()
+		{
+			METASOUND_LLM_SCOPE;
+			METASOUND_TRACE_CPUPROFILER_EVENT_SCOPE(metasound::FRegistryContainerImpl::RegisterPendingNodes);
+			{
+				FScopeLock ScopeLock(&LazyInitCommandCritSection);
 
-				// Find Frontend Document data.
-				bool FindFrontendClassFromRegistered(const FNodeRegistryKey& InKey, FMetasoundFrontendClass& OutClass) override;
-				virtual bool FindNodeClassInfoFromRegistered(const Metasound::Frontend::FNodeRegistryKey& InKey, FNodeClassInfo& OutInfo) override;
-				bool FindInputNodeRegistryKeyForDataType(const FName& InDataTypeName, FNodeRegistryKey& OutKey) override;
-				bool FindVariableNodeRegistryKeyForDataType(const FName& InDataTypeName, FNodeRegistryKey& OutKey) override;
-				bool FindOutputNodeRegistryKeyForDataType(const FName& InDataTypeName, FNodeRegistryKey& OutKey) override;
+				for (TUniqueFunction<void()>& Command : LazyInitCommands)
+				{
+					Command();
+				}
 
-				// Create a new instance of a C++ implemented node from the registry.
-				virtual TUniquePtr<Metasound::INode> CreateNode(const FNodeRegistryKey& InKey, const Metasound::FNodeInitData& InInitData) const override;
-				virtual TUniquePtr<INode> CreateNode(const FNodeRegistryKey& InKey, FDefaultLiteralNodeConstructorParams&&) const override;
-				virtual TUniquePtr<INode> CreateNode(const FNodeRegistryKey& InKey, FDefaultNamedVertexNodeConstructorParams&&) const override;
-				virtual TUniquePtr<INode> CreateNode(const FNodeRegistryKey& InKey, FDefaultNamedVertexWithLiteralNodeConstructorParams&&) const override;
+				LazyInitCommands.Empty();
+			}
 
-				// Returns a list of possible nodes to use to convert from FromDataType to ToDataType.
-				// Returns an empty array if none are available.
-				TArray<FConverterNodeInfo> GetPossibleConverterNodes(const FName& FromDataType, const FName& ToDataType) override;
+			// Prime search engine after bulk registration.
+			ISearchEngine::Get().Prime();
+		}
 
-			private:
+		bool FRegistryContainerImpl::EnqueueInitCommand(TUniqueFunction<void()>&& InFunc)
+		{
 
-				const INodeRegistryEntry* FindNodeEntry(const FNodeRegistryKey& InKey) const;
+			FScopeLock ScopeLock(&LazyInitCommandCritSection);
+			if (LazyInitCommands.Num() >= MaxNumNodesAndDatatypesToInitialize)
+			{
+				UE_LOG(LogMetaSound, Warning, TEXT("Registering more that %d nodes and datatypes for metasounds! Consider increasing MetasoundFrontendRegistryContainer::MaxNumNodesAndDatatypesToInitialize."));
+			}
 
+			LazyInitCommands.Add(MoveTemp(InFunc));
+			return true;
+		}
 
-				// This buffer is used to enqueue nodes and datatypes to register when the module has been initialized,
-				// in order to avoid bad behavior with ensures, logs, etc. on static initialization.
-				// The bad news is that TInlineAllocator is the safest allocator to use on static init.
-				// The good news is that none of these lambdas typically have captures, so this should have low memory overhead.
-				static constexpr int32 MaxNumNodesAndDatatypesToInitialize = 8192;
-				TArray<TUniqueFunction<void()>, TInlineAllocator<MaxNumNodesAndDatatypesToInitialize>> LazyInitCommands;
-				
-				FCriticalSection LazyInitCommandCritSection;
+		void FRegistryContainerImpl::ForEachNodeRegistryTransactionSince(FRegistryTransactionID InSince, FRegistryTransactionID* OutCurrentRegistryTransactionID, TFunctionRef<void(const FNodeRegistryTransaction&)> InFunc) const 
+		{
+			UE_LOG(LogMetaSound, Error, TEXT("FMetasoundFrontendRegistryContainer::ForEachNodeRegistryTransactionSince is no longer supported and will not function as expected."));
+		}
 
-				// Registry in which we keep all information about nodes implemented in C++.
-				TMap<FNodeRegistryKey, TSharedRef<INodeRegistryEntry, ESPMode::ThreadSafe>> RegisteredNodes;
+		TUniquePtr<Metasound::INode> FRegistryContainerImpl::CreateNode(const FNodeRegistryKey& InKey, const Metasound::FNodeInitData& InInitData) const
+		{
+			if (const INodeRegistryEntry* Entry = FindNodeEntry(InKey))
+			{
+				return Entry->CreateNode(InInitData);
+			}
 
-				// Registry in which we keep lists of possible nodes to use to convert between two datatypes
-				TMap<FConverterNodeRegistryKey, FConverterNodeRegistryValue> ConverterNodeRegistry;
+			// Because creation of external nodes can rely on assets being unavailable due to errors in loading order, asset(s)
+			// missing, etc. only log error and don't throw ensure to avoid blocking start-up if assets are missing. All other
+			// CreateNode calls are natively managed and thus better suited to throw ensures.
+			UE_LOG(LogMetaSound, Error, TEXT("Could not find node [RegistryKey:%s]"), *InKey);
+			return nullptr;
+		}
 
-				TRegistryTransactionHistory<FNodeRegistryTransaction> RegistryTransactionHistory;
-			};
+		TUniquePtr<Metasound::INode> FRegistryContainerImpl::CreateNode(const FNodeRegistryKey& InKey, FDefaultLiteralNodeConstructorParams&& InParams) const
+		{
+			const INodeRegistryEntry* Entry = FindNodeEntry(InKey);
 
-			void FRegistryContainerImpl::RegisterPendingNodes()
+			if (ensureAlwaysMsgf(nullptr != Entry, TEXT("Could not find node [RegistryKey:%s]"), *InKey))
 			{
-				METASOUND_LLM_SCOPE;
-				METASOUND_TRACE_CPUPROFILER_EVENT_SCOPE(metasound::FRegistryContainerImpl::RegisterPendingNodes);
-				{
-					FScopeLock ScopeLock(&LazyInitCommandCritSection);
-
-					for (TUniqueFunction<void()>& Command : LazyInitCommands)
-					{
-						Command();
-					}
+				return Entry->CreateNode(MoveTemp(InParams));
+			}
 
-					LazyInitCommands.Empty();
-				}
+			return nullptr;
+		}
 
-				// Prime search engine after bulk registration.
-				ISearchEngine::Get().Prime();
-			}
+		TUniquePtr<Metasound::INode> FRegistryContainerImpl::CreateNode(const FNodeRegistryKey& InKey, FDefaultNamedVertexNodeConstructorParams&& InParams) const
+		{
+			const INodeRegistryEntry* Entry = FindNodeEntry(InKey);
 
-			bool FRegistryContainerImpl::EnqueueInitCommand(TUniqueFunction<void()>&& InFunc)
+			if (ensureAlwaysMsgf(nullptr != Entry, TEXT("Could not find node [RegistryKey:%s]"), *InKey))
 			{
+				return Entry->CreateNode(MoveTemp(InParams));
+			}
 
-				FScopeLock ScopeLock(&LazyInitCommandCritSection);
-				if (LazyInitCommands.Num() >= MaxNumNodesAndDatatypesToInitialize)
-				{
-					UE_LOG(LogMetaSound, Warning, TEXT("Registering more that %d nodes and datatypes for metasounds! Consider increasing MetasoundFrontendRegistryContainer::MaxNumNodesAndDatatypesToInitialize."));
-				}
+			return nullptr;
+		}
 
-				LazyInitCommands.Add(MoveTemp(InFunc));
-				return true;
-			}
+		TUniquePtr<Metasound::INode> FRegistryContainerImpl::CreateNode(const FNodeRegistryKey& InKey, FDefaultNamedVertexWithLiteralNodeConstructorParams&& InParams) const
+		{
+			const INodeRegistryEntry* Entry = FindNodeEntry(InKey);
 
-			void FRegistryContainerImpl::ForEachNodeRegistryTransactionSince(FRegistryTransactionID InSince, FRegistryTransactionID* OutCurrentRegistryTransactionID, TFunctionRef<void(const FNodeRegistryTransaction&)> InFunc) const 
+			if (ensureAlwaysMsgf(nullptr != Entry, TEXT("Could not find node [RegistryKey:%s]"), *InKey))
 			{
-				return RegistryTransactionHistory.ForEachTransactionSince(InSince, OutCurrentRegistryTransactionID, InFunc);
+				return Entry->CreateNode(MoveTemp(InParams));
 			}
 
+			return nullptr;
+		}
 
-			TUniquePtr<Metasound::INode> FRegistryContainerImpl::CreateNode(const FNodeRegistryKey& InKey, const Metasound::FNodeInitData& InInitData) const
+		TArray<::Metasound::Frontend::FConverterNodeInfo> FRegistryContainerImpl::GetPossibleConverterNodes(const FName& FromDataType, const FName& ToDataType)
+		{
+			FConverterNodeRegistryKey InKey = { FromDataType, ToDataType };
+			if (!ConverterNodeRegistry.Contains(InKey))
 			{
-				if (const INodeRegistryEntry* Entry = FindNodeEntry(InKey))
-				{
-					return Entry->CreateNode(InInitData);
-				}
-
-				// Because creation of external nodes can rely on assets being unavailable due to errors in loading order, asset(s)
-				// missing, etc. only log error and don't throw ensure to avoid blocking start-up if assets are missing. All other
-				// CreateNode calls are natively managed and thus better suited to throw ensures.
-				UE_LOG(LogMetaSound, Error, TEXT("Could not find node [RegistryKey:%s]"), *InKey);
-				return nullptr;
+				return TArray<FConverterNodeInfo>();
 			}
-
-			TUniquePtr<Metasound::INode> FRegistryContainerImpl::CreateNode(const FNodeRegistryKey& InKey, FDefaultLiteralNodeConstructorParams&& InParams) const
+			else
 			{
-				const INodeRegistryEntry* Entry = FindNodeEntry(InKey);
-
-				if (ensureAlwaysMsgf(nullptr != Entry, TEXT("Could not find node [RegistryKey:%s]"), *InKey))
-				{
-					return Entry->CreateNode(MoveTemp(InParams));
-				}
-
-				return nullptr;
+				return ConverterNodeRegistry[InKey].PotentialConverterNodes;
 			}
+		}
 
-			TUniquePtr<Metasound::INode> FRegistryContainerImpl::CreateNode(const FNodeRegistryKey& InKey, FDefaultNamedVertexNodeConstructorParams&& InParams) const
-			{
-				const INodeRegistryEntry* Entry = FindNodeEntry(InKey);
+		TUniquePtr<FNodeRegistryTransactionStream> FRegistryContainerImpl::CreateTransactionStream()
+		{
+			return MakeUnique<FNodeRegistryTransactionStream>(TransactionBuffer);
+		}
 
-				if (ensureAlwaysMsgf(nullptr != Entry, TEXT("Could not find node [RegistryKey:%s]"), *InKey))
-				{
-					return Entry->CreateNode(MoveTemp(InParams));
-				}
+		FNodeRegistryKey FRegistryContainerImpl::RegisterNode(TUniquePtr<INodeRegistryEntry>&& InEntry)
+		{
+			METASOUND_LLM_SCOPE;
 
-				return nullptr;
-			}
+			FNodeRegistryKey Key;
 
-			TUniquePtr<Metasound::INode> FRegistryContainerImpl::CreateNode(const FNodeRegistryKey& InKey, FDefaultNamedVertexWithLiteralNodeConstructorParams&& InParams) const
+			if (InEntry.IsValid())
 			{
-				const INodeRegistryEntry* Entry = FindNodeEntry(InKey);
+				TSharedRef<INodeRegistryEntry, ESPMode::ThreadSafe> Entry(InEntry.Release());
 
-				if (ensureAlwaysMsgf(nullptr != Entry, TEXT("Could not find node [RegistryKey:%s]"), *InKey))
-				{
-					return Entry->CreateNode(MoveTemp(InParams));
-				}
+				FNodeRegistryTransaction::FTimeType Timestamp = FPlatformTime::Cycles64();
 
-				return nullptr;
-			}
+				Key = NodeRegistryKey::CreateKey(Entry->GetClassInfo());
 
-			TArray<::Metasound::Frontend::FConverterNodeInfo> FRegistryContainerImpl::GetPossibleConverterNodes(const FName& FromDataType, const FName& ToDataType)
-			{
-				FConverterNodeRegistryKey InKey = { FromDataType, ToDataType };
-				if (!ConverterNodeRegistry.Contains(InKey))
-				{
-					return TArray<FConverterNodeInfo>();
-				}
-				else
-				{
-					return ConverterNodeRegistry[InKey].PotentialConverterNodes;
-				}
+				// check to see if an identical node was already registered, and log
+				ensureAlwaysMsgf(
+					!RegisteredNodes.Contains(Key),
+					TEXT("Node with registry key '%s' already registered. "
+						"The previously registered node will be overwritten. "
+						"This can happen if two classes share the same name or if METASOUND_REGISTER_NODE is defined in a public header."),
+					*Key);
+
+				// Store update to newly registered node in history so nodes
+				// can be queried by transaction ID
+				
+				TransactionBuffer->AddTransaction(FNodeRegistryTransaction(FNodeRegistryTransaction::ETransactionType::NodeRegistration, Entry->GetClassInfo(), Timestamp));
+
+				// Store registry elements in map so nodes can be queried using registry key.
+				RegisteredNodes.Add(Key, Entry);
 			}
 
-			FNodeRegistryKey FRegistryContainerImpl::RegisterNode(TUniquePtr<INodeRegistryEntry>&& InEntry)
-			{
-				METASOUND_LLM_SCOPE;
+			return Key;
+		}
 
-				FNodeRegistryKey Key;
+		bool FRegistryContainerImpl::UnregisterNode(const FNodeRegistryKey& InKey)
+		{
+			METASOUND_LLM_SCOPE;
 
-				if (InEntry.IsValid())
+			if (NodeRegistryKey::IsValid(InKey))
+			{
+				if (const INodeRegistryEntry* Entry = FindNodeEntry(InKey))
 				{
-					TSharedRef<INodeRegistryEntry, ESPMode::ThreadSafe> Entry(InEntry.Release());
-
 					FNodeRegistryTransaction::FTimeType Timestamp = FPlatformTime::Cycles64();
 
-					Key = NodeRegistryKey::CreateKey(Entry->GetClassInfo());
-
-					// check to see if an identical node was already registered, and log
-					ensureAlwaysMsgf(
-						!RegisteredNodes.Contains(Key),
-						TEXT("Node with registry key '%s' already registered. "
-							"The previously registered node will be overwritten. "
-							"This can happen if two classes share the same name or if METASOUND_REGISTER_NODE is defined in a public header."),
-						*Key);
+					TransactionBuffer->AddTransaction(FNodeRegistryTransaction(FNodeRegistryTransaction::ETransactionType::NodeUnregistration, Entry->GetClassInfo(), Timestamp));
 
-					// Store update to newly registered node in history so nodes
-					// can be queried by transaction ID
-					RegistryTransactionHistory.Add(FNodeRegistryTransaction(FNodeRegistryTransaction::ETransactionType::NodeRegistration, Entry->GetClassInfo(), Timestamp));
-
-					// Store registry elements in map so nodes can be queried using registry key.
-					RegisteredNodes.Add(Key, Entry);
+					RegisteredNodes.Remove(InKey);
+					return true;
 				}
-
-				return Key;
 			}
 
-			bool FRegistryContainerImpl::UnregisterNode(const FNodeRegistryKey& InKey)
-			{
-				METASOUND_LLM_SCOPE;
-
-				if (NodeRegistryKey::IsValid(InKey))
-				{
-					if (const INodeRegistryEntry* Entry = FindNodeEntry(InKey))
-					{
-						FNodeRegistryTransaction::FTimeType Timestamp = FPlatformTime::Cycles64();
+			return false;
+		}
 
-						RegistryTransactionHistory.Add(FNodeRegistryTransaction(FNodeRegistryTransaction::ETransactionType::NodeUnregistration, Entry->GetClassInfo(), Timestamp));
+		bool FRegistryContainerImpl::RegisterConversionNode(const FConverterNodeRegistryKey& InNodeKey, const FConverterNodeInfo& InNodeInfo)
+		{
+			if (!ConverterNodeRegistry.Contains(InNodeKey))
+			{
+				ConverterNodeRegistry.Add(InNodeKey);
+			}
 
-						RegisteredNodes.Remove(InKey);
-						return true;
-					}
-				}
+			FConverterNodeRegistryValue& ConverterNodeList = ConverterNodeRegistry[InNodeKey];
 
+			if (ensureAlways(!ConverterNodeList.PotentialConverterNodes.Contains(InNodeInfo)))
+			{
+				ConverterNodeList.PotentialConverterNodes.Add(InNodeInfo);
+				return true;
+			}
+			else
+			{
+				// If we hit this, someone attempted to add the same converter node to our list multiple times.
 				return false;
 			}
+		}
+
+		bool FRegistryContainerImpl::IsNodeRegistered(const FNodeRegistryKey& InKey) const
+		{
+			return RegisteredNodes.Contains(InKey);
+		}
 
-			bool FRegistryContainerImpl::RegisterConversionNode(const FConverterNodeRegistryKey& InNodeKey, const FConverterNodeInfo& InNodeInfo)
+		bool FRegistryContainerImpl::IsNodeNative(const FNodeRegistryKey& InKey) const
+		{
+			if (const INodeRegistryEntry* Entry = FindNodeEntry(InKey))
 			{
-				if (!ConverterNodeRegistry.Contains(InNodeKey))
-				{
-					ConverterNodeRegistry.Add(InNodeKey);
-				}
+				return Entry->IsNative();
+			}
 
-				FConverterNodeRegistryValue& ConverterNodeList = ConverterNodeRegistry[InNodeKey];
+			return false;
+		}
 
-				if (ensureAlways(!ConverterNodeList.PotentialConverterNodes.Contains(InNodeInfo)))
-				{
-					ConverterNodeList.PotentialConverterNodes.Add(InNodeInfo);
-					return true;
-				}
-				else
-				{
-					// If we hit this, someone attempted to add the same converter node to our list multiple times.
-					return false;
-				}
-			}
 
-			bool FRegistryContainerImpl::IsNodeRegistered(const FNodeRegistryKey& InKey) const
+		bool FRegistryContainerImpl::FindFrontendClassFromRegistered(const FNodeRegistryKey& InKey, FMetasoundFrontendClass& OutClass)
+		{
+			if (const INodeRegistryEntry* Entry = FindNodeEntry(InKey))
 			{
-				return RegisteredNodes.Contains(InKey);
+				OutClass = Entry->GetFrontendClass();
+				return true;
 			}
+			return false;
+		}
 
-			bool FRegistryContainerImpl::IsNodeNative(const FNodeRegistryKey& InKey) const
+		bool FRegistryContainerImpl::FindNodeClassInfoFromRegistered(const Metasound::Frontend::FNodeRegistryKey& InKey, FNodeClassInfo& OutInfo)
+		{
+			if (const INodeRegistryEntry* Entry = FindNodeEntry(InKey))
 			{
-				if (const INodeRegistryEntry* Entry = FindNodeEntry(InKey))
-				{
-					return Entry->IsNative();
-				}
-
-				return false;
+				OutInfo = Entry->GetClassInfo();
+				return true;
 			}
 
+			return false;
+		}
 
-			bool FRegistryContainerImpl::FindFrontendClassFromRegistered(const FNodeRegistryKey& InKey, FMetasoundFrontendClass& OutClass)
+		bool FRegistryContainerImpl::FindInputNodeRegistryKeyForDataType(const FName& InDataTypeName, FNodeRegistryKey& OutKey)
+		{
+			FMetasoundFrontendClass Class;
+			if (IDataTypeRegistry::Get().GetFrontendInputClass(InDataTypeName, Class))
 			{
-				if (const INodeRegistryEntry* Entry = FindNodeEntry(InKey))
-				{
-					OutClass = Entry->GetFrontendClass();
-					return true;
-				}
-				return false;
+				OutKey = NodeRegistryKey::CreateKey(Class.Metadata);
+				return true;
 			}
+			return false;
+		}
 
-			bool FRegistryContainerImpl::FindNodeClassInfoFromRegistered(const Metasound::Frontend::FNodeRegistryKey& InKey, FNodeClassInfo& OutInfo)
+		bool FRegistryContainerImpl::FindVariableNodeRegistryKeyForDataType(const FName& InDataTypeName, FNodeRegistryKey& OutKey)
+		{
+			FMetasoundFrontendClass Class;
+			if (IDataTypeRegistry::Get().GetFrontendLiteralClass(InDataTypeName, Class))
 			{
-				if (const INodeRegistryEntry* Entry = FindNodeEntry(InKey))
-				{
-					OutInfo = Entry->GetClassInfo();
-					return true;
-				}
-
-				return false;
+				OutKey = NodeRegistryKey::CreateKey(Class.Metadata);
+				return true;
 			}
+			return false;
+		}
 
-			bool FRegistryContainerImpl::FindInputNodeRegistryKeyForDataType(const FName& InDataTypeName, FNodeRegistryKey& OutKey)
+		bool FRegistryContainerImpl::FindOutputNodeRegistryKeyForDataType(const FName& InDataTypeName, FNodeRegistryKey& OutKey)
+		{
+			FMetasoundFrontendClass Class;
+			if (IDataTypeRegistry::Get().GetFrontendOutputClass(InDataTypeName, Class))
 			{
-				FMetasoundFrontendClass Class;
-				if (IDataTypeRegistry::Get().GetFrontendInputClass(InDataTypeName, Class))
-				{
-					OutKey = NodeRegistryKey::CreateKey(Class.Metadata);
-					return true;
-				}
-				return false;
+				OutKey = NodeRegistryKey::CreateKey(Class.Metadata);
+				return true;
 			}
+			return false;
+		}
 
-			bool FRegistryContainerImpl::FindVariableNodeRegistryKeyForDataType(const FName& InDataTypeName, FNodeRegistryKey& OutKey)
+		void FRegistryContainerImpl::IterateRegistry(Metasound::FIterateMetasoundFrontendClassFunction InIterFunc, EMetasoundFrontendClassType InClassType) const
+		{
+			auto WrappedFunc = [&](const TPair<FNodeRegistryKey, TSharedPtr<INodeRegistryEntry, ESPMode::ThreadSafe>>& Pair)
 			{
-				FMetasoundFrontendClass Class;
-				if (IDataTypeRegistry::Get().GetFrontendLiteralClass(InDataTypeName, Class))
-				{
-					OutKey = NodeRegistryKey::CreateKey(Class.Metadata);
-					return true;
-				}
-				return false;
-			}
+				InIterFunc(Pair.Value->GetFrontendClass());
+			};
 
-			bool FRegistryContainerImpl::FindOutputNodeRegistryKeyForDataType(const FName& InDataTypeName, FNodeRegistryKey& OutKey)
+			if (EMetasoundFrontendClassType::Invalid == InClassType)
 			{
-				FMetasoundFrontendClass Class;
-				if (IDataTypeRegistry::Get().GetFrontendOutputClass(InDataTypeName, Class))
-				{
-					OutKey = NodeRegistryKey::CreateKey(Class.Metadata);
-					return true;
-				}
-				return false;
+				// Iterate through all classes. 
+				Algo::ForEach(RegisteredNodes, WrappedFunc);
 			}
-
-			void FRegistryContainerImpl::IterateRegistry(Metasound::FIterateMetasoundFrontendClassFunction InIterFunc, EMetasoundFrontendClassType InClassType) const
+			else
 			{
-				auto WrappedFunc = [&](const TPair<FNodeRegistryKey, TSharedRef<INodeRegistryEntry, ESPMode::ThreadSafe>>& Pair)
+				// Only call function on classes of certain type.
+				auto IsMatchingClassType = [&](const TPair<FNodeRegistryKey, TSharedPtr<INodeRegistryEntry, ESPMode::ThreadSafe>>& Pair)
 				{
-					InIterFunc(Pair.Value->GetFrontendClass());
+					return Pair.Value->GetClassInfo().Type == InClassType;
 				};
-
-				if (EMetasoundFrontendClassType::Invalid == InClassType)
-				{
-					// Iterate through all classes. 
-					Algo::ForEach(RegisteredNodes, WrappedFunc);
-				}
-				else
-				{
-					// Only call function on classes of certain type.
-					auto IsMatchingClassType = [&](const TPair<FNodeRegistryKey, TSharedRef<INodeRegistryEntry, ESPMode::ThreadSafe>>& Pair)
-					{
-						return Pair.Value->GetClassInfo().Type == InClassType;
-					};
-					Algo::ForEachIf(RegisteredNodes, IsMatchingClassType, WrappedFunc);
-				}
+				Algo::ForEachIf(RegisteredNodes, IsMatchingClassType, WrappedFunc);
 			}
+		}
 
 
-			const INodeRegistryEntry* FRegistryContainerImpl::FindNodeEntry(const FNodeRegistryKey& InKey) const
+		const INodeRegistryEntry* FRegistryContainerImpl::FindNodeEntry(const FNodeRegistryKey& InKey) const
+		{
+			if (const TSharedRef<INodeRegistryEntry, ESPMode::ThreadSafe>* Entry = RegisteredNodes.Find(InKey))
 			{
-				if (const TSharedRef<INodeRegistryEntry, ESPMode::ThreadSafe>* Entry = RegisteredNodes.Find(InKey))
-				{
-					return &Entry->Get();
-				}
-
-				return nullptr;
+				return &Entry->Get();
 			}
-		} // namespace MetasoundFrontendRegistriesPrivate
+
+			return nullptr;
+		}
 
 
 		FNodeRegistryTransaction::FNodeRegistryTransaction(ETransactionType InType, const FNodeClassInfo& InNodeClassInfo, FNodeRegistryTransaction::FTimeType InTimestamp)
@@ -569,25 +515,15 @@ namespace Metasound
 	} // namespace Frontend
 } // namespace Metasound
 
-FMetasoundFrontendRegistryContainer* FMetasoundFrontendRegistryContainer::LazySingleton = nullptr;
 
 FMetasoundFrontendRegistryContainer* FMetasoundFrontendRegistryContainer::Get()
 {
-	if (!LazySingleton)
-	{
-		LazySingleton = new Metasound::Frontend::MetasoundFrontendRegistryPrivate::FRegistryContainerImpl();
-	}
-
-	return LazySingleton;
+	return &Metasound::Frontend::FRegistryContainerImpl::Get();
 }
 
 void FMetasoundFrontendRegistryContainer::ShutdownMetasoundFrontend()
 {
-	if (LazySingleton)
-	{
-		delete LazySingleton;
-		LazySingleton = nullptr;
-	}
+	Metasound::Frontend::FRegistryContainerImpl::Shutdown();
 }
 
 Metasound::Frontend::FNodeRegistryKey FMetasoundFrontendRegistryContainer::GetRegistryKey(const FNodeClassMetadata& InNodeMetadata)
diff --git a/Engine/Plugins/Runtime/Metasound/Source/MetasoundFrontend/Private/MetasoundFrontendRegistryTransaction.cpp b/Engine/Plugins/Runtime/Metasound/Source/MetasoundFrontend/Private/MetasoundFrontendRegistryTransaction.cpp
--- a/Engine/Plugins/Runtime/Metasound/Source/MetasoundFrontend/Private/MetasoundFrontendRegistryTransaction.cpp
+++ b/Engine/Plugins/Runtime/Metasound/Source/MetasoundFrontend/Private/MetasoundFrontendRegistryTransaction.cpp
@@ -3,16 +3,18 @@
 
 #include "MetasoundFrontendRegistryTransaction.h"
 
-#include "CoreMinimal.h"
-
 namespace Metasound
 {
 	namespace Frontend
 	{
-		FRegistryTransactionID GetOriginRegistryTransactionID()
-		{
-			return 0;
-		}
+		int32 MetaSoundFrontendDiscardStreamedRegistryTransactionsCVar = 1;
+
+		FAutoConsoleVariableRef CVarMetaSoundFrontendDiscardStreamedRegistryTransactions(
+			TEXT("au.MetaSounds.Frontend.DiscardStreamedRegistryTransactions"),
+			MetaSoundFrontendDiscardStreamedRegistryTransactionsCVar,
+			TEXT("If enabled, MetaSound registry transactions are discarded after they have been streamed.\n")
+			TEXT("0: Disabled, !0: Enabled (default)"),
+			ECVF_Default);
 	}
 }
 
diff --git a/Engine/Plugins/Runtime/Metasound/Source/MetasoundFrontend/Private/MetasoundFrontendRegistryTransaction.h b/Engine/Plugins/Runtime/Metasound/Source/MetasoundFrontend/Private/MetasoundFrontendRegistryTransaction.h
--- a/Engine/Plugins/Runtime/Metasound/Source/MetasoundFrontend/Private/MetasoundFrontendRegistryTransaction.h
+++ b/Engine/Plugins/Runtime/Metasound/Source/MetasoundFrontend/Private/MetasoundFrontendRegistryTransaction.h
@@ -2,109 +2,216 @@
 
 #pragma once
 
-#include "CoreMinimal.h"
-#include "MetasoundFrontend.h"
-#include "MetasoundFrontendRegistries.h"
 #include "Algo/ForEach.h"
+#include "Algo/MinElement.h"
+#include "Containers/Array.h"
+#include "HAL/CriticalSection.h"
+#include "HAL/IConsoleManager.h"
+#include "MetasoundLog.h"
+#include "Misc/ScopeLock.h"
+
 
 namespace Metasound
 {
 	namespace Frontend
 	{
-		/** Returns an ID representing the beginning of the transaction history. */
-		FRegistryTransactionID GetOriginRegistryTransactionID();
+		extern int32 MetaSoundFrontendDiscardStreamedRegistryTransactionsCVar;
+		extern FAutoConsoleVariableRef CVarMetaSoundFrontendDiscardStreamedRegistryTransactions;
 
+		using FTransactionReaderHandle = uint32;
 
-		/** Maintains a history of TransactionTypes. Calls are threadsafe (excluding
+		/** Maintains a limited history of TransactionTypes. Calls are threadsafe (excluding
 		 * the constructor and destructor.)
 		 */
 		template<typename TransactionType>
-		class TRegistryTransactionHistory
+		class TTransactionBuffer
 		{
+
 		public:
-			TRegistryTransactionHistory()
-			: Current(GetOriginRegistryTransactionID())
+			TTransactionBuffer()
+			: NumRemovedTransactions(0)
+			, CurrentReaderHandle(0)
 			{
 			}
 
 			/** Add a transaction to the history. Threadsafe. 
 			 *
 			 * @return The transaction ID associated with the action. */
-			FRegistryTransactionID Add(TransactionType&& InRegistryTransaction)
+			int32 AddTransaction(TransactionType&& InRegistryTransaction)
 			{
 				FScopeLock Lock(&RegistryTransactionMutex);
+				{
+					TransactionHistory.Add(MoveTemp(InRegistryTransaction));
+					return TransactionHistory.Num() + NumRemovedTransactions;
+				}
+			}
 
-				Current++;
-
-				int32 Index = RegistryTransactions.Num();
+			/** Create a new transaction reader for this buffer. */
+			FTransactionReaderHandle CreateReader()
+			{
+				FScopeLock Lock(&RegistryTransactionMutex);
+				{
+					FTransactionReaderHandle NewReaderHandle = CurrentReaderHandle;
+					CurrentReaderHandle++;
+					ReaderPositions.Add(NewReaderHandle, NumRemovedTransactions);
+					return NewReaderHandle;
+				}
+			}
 
-				RegistryTransactions.Add(MoveTemp(InRegistryTransaction));
-				RegistryTransactionIndexMap.Add(Current, Index);
+			/** Create a duplicate reader which has the same read position as
+			 * the provided reader handle.  */
+			FTransactionReaderHandle DuplicateReader(FTransactionReaderHandle InReaderHandle)
+			{
+				FScopeLock Lock(&RegistryTransactionMutex);
+				{
+					int32 Position = NumRemovedTransactions;
+					if (int32* ExistingReaderPosition = ReaderPositions.Find(InReaderHandle))
+					{
+						Position = *ExistingReaderPosition;
+					}
 
-				return Current;
+					FTransactionReaderHandle NewReaderHandle = CurrentReaderHandle;
+					CurrentReaderHandle++;
+					ReaderPositions.Add(NewReaderHandle, Position);
+				}
 			}
 
-			/** Gets the transaction ID of the most recent transaction. Threadsafe. */
-			FRegistryTransactionID GetCurrent() const
+			/** Destroy a reader associated with this buffer. */
+			void DestroyReader(FTransactionReaderHandle ReaderHandle)
 			{
 				FScopeLock Lock(&RegistryTransactionMutex);
 				{
-					return Current;
+					ReaderPositions.Remove(ReaderHandle);
+					RemoveTransactionsStreamedToAllReaders();
 				}
 			}
 
-			/** Invoke a function on all transactions since transaction ID. 
+
+			/** Invoke a function on all unread transactions.
 			 *
-			 * @param InSince - All transactions occurring after this transaction ID will be returned.
-			 * @param OutCurrent - If not null, will be set to the most recent transaction ID returned. 
+			 * @param InReaderHandle - Handle of a valid reader for this buffer.
 			 * @param InCallable - A callable of the form Callable(const TransactionType&)
 			 */
 			template<typename CallableType>
-			void ForEachTransactionSince(FRegistryTransactionID InSince, FRegistryTransactionID* OutCurrent, CallableType InCallable) const
+			void StreamTransactions(FTransactionReaderHandle InReaderHandle, CallableType InCallable)
 			{
 				FScopeLock Lock(&RegistryTransactionMutex);
 				{
-					if (nullptr != OutCurrent)
+					if (ensure(ReaderPositions.Contains(InReaderHandle)))
 					{
-						*OutCurrent = Current;
-					}
+						int32 Start = ReaderPositions[InReaderHandle] - NumRemovedTransactions;
+						if (Start < 0)
+						{
+							UE_LOG(LogMetaSound, Warning, TEXT("Missed transactions in transaction stream. MetaSound registries may be incorrect"));
+							Start = 0;
+						}
 
-					int32 Start = INDEX_NONE;
-					
-					if (GetOriginRegistryTransactionID() == InSince)
-					{
-						Start = 0;
-					}
-					else if (const int32* Pos = RegistryTransactionIndexMap.Find(InSince))
-					{
-						Start = *Pos + 1;
+						const int32 OutNum = TransactionHistory.Num() - Start;
+						if (OutNum > 0)
+						{
+							TArrayView<const TransactionType> TransactionView = MakeArrayView(&TransactionHistory[Start], OutNum);
+							Algo::ForEach(TransactionView, InCallable);
+							ReaderPositions[InReaderHandle] += OutNum;
+
+							RemoveTransactionsStreamedToAllReaders();
+						}
 					}
-					
-					if (INDEX_NONE != Start)
-					{
-						const int32 Num = RegistryTransactions.Num();
+				}
+			}
 
-						if (ensure(Start <= Num))
+		private:
+
+			// Remove transactions that have been streamed by all valid readers.
+			void RemoveTransactionsStreamedToAllReaders()
+			{
+				if (0 != MetaSoundFrontendDiscardStreamedRegistryTransactionsCVar)
+				{
+					const auto* MinElement = Algo::MinElementBy(ReaderPositions, [](const auto& Pair){return Pair.Value;});
+					if (nullptr != MinElement)
+					{
+						int32 NumToRemove = (MinElement->Value) - NumRemovedTransactions;
+						if (NumToRemove > 0)
 						{
-							const int32 OutNum = Num - Start;
-							if (OutNum > 0)
-							{
-								TArrayView<const TransactionType> TransactionsSince = MakeArrayView(&RegistryTransactions[Start], OutNum);
-								Algo::ForEach(TransactionsSince, InCallable);
-							}
+							TransactionHistory.RemoveAt(0, NumToRemove);
+							NumRemovedTransactions += NumToRemove;
 						}
 					}
+					else if (ReaderPositions.Num() == 0)
+					{
+						NumRemovedTransactions += TransactionHistory.Num();
+						TransactionHistory.Reset();
+					}
 				}
 			}
 
-		private:
+			FCriticalSection RegistryTransactionMutex;
 
-			mutable FCriticalSection RegistryTransactionMutex;
+			int32 NumRemovedTransactions = 0;
+			FTransactionReaderHandle CurrentReaderHandle = 0;
+			TSortedMap<FTransactionReaderHandle, int32> ReaderPositions;
+			TArray<TransactionType> TransactionHistory;
+		};
 
-			FRegistryTransactionID Current;
 
-			TMap<FRegistryTransactionID, int32> RegistryTransactionIndexMap;
-			TArray<TransactionType> RegistryTransactions;
+		/** A transaction stream which streams data from a transaction buffer. */
+		template<typename TransactionType>
+		class TTransactionStream
+		{
+			using TransactionBufferType = TTransactionBuffer<TransactionType>;
+		public:
+			TTransactionStream(TSharedRef<TransactionBufferType, ESPMode::ThreadSafe> InBuffer)
+			: Buffer(InBuffer)
+			{
+				ReaderHandle = Buffer->CreateReader();
+			}
+
+			TTransactionStream(const TTransactionStream& InOther)
+			: Buffer(InOther.Buffer)
+			{
+				ReaderHandle = Buffer->DuplicateReader(InOther.ReaderHandle);
+			}
+
+			TTransactionStream& operator=(const TTransactionStream& InOther)
+			{
+				Buffer.RemoveReader(ReaderHandle);
+				Buffer = InOther.Buffer;
+				ReaderHandle = Buffer.DuplicateReader(InOther.ReaderHandle);
+				return *this;
+			}
+
+			TTransactionStream(TTransactionStream&& InOther)
+			: Buffer(MoveTemp(InOther.Buffer))
+			, ReaderHandle(InOther.ReaderHandle)
+			{
+				InOther.ReaderHandle = -1;
+			}
+
+			TTransactionStream& operator=(TTransactionStream&& InOther)
+			{
+				Buffer = MoveTemp(InOther.Buffer);
+				ReaderHandle = InOther.ReaderHandle;
+				InOther.ReaderHandle = -1;
+				return *this;
+			}
+
+			~TTransactionStream()
+			{
+				Buffer->DestroyReader(ReaderHandle);
+			}
+
+			/** Invoke a function on all unstreamed transactions.
+			 *
+			 * @param InCallable - A callable of the form Callable(const TransactionType&)
+			 */
+			template <typename CallableType>
+			void Stream(CallableType Callable)
+			{
+				Buffer->StreamTransactions(ReaderHandle, Callable);
+			}
+
+		private:
+			TSharedRef<TransactionBufferType, ESPMode::ThreadSafe> Buffer;
+			FTransactionReaderHandle ReaderHandle = -1;
 		};
 	}
 }
diff --git a/Engine/Plugins/Runtime/Metasound/Source/MetasoundFrontend/Private/MetasoundFrontendSearchEngineCore.cpp b/Engine/Plugins/Runtime/Metasound/Source/MetasoundFrontend/Private/MetasoundFrontendSearchEngineCore.cpp
--- a/Engine/Plugins/Runtime/Metasound/Source/MetasoundFrontend/Private/MetasoundFrontendSearchEngineCore.cpp
+++ b/Engine/Plugins/Runtime/Metasound/Source/MetasoundFrontend/Private/MetasoundFrontendSearchEngineCore.cpp
@@ -6,6 +6,7 @@
 #include "Containers/Array.h"
 #include "Containers/UnrealString.h"
 #include "MetasoundFrontendArchetypeRegistry.h"
+#include "MetasoundFrontendInterfaceRegistryPrivate.h"
 #include "MetasoundFrontendQuery.h"
 #include "MetasoundFrontendQuerySteps.h"
 #include "MetasoundFrontendRegistryTransaction.h"
@@ -33,7 +34,7 @@ namespace Metasound
 			}
 
 			FInterfaceRegistryTransactionSource::FInterfaceRegistryTransactionSource()
-			: CurrentTransactionID(GetOriginRegistryTransactionID())
+			: TransactionStream(FInterfaceRegistry::Get().CreateTransactionStream())
 			{
 			}
 
@@ -43,8 +44,11 @@ namespace Metasound
 				{
 					OutEntries.Emplace(TInPlaceType<FInterfaceRegistryTransaction>(), InTransaction);
 				};
-				
-				IInterfaceRegistry::Get().ForEachRegistryTransactionSince(CurrentTransactionID, &CurrentTransactionID, AddValue);
+
+				if (TransactionStream.IsValid())
+				{
+					TransactionStream->Stream(AddValue);
+				}
 			}
 
 			FFrontendQueryKey FMapInterfaceRegistryTransactionsToInterfaceRegistryKeys::Map(const FFrontendQueryEntry& InEntry) const
diff --git a/Engine/Plugins/Runtime/Metasound/Source/MetasoundFrontend/Private/MetasoundFrontendSearchEngineCore.h b/Engine/Plugins/Runtime/Metasound/Source/MetasoundFrontend/Private/MetasoundFrontendSearchEngineCore.h
--- a/Engine/Plugins/Runtime/Metasound/Source/MetasoundFrontend/Private/MetasoundFrontendSearchEngineCore.h
+++ b/Engine/Plugins/Runtime/Metasound/Source/MetasoundFrontend/Private/MetasoundFrontendSearchEngineCore.h
@@ -5,6 +5,7 @@
 #include "Containers/Array.h"
 #include "HAL/CriticalSection.h"
 #include "MetasoundFrontendArchetypeRegistry.h"
+#include "MetasoundFrontendInterfaceRegistryPrivate.h"
 #include "MetasoundFrontendQuery.h"
 #include "MetasoundFrontendQuerySteps.h"
 #include "MetasoundFrontendRegistries.h"
@@ -33,7 +34,7 @@ namespace Metasound
 				virtual void Stream(TArray<FFrontendQueryValue>& OutEntries) override;
 
 			private:
-				FRegistryTransactionID CurrentTransactionID;
+				TUniquePtr<FInterfaceTransactionStream> TransactionStream;
 			};
 
 			class FMapInterfaceRegistryTransactionsToInterfaceRegistryKeys : public IFrontendQueryMapStep
diff --git a/Engine/Plugins/Runtime/Metasound/Source/MetasoundFrontend/Public/MetasoundFrontendArchetypeRegistry.h b/Engine/Plugins/Runtime/Metasound/Source/MetasoundFrontend/Public/MetasoundFrontendArchetypeRegistry.h
--- a/Engine/Plugins/Runtime/Metasound/Source/MetasoundFrontend/Public/MetasoundFrontendArchetypeRegistry.h
+++ b/Engine/Plugins/Runtime/Metasound/Source/MetasoundFrontend/Public/MetasoundFrontendArchetypeRegistry.h
@@ -87,6 +87,7 @@ namespace Metasound
 			virtual bool FindInterface(const FInterfaceRegistryKey& InKey, FMetasoundFrontendInterface& OutInterface) const = 0;
 
 			// Executes a function against all transactions since a provided ID and provides the current transaction ID.
+			UE_DEPRECATED(5.1, "ForEachRegistryTransactionSince is no longer supported. Use Metasound::Frontend::ISearchEngine to query for all interfaces.")
 			virtual void ForEachRegistryTransactionSince(FRegistryTransactionID InSince, FRegistryTransactionID* OutCurrentRegistryTransactionID, TFunctionRef<void(const FInterfaceRegistryTransaction&)> InFunc) const = 0;
 		};
 	}
diff --git a/Engine/Plugins/Runtime/Metasound/Source/MetasoundFrontend/Public/MetasoundFrontendQuerySteps.h b/Engine/Plugins/Runtime/Metasound/Source/MetasoundFrontend/Public/MetasoundFrontendQuerySteps.h
--- a/Engine/Plugins/Runtime/Metasound/Source/MetasoundFrontend/Public/MetasoundFrontendQuerySteps.h
+++ b/Engine/Plugins/Runtime/Metasound/Source/MetasoundFrontend/Public/MetasoundFrontendQuerySteps.h
@@ -7,9 +7,13 @@
 #include "MetasoundFrontendQuery.h"
 #include "MetasoundFrontendRegistries.h"
 #include "Misc/Guid.h"
+#include "Templates/PimplPtr.h"
 
 namespace Metasound
 {
+	// Forward declare private implementation
+	class FNodeClassRegistrationEventsPimpl;
+
 	class FMapToNull : public IFrontendQueryMapStep
 	{
 	public:
@@ -28,7 +32,7 @@ namespace Metasound
 		virtual void Stream(TArray<FFrontendQueryValue>& OutValues) override;
 
 	private:
-		Frontend::FRegistryTransactionID CurrentTransactionID;
+		TPimplPtr<FNodeClassRegistrationEventsPimpl> Pimpl;
 	};
 
 	/** Partitions node registration events by their node registration keys. */
diff --git a/Engine/Plugins/Runtime/Metasound/Source/MetasoundFrontend/Public/MetasoundFrontendRegistries.h b/Engine/Plugins/Runtime/Metasound/Source/MetasoundFrontend/Public/MetasoundFrontendRegistries.h
--- a/Engine/Plugins/Runtime/Metasound/Source/MetasoundFrontend/Public/MetasoundFrontendRegistries.h
+++ b/Engine/Plugins/Runtime/Metasound/Source/MetasoundFrontend/Public/MetasoundFrontendRegistries.h
@@ -306,6 +306,7 @@ public:
 	virtual void RegisterPendingNodes() = 0;
 
 	/** Perform function for each registry transaction since a given transaction ID. */
+	UE_DEPRECATED(5.1, "ForEachNodeRegistryTransactionSince is no longer be supported")
 	virtual void ForEachNodeRegistryTransactionSince(Metasound::Frontend::FRegistryTransactionID InSince, Metasound::Frontend::FRegistryTransactionID* OutCurrentRegistryTransactionID, TFunctionRef<void(const Metasound::Frontend::FNodeRegistryTransaction&)> InFunc) const = 0;
 
 
@@ -344,7 +345,6 @@ public:
 	virtual bool FindVariableNodeRegistryKeyForDataType(const FName& InDataTypeName, FNodeRegistryKey& OutKey) = 0;
 	virtual bool FindOutputNodeRegistryKeyForDataType(const FName& InDataTypeName, FNodeRegistryKey& OutKey) = 0;
 
-
 	virtual TUniquePtr<Metasound::INode> CreateNode(const FNodeRegistryKey& InKey, const Metasound::FNodeInitData&) const = 0;
 	virtual TUniquePtr<Metasound::INode> CreateNode(const FNodeRegistryKey& InKey, Metasound::FDefaultLiteralNodeConstructorParams&&) const = 0;
 	virtual TUniquePtr<Metasound::INode> CreateNode(const FNodeRegistryKey& InKey, Metasound::FDefaultNamedVertexNodeConstructorParams&&) const = 0;
@@ -357,7 +357,6 @@ public:
 	virtual TArray<FConverterNodeInfo> GetPossibleConverterNodes(const FName& FromDataType, const FName& ToDataType) = 0;
 
 private:
-	static FMetasoundFrontendRegistryContainer* LazySingleton;
 };
 
 
