Reducing internal container sizes in Metasound Frontend Query

diff --git a/Engine/Plugins/Runtime/Metasound/Source/MetasoundFrontend/Private/MetasoundFrontendQuery.cpp b/Engine/Plugins/Runtime/Metasound/Source/MetasoundFrontend/Private/MetasoundFrontendQuery.cpp
--- a/Engine/Plugins/Runtime/Metasound/Source/MetasoundFrontend/Private/MetasoundFrontendQuery.cpp
+++ b/Engine/Plugins/Runtime/Metasound/Source/MetasoundFrontend/Private/MetasoundFrontendQuery.cpp
@@ -7,6 +7,7 @@
 #include "MetasoundTrace.h"
 #include "Misc/Guid.h"
 #include "Templates/TypeHash.h"
+#include "Traits/IsContiguousContainer.h"
 
 namespace Metasound
 {
@@ -14,10 +15,51 @@ namespace Metasound
 	{
 		void CompactSelection(FFrontendQuerySelection& InSelection)
 		{
-			for (auto& Pair : InSelection)
+			InSelection.Shrink();
+		}
+
+		// Not symmetric. Elements in LHS that are not in RHS. Modifies partition by resorting.
+		FFrontendQueryPartition Difference(FFrontendQueryPartition& InLHS, FFrontendQueryPartition& InRHS)
+		{
+			static_assert(TIsContiguousContainer<FFrontendQueryPartition>::Value, "Partitions must be a contiguous container for Difference algorithm to no access invalid memory");
+
+			if ((InLHS.Num() == 0) || (InRHS.Num() == 0))
+			{
+				return InLHS;
+			}
+
+			auto IsIDLessThan = [](const FFrontendQueryEntry& InLHS, const FFrontendQueryEntry& InRHS) { return InLHS.ID < InRHS.ID; };
+			auto IsIDEqual = [](const FFrontendQueryEntry& InLHS, const FFrontendQueryEntry& InRHS) { return InLHS.ID == InRHS.ID; };
+
+			InLHS.Sort(IsIDLessThan);
+			InRHS.Sort(IsIDLessThan);
+
+			FFrontendQueryPartition Result;
+			const FFrontendQueryEntry* LHSPtr = InLHS.GetData();
+			const FFrontendQueryEntry* LHSPtrEnd = LHSPtr + InLHS.Num();
+			const FFrontendQueryEntry* RHSPtr = InRHS.GetData();
+			const FFrontendQueryEntry* RHSPtrEnd = RHSPtr + InRHS.Num();
+
+			while (LHSPtr != LHSPtrEnd)
 			{
-				Pair.Value.Compact();
+				if ((RHSPtr == RHSPtrEnd) || IsIDLessThan(*LHSPtr, *RHSPtr))
+				{
+					Result.Add(*LHSPtr);
+					LHSPtr++;
+				}
+				else if (IsIDLessThan(*RHSPtr, *LHSPtr))
+				{
+					RHSPtr++;
+				}
+				else
+				{
+					// Values are equal
+					LHSPtr++;
+					RHSPtr++;
+				}
 			}
+			
+			return Result;
 		}
 
 		// Wrapper for step defined by function
@@ -340,7 +382,10 @@ namespace Metasound
 
 						// Get the updated set of keys in the output.
 						InOutUpdatedKeys.Reset();
-						Result.GetKeys(InOutUpdatedKeys);
+						for (const auto& Pair : Result)
+						{
+							InOutUpdatedKeys.Add(Pair.Key);
+						}
 
 						// Append the new values to the final result.
 						Append(InOutUpdatedKeys, Result, InOutResult);
@@ -376,7 +421,8 @@ namespace Metasound
 
 								Step->Reduce(Key, Partition);
 
-								FFrontendQueryPartition Removed = OriginalPartition.Difference(Partition);
+								FFrontendQueryPartition Removed = FrontendQueryPrivate::Difference(OriginalPartition, Partition);
+
 								if (Removed.Num() > 0)
 								{
 									InOutIncremental.ActiveRemovalSelection.FindOrAdd(Key).Append(Removed);
diff --git a/Engine/Plugins/Runtime/Metasound/Source/MetasoundFrontend/Public/MetasoundFrontendQuery.h b/Engine/Plugins/Runtime/Metasound/Source/MetasoundFrontend/Public/MetasoundFrontendQuery.h
--- a/Engine/Plugins/Runtime/Metasound/Source/MetasoundFrontend/Public/MetasoundFrontendQuery.h
+++ b/Engine/Plugins/Runtime/Metasound/Source/MetasoundFrontend/Public/MetasoundFrontendQuery.h
@@ -2,7 +2,8 @@
 
 #pragma once
 
-#include "CoreMinimal.h"
+#include "Containers/Array.h"
+#include "Containers/SortedMap.h"
 #include "Math/NumericLimits.h"
 #include "MetasoundFrontendDocument.h"
 #include "MetasoundFrontendRegistries.h"
@@ -95,10 +96,11 @@ namespace Metasound
 
 	/** A FFrontendQueryPartition represents a set of entries associated with a 
 	 * single FFrontendQueryKey. */
-	using FFrontendQueryPartition = TSet<FFrontendQueryEntry>;
+	
+	using FFrontendQueryPartition = TArray<FFrontendQueryEntry, TInlineAllocator<1>>;
 
 	/** A FFrontendQuerySelection holds a map of keys to partitions. */
-	using FFrontendQuerySelection = TMap<FFrontendQueryKey, FFrontendQueryPartition>;
+	using FFrontendQuerySelection = TSortedMap<FFrontendQueryKey, FFrontendQueryPartition>;
 
 	/** Interface for an individual step in a query */
 	class IFrontendQueryStep
@@ -162,7 +164,7 @@ namespace Metasound
 			virtual ~IFrontendQuerySortStep() = default;
 			virtual bool Sort(const FFrontendQueryEntry& InEntryLHS, const FFrontendQueryEntry& InEntryRHS) const = 0;
 	};
-
+	
 	class IFrontendQueryLimitStep : public IFrontendQueryStep
 	{
 		public:
