MetaSound Analyzer Frontend/Graph Core Checkpoint
- Add ability to track all internal data references to a core graph
- Add AnalyzerRegistry & first pass at value/envelope analyzers
- Add support for enabling analyzing from MetaSound Generator
- Add ability for tracking colorized bounds for edges as literal-color array

diff --git a/Engine/Plugins/Runtime/Metasound/Source/MetasoundEngine/Private/MetasoundEngineModule.cpp b/Engine/Plugins/Runtime/Metasound/Source/MetasoundEngine/Private/MetasoundEngineModule.cpp
--- a/Engine/Plugins/Runtime/Metasound/Source/MetasoundEngine/Private/MetasoundEngineModule.cpp
+++ b/Engine/Plugins/Runtime/Metasound/Source/MetasoundEngine/Private/MetasoundEngineModule.cpp
@@ -1,6 +1,7 @@
 // Copyright Epic Games, Inc. All Rights Reserved.
 #include "MetasoundEngineModule.h"
 
+#include "Analysis/MetasoundFrontendAnalyzerRegistry.h"
 #include "Metasound.h"
 #include "MetasoundDataReference.h"
 #include "MetasoundDataTypeRegistrationMacro.h"
@@ -38,6 +39,9 @@ class FMetasoundEngineModule : public IMetasoundEngineModule
 		// Flush node registration queue
 		FMetasoundFrontendRegistryContainer::Get()->RegisterPendingNodes();
 
+		// Register Analyzers
+		Metasound::Frontend::IVertexAnalyzerRegistry::Get().RegisterAnalyzerFactories();
+
 		UE_LOG(LogMetasoundEngine, Log, TEXT("MetaSound Engine Initialized"));
 	}
 };
diff --git a/Engine/Plugins/Runtime/Metasound/Source/MetasoundEngine/Private/MetasoundSource.cpp b/Engine/Plugins/Runtime/Metasound/Source/MetasoundEngine/Private/MetasoundSource.cpp
--- a/Engine/Plugins/Runtime/Metasound/Source/MetasoundEngine/Private/MetasoundSource.cpp
+++ b/Engine/Plugins/Runtime/Metasound/Source/MetasoundEngine/Private/MetasoundSource.cpp
@@ -19,6 +19,7 @@
 #include "MetasoundFrontendTransform.h"
 #include "MetasoundGenerator.h"
 #include "MetasoundLog.h"
+#include "MetasoundOperatorBuilderSettings.h"
 #include "MetasoundOperatorSettings.h"
 #include "MetasoundOutputFormatInterfaces.h"
 #include "MetasoundParameterTransmitter.h"
@@ -72,12 +73,6 @@ namespace Metasound
 	} // namespace SourcePrivate
 } // namespace Metasound
 
-FAutoConsoleVariableRef CVarMetaSoundBlockRate(
-	TEXT("au.MetaSound.BlockRate"),
-	Metasound::ConsoleVariables::BlockRate,
-	TEXT("Sets block rate (blocks per second) of MetaSounds.\n")
-	TEXT("Default: 100.0f"),
-	ECVF_Default);
 
 UMetaSoundSource::UMetaSoundSource(const FObjectInitializer& ObjectInitializer)
 	: Super(ObjectInitializer)
@@ -531,9 +526,15 @@ ISoundGeneratorPtr UMetaSoundSource::CreateSoundGenerator(const FSoundGeneratorI
 		return ISoundGeneratorPtr(nullptr);
 	}
 
+	FOperatorBuilderSettings BuilderSettings = FOperatorBuilderSettings::GetDefaultSettings();
+
+	// Graph analyzer currently only enabled for preview sounds (but can theoretically be supported for all sounds)
+	BuilderSettings.bPopulateInternalDataReferences = InParams.bIsPreviewSound;
+
 	FMetasoundGeneratorInitParams InitParams =
 	{
 		InSettings,
+		MoveTemp(BuilderSettings),
 		MetasoundGraph,
 		Environment,
 		GetName(),
@@ -817,7 +818,7 @@ TUniquePtr<Audio::IParameterTransmitter> UMetaSoundSource::CreateParameterTransm
 
 Metasound::FOperatorSettings UMetaSoundSource::GetOperatorSettings(Metasound::FSampleRate InSampleRate) const
 {
-	const float BlockRate = FMath::Clamp(Metasound::ConsoleVariables::BlockRate, 1.0f, 1000.0f);
+	const float BlockRate = Metasound::Frontend::GetDefaultBlockRate();
 	return Metasound::FOperatorSettings(InSampleRate, BlockRate);
 }
 
diff --git a/Engine/Plugins/Runtime/Metasound/Source/MetasoundEngine/Private/MetasoundWavePlayerNode.cpp b/Engine/Plugins/Runtime/Metasound/Source/MetasoundEngine/Private/MetasoundWavePlayerNode.cpp
--- a/Engine/Plugins/Runtime/Metasound/Source/MetasoundEngine/Private/MetasoundWavePlayerNode.cpp
+++ b/Engine/Plugins/Runtime/Metasound/Source/MetasoundEngine/Private/MetasoundWavePlayerNode.cpp
@@ -7,6 +7,7 @@
 #include "DSP/MultichannelLinearResampler.h"
 #include "IAudioCodec.h"
 #include "MetasoundBuildError.h"
+#include "MetasoundBuilderInterface.h"
 #include "MetasoundEngineNodesNames.h"
 #include "MetasoundExecutableOperator.h"
 #include "MetasoundLog.h"
@@ -51,7 +52,7 @@ namespace Metasound
 	{
 		class FOperatorFactory : public IOperatorFactory
 		{
-			virtual TUniquePtr<IOperator> CreateOperator(const FCreateOperatorParams& InParams, FBuildErrorArray& OutErrors) override;
+			virtual TUniquePtr<IOperator> CreateOperator(const FCreateOperatorParams& InParams, FBuildGraphResults& OutResults) override;
 		};
 
 	public:
@@ -982,9 +983,7 @@ namespace Metasound
 		
 	};
 
-	TUniquePtr<IOperator> FWavePlayerNode::FOperatorFactory::CreateOperator(
-		const FCreateOperatorParams& InParams, 
-		FBuildErrorArray& OutErrors) 
+	TUniquePtr<IOperator> FWavePlayerNode::FOperatorFactory::CreateOperator(const FCreateOperatorParams& InParams, FBuildGraphResults& OutResults)
 	{
 		using namespace Audio;
 		using namespace WavePlayerVertexNames;
diff --git a/Engine/Plugins/Runtime/Metasound/Source/MetasoundEngine/Public/Metasound.h b/Engine/Plugins/Runtime/Metasound/Source/MetasoundEngine/Public/Metasound.h
--- a/Engine/Plugins/Runtime/Metasound/Source/MetasoundEngine/Public/Metasound.h
+++ b/Engine/Plugins/Runtime/Metasound/Source/MetasoundEngine/Public/Metasound.h
@@ -45,7 +45,7 @@ namespace Metasound
 {
 	namespace ConsoleVariables
 	{
-		static float BlockRate = 100.f;
+		static float GetBlockRate();
 	} // namespace ConsoleVariables
 
 #if WITH_EDITOR
diff --git a/Engine/Plugins/Runtime/Metasound/Source/MetasoundFrontend/Private/Analysis/MetasoundFrontendAnalyzerAddress.cpp b/Engine/Plugins/Runtime/Metasound/Source/MetasoundFrontend/Private/Analysis/MetasoundFrontendAnalyzerAddress.cpp
new file mode 100644
--- /dev/null
+++ b/Engine/Plugins/Runtime/Metasound/Source/MetasoundFrontend/Private/Analysis/MetasoundFrontendAnalyzerAddress.cpp
@@ -0,0 +1,61 @@
+// Copyright Epic Games, Inc. All Rights Reserved.
+#include "Analysis/MetasoundFrontendAnalyzerAddress.h"
+
+
+namespace Metasound
+{
+	namespace Frontend
+	{
+		const FString FAnalyzerAddress::PathSeparator = TEXT("/");
+
+		FString FAnalyzerAddress::ToString() const
+		{
+			return FString::Join(TArray<FString>
+			{
+				*FString::Printf(TEXT("%lld"), InstanceID),
+				*NodeID.ToString(),
+				*OutputName.ToString(),
+				*DataType.ToString(),
+				*AnalyzerName.ToString(),
+				*AnalyzerInstanceID.ToString(),
+				*AnalyzerMemberName.ToString()
+			}, *PathSeparator);
+		}
+
+		FSendAddress FAnalyzerAddress::ToSendAddress() const
+		{
+			const TArray<FString> ChannelTokens
+			{
+				NodeID.ToString(),
+				OutputName.ToString(),
+				AnalyzerName.ToString(),
+				AnalyzerInstanceID.ToString(),
+				AnalyzerMemberName.ToString()
+			};
+
+			// TODO: This is bad as its generating FNames like crazy. One idea was to include guid support in FSendAddresses to avoid FName generation,
+			// which would enable this factory to map analyzer keys to send guid.
+			const FName Channel = FName(*FString::Join(ChannelTokens, *PathSeparator));
+			return FSendAddress { Channel, DataType, InstanceID };
+		}
+
+		bool FAnalyzerAddress::ParseKey(const FString& InAnalyzerKey, FAnalyzerAddress& OutAddress)
+		{
+			TArray<FString> Tokens;
+			if (InAnalyzerKey.ParseIntoArray(Tokens, *PathSeparator) == 7)
+			{
+				OutAddress.InstanceID = static_cast<uint64>(FCString::Atoi64(*Tokens[0]));
+				OutAddress.NodeID = FGuid(Tokens[1]);
+				OutAddress.OutputName = *Tokens[2];
+				OutAddress.DataType = *Tokens[3];
+				OutAddress.AnalyzerName = *Tokens[4];
+				OutAddress.AnalyzerInstanceID = FGuid(Tokens[5]);
+				OutAddress.AnalyzerMemberName = *Tokens[6];
+
+				return true;
+			}
+
+			return false;
+		}
+	} // namespace Frontend
+} // namespace Metasound
diff --git a/Engine/Plugins/Runtime/Metasound/Source/MetasoundFrontend/Private/Analysis/MetasoundFrontendAnalyzerRegistry.cpp b/Engine/Plugins/Runtime/Metasound/Source/MetasoundFrontend/Private/Analysis/MetasoundFrontendAnalyzerRegistry.cpp
new file mode 100644
--- /dev/null
+++ b/Engine/Plugins/Runtime/Metasound/Source/MetasoundFrontend/Private/Analysis/MetasoundFrontendAnalyzerRegistry.cpp
@@ -0,0 +1,61 @@
+// Copyright Epic Games, Inc. All Rights Reserved.
+#include "Analysis/MetasoundFrontendAnalyzerRegistry.h"
+
+#include "Analysis/MetasoundFrontendAnalyzerFactory.h"
+#include "Analysis/MetasoundFrontendVertexAnalyzerEnvelopeFollower.h"
+#include "Analysis/MetasoundFrontendVertexAnalyzerTriggerDensity.h"
+#include "Analysis/MetasoundFrontendVertexAnalyzerValue.h"
+#include "MetasoundPrimitives.h"
+
+
+namespace Metasound
+{
+	namespace Frontend
+	{
+		class FVertexAnalyzerRegistry : public IVertexAnalyzerRegistry
+		{
+			TMap<FName, TUniquePtr<IVertexAnalyzerFactory>> AnalyzerFactoryRegistry;
+
+			template<typename TAnalyzerFactoryClass>
+			void RegisterAnalyzerFactory()
+			{
+				TUniquePtr<IVertexAnalyzerFactory> Factory(new TAnalyzerFactoryClass());
+				AnalyzerFactoryRegistry.Emplace(TAnalyzerFactoryClass::GetAnalyzerName(), MoveTemp(Factory));
+			}
+
+		public:
+			FVertexAnalyzerRegistry() = default;
+			virtual ~FVertexAnalyzerRegistry() = default;
+
+			virtual const TUniquePtr<IVertexAnalyzerFactory>& FindAnalyzerFactory(FName InAnalyzerName) const override
+			{
+				const TUniquePtr<IVertexAnalyzerFactory>* Factory = AnalyzerFactoryRegistry.Find(InAnalyzerName);
+				if (ensureMsgf(Factory, TEXT("Failed to find registered MetaSound Analyzer Factory with name '%s'"), *InAnalyzerName.ToString()))
+				{
+					return *Factory;
+				}
+
+				static const TUniquePtr<IVertexAnalyzerFactory> InvalidFactory;
+				return InvalidFactory;
+			}
+
+			virtual void RegisterAnalyzerFactories() override
+			{
+				RegisterAnalyzerFactory<FVertexAnalyzerEnvelopeFollower::FFactory>();
+				RegisterAnalyzerFactory<FVertexAnalyzerTriggerDensity::FFactory>();
+
+				// Primitives
+				RegisterAnalyzerFactory<FVertexAnalyzerBool::FFactory>();
+				RegisterAnalyzerFactory<FVertexAnalyzerFloat::FFactory>();
+				RegisterAnalyzerFactory<FVertexAnalyzerInt::FFactory>();
+				RegisterAnalyzerFactory<FVertexAnalyzerString::FFactory>();
+			}
+		};
+
+		IVertexAnalyzerRegistry& IVertexAnalyzerRegistry::Get()
+		{
+			static FVertexAnalyzerRegistry Registry;
+			return Registry;
+		}
+	} // namespace Frontend
+} // namespace Metasound
diff --git a/Engine/Plugins/Runtime/Metasound/Source/MetasoundFrontend/Private/Analysis/MetasoundFrontendAnalyzerView.cpp b/Engine/Plugins/Runtime/Metasound/Source/MetasoundFrontend/Private/Analysis/MetasoundFrontendAnalyzerView.cpp
new file mode 100644
--- /dev/null
+++ b/Engine/Plugins/Runtime/Metasound/Source/MetasoundFrontend/Private/Analysis/MetasoundFrontendAnalyzerView.cpp
@@ -0,0 +1,38 @@
+// Copyright Epic Games, Inc. All Rights Reserved.
+#include "Analysis/MetasoundFrontendAnalyzerView.h"
+
+#include "Analysis/MetasoundFrontendVertexAnalyzer.h"
+#include "Analysis/MetasoundFrontendVertexAnalyzerEnvelopeFollower.h"
+#include "Analysis/MetasoundFrontendGraphAnalyzer.h"
+#include "Math/UnrealMathUtility.h"
+#include "MetasoundAssetBase.h"
+#include "MetasoundFrontendController.h"
+#include "MetasoundOperatorSettings.h"
+#include "Templates/Function.h"
+
+
+namespace Metasound
+{
+	namespace Frontend
+	{
+		FMetasoundAnalyzerView::FMetasoundAnalyzerView(const FOperatorSettings& InOperatorSettings, FAnalyzerAddress&& InAnalyzerAddress)
+			: AnalyzerAddress(MoveTemp(InAnalyzerAddress))
+		{
+			const FSendAddress SendAddress = AnalyzerAddress.ToSendAddress();
+			const FReceiverInitParams ReceiverParams { InOperatorSettings };
+			Receiver = FDataTransmissionCenter::Get().RegisterNewReceiver(SendAddress, ReceiverParams);
+		}
+
+		IReceiver& FMetasoundAnalyzerView::GetReceiverChecked()
+		{
+			check(Receiver.IsValid());
+			return *Receiver;
+		}
+
+		const IReceiver& FMetasoundAnalyzerView::GetReceiverChecked() const
+		{
+			check(Receiver.IsValid());
+			return *Receiver;
+		}
+	} // namespace Frontend
+} // namespace Metasound
diff --git a/Engine/Plugins/Runtime/Metasound/Source/MetasoundFrontend/Private/Analysis/MetasoundFrontendGraphAnalyzer.cpp b/Engine/Plugins/Runtime/Metasound/Source/MetasoundFrontend/Private/Analysis/MetasoundFrontendGraphAnalyzer.cpp
new file mode 100644
--- /dev/null
+++ b/Engine/Plugins/Runtime/Metasound/Source/MetasoundFrontend/Private/Analysis/MetasoundFrontendGraphAnalyzer.cpp
@@ -0,0 +1,94 @@
+// Copyright Epic Games, Inc. All Rights Reserved.
+#include "Analysis/MetasoundFrontendGraphAnalyzer.h"
+
+#include "Analysis/MetasoundFrontendAnalyzerFactory.h"
+#include "Analysis/MetasoundFrontendAnalyzerRegistry.h"
+#include "Analysis/MetasoundFrontendVertexAnalyzer.h"
+#include "MetasoundDataReference.h"
+#include "MetasoundPrimitives.h"
+
+
+namespace Metasound
+{
+	namespace Frontend
+	{
+		FGraphAnalyzer::FGraphAnalyzer(const FOperatorSettings& InSettings, uint64 InInstanceID, FNodeVertexDataMap&& InReferences)
+			: OperatorSettings(InSettings)
+			, InstanceID(InInstanceID)
+			, InternalDataReferences(MoveTemp(InReferences))
+		{
+			using namespace Frontend;
+			FSendAddress ActiveAnalyzerAddr(GetAnalyzerArraySendChannelName(InInstanceID), GetMetasoundDataTypeName<TArray<FString>>(), InInstanceID);
+			ActiveAnalyzerReceiver = FDataTransmissionCenter::Get().RegisterNewReceiver<TArray<FString>>(ActiveAnalyzerAddr, FReceiverInitParams { InSettings });
+		}
+
+		void FGraphAnalyzer::Execute()
+		{
+			// 1. Check if any message has been received that determines what analyzers are active.
+			TSet<FString> ReceiverKeys;
+			bool bUpdateReceivers = false;
+			if (ActiveAnalyzerReceiver.IsValid())
+			{
+				TReceiver<TArray<FString>>& AnalyzerReceiver = ActiveAnalyzerReceiver->GetAs<TReceiver<TArray<FString>>>();
+				if (AnalyzerReceiver.CanPop())
+				{
+					bUpdateReceivers = true;
+					TArray<FString> AnalyzerAddresses;
+					AnalyzerReceiver.Pop(AnalyzerAddresses);
+					ReceiverKeys.Append(AnalyzerAddresses);
+				}
+			}
+
+			// 2. Check if any message has been received that determines what analyzers are active and remove stale analyzers.
+			for (int32 i = Analyzers.Num() - 1; i >= 0; --i)
+			{
+				TUniquePtr<IVertexAnalyzer>& Analyzer = Analyzers[i];
+				check(Analyzer.IsValid());
+				if (bUpdateReceivers)
+				{
+					const FString AnalyzerKey = Analyzer->GetAnalyzerAddress().ToString();
+					if (!ReceiverKeys.Contains(AnalyzerKey))
+					{
+						constexpr bool bAllowShrinking = false;
+						Analyzers.RemoveAtSwap(i, 1, bAllowShrinking);
+						continue;
+					}
+
+					ReceiverKeys.Remove(AnalyzerKey);
+				}
+			}
+
+			// 3. If message received to update analyzers, create missing analyzers
+			if (bUpdateReceivers)
+			{
+				for (const FString& AnalyzerKey : ReceiverKeys)
+				{
+					FAnalyzerAddress AnalyzerAddress;
+					const bool bAnalyzerKeyValid = FAnalyzerAddress::ParseKey(AnalyzerKey, AnalyzerAddress);
+					if (ensureMsgf(bAnalyzerKeyValid, TEXT("Failed to create MetaSoundAnalyzer: AnalyzerKey '%s' invalid."), *AnalyzerKey))
+					{
+						const FDataReferenceCollection* Collection = InternalDataReferences.Find(AnalyzerAddress.NodeID);
+						// TODO: This currently fails for composed graphs.  Figure out how to differentiate addresses for composed graphs or refactor to only support parent graph analysis.
+// 						if (ensureMsgf(Collection != nullptr, TEXT("Failed to create MetaSoundAnalyzer: DataReferenceCollection for node analyzer at address '%s' not found."), *AnalyzerKey))
+						if (Collection)
+						{
+							const TUniquePtr<IVertexAnalyzerFactory>& Factory = IVertexAnalyzerRegistry::Get().FindAnalyzerFactory(AnalyzerAddress.AnalyzerName);
+							if (Factory.IsValid())
+							{
+								FCreateAnalyzerParams Params { AnalyzerAddress, OperatorSettings, *Collection };
+								TUniquePtr<IVertexAnalyzer> NewAnalyzer = Factory->CreateAnalyzer(Params);
+								Analyzers.Add(MoveTemp(NewAnalyzer));
+							}
+						}
+					}
+				}
+			}
+
+			// 4. Execute active analyzers post update.
+			for (TUniquePtr<Frontend::IVertexAnalyzer>& Analyzer : Analyzers)
+			{
+				Analyzer->Execute();
+			}
+		}
+	} // namespace Frontend
+} // namespace Metasound
diff --git a/Engine/Plugins/Runtime/Metasound/Source/MetasoundFrontend/Private/Analysis/MetasoundFrontendGraphAnalyzerView.cpp b/Engine/Plugins/Runtime/Metasound/Source/MetasoundFrontend/Private/Analysis/MetasoundFrontendGraphAnalyzerView.cpp
new file mode 100644
--- /dev/null
+++ b/Engine/Plugins/Runtime/Metasound/Source/MetasoundFrontend/Private/Analysis/MetasoundFrontendGraphAnalyzerView.cpp
@@ -0,0 +1,217 @@
+// Copyright Epic Games, Inc. All Rights Reserved.
+#include "Analysis/MetasoundFrontendGraphAnalyzerView.h"
+
+#include "Analysis/MetasoundFrontendVertexAnalyzer.h"
+#include "Analysis/MetasoundFrontendAnalyzerFactory.h"
+#include "Analysis/MetasoundFrontendAnalyzerRegistry.h"
+#include "Analysis/MetasoundFrontendGraphAnalyzer.h"
+#include "Math/UnrealMathUtility.h"
+#include "MetasoundAssetBase.h"
+#include "MetasoundFrontendController.h"
+#include "MetasoundOperatorSettings.h"
+#include "Templates/Function.h"
+
+
+namespace Metasound
+{
+	namespace Frontend
+	{
+		namespace GraphAnalyzerViewPrivate
+		{
+			using FIterateOutputFunctionRef = TFunctionRef<void(const FConstOutputHandle&, const FAnalyzerAddress& AnalyzerAddress)>;
+			void IterateOutputsSupportingAnalyzer(const FConstGraphHandle& GraphHandle, uint64 InInstanceID, FName InAnalyzerName, bool bInRequiresConnection, FIterateOutputFunctionRef InFunc)
+			{
+				const TUniquePtr<IVertexAnalyzerFactory>& Factory = IVertexAnalyzerRegistry::Get().FindAnalyzerFactory(InAnalyzerName);
+				if (!Factory.IsValid())
+				{
+					return;
+				}
+
+				const FName AnalyzerDataType = Factory->GetDataType();
+				GraphHandle->IterateConstNodes([&](FConstNodeHandle NodeHandle)
+				{
+					if (NodeHandle->GetClassMetadata().GetType() != EMetasoundFrontendClassType::Output)
+					{
+						NodeHandle->IterateConstOutputs([&](FConstOutputHandle OutputHandle)
+						{
+							const bool bInclude = !bInRequiresConnection || (OutputHandle->IsConnected() && bInRequiresConnection);
+							if (bInclude && OutputHandle->GetDataType() == AnalyzerDataType)
+							{
+								FAnalyzerAddress AnalyzerAddress;
+								AnalyzerAddress.DataType = AnalyzerDataType;
+								AnalyzerAddress.InstanceID = InInstanceID;
+								AnalyzerAddress.NodeID = NodeHandle->GetID();
+								AnalyzerAddress.OutputName = OutputHandle->GetName();
+								AnalyzerAddress.AnalyzerName = InAnalyzerName;
+								AnalyzerAddress.AnalyzerInstanceID = FGuid::NewGuid();
+
+								InFunc(OutputHandle, AnalyzerAddress);
+							}
+						});
+					}
+				});
+			}
+		} // namespace GraphAnalyzerViewPrivate
+
+		FMetasoundGraphAnalyzerView::FMetasoundGraphAnalyzerView(const FMetasoundAssetBase& InAssetBase, uint64 InInstanceID, FSampleRate InSampleRate)
+			: InstanceID(InInstanceID)
+			, MetaSoundAsset(&InAssetBase)
+			, OperatorSettings({ InSampleRate, GetDefaultBlockRate() })
+		{
+			// TODO: Mirrored from FMetaSoundParameterTransmitter. Fix here when this is refactored.
+			const float DelayTimeInSeconds = 0.1f;
+			FSendAddress ActiveAnalyzerAddr(FGraphAnalyzer::GetAnalyzerArraySendChannelName(InInstanceID), GetMetasoundDataTypeName<TArray<FString>>(), InstanceID);
+			FSenderInitParams SenderParams{ OperatorSettings, DelayTimeInSeconds };
+
+			ActiveAnalyzerSender = FDataTransmissionCenter::Get().RegisterNewSender<TArray<FString>>(ActiveAnalyzerAddr, SenderParams);
+		}
+
+		void FMetasoundGraphAnalyzerView::AddAnalyzerForAllSupportedOutputs(FName InAnalyzerName, bool bInRequiresConnection)
+		{
+			using namespace GraphAnalyzerViewPrivate;
+
+			FConstGraphHandle GraphHandle = GetMetaSoundAssetChecked().GetRootGraphHandle();
+			const TUniquePtr<IVertexAnalyzerFactory>& Factory = IVertexAnalyzerRegistry::Get().FindAnalyzerFactory(InAnalyzerName);
+			if (!Factory.IsValid())
+			{
+				return;
+			}
+
+			IterateOutputsSupportingAnalyzer(GraphHandle, InstanceID, InAnalyzerName, bInRequiresConnection, [this, FactoryPtr = Factory.Get(), &InAnalyzerName](FConstOutputHandle OutputHandle, const FAnalyzerAddress& AnalyzerAddress)
+			{
+				FString AnalyzerKey = AnalyzerAddress.ToString();
+				ActiveAnalyzerKeys.Add(MoveTemp(AnalyzerKey));
+
+				const TArray<FAnalyzerOutput>& AnalyzerOutputs = FactoryPtr->GetAnalyzerOutputs();
+				for (const FAnalyzerOutput& AnalyzerOutput : AnalyzerOutputs)
+				{
+					FAnalyzerAddress OutputReceiverAddress = AnalyzerAddress;
+					OutputReceiverAddress.AnalyzerMemberName = AnalyzerOutput.Name;
+					OutputReceiverAddress.DataType = AnalyzerOutput.DataType;
+
+					FMetasoundGraphAnalyzerOutputKey OutputKey { AnalyzerAddress.NodeID, AnalyzerAddress.OutputName };
+					FMetasoundAnalyzerView NewView(OperatorSettings, MoveTemp(OutputReceiverAddress));
+					AnalyzerViews.FindOrAdd(OutputKey).Add(MoveTemp(NewView));
+				}
+			});
+			ActiveAnalyzerSender->PushLiteral(ActiveAnalyzerKeys.Array());
+		}
+
+		TArray<FMetasoundAnalyzerView*> FMetasoundGraphAnalyzerView::GetAnalyzerViews(FName InAnalyzerName)
+		{
+			TArray<FMetasoundAnalyzerView*> ViewsToReturn;
+
+			for (TPair<FMetasoundGraphAnalyzerOutputKey, TArray<FMetasoundAnalyzerView>>& Pair : AnalyzerViews)
+			{
+				for (FMetasoundAnalyzerView& View : Pair.Value)
+				{
+					if (View.AnalyzerAddress.AnalyzerName == InAnalyzerName)
+					{
+						ViewsToReturn.Add(&View);
+					}
+				}
+			}
+
+			return ViewsToReturn;
+		}
+
+		TArray<const FMetasoundAnalyzerView*> FMetasoundGraphAnalyzerView::GetAnalyzerViews(FName InAnalyzerName) const
+		{
+			TArray<const FMetasoundAnalyzerView*> ViewsToReturn;
+
+			for (const TPair<FMetasoundGraphAnalyzerOutputKey, TArray<FMetasoundAnalyzerView>>& Pair : AnalyzerViews)
+			{
+				for (const FMetasoundAnalyzerView& View : Pair.Value)
+				{
+					if (View.AnalyzerAddress.AnalyzerName == InAnalyzerName)
+					{
+						ViewsToReturn.Add(&View);
+					}
+				}
+			}
+
+			return ViewsToReturn;
+		}
+
+		TArray<FMetasoundAnalyzerView*> FMetasoundGraphAnalyzerView::GetAnalyzerViewsForOutput(const FGuid& InNodeID, FName InOutputName, FName InAnalyzerName)
+		{
+			TArray<FMetasoundAnalyzerView*> ViewsToReturn;
+
+			const FMetasoundGraphAnalyzerOutputKey OutputKey { InNodeID, InOutputName };
+			if (TArray<FMetasoundAnalyzerView>* OutputAnalyzerViews = AnalyzerViews.Find(OutputKey))
+			{
+				for (FMetasoundAnalyzerView& View : *OutputAnalyzerViews)
+				{
+					if (View.AnalyzerAddress.AnalyzerName == InAnalyzerName)
+					{
+						ViewsToReturn.Add(&View);
+					}
+				}
+			}
+
+			return ViewsToReturn;
+		}
+
+		TArray<const FMetasoundAnalyzerView*> FMetasoundGraphAnalyzerView::GetAnalyzerViewsForOutput(const FGuid& InNodeID, FName InOutputName, FName InAnalyzerName) const
+		{
+			TArray<const FMetasoundAnalyzerView*> ViewsToReturn;
+
+			const FMetasoundGraphAnalyzerOutputKey OutputKey { InNodeID, InOutputName };
+			if (const TArray<FMetasoundAnalyzerView>* OutputAnalyzerViews = AnalyzerViews.Find(OutputKey))
+			{
+				for (const FMetasoundAnalyzerView& View : *OutputAnalyzerViews)
+				{
+					if (View.AnalyzerAddress.AnalyzerName == InAnalyzerName)
+					{
+						ViewsToReturn.Add(&View);
+					}
+				}
+			}
+
+			return ViewsToReturn;
+		}
+
+		void FMetasoundGraphAnalyzerView::RemoveAnalyzerForAllSupportedOutputs(FName InAnalyzerName)
+		{
+			TArray<FMetasoundGraphAnalyzerOutputKey> OutputsToRemove;
+			for (TPair<FMetasoundGraphAnalyzerOutputKey, TArray<FMetasoundAnalyzerView>>& AnalyzerPair : AnalyzerViews)
+			{
+				TArray<FMetasoundAnalyzerView>& Views = AnalyzerPair.Value;
+				Views.RemoveAll([InAnalyzerName](const FMetasoundAnalyzerView& View)
+				{
+					return View.AnalyzerAddress.AnalyzerName == InAnalyzerName;
+				});
+
+				if (Views.IsEmpty())
+				{
+					OutputsToRemove.Add(AnalyzerPair.Key);
+				}
+			}
+
+			for (const FMetasoundGraphAnalyzerOutputKey& Key : OutputsToRemove)
+			{
+				AnalyzerViews.Remove(Key);
+			}
+
+			TArray<FString> KeyArray = ActiveAnalyzerKeys.Array();
+			for (FString& AnalyzerKey : KeyArray)
+			{
+				FAnalyzerAddress AnalyzerAddress;
+				if (ensureAlways(FAnalyzerAddress::ParseKey(AnalyzerKey, AnalyzerAddress)))
+				{
+					if (AnalyzerAddress.AnalyzerName == InAnalyzerName)
+					{
+						ActiveAnalyzerKeys.Remove(AnalyzerKey);
+					}
+				}
+			}
+			ActiveAnalyzerSender->PushLiteral(ActiveAnalyzerKeys.Array());
+		}
+
+		const FMetasoundAssetBase& FMetasoundGraphAnalyzerView::GetMetaSoundAssetChecked() const
+		{
+			check(MetaSoundAsset);
+			return *MetaSoundAsset;
+		}
+	} // namespace Frontend
+} // namespace Metasound
diff --git a/Engine/Plugins/Runtime/Metasound/Source/MetasoundFrontend/Private/Analysis/MetasoundFrontendVertexAnalyzerEnvelopeFollower.cpp b/Engine/Plugins/Runtime/Metasound/Source/MetasoundFrontend/Private/Analysis/MetasoundFrontendVertexAnalyzerEnvelopeFollower.cpp
new file mode 100644
--- /dev/null
+++ b/Engine/Plugins/Runtime/Metasound/Source/MetasoundFrontend/Private/Analysis/MetasoundFrontendVertexAnalyzerEnvelopeFollower.cpp
@@ -0,0 +1,49 @@
+// Copyright Epic Games, Inc. All Rights Reserved.
+
+#include "Analysis/MetasoundFrontendVertexAnalyzerEnvelopeFollower.h"
+
+#include "MetasoundDataReference.h"
+#include "MetasoundPrimitives.h"
+#include "MetasoundRouter.h"
+
+
+namespace Metasound
+{
+	namespace Frontend
+	{
+		const FAnalyzerOutput FVertexAnalyzerEnvelopeFollower::FOutputs::Value = { "EnvelopeValue", GetMetasoundDataTypeName<float>() };
+
+		FVertexAnalyzerEnvelopeFollower::FVertexAnalyzerEnvelopeFollower(const FCreateAnalyzerParams& InParams)
+			: TVertexAnalyzer(InParams)
+		{
+			FAnalyzerAddress OutputAddress = InParams.AnalyzerAddress;
+			OutputAddress.DataType = FOutputs::Value.DataType;
+			OutputAddress.AnalyzerMemberName = FOutputs::Value.Name;
+			FSendAddress SendAddress = OutputAddress.ToSendAddress();
+
+			const FSenderInitParams InitParams { InParams.OperatorSettings, 0.0f };
+			Sender = FDataTransmissionCenter::Get().RegisterNewSender(MoveTemp(SendAddress), InitParams);
+
+			Audio::FEnvelopeFollowerInitParams Params;
+			Params.Mode = Audio::EPeakMode::RootMeanSquared;
+			Params.SampleRate = OperatorSettings.GetSampleRate();
+			Params.NumChannels = 1;
+			Params.AttackTimeMsec = 10;
+			Params.ReleaseTimeMsec = 10;
+			EnvelopeFollower.Init(Params);
+		}
+
+		void FVertexAnalyzerEnvelopeFollower::Execute()
+		{
+			const FAudioBuffer& InputData = GetAnalysisData();
+			EnvelopeFollower.ProcessAudio(InputData.GetData(), InputData.Num());
+
+			check(EnvelopeFollower.GetEnvelopeValues().Num() == 1);
+			const float EnvelopeValue = EnvelopeFollower.GetEnvelopeValues().Last();
+
+			FLiteral Literal;
+			Literal.Set(EnvelopeValue);
+			Sender->PushLiteral(Literal);
+		}
+	} // namespace Frontend
+} // namespace Metasound
diff --git a/Engine/Plugins/Runtime/Metasound/Source/MetasoundFrontend/Private/Analysis/MetasoundFrontendVertexAnalyzerTriggerDensity.cpp b/Engine/Plugins/Runtime/Metasound/Source/MetasoundFrontend/Private/Analysis/MetasoundFrontendVertexAnalyzerTriggerDensity.cpp
new file mode 100644
--- /dev/null
+++ b/Engine/Plugins/Runtime/Metasound/Source/MetasoundFrontend/Private/Analysis/MetasoundFrontendVertexAnalyzerTriggerDensity.cpp
@@ -0,0 +1,63 @@
+// Copyright Epic Games, Inc. All Rights Reserved.
+
+#include "Analysis/MetasoundFrontendVertexAnalyzerTriggerDensity.h"
+
+#include "MetasoundDataReference.h"
+#include "MetasoundPrimitives.h"
+#include "MetasoundRouter.h"
+#include "MetasoundTrigger.h"
+
+
+namespace Metasound
+{
+	namespace Frontend
+	{
+		const FAnalyzerOutput FVertexAnalyzerTriggerDensity::FOutputs::Value = { "TriggerDensity", GetMetasoundDataTypeName<float>() };
+
+		FVertexAnalyzerTriggerDensity::FVertexAnalyzerTriggerDensity(const FCreateAnalyzerParams& InParams)
+			: TVertexAnalyzer(InParams)
+			, NumFramesPerBlock(InParams.OperatorSettings.GetNumFramesPerBlock())
+		{
+			Audio::FEnvelopeFollowerInitParams Params;
+			Params.Mode = Audio::EPeakMode::Peak;
+			Params.SampleRate = OperatorSettings.GetSampleRate();
+			Params.NumChannels = 1;
+			Params.AttackTimeMsec = 0;
+			Params.ReleaseTimeMsec = 120;
+			EnvelopeFollower.Init(Params);
+
+			FAnalyzerAddress OutputAddress = InParams.AnalyzerAddress;
+			OutputAddress.DataType = FOutputs::Value.DataType;
+			OutputAddress.AnalyzerMemberName = FOutputs::Value.Name;
+			FSendAddress SendAddress = OutputAddress.ToSendAddress();
+
+			const FSenderInitParams InitParams { InParams.OperatorSettings, 0.0f };
+			Sender = FDataTransmissionCenter::Get().RegisterNewSender(MoveTemp(SendAddress), InitParams);
+		}
+
+		void FVertexAnalyzerTriggerDensity::Execute()
+		{
+			const FTrigger& Trigger = GetAnalysisData();
+
+			ScratchBuffer.Reset();
+			ScratchBuffer.AddZeroed(NumFramesPerBlock);
+			for (int32 i = 0; i < Trigger.Num(); ++i)
+			{
+				const int32 TriggerIndex = Trigger[i];
+				// Can trigger in the future so ignore those beyond the current block
+				if (TriggerIndex < ScratchBuffer.Num())
+				{
+					ScratchBuffer[TriggerIndex] = 1.0f;
+				}
+			}
+
+			EnvelopeFollower.ProcessAudio(ScratchBuffer.GetData(), ScratchBuffer.Num());
+
+			check(EnvelopeFollower.GetEnvelopeValues().Num() == 1);
+			const float EnvelopeValue = EnvelopeFollower.GetEnvelopeValues().Last();
+			FLiteral Literal;
+			Literal.Set(EnvelopeValue);
+			Sender->PushLiteral(Literal);
+		}
+	} // namespace Frontend
+} // namespace Metasound
diff --git a/Engine/Plugins/Runtime/Metasound/Source/MetasoundFrontend/Private/Analysis/MetasoundFrontendVertexAnalyzerValue.cpp b/Engine/Plugins/Runtime/Metasound/Source/MetasoundFrontend/Private/Analysis/MetasoundFrontendVertexAnalyzerValue.cpp
new file mode 100644
--- /dev/null
+++ b/Engine/Plugins/Runtime/Metasound/Source/MetasoundFrontend/Private/Analysis/MetasoundFrontendVertexAnalyzerValue.cpp
@@ -0,0 +1,21 @@
+// Copyright Epic Games, Inc. All Rights Reserved.
+
+#include "Analysis/MetasoundFrontendVertexAnalyzerValue.h"
+
+#include "Analysis/MetasoundFrontendAnalyzerAddress.h"
+#include "MetasoundDataReference.h"
+#include "MetasoundPrimitives.h"
+#include "MetasoundRouter.h"
+#include "MetasoundTrigger.h"
+
+
+namespace Metasound
+{
+	namespace Frontend
+	{
+		METASOUND_DEFINE_VALUE_ANALYZER(FVertexAnalyzerBool, bool)
+		METASOUND_DEFINE_VALUE_ANALYZER(FVertexAnalyzerFloat, float)
+		METASOUND_DEFINE_VALUE_ANALYZER(FVertexAnalyzerInt, int32)
+		METASOUND_DEFINE_VALUE_ANALYZER(FVertexAnalyzerString, FString)
+	} // namespace Frontend
+} // namespace Metasound
diff --git a/Engine/Plugins/Runtime/Metasound/Source/MetasoundFrontend/Private/MetasoundAssetBase.cpp b/Engine/Plugins/Runtime/Metasound/Source/MetasoundFrontend/Private/MetasoundAssetBase.cpp
--- a/Engine/Plugins/Runtime/Metasound/Source/MetasoundFrontend/Private/MetasoundAssetBase.cpp
+++ b/Engine/Plugins/Runtime/Metasound/Source/MetasoundFrontend/Private/MetasoundAssetBase.cpp
@@ -33,27 +33,44 @@
 
 namespace Metasound
 {
-	namespace AssetBasePrivate
+	namespace Frontend
 	{
-		void DepthFirstTraversal(const FMetasoundAssetBase& InInitAsset, TFunctionRef<TSet<const FMetasoundAssetBase*>(const FMetasoundAssetBase&)> InVisitFunction)
+		namespace AssetBasePrivate
 		{
-			// Non recursive depth first traversal.
-			TArray<const FMetasoundAssetBase*> Stack({ &InInitAsset });
-			TSet<const FMetasoundAssetBase*> Visited;
+			static float BlockRate = 100.f;
 
-			while (!Stack.IsEmpty())
+			void DepthFirstTraversal(const FMetasoundAssetBase& InInitAsset, TFunctionRef<TSet<const FMetasoundAssetBase*>(const FMetasoundAssetBase&)> InVisitFunction)
 			{
-				const FMetasoundAssetBase* CurrentNode = Stack.Pop();
-				if (!Visited.Contains(CurrentNode))
-				{
-					TArray<const FMetasoundAssetBase*> Children = InVisitFunction(*CurrentNode).Array();
-					Stack.Append(Children);
+				// Non recursive depth first traversal.
+				TArray<const FMetasoundAssetBase*> Stack({ &InInitAsset });
+				TSet<const FMetasoundAssetBase*> Visited;
 
-					Visited.Add(CurrentNode);
+				while (!Stack.IsEmpty())
+				{
+					const FMetasoundAssetBase* CurrentNode = Stack.Pop();
+					if (!Visited.Contains(CurrentNode))
+					{
+						TArray<const FMetasoundAssetBase*> Children = InVisitFunction(*CurrentNode).Array();
+						Stack.Append(Children);
+
+						Visited.Add(CurrentNode);
+					}
 				}
 			}
+		} // namespace AssetBasePrivate
+
+		FAutoConsoleVariableRef CVarMetaSoundBlockRate(
+			TEXT("au.MetaSound.BlockRate"),
+			AssetBasePrivate::BlockRate,
+			TEXT("Sets block rate (blocks per second) of MetaSounds.\n")
+			TEXT("Default: 100.0f, Min: 1.0f, Max: 1000.0f"),
+			ECVF_Default);
+
+		float GetDefaultBlockRate()
+		{
+			return FMath::Clamp(AssetBasePrivate::BlockRate, 1.0f, 1000.0f);
 		}
-	} // namespace AssetBasePrivate
+	} // namespace Frontend
 } // namespace Metasound
 
 const FString FMetasoundAssetBase::FileExtension(TEXT(".metasound"));
@@ -585,7 +602,7 @@ bool FMetasoundAssetBase::IsReferencedAsset(const FMetasoundAssetBase& InAsset)
 	using namespace Metasound::Frontend;
 
 	bool bIsReferenced = false;
-	Metasound::AssetBasePrivate::DepthFirstTraversal(*this, [&](const FMetasoundAssetBase& ChildAsset)
+	AssetBasePrivate::DepthFirstTraversal(*this, [&](const FMetasoundAssetBase& ChildAsset)
 	{
 		TSet<const FMetasoundAssetBase*> Children;
 		if (&ChildAsset == &InAsset)
@@ -616,7 +633,7 @@ bool FMetasoundAssetBase::AddingReferenceCausesLoop(const FSoftObjectPath& InRef
 
 	bool bCausesLoop = false;
 	const FMetasoundAssetBase* Parent = this;
-	Metasound::AssetBasePrivate::DepthFirstTraversal(*ReferenceAsset, [&](const FMetasoundAssetBase& ChildAsset)
+	AssetBasePrivate::DepthFirstTraversal(*ReferenceAsset, [&](const FMetasoundAssetBase& ChildAsset)
 	{
 		TSet<const FMetasoundAssetBase*> Children;
 		if (Parent == &ChildAsset)
diff --git a/Engine/Plugins/Runtime/Metasound/Source/MetasoundFrontend/Private/MetasoundFrontendDocument.cpp b/Engine/Plugins/Runtime/Metasound/Source/MetasoundFrontend/Private/MetasoundFrontendDocument.cpp
--- a/Engine/Plugins/Runtime/Metasound/Source/MetasoundFrontend/Private/MetasoundFrontendDocument.cpp
+++ b/Engine/Plugins/Runtime/Metasound/Source/MetasoundFrontend/Private/MetasoundFrontendDocument.cpp
@@ -20,6 +20,11 @@ namespace Metasound
 	{
 		namespace DisplayStyle
 		{
+			namespace EdgeAnimation
+			{
+				const FLinearColor DefaultColor = FLinearColor::Transparent;
+			}
+
 			namespace NodeLayout
 			{
 				const FVector2D DefaultOffsetX { 300.0f, 0.0f };
diff --git a/Engine/Plugins/Runtime/Metasound/Source/MetasoundFrontend/Private/MetasoundFrontendGraphController.cpp b/Engine/Plugins/Runtime/Metasound/Source/MetasoundFrontend/Private/MetasoundFrontendGraphController.cpp
--- a/Engine/Plugins/Runtime/Metasound/Source/MetasoundFrontend/Private/MetasoundFrontendGraphController.cpp
+++ b/Engine/Plugins/Runtime/Metasound/Source/MetasoundFrontend/Private/MetasoundFrontendGraphController.cpp
@@ -1508,7 +1508,7 @@ namespace Metasound
 			return INodeController::GetInvalidHandle();
 		}
 
-		TUniquePtr<IOperator> FGraphController::BuildOperator(const FOperatorSettings& InSettings, const FMetasoundEnvironment& InEnvironment, TArray<IOperatorBuilder::FBuildErrorPtr>& OutBuildErrors) const
+		TUniquePtr<IOperator> FGraphController::BuildOperator(const FOperatorSettings& InSettings, const FMetasoundEnvironment& InEnvironment, FBuildGraphResults& OutResults) const
 		{
 			if (const FMetasoundFrontendGraphClass* GraphClass = GraphClassPtr.Get())
 			{
@@ -1527,9 +1527,8 @@ namespace Metasound
 					return TUniquePtr<IOperator>(nullptr);
 				}
 
-				FOperatorBuilder OperatorBuilder(FOperatorBuilderSettings::GetDefaultSettings());
-				FBuildGraphParams BuildParams{*Graph, InSettings, FDataReferenceCollection{}, InEnvironment};
-				return OperatorBuilder.BuildGraphOperator(BuildParams, OutBuildErrors);
+				FBuildGraphParams BuildParams { *Graph, InSettings, FDataReferenceCollection { }, InEnvironment, FOperatorBuilderSettings::GetDefaultSettings() };
+				return FOperatorBuilder().BuildGraphOperator(BuildParams, OutResults);
 			}
 			else
 			{
diff --git a/Engine/Plugins/Runtime/Metasound/Source/MetasoundFrontend/Private/MetasoundFrontendGraphController.h b/Engine/Plugins/Runtime/Metasound/Source/MetasoundFrontend/Private/MetasoundFrontendGraphController.h
--- a/Engine/Plugins/Runtime/Metasound/Source/MetasoundFrontend/Private/MetasoundFrontendGraphController.h
+++ b/Engine/Plugins/Runtime/Metasound/Source/MetasoundFrontend/Private/MetasoundFrontendGraphController.h
@@ -185,7 +185,7 @@ namespace Metasound
 
 			FNodeHandle CreateEmptySubgraph(const FMetasoundFrontendClassMetadata& InInfo) override;
 
-			TUniquePtr<IOperator> BuildOperator(const FOperatorSettings& InSettings, const FMetasoundEnvironment& InEnvironment, TArray<IOperatorBuilder::FBuildErrorPtr>& OutBuildErrors) const override;
+			TUniquePtr<IOperator> BuildOperator(const FOperatorSettings& InSettings, const FMetasoundEnvironment& InEnvironment, FBuildGraphResults& OutBuildErrors) const override;
 
 			FDocumentHandle GetOwningDocument() override;
 			FConstDocumentHandle GetOwningDocument() const override;
diff --git a/Engine/Plugins/Runtime/Metasound/Source/MetasoundFrontend/Private/MetasoundFrontendInputController.cpp b/Engine/Plugins/Runtime/Metasound/Source/MetasoundFrontend/Private/MetasoundFrontendInputController.cpp
--- a/Engine/Plugins/Runtime/Metasound/Source/MetasoundFrontend/Private/MetasoundFrontendInputController.cpp
+++ b/Engine/Plugins/Runtime/Metasound/Source/MetasoundFrontend/Private/MetasoundFrontendInputController.cpp
@@ -303,7 +303,16 @@ namespace Metasound
 					return (Edge.FromNodeID == FromNodeID) && (Edge.FromVertexID == FromVertexID) && (Edge.ToNodeID == ToNodeID) && (Edge.ToVertexID == ToVertexID);
 				};
 
-				int32 NumRemoved = Graph->Edges.RemoveAllSwap(IsMatchingEdge);
+				const int32 NumRemoved = Graph->Edges.RemoveAllSwap(IsMatchingEdge);
+
+#if WITH_EDITOR
+				auto IsMatchingStyle = [&](const FMetasoundFrontendEdgeStyle& EdgeStyle)
+				{
+					return EdgeStyle.NodeID == FromNodeID && InController.GetName() == EdgeStyle.OutputName;
+				};
+				Graph->Style.EdgeStyles.RemoveAllSwap(IsMatchingStyle);
+#endif // WITH_EDITOR
+
 				return NumRemoved > 0;
 			}
 
diff --git a/Engine/Plugins/Runtime/Metasound/Source/MetasoundFrontend/Private/MetasoundFrontendInvalidController.h b/Engine/Plugins/Runtime/Metasound/Source/MetasoundFrontend/Private/MetasoundFrontendInvalidController.h
--- a/Engine/Plugins/Runtime/Metasound/Source/MetasoundFrontend/Private/MetasoundFrontendInvalidController.h
+++ b/Engine/Plugins/Runtime/Metasound/Source/MetasoundFrontend/Private/MetasoundFrontendInvalidController.h
@@ -419,10 +419,7 @@ namespace Metasound
 
 			virtual FNodeHandle CreateEmptySubgraph(const FMetasoundFrontendClassMetadata& InInfo) override { return INodeController::GetInvalidHandle(); }
 
-			virtual TUniquePtr<IOperator> BuildOperator(const FOperatorSettings& InSettings, const FMetasoundEnvironment& InEnvironment, TArray<IOperatorBuilder::FBuildErrorPtr>& OutBuildErrors) const override
-			{
-				return TUniquePtr<IOperator>(nullptr);
-			}
+			virtual TUniquePtr<IOperator> BuildOperator(const FOperatorSettings& InSettings, const FMetasoundEnvironment& InEnvironment, FBuildGraphResults& OutBuildErrors) const override { return { }; }
 
 			virtual FDocumentHandle GetOwningDocument() override;
 			virtual FConstDocumentHandle GetOwningDocument() const override;
diff --git a/Engine/Plugins/Runtime/Metasound/Source/MetasoundFrontend/Public/Analysis/MetasoundFrontendAnalyzerAddress.h b/Engine/Plugins/Runtime/Metasound/Source/MetasoundFrontend/Public/Analysis/MetasoundFrontendAnalyzerAddress.h
new file mode 100644
--- /dev/null
+++ b/Engine/Plugins/Runtime/Metasound/Source/MetasoundFrontend/Public/Analysis/MetasoundFrontendAnalyzerAddress.h
@@ -0,0 +1,56 @@
+// Copyright Epic Games, Inc. All Rights Reserved.
+#pragma once
+
+#include "Containers/UnrealString.h"
+#include "HAL/Platform.h"
+#include "MetasoundRouter.h"
+#include "MetasoundVertex.h"
+#include "Misc/Guid.h"
+
+
+namespace Metasound
+{
+	namespace Frontend
+	{
+		// String serializable (as key) channel of analyzer or its internal members
+		// that can be written to or read from using the Transmission System.
+		struct METASOUNDFRONTEND_API FAnalyzerAddress
+		{
+			static const FString PathSeparator;
+
+			// Active Instance ID to monitor
+			uint64 InstanceID = TNumericLimits<uint64>::Max();
+
+			// ID of Node being monitored
+			FGuid NodeID;
+
+			// Name of output to monitor (not to be confused with the Analyzer's members,
+			// which are specific to the analyzer instance being addressed)
+			FVertexName OutputName;
+
+			// DataType of the given channel
+			FName DataType;
+
+			// Name of Analyzer
+			FName AnalyzerName;
+
+			// Instance ID of analyzer (allowing for multiple analyzer of the same type to be
+			// addressed at the same output).
+			FGuid AnalyzerInstanceID;
+
+			// Optional name used to specify a channel for a given analyzer's inputs/outputs.
+			// If not provided (i.e. 'none'), single input & output are assumed to share
+			// the same name. Useful if the analyzer requires outputting multiple analysis values.
+			// Can potentially be used as an input as well to modify analyzer settings.
+			FName AnalyzerMemberName;
+
+			// Converts AnalyzerAddress to String representation using the PathSeparator
+			FString ToString() const;
+
+			// Converts AnalyzerAddress to SendAddress
+			FSendAddress ToSendAddress() const;
+
+			static bool ParseKey(const FString& InAnalyzerKey, FAnalyzerAddress& OutAnalyzer);
+		};
+	} // namespace Frontend
+} // namespace Metasound
diff --git a/Engine/Plugins/Runtime/Metasound/Source/MetasoundFrontend/Public/Analysis/MetasoundFrontendAnalyzerFactory.h b/Engine/Plugins/Runtime/Metasound/Source/MetasoundFrontend/Public/Analysis/MetasoundFrontendAnalyzerFactory.h
new file mode 100644
--- /dev/null
+++ b/Engine/Plugins/Runtime/Metasound/Source/MetasoundFrontend/Public/Analysis/MetasoundFrontendAnalyzerFactory.h
@@ -0,0 +1,37 @@
+// Copyright Epic Games, Inc. All Rights Reserved.
+#pragma once
+
+#include "Analysis/MetasoundFrontendVertexAnalyzer.h"
+#include "Templates/UniquePtr.h"
+
+
+namespace Metasound
+{
+	namespace Frontend
+	{
+		class METASOUNDFRONTEND_API IVertexAnalyzerFactory
+		{
+		public:
+			virtual ~IVertexAnalyzerFactory() = default;
+
+			virtual const TArray<FAnalyzerOutput>& GetAnalyzerOutputs() const = 0;
+			virtual FName GetDataType() const = 0;
+			virtual TUniquePtr<IVertexAnalyzer> CreateAnalyzer(const FCreateAnalyzerParams& InParams) = 0;
+		};
+
+		template<typename TAnalyzerClass>
+		class TVertexAnalyzerFactory : public IVertexAnalyzerFactory
+		{
+		public:
+			static const FName& GetAnalyzerName() { return TAnalyzerClass::GetAnalyzerName(); }
+
+			virtual ~TVertexAnalyzerFactory() = default;
+
+			virtual FName GetDataType() const override { return TAnalyzerClass::GetDataType(); }
+			virtual TUniquePtr<IVertexAnalyzer> CreateAnalyzer(const FCreateAnalyzerParams& InParams) override
+			{
+				return TUniquePtr<IVertexAnalyzer>(new TAnalyzerClass(InParams));
+			}
+		};
+	} // namespace Frontend
+} // namespace Metasound
diff --git a/Engine/Plugins/Runtime/Metasound/Source/MetasoundFrontend/Public/Analysis/MetasoundFrontendAnalyzerRegistry.h b/Engine/Plugins/Runtime/Metasound/Source/MetasoundFrontend/Public/Analysis/MetasoundFrontendAnalyzerRegistry.h
new file mode 100644
--- /dev/null
+++ b/Engine/Plugins/Runtime/Metasound/Source/MetasoundFrontend/Public/Analysis/MetasoundFrontendAnalyzerRegistry.h
@@ -0,0 +1,26 @@
+// Copyright Epic Games, Inc. All Rights Reserved.
+#pragma once
+
+#include "UObject/NameTypes.h"
+#include "Templates/UniquePtr.h"
+
+namespace Metasound
+{
+	// Forward Declarations
+	class FDataReferenceCollection;
+
+	namespace Frontend
+	{
+		// Forward Declarations
+		class IVertexAnalyzerFactory;
+
+		class METASOUNDFRONTEND_API IVertexAnalyzerRegistry
+		{
+		public:
+			static IVertexAnalyzerRegistry& Get();
+
+			virtual const TUniquePtr<IVertexAnalyzerFactory>& FindAnalyzerFactory(FName InAnalyzerName) const = 0;
+			virtual void RegisterAnalyzerFactories() = 0;
+		};
+	} // namespace Frontend
+} // namespace Metasound
diff --git a/Engine/Plugins/Runtime/Metasound/Source/MetasoundFrontend/Public/Analysis/MetasoundFrontendAnalyzerView.h b/Engine/Plugins/Runtime/Metasound/Source/MetasoundFrontend/Public/Analysis/MetasoundFrontendAnalyzerView.h
new file mode 100644
--- /dev/null
+++ b/Engine/Plugins/Runtime/Metasound/Source/MetasoundFrontend/Public/Analysis/MetasoundFrontendAnalyzerView.h
@@ -0,0 +1,36 @@
+// Copyright Epic Games, Inc. All Rights Reserved.
+#pragma once
+
+#include "Analysis/MetasoundFrontendVertexAnalyzer.h"
+#include "Containers/Map.h"
+#include "HAL/Platform.h"
+#include "MetasoundAssetBase.h"
+#include "MetasoundRouter.h"
+#include "Templates/UniquePtr.h"
+
+
+// Forward Declarations
+class FMetasoundAssetBase;
+class UAudioComponent;
+
+namespace Metasound
+{
+	namespace Frontend
+	{
+		// Pairs an IReceiver with a given AnalyzerAddress, which enables
+		// watching a particular analyzer result on any given thread.
+		class METASOUNDFRONTEND_API FMetasoundAnalyzerView
+		{
+			TUniquePtr<IReceiver> Receiver;
+
+		public:
+			const FAnalyzerAddress AnalyzerAddress = { };
+
+			FMetasoundAnalyzerView() = default;
+			FMetasoundAnalyzerView(const FOperatorSettings& InOperatorSettings, FAnalyzerAddress&& InAnalyzerAddress);
+
+			IReceiver& GetReceiverChecked();
+			const IReceiver& GetReceiverChecked() const;
+		};
+	}
+}
diff --git a/Engine/Plugins/Runtime/Metasound/Source/MetasoundFrontend/Public/Analysis/MetasoundFrontendGraphAnalyzer.h b/Engine/Plugins/Runtime/Metasound/Source/MetasoundFrontend/Public/Analysis/MetasoundFrontendGraphAnalyzer.h
new file mode 100644
--- /dev/null
+++ b/Engine/Plugins/Runtime/Metasound/Source/MetasoundFrontend/Public/Analysis/MetasoundFrontendGraphAnalyzer.h
@@ -0,0 +1,42 @@
+// Copyright Epic Games, Inc. All Rights Reserved.
+#pragma once
+
+#include "Analysis/MetasoundFrontendVertexAnalyzer.h"
+#include "MetasoundBuilderInterface.h"
+#include "MetasoundOperatorSettings.h"
+#include "MetasoundRouter.h"
+#include "Templates/UniquePtr.h"
+
+
+namespace Metasound
+{
+	namespace Frontend
+	{
+		// Handles intrinsic analysis operations within a given graph
+		// should the graph's operator be enabled for analysis.
+		class METASOUNDFRONTEND_API FGraphAnalyzer
+		{
+			const FOperatorSettings OperatorSettings;
+			const uint64 InstanceID = INDEX_NONE;
+
+			TUniquePtr<IReceiver> ActiveAnalyzerReceiver;
+			TArray<TUniquePtr<Frontend::IVertexAnalyzer>> Analyzers;
+			FNodeVertexDataMap InternalDataReferences;
+
+		public:
+			FGraphAnalyzer(const FOperatorSettings& InSettings, uint64 InInstanceID, FNodeVertexDataMap&& InGraphReferences);
+			~FGraphAnalyzer() = default;
+
+			// Creates a send channel name unique for the given sound instance used to send array of analyzer
+			// addresses to the profiler of what analyzers are expected to be active for a given instance.
+			static const FName GetAnalyzerArraySendChannelName(uint64 InInstanceID)
+			{
+				const FString ChannelName = FString::Printf(TEXT("AnalyzerArray%s%lld"), *FAnalyzerAddress::PathSeparator, InInstanceID);
+				return *ChannelName;
+			}
+
+			// Execute analysis for the current block
+			void Execute();
+		};
+	} // namespace Frontend
+} // namespace Metasound
diff --git a/Engine/Plugins/Runtime/Metasound/Source/MetasoundFrontend/Public/Analysis/MetasoundFrontendGraphAnalyzerView.h b/Engine/Plugins/Runtime/Metasound/Source/MetasoundFrontend/Public/Analysis/MetasoundFrontendGraphAnalyzerView.h
new file mode 100644
--- /dev/null
+++ b/Engine/Plugins/Runtime/Metasound/Source/MetasoundFrontend/Public/Analysis/MetasoundFrontendGraphAnalyzerView.h
@@ -0,0 +1,51 @@
+// Copyright Epic Games, Inc. All Rights Reserved.
+#pragma once
+
+#include "Analysis/MetasoundFrontendAnalyzerView.h"
+#include "Analysis/MetasoundFrontendGraphAnalyzer.h"
+#include "Containers/Map.h"
+#include "HAL/Platform.h"
+#include "MetasoundAssetBase.h"
+#include "MetasoundRouter.h"
+#include "Templates/UniquePtr.h"
+
+
+// Forward Declarations
+class FMetasoundAssetBase;
+class UAudioComponent;
+
+namespace Metasound
+{
+	namespace Frontend
+	{
+		class METASOUNDFRONTEND_API FMetasoundGraphAnalyzerView
+		{
+			// Sender in charge of supplying expected vertex analyzers currently being analyzed.
+			TUniquePtr<ISender> ActiveAnalyzerSender;
+
+			// Set of active keys used for what senders should be active on the instance graph
+			TSet<FString> ActiveAnalyzerKeys;
+
+			const FMetasoundAssetBase& GetMetaSoundAssetChecked() const;
+
+			uint64 InstanceID = INDEX_NONE;
+			const FMetasoundAssetBase* MetaSoundAsset = nullptr;
+
+			FOperatorSettings OperatorSettings = { 48000, 100 };
+
+			using FMetasoundGraphAnalyzerOutputKey = TTuple<FGuid /* NodeID */, FName /* OutputName */>;
+			TMap<FMetasoundGraphAnalyzerOutputKey, TArray<FMetasoundAnalyzerView>> AnalyzerViews;
+
+		public:
+			FMetasoundGraphAnalyzerView() = default;
+			FMetasoundGraphAnalyzerView(const FMetasoundAssetBase& InAssetBase, uint64 InInstanceID, FSampleRate InSampleRate);
+
+			void AddAnalyzerForAllSupportedOutputs(FName InAnalyzerName, bool bInRequiresConnection = true);
+			void RemoveAnalyzerForAllSupportedOutputs(FName InAnalyzerName);
+			TArray<FMetasoundAnalyzerView*> GetAnalyzerViews(FName InAnalyzerName);
+			TArray<const FMetasoundAnalyzerView*> GetAnalyzerViews(FName InAnalyzerName) const;
+			TArray<FMetasoundAnalyzerView*> GetAnalyzerViewsForOutput(const FGuid& InNodeID, FName InOutputName, FName InAnalyzerName);
+			TArray<const FMetasoundAnalyzerView*> GetAnalyzerViewsForOutput(const FGuid& InNodeID, FName InOutputName, FName InAnalyzerName) const;
+		};
+	}
+}
diff --git a/Engine/Plugins/Runtime/Metasound/Source/MetasoundFrontend/Public/Analysis/MetasoundFrontendVertexAnalyzer.h b/Engine/Plugins/Runtime/Metasound/Source/MetasoundFrontend/Public/Analysis/MetasoundFrontendVertexAnalyzer.h
new file mode 100644
--- /dev/null
+++ b/Engine/Plugins/Runtime/Metasound/Source/MetasoundFrontend/Public/Analysis/MetasoundFrontendVertexAnalyzer.h
@@ -0,0 +1,72 @@
+// Copyright Epic Games, Inc. All Rights Reserved.
+#pragma once
+
+#include "Analysis/MetasoundFrontendAnalyzerAddress.h"
+#include "HAL/Platform.h"
+#include "MetasoundDataReference.h"
+#include "MetasoundDataReferenceCollection.h"
+#include "Misc/Guid.h"
+
+
+namespace Metasound
+{
+	namespace Frontend
+	{
+		struct METASOUNDFRONTEND_API FAnalyzerOutput
+		{
+			FName Name;
+			FName DataType;
+		};
+
+		struct METASOUNDFRONTEND_API FCreateAnalyzerParams
+		{
+			const FAnalyzerAddress& AnalyzerAddress;
+			const FOperatorSettings& OperatorSettings;
+			const FDataReferenceCollection& Collection;
+		};
+
+		class METASOUNDFRONTEND_API IVertexAnalyzer
+		{
+		public:
+			virtual ~IVertexAnalyzer() = default;
+
+			virtual const FAnalyzerAddress& GetAnalyzerAddress() const = 0;
+			virtual void Execute() = 0;
+		};
+
+		template<typename TChildClass, typename TDataType>
+		class TVertexAnalyzer : public IVertexAnalyzer
+		{
+		protected:
+			FAnalyzerAddress Address;
+			FOperatorSettings OperatorSettings { 48000, 100 };
+			TDataReadReference<TDataType> Reference;
+
+		public:
+			static FName GetDataType()
+			{
+				return GetMetasoundDataTypeName<TDataType>();
+			}
+
+			TVertexAnalyzer(const FCreateAnalyzerParams& InParams)
+				: Address(InParams.AnalyzerAddress)
+				, OperatorSettings(InParams.OperatorSettings)
+				, Reference(InParams.Collection.GetDataReadReference<TDataType>(InParams.AnalyzerAddress.OutputName))
+			{
+				check(Address.AnalyzerName == TChildClass::GetAnalyzerName());
+			}
+
+			virtual ~TVertexAnalyzer() = default;
+
+			const TDataType& GetAnalysisData() const
+			{
+				return *Reference;
+			}
+
+			virtual const FAnalyzerAddress& GetAnalyzerAddress() const override
+			{
+				return Address;
+			}
+		};
+	} // namespace Frontend
+} // namespace Metasound
diff --git a/Engine/Plugins/Runtime/Metasound/Source/MetasoundFrontend/Public/Analysis/MetasoundFrontendVertexAnalyzerEnvelopeFollower.h b/Engine/Plugins/Runtime/Metasound/Source/MetasoundFrontend/Public/Analysis/MetasoundFrontendVertexAnalyzerEnvelopeFollower.h
new file mode 100644
--- /dev/null
+++ b/Engine/Plugins/Runtime/Metasound/Source/MetasoundFrontend/Public/Analysis/MetasoundFrontendVertexAnalyzerEnvelopeFollower.h
@@ -0,0 +1,49 @@
+// Copyright Epic Games, Inc. All Rights Reserved.
+#pragma once
+
+#include "Analysis/MetasoundFrontendAnalyzerFactory.h"
+#include "Analysis/MetasoundFrontendVertexAnalyzer.h"
+#include "Containers/Array.h"
+#include "DSP/EnvelopeFollower.h"
+#include "MetasoundDataReferenceCollection.h"
+#include "MetasoundRouter.h"
+
+
+namespace Metasound
+{
+	namespace Frontend
+	{
+		class METASOUNDFRONTEND_API FVertexAnalyzerEnvelopeFollower : public TVertexAnalyzer<FVertexAnalyzerEnvelopeFollower, FAudioBuffer>
+		{
+			TUniquePtr<ISender> Sender;
+			Audio::FEnvelopeFollower EnvelopeFollower;
+
+		public:
+			static const FName& GetAnalyzerName()
+			{
+				static const FName AnalyzerName = "UE.Audio.EnvelopeFollower";
+				return AnalyzerName;
+			}
+
+			struct METASOUNDFRONTEND_API FOutputs
+			{
+				static const FAnalyzerOutput Value;
+			};
+
+			class METASOUNDFRONTEND_API FFactory : public TVertexAnalyzerFactory<FVertexAnalyzerEnvelopeFollower>
+			{
+			public:
+				virtual const TArray<FAnalyzerOutput>& GetAnalyzerOutputs() const override
+				{
+					static const TArray<FAnalyzerOutput> Outputs { FOutputs::Value };
+					return Outputs;
+				}
+			};
+
+			FVertexAnalyzerEnvelopeFollower(const FCreateAnalyzerParams& InParams);
+			virtual ~FVertexAnalyzerEnvelopeFollower() = default;
+
+			virtual void Execute() override;
+		};
+	} // namespace Frontend
+} // namespace Metasound
diff --git a/Engine/Plugins/Runtime/Metasound/Source/MetasoundFrontend/Public/Analysis/MetasoundFrontendVertexAnalyzerTriggerDensity.h b/Engine/Plugins/Runtime/Metasound/Source/MetasoundFrontend/Public/Analysis/MetasoundFrontendVertexAnalyzerTriggerDensity.h
new file mode 100644
--- /dev/null
+++ b/Engine/Plugins/Runtime/Metasound/Source/MetasoundFrontend/Public/Analysis/MetasoundFrontendVertexAnalyzerTriggerDensity.h
@@ -0,0 +1,52 @@
+// Copyright Epic Games, Inc. All Rights Reserved.
+#pragma once
+
+#include "Analysis/MetasoundFrontendAnalyzerFactory.h"
+#include "Containers/Array.h"
+#include "DSP/BufferVectorOperations.h"
+#include "DSP/EnvelopeFollower.h"
+#include "MetasoundTrigger.h"
+
+
+namespace Metasound
+{
+	namespace Frontend
+	{
+		class METASOUNDFRONTEND_API FVertexAnalyzerTriggerDensity : public TVertexAnalyzer<FVertexAnalyzerTriggerDensity, FTrigger>
+		{
+			TUniquePtr<ISender> Sender;
+			int32 NumFramesPerBlock = 0;
+
+			Audio::FEnvelopeFollower EnvelopeFollower;
+			Audio::FAlignedFloatBuffer ScratchBuffer;
+
+		public:
+			static const FName& GetAnalyzerName()
+			{
+				static const FName AnalyzerName = "UE.Trigger.Density";
+				return AnalyzerName;
+			}
+
+			struct METASOUNDFRONTEND_API FOutputs
+			{
+				static const FAnalyzerOutput Value;
+
+			};
+
+			class METASOUNDFRONTEND_API FFactory : public TVertexAnalyzerFactory<FVertexAnalyzerTriggerDensity>
+			{
+			public:
+				virtual const TArray<FAnalyzerOutput>& GetAnalyzerOutputs() const override
+				{
+					static const TArray<FAnalyzerOutput> Outputs { FOutputs::Value };
+					return Outputs;
+				}
+			};
+
+			FVertexAnalyzerTriggerDensity(const FCreateAnalyzerParams& InParams);
+			virtual ~FVertexAnalyzerTriggerDensity() = default;
+
+			virtual void Execute() override;
+		};
+	} // namespace Frontend
+} // namespace Metasound
diff --git a/Engine/Plugins/Runtime/Metasound/Source/MetasoundFrontend/Public/Analysis/MetasoundFrontendVertexAnalyzerValue.h b/Engine/Plugins/Runtime/Metasound/Source/MetasoundFrontend/Public/Analysis/MetasoundFrontendVertexAnalyzerValue.h
new file mode 100644
--- /dev/null
+++ b/Engine/Plugins/Runtime/Metasound/Source/MetasoundFrontend/Public/Analysis/MetasoundFrontendVertexAnalyzerValue.h
@@ -0,0 +1,83 @@
+// Copyright Epic Games, Inc. All Rights Reserved.
+#pragma once
+
+#include "Analysis/MetasoundFrontendAnalyzerAddress.h"
+#include "Analysis/MetasoundFrontendAnalyzerFactory.h"
+#include "Analysis/MetasoundFrontendVertexAnalyzer.h"
+#include "Containers/Array.h"
+#include "MetasoundDataReference.h"
+#include "MetasoundPrimitives.h"
+#include "MetasoundRouter.h"
+#include "MetasoundTrigger.h"
+
+
+namespace Metasound
+{
+	namespace Frontend
+	{
+		template <typename TChildClass, typename TDataType>
+		class TVertexAnalyzerValue : public TVertexAnalyzer<TChildClass, TDataType>
+		{
+			TUniquePtr<ISender> Sender;
+			TDataType LastValue;
+
+		public:
+			class FFactory : public TVertexAnalyzerFactory<TChildClass>
+			{
+			public:
+				virtual const TArray<FAnalyzerOutput>& GetAnalyzerOutputs() const override
+				{
+					static const TArray<FAnalyzerOutput> Outputs { TChildClass::FOutputs::Value };
+					return Outputs;
+				}
+			};
+
+			TVertexAnalyzerValue(const FCreateAnalyzerParams& InParams)
+ 				: TVertexAnalyzer<TChildClass, TDataType>(InParams)
+			{
+				FAnalyzerAddress OutputAddress = InParams.AnalyzerAddress;
+				OutputAddress.DataType = TChildClass::FOutputs::Value.DataType;
+				OutputAddress.AnalyzerMemberName = TChildClass::FOutputs::Value.Name;
+				FSendAddress SendAddress = OutputAddress.ToSendAddress();
+
+				const FSenderInitParams InitParams{ InParams.OperatorSettings, 0.0f };
+				Sender = FDataTransmissionCenter::Get().RegisterNewSender(MoveTemp(SendAddress), InitParams);
+			}
+
+			virtual ~TVertexAnalyzerValue() = default;
+
+			virtual void Execute() override
+			{
+				const TDataType& Value = TVertexAnalyzer<TChildClass, TDataType>::GetAnalysisData();
+				if (LastValue != Value)
+				{
+					FLiteral Literal;
+					Literal.Set(Value);
+					Sender->PushLiteral(Literal);
+				}
+			}
+		};
+
+		#define METASOUND_DECLARE_VALUE_ANALYZER(ClassType, AnalyzerFName, DataType) \
+			class METASOUNDFRONTEND_API ClassType : public TVertexAnalyzerValue<ClassType, DataType> \
+			{ \
+			public: \
+				static const FName& GetAnalyzerName() \
+				{ \
+					static const FName AnalyzerName = AnalyzerFName; return AnalyzerName; \
+				} \
+				struct METASOUNDFRONTEND_API FOutputs { static const FAnalyzerOutput Value; }; \
+				ClassType(const FCreateAnalyzerParams& InParams) \
+					: TVertexAnalyzerValue(InParams) { } \
+				virtual ~ClassType() = default; \
+			};
+
+		#define METASOUND_DEFINE_VALUE_ANALYZER(ClassType, DataType) \
+			const FAnalyzerOutput ClassType::FOutputs::Value = { "Value", GetMetasoundDataTypeName<DataType>() };
+
+		METASOUND_DECLARE_VALUE_ANALYZER(FVertexAnalyzerBool, "UE.Bool", bool)
+		METASOUND_DECLARE_VALUE_ANALYZER(FVertexAnalyzerFloat, "UE.Float", float)
+		METASOUND_DECLARE_VALUE_ANALYZER(FVertexAnalyzerInt, "UE.Int32", int32)
+		METASOUND_DECLARE_VALUE_ANALYZER(FVertexAnalyzerString, "UE.String", FString)
+	} // namespace Frontend
+} // namespace Metasound
diff --git a/Engine/Plugins/Runtime/Metasound/Source/MetasoundFrontend/Public/MetasoundAssetBase.h b/Engine/Plugins/Runtime/Metasound/Source/MetasoundFrontend/Public/MetasoundAssetBase.h
--- a/Engine/Plugins/Runtime/Metasound/Source/MetasoundFrontend/Public/MetasoundAssetBase.h
+++ b/Engine/Plugins/Runtime/Metasound/Source/MetasoundFrontend/Public/MetasoundAssetBase.h
@@ -19,8 +19,12 @@ class UEdGraph;
 
 namespace Metasound
 {
-	namespace Frontend {
+	namespace Frontend
+	{
+		// Forward Declarations
 		class IInterfaceRegistryEntry;
+
+		METASOUNDFRONTEND_API float GetDefaultBlockRate();
 	} // namespace Frontend
 } // namespace Metasound
 
diff --git a/Engine/Plugins/Runtime/Metasound/Source/MetasoundFrontend/Public/MetasoundAutoConverterNode.h b/Engine/Plugins/Runtime/Metasound/Source/MetasoundFrontend/Public/MetasoundAutoConverterNode.h
--- a/Engine/Plugins/Runtime/Metasound/Source/MetasoundFrontend/Public/MetasoundAutoConverterNode.h
+++ b/Engine/Plugins/Runtime/Metasound/Source/MetasoundFrontend/Public/MetasoundAutoConverterNode.h
@@ -212,7 +212,7 @@ namespace Metasound
 			public:
 				FCoverterOperatorFactory() = default;
 
-				virtual TUniquePtr<IOperator> CreateOperator(const FCreateOperatorParams& InParams, FBuildErrorArray& OutErrors) override
+				virtual TUniquePtr<IOperator> CreateOperator(const FCreateOperatorParams& InParams, FBuildGraphResults& OutResults) override
 				{
 					TDataWriteReference<ToDataType> WriteReference = TDataWriteReferenceFactory<ToDataType>::CreateAny(InParams.OperatorSettings);
 
@@ -234,7 +234,7 @@ namespace Metasound
 					if (ensure(InParams.Node.GetVertexInterface().ContainsInputVertex(InputName)))
 					{
 						FInputDataDestination Dest(InParams.Node, InParams.Node.GetVertexInterface().GetInputVertex(GetInputName()));
-						AddBuildError<FMissingInputDataReferenceError>(OutErrors, Dest);
+						AddBuildError<FMissingInputDataReferenceError>(OutResults.Errors, Dest);
 					}
 
 					return TUniquePtr<IOperator>(nullptr);
diff --git a/Engine/Plugins/Runtime/Metasound/Source/MetasoundFrontend/Public/MetasoundFrontendController.h b/Engine/Plugins/Runtime/Metasound/Source/MetasoundFrontend/Public/MetasoundFrontendController.h
--- a/Engine/Plugins/Runtime/Metasound/Source/MetasoundFrontend/Public/MetasoundFrontendController.h
+++ b/Engine/Plugins/Runtime/Metasound/Source/MetasoundFrontend/Public/MetasoundFrontendController.h
@@ -958,11 +958,11 @@ namespace Metasound
 			 *
 			 * @param InSettings - Settings to use when creating operators.
 			 * @param InEnvironment - Environment variables available during creation.
-			 * @param OutBuildErrors - An array to populate with errors encountered during the build process.
+			 * @param OutResults - Results pertaining to operator build process & resulting IOperator instance.
 			 *
-			 * @return On success, a valid pointer to a Metasound operator. An invalid pointer on failure. 
+			 * @return On success, a valid pointer to a Metasound operator. An invalid pointer on failure.
 			 */
-			virtual TUniquePtr<IOperator> BuildOperator(const FOperatorSettings& InSettings, const FMetasoundEnvironment& InEnvironment, TArray<IOperatorBuilder::FBuildErrorPtr>& OutBuildErrors) const = 0;
+			virtual TUniquePtr<IOperator> BuildOperator(const FOperatorSettings& InSettings, const FMetasoundEnvironment& InEnvironment, FBuildGraphResults& OutResults) const = 0;
 
 			/** Returns a handle to the document owning this graph. */
 			virtual FDocumentHandle GetOwningDocument() = 0;
diff --git a/Engine/Plugins/Runtime/Metasound/Source/MetasoundFrontend/Public/MetasoundFrontendDocument.h b/Engine/Plugins/Runtime/Metasound/Source/MetasoundFrontend/Public/MetasoundFrontendDocument.h
--- a/Engine/Plugins/Runtime/Metasound/Source/MetasoundFrontend/Public/MetasoundFrontendDocument.h
+++ b/Engine/Plugins/Runtime/Metasound/Source/MetasoundFrontend/Public/MetasoundFrontendDocument.h
@@ -15,6 +15,7 @@
 #include "Templates/Function.h"
 #include "Templates/Invoke.h"
 #include "Templates/TypeHash.h"
+#include "UObject/NoExportTypes.h"
 
 #include "MetasoundFrontendDocument.generated.h"
 
@@ -32,6 +33,11 @@ namespace Metasound
 	{
 		namespace DisplayStyle
 		{
+			namespace EdgeAnimation
+			{
+				extern const FLinearColor METASOUNDFRONTEND_API DefaultColor;
+			} // namespace EdgeStyle
+
 			namespace NodeLayout
 			{
 				extern const FVector2D METASOUNDFRONTEND_API DefaultOffsetX;
@@ -475,65 +481,55 @@ struct FMetasoundFrontendEdge
 	FGuid ToVertexID = Metasound::FrontendInvalidID;
 };
 
-
-// Display style for an edge.
-UENUM()
-enum class EMetasoundFrontendStyleEdgeDisplay : uint8
-{
-	Default,
-	Inherited,
-	Hidden
-};
-
-// Styling for edges
 USTRUCT()
-struct FMetasoundFrontendStyleEdge
+struct METASOUNDFRONTEND_API FMetasoundFrontendEdgeStyleLiteralColorPair
 {
 	GENERATED_BODY()
 
-#if WITH_EDITORONLY_DATA
 	UPROPERTY()
-	EMetasoundFrontendStyleEdgeDisplay Display = EMetasoundFrontendStyleEdgeDisplay::Default;
-#endif // WITH_EDITORONLY_DATA
+	FMetasoundFrontendLiteral Value;
+
+	UPROPERTY()
+	FLinearColor Color = Metasound::Frontend::DisplayStyle::EdgeAnimation::DefaultColor;
 };
 
-// Styling for a class of edges dependent upon edge data type.
+// Styling for all edges associated with a given output (characterized by NodeID & Name)
 USTRUCT()
-struct FMetasoundFrontendStyleEdgeClass
+struct METASOUNDFRONTEND_API FMetasoundFrontendEdgeStyle
 {
 	GENERATED_BODY()
 
-#if WITH_EDITORONLY_DATA
-	// Datatype of edge to apply style to
+	// Node ID for associated edge(s) that should use the given style data.
 	UPROPERTY()
-	FName TypeName;
+	FGuid NodeID;
 
-	// Style information for edge.
+	// Name of node's output to associate style information for its associated edge(s).
 	UPROPERTY()
-	FMetasoundFrontendStyleEdge Style;
-#endif // WITH_EDITORONLY_DATA
+	FName OutputName;
+
+	// Array of colors used to animate given output's associated edge(s). Interpolation
+	// between values dependent on value used.
+	UPROPERTY()
+	TArray<FMetasoundFrontendEdgeStyleLiteralColorPair> LiteralColorPairs;
 };
 
 // Styling for a class
 USTRUCT()
-struct FMetasoundFrontendGraphStyle 
+struct METASOUNDFRONTEND_API FMetasoundFrontendGraphStyle
 {
 	GENERATED_BODY()
 
-#if WITH_EDITORONLY_DATA
 	// Whether or not the graph is editable by a user
 	UPROPERTY()
 	bool bIsGraphEditable = true;
 
-	// Edge styles for graph.
+	// Styles for graph edges.
 	UPROPERTY()
-	TArray<FMetasoundFrontendStyleEdgeClass> EdgeStyles;
-#endif // WITH_EDITORONLY_DATA
+	TArray<FMetasoundFrontendEdgeStyle> EdgeStyles;
 };
 
-
 USTRUCT()
-struct FMetasoundFrontendGraph
+struct METASOUNDFRONTEND_API FMetasoundFrontendGraph
 {
 	GENERATED_BODY()
 
@@ -560,7 +556,7 @@ struct FMetasoundFrontendGraph
 
 // Metadata associated with a vertex.
 USTRUCT()
-struct FMetasoundFrontendVertexMetadata
+struct METASOUNDFRONTEND_API FMetasoundFrontendVertexMetadata
 {
 	GENERATED_BODY()
 
@@ -1347,7 +1343,7 @@ struct METASOUNDFRONTEND_API FMetasoundFrontendGraphClass : public FMetasoundFro
 };
 
 USTRUCT()
-struct FMetasoundFrontendDocumentMetadata
+struct METASOUNDFRONTEND_API FMetasoundFrontendDocumentMetadata
 {
 	GENERATED_BODY()
 
@@ -1355,7 +1351,6 @@ struct FMetasoundFrontendDocumentMetadata
 	FMetasoundFrontendVersion Version;
 };
 
-
 USTRUCT()
 struct METASOUNDFRONTEND_API FMetasoundFrontendDocument
 {
diff --git a/Engine/Plugins/Runtime/Metasound/Source/MetasoundFrontend/Public/MetasoundInputNode.h b/Engine/Plugins/Runtime/Metasound/Source/MetasoundFrontend/Public/MetasoundInputNode.h
--- a/Engine/Plugins/Runtime/Metasound/Source/MetasoundFrontend/Public/MetasoundInputNode.h
+++ b/Engine/Plugins/Runtime/Metasound/Source/MetasoundFrontend/Public/MetasoundInputNode.h
@@ -246,7 +246,7 @@ namespace Metasound
 			{
 			}
 
-			virtual TUniquePtr<IOperator> CreateOperator(const FCreateOperatorParams& InParams, FBuildErrorArray& OutErrors) override;
+			virtual TUniquePtr<IOperator> CreateOperator(const FCreateOperatorParams& InParams, FBuildGraphResults& OutResults) override;
 
 		private:
 			ReferenceCreatorType ReferenceCreator;
@@ -361,7 +361,7 @@ namespace Metasound
 
 
 	template<typename DataType, typename ReferenceCreatorType>
-	TUniquePtr<IOperator> TInputOperatorFactory<DataType, ReferenceCreatorType>::CreateOperator(const FCreateOperatorParams& InParams, FBuildErrorArray& OutErrors) 
+	TUniquePtr<IOperator> TInputOperatorFactory<DataType, ReferenceCreatorType>::CreateOperator(const FCreateOperatorParams& InParams, FBuildGraphResults& OutResults)
 	{
 		using FInputNodeType = TInputNode<DataType>;
 
diff --git a/Engine/Plugins/Runtime/Metasound/Source/MetasoundFrontend/Public/MetasoundOutputNode.h b/Engine/Plugins/Runtime/Metasound/Source/MetasoundFrontend/Public/MetasoundOutputNode.h
--- a/Engine/Plugins/Runtime/Metasound/Source/MetasoundFrontend/Public/MetasoundOutputNode.h
+++ b/Engine/Plugins/Runtime/Metasound/Source/MetasoundFrontend/Public/MetasoundOutputNode.h
@@ -75,7 +75,7 @@ namespace Metasound
 				{
 				}
 
-				virtual TUniquePtr<IOperator> CreateOperator(const FCreateOperatorParams& InParams, FBuildErrorArray& OutErrors) override
+				virtual TUniquePtr<IOperator> CreateOperator(const FCreateOperatorParams& InParams, FBuildGraphResults& OutResults) override
 				{
 					const bool bContainsRef = InParams.InputDataReferences.ContainsDataReadReference<DataType>(DataReferenceName);
 					if (bContainsRef)
@@ -91,7 +91,7 @@ namespace Metasound
 						return MakeUnique<FOutputOperator>(DataReferenceName, DefaultReadRef);
 					}
 
-					OutErrors.Emplace(MakeUnique<FMissingOutputNodeInputReferenceError>(InParams.Node, GetMetasoundDataTypeDisplayText<DataType>()));
+					OutResults.Errors.Emplace(MakeUnique<FMissingOutputNodeInputReferenceError>(InParams.Node, GetMetasoundDataTypeDisplayText<DataType>()));
 					return TUniquePtr<IOperator>(nullptr);
 				}
 
diff --git a/Engine/Plugins/Runtime/Metasound/Source/MetasoundFrontend/Public/MetasoundReceiveNode.h b/Engine/Plugins/Runtime/Metasound/Source/MetasoundFrontend/Public/MetasoundReceiveNode.h
--- a/Engine/Plugins/Runtime/Metasound/Source/MetasoundFrontend/Public/MetasoundReceiveNode.h
+++ b/Engine/Plugins/Runtime/Metasound/Source/MetasoundFrontend/Public/MetasoundReceiveNode.h
@@ -212,7 +212,7 @@ namespace Metasound
 			public:
 				FReceiverOperatorFactory() = default;
 
-				virtual TUniquePtr<IOperator> CreateOperator(const FCreateOperatorParams& InParams, FBuildErrorArray& OutErrors) override
+				virtual TUniquePtr<IOperator> CreateOperator(const FCreateOperatorParams& InParams, FBuildGraphResults& OutResults) override
 				{
 					using namespace ReceiveNodeInfo; 
 
diff --git a/Engine/Plugins/Runtime/Metasound/Source/MetasoundFrontend/Public/MetasoundSendNode.h b/Engine/Plugins/Runtime/Metasound/Source/MetasoundFrontend/Public/MetasoundSendNode.h
--- a/Engine/Plugins/Runtime/Metasound/Source/MetasoundFrontend/Public/MetasoundSendNode.h
+++ b/Engine/Plugins/Runtime/Metasound/Source/MetasoundFrontend/Public/MetasoundSendNode.h
@@ -169,9 +169,9 @@ namespace Metasound
 			public:
 				FSendOperatorFactory() = default;
 
-				virtual TUniquePtr<IOperator> CreateOperator(const FCreateOperatorParams& InParams, FBuildErrorArray& OutErrors) override
+				virtual TUniquePtr<IOperator> CreateOperator(const FCreateOperatorParams& InParams, FBuildGraphResults& OutResults) override
 				{
-					using namespace SendVertexNames; 
+					using namespace SendVertexNames;
 
 					if (InParams.InputDataReferences.ContainsDataReadReference<TDataType>(GetSendInputName()))
 					{
diff --git a/Engine/Plugins/Runtime/Metasound/Source/MetasoundGenerator/Private/MetasoundGenerator.cpp b/Engine/Plugins/Runtime/Metasound/Source/MetasoundGenerator/Private/MetasoundGenerator.cpp
--- a/Engine/Plugins/Runtime/Metasound/Source/MetasoundGenerator/Private/MetasoundGenerator.cpp
+++ b/Engine/Plugins/Runtime/Metasound/Source/MetasoundGenerator/Private/MetasoundGenerator.cpp
@@ -1,6 +1,7 @@
 // Copyright Epic Games, Inc. All Rights Reserved.
 #include "MetasoundGenerator.h"
 
+#include "Analysis/MetasoundFrontendVertexAnalyzer.h"
 #include "AudioParameter.h"
 #include "DSP/Dsp.h"
 #include "MetasoundGraph.h"
@@ -8,6 +9,7 @@
 #include "MetasoundOperatorBuilder.h"
 #include "MetasoundOperatorInterface.h"
 #include "MetasoundOutputNode.h"
+#include "MetasoundRouter.h"
 #include "MetasoundSourceInterface.h"
 #include "MetasoundTrace.h"
 #include "MetasoundTrigger.h"
@@ -54,15 +56,21 @@ namespace Metasound
 		METASOUND_TRACE_CPUPROFILER_EVENT_SCOPE_TEXT(*FString::Printf(TEXT("AsyncMetaSoundBuilder::DoWork %s"), *InitParams.MetaSoundName));
 
 		// Create an instance of the new graph
-		FOperatorBuilder OperatorBuilder(FOperatorBuilderSettings::GetDefaultSettings());
-		FDataReferenceCollection DataReferenceCollection{};
-		FBuildGraphParams BuildParams{*InitParams.Graph, InitParams.OperatorSettings, DataReferenceCollection, InitParams.Environment};
-		TArray<IOperatorBuilder::FBuildErrorPtr> BuildErrors;
+		FDataReferenceCollection DataReferenceCollection;
+		FBuildGraphParams BuildParams { *InitParams.Graph, InitParams.OperatorSettings, DataReferenceCollection, InitParams.Environment, InitParams.BuilderSettings };
+		FBuildGraphResults BuildResults;
 
-		TUniquePtr<IOperator> GraphOperator = OperatorBuilder.BuildGraphOperator(BuildParams, BuildErrors);
+		TUniquePtr<IOperator> GraphOperator = FOperatorBuilder().BuildGraphOperator(BuildParams, BuildResults);
+
+		TUniquePtr<FGraphAnalyzer> GraphAnalyzer;
+		if (InitParams.BuilderSettings.bPopulateInternalDataReferences)
+		{
+			const uint64 InstanceID = BuildParams.Environment.GetValue<uint64>(SourceInterface::Environment::TransmitterID);
+			GraphAnalyzer = MakeUnique<FGraphAnalyzer>(InitParams.OperatorSettings, InstanceID, MoveTemp(BuildResults.InternalDataReferences));
+		}
 
 		// Log build errors
-		for (const IOperatorBuilder::FBuildErrorPtr& Error : BuildErrors)
+		for (const IOperatorBuilder::FBuildErrorPtr& Error : BuildResults.Errors)
 		{
 			if (Error.IsValid())
 			{
@@ -94,9 +102,10 @@ namespace Metasound
 			{
 				InitParams.OperatorSettings,
 				MoveTemp(GraphOperator),
+				MoveTemp(GraphAnalyzer),
 				MoveTemp(OutputBuffers),
 				MoveTemp(PlayTrigger),
-				MoveTemp(FinishTrigger)
+				MoveTemp(FinishTrigger),
 			};
 
 			Generator->SetPendingGraph(MoveTemp(GeneratorData), bTriggerGenerator);
@@ -169,7 +178,7 @@ namespace Metasound
 	{
 		FScopeLock SetPendingGraphLock(&PendingGraphMutex);
 		{
-			PendingGraphData = TUniquePtr<FMetasoundGeneratorData>(nullptr);
+			PendingGraphData.Reset();
 			bPendingGraphTrigger = false;
 			bIsNewGraphPending = true;
 		}
@@ -242,6 +251,8 @@ namespace Metasound
 			InterleavedAudioBuffer.AddUninitialized(NumSamplesPerExecute);
 		}
 
+		GraphAnalyzer = MoveTemp(InData->GraphAnalyzer);
+
 		if (bTriggerGraph)
 		{
 			OnPlayTriggerRef->TriggerFrame(0);
@@ -302,6 +313,11 @@ namespace Metasound
 			// Call metasound graph operator.
 			RootExecuter.Execute();
 
+			if (GraphAnalyzer.IsValid())
+			{
+				GraphAnalyzer->Execute();
+			}
+
 			// Interleave audio because ISoundGenerator interface expects interleaved audio.
 			InterleaveGeneratedAudio();
 
@@ -322,11 +338,11 @@ namespace Metasound
 
 				FMemory::Memcpy(OverflowBuffer.GetData(), &InterleavedAudioBuffer.GetData()[ThisLoopNumSamplesWritten], OverflowCount * sizeof(float));
 			}
-		}
 
-		if (*OnFinishedTriggerRef)
-		{
-			bIsFinishTriggered = true;
+			if (*OnFinishedTriggerRef)
+			{
+				bIsFinishTriggered = true;
+			}
 		}
 
 		return NumSamplesWritten;
diff --git a/Engine/Plugins/Runtime/Metasound/Source/MetasoundGenerator/Public/MetasoundGenerator.h b/Engine/Plugins/Runtime/Metasound/Source/MetasoundGenerator/Public/MetasoundGenerator.h
--- a/Engine/Plugins/Runtime/Metasound/Source/MetasoundGenerator/Public/MetasoundGenerator.h
+++ b/Engine/Plugins/Runtime/Metasound/Source/MetasoundGenerator/Public/MetasoundGenerator.h
@@ -1,14 +1,17 @@
 // Copyright Epic Games, Inc. All Rights Reserved.
 #pragma once
 
+#include "Analysis/MetasoundFrontendGraphAnalyzer.h"
 #include "Async/Async.h"
-#include "CoreMinimal.h"
 #include "DSP/Dsp.h"
 #include "MetasoundAudioFormats.h"
 #include "MetasoundExecutableOperator.h"
 #include "MetasoundFrontendController.h"
+#include "MetasoundGraphOperator.h"
+#include "MetasoundOperatorBuilder.h"
 #include "MetasoundOperatorInterface.h"
 #include "MetasoundPrimitives.h"
+#include "MetasoundRouter.h"
 #include "MetasoundTrigger.h"
 #include "MetasoundVertex.h"
 #include "Sound/SoundGenerator.h"
@@ -21,6 +24,7 @@ namespace Metasound
 	struct METASOUNDGENERATOR_API FMetasoundGeneratorInitParams
 	{
 		FOperatorSettings OperatorSettings;
+		FOperatorBuilderSettings BuilderSettings;
 		TSharedPtr<const IGraph, ESPMode::ThreadSafe> Graph;
 		FMetasoundEnvironment Environment;
 		FString MetaSoundName;
@@ -32,7 +36,8 @@ namespace Metasound
 	struct FMetasoundGeneratorData
 	{
 		FOperatorSettings OperatorSettings;
-		TUniquePtr<Metasound::IOperator> GraphOperator;
+		TUniquePtr<IOperator> GraphOperator;
+		TUniquePtr<Frontend::FGraphAnalyzer> GraphAnalyzer;
 		TArray<TDataReadReference<FAudioBuffer>> OutputBuffers;
 		FTriggerWriteRef TriggerOnPlayRef;
 		FTriggerReadRef TriggerOnFinishRef;
@@ -190,6 +195,7 @@ namespace Metasound
 		bool bPendingGraphTrigger;
 		bool bIsNewGraphPending;
 		bool bIsWaitingForFirstGraph;
+
+		TUniquePtr<Frontend::FGraphAnalyzer> GraphAnalyzer;
 	};
 }
-
diff --git a/Engine/Plugins/Runtime/Metasound/Source/MetasoundGraphCore/Private/MetasoundFacade.cpp b/Engine/Plugins/Runtime/Metasound/Source/MetasoundGraphCore/Private/MetasoundFacade.cpp
--- a/Engine/Plugins/Runtime/Metasound/Source/MetasoundGraphCore/Private/MetasoundFacade.cpp
+++ b/Engine/Plugins/Runtime/Metasound/Source/MetasoundGraphCore/Private/MetasoundFacade.cpp
@@ -9,9 +9,9 @@ namespace Metasound
 	{
 	}
 
-	TUniquePtr<IOperator> FNodeFacade::FFactory::CreateOperator(const FCreateOperatorParams& InParams, FBuildErrorArray& OutErrors)
+	TUniquePtr<IOperator> FNodeFacade::FFactory::CreateOperator(const FCreateOperatorParams& InParams, FBuildGraphResults& OutResults)
 	{
-		return CreateFunc(InParams, OutErrors);
+		return CreateFunc(InParams, OutResults.Errors);
 	}
 
 	const FVertexInterface& FNodeFacade::GetVertexInterface() const
diff --git a/Engine/Plugins/Runtime/Metasound/Source/MetasoundGraphCore/Private/MetasoundGraph.cpp b/Engine/Plugins/Runtime/Metasound/Source/MetasoundGraphCore/Private/MetasoundGraph.cpp
--- a/Engine/Plugins/Runtime/Metasound/Source/MetasoundGraphCore/Private/MetasoundGraph.cpp
+++ b/Engine/Plugins/Runtime/Metasound/Source/MetasoundGraphCore/Private/MetasoundGraph.cpp
@@ -139,21 +139,20 @@ namespace Metasound
 		return MakeShared<FGraph::FFactory>();
 	}
 
-	TUniquePtr<IOperator> FGraph::FFactory::CreateOperator(const FCreateOperatorParams& InParams, FBuildErrorArray& OutErrors) 
+	TUniquePtr<IOperator> FGraph::FFactory::CreateOperator(const FCreateOperatorParams& InParams, FBuildGraphResults& OutResults)
 	{
 		const FGraph& Graph = static_cast<const FGraph&>(InParams.Node);
 
-		FBuildGraphParams BuildParams{Graph, InParams.OperatorSettings, InParams.InputDataReferences, InParams.Environment};
+		FBuildGraphParams BuildParams { Graph, InParams.OperatorSettings, InParams.InputDataReferences, InParams.Environment, InParams.BuilderSettings };
 
 		if (nullptr != InParams.Builder)
 		{
-			// Use the provided builder if it exists. 
-			return InParams.Builder->BuildGraphOperator(BuildParams, OutErrors);
+			// Use the provided builder if it exists.
+			return InParams.Builder->BuildGraphOperator(BuildParams, OutResults);
 		}
 		else
 		{
-			// Use a builder with default settings if no builder was provided. 
-			return FOperatorBuilder(FOperatorBuilderSettings::GetDefaultSettings()).BuildGraphOperator(BuildParams, OutErrors);
+			return FOperatorBuilder().BuildGraphOperator(BuildParams, OutResults);
 		}
 	}
 }
diff --git a/Engine/Plugins/Runtime/Metasound/Source/MetasoundGraphCore/Private/MetasoundOperatorBuilder.cpp b/Engine/Plugins/Runtime/Metasound/Source/MetasoundGraphCore/Private/MetasoundOperatorBuilder.cpp
--- a/Engine/Plugins/Runtime/Metasound/Source/MetasoundGraphCore/Private/MetasoundOperatorBuilder.cpp
+++ b/Engine/Plugins/Runtime/Metasound/Source/MetasoundGraphCore/Private/MetasoundOperatorBuilder.cpp
@@ -4,7 +4,6 @@
 
 #include "Algo/Count.h"
 #include "Algo/Transform.h"
-#include "CoreMinimal.h"
 
 #include "MetasoundBuildError.h"
 #include "MetasoundBuilderInterface.h"
@@ -33,122 +32,40 @@ namespace Metasound
 		}
 	}
 
-	// Default settings
-	FOperatorBuilderSettings FOperatorBuilderSettings::GetDefaultSettings()
-	{
-#if UE_BUILD_SHIPPING
-		return GetDefaultShippingSettings();
-#elif UE_BUILD_TEST
-		return GetDefaultTestSettings();
-#elif UE_BUILD_DEVELOPMENT
-		return GetDefaultDevelopementSettings();
-#elif UE_BUILD_DEBUG
-		return GetDefaultDebugSettings();
-#else
-		return GetDefaultShippingSettings();
-#endif
-	}
-
-	// Debug settings
-	FOperatorBuilderSettings FOperatorBuilderSettings::GetDefaultDebugSettings()
-	{
-		FOperatorBuilderSettings Settings;
-
-		Settings.PruningMode = EOperatorBuilderNodePruning::None;
-		Settings.bValidateNoCyclesInGraph = true;
-		Settings.bValidateNoDuplicateInputs = true;
-		Settings.bValidateVerticesExist = true;
-		Settings.bValidateEdgeDataTypesMatch = true;
-		Settings.bFailOnAnyError = false;
-
-		return Settings;
-	}
-
-	// Development settings
-	FOperatorBuilderSettings FOperatorBuilderSettings::GetDefaultDevelopementSettings()
-	{
-		FOperatorBuilderSettings Settings;
-
-		Settings.PruningMode = EOperatorBuilderNodePruning::PruneNodesWithoutExternalDependency;
-		Settings.bValidateNoCyclesInGraph = true;
-		Settings.bValidateNoDuplicateInputs = true;
-		Settings.bValidateVerticesExist = true;
-		Settings.bValidateEdgeDataTypesMatch = true;
-		Settings.bFailOnAnyError = false;
-
-		return Settings;
-	}
-
-	// Testing settings
-	FOperatorBuilderSettings FOperatorBuilderSettings::GetDefaultTestSettings()
-	{
-		FOperatorBuilderSettings Settings;
-
-		Settings.PruningMode = EOperatorBuilderNodePruning::PruneNodesWithoutExternalDependency;
-		Settings.bValidateNoCyclesInGraph = false;
-		Settings.bValidateNoDuplicateInputs = false;
-		Settings.bValidateVerticesExist = false;
-		Settings.bValidateEdgeDataTypesMatch = false;
-		Settings.bFailOnAnyError = false;
-
-		return Settings;
-	}
-
-	// Shipping settings
-	FOperatorBuilderSettings FOperatorBuilderSettings::GetDefaultShippingSettings()
-	{
-		FOperatorBuilderSettings Settings;
-
-		Settings.PruningMode = EOperatorBuilderNodePruning::PruneNodesWithoutExternalDependency;
-		Settings.bValidateNoCyclesInGraph = false;
-		Settings.bValidateNoDuplicateInputs = false;
-		Settings.bValidateVerticesExist = false;
-		Settings.bValidateEdgeDataTypesMatch = false;
-		Settings.bFailOnAnyError = false;
-
-		return Settings;
-	}
-
-
-	FOperatorBuilder::FOperatorBuilder(const FOperatorBuilderSettings& InBuilderSettings)
-	:	BuilderSettings(InBuilderSettings)
-	{
-		MaxBuildStatusErrorLevel = BuilderSettings.bFailOnAnyError ? FBuildStatus::NoError : FBuildStatus::NonFatalError;
-	}
-
 	FOperatorBuilder::~FOperatorBuilder()
 	{
 	}
 
-	TUniquePtr<IOperator> FOperatorBuilder::BuildGraphOperator(const FBuildGraphParams& InParams, FBuildErrorArray& OutErrors) const
+	TUniquePtr<IOperator> FOperatorBuilder::BuildGraphOperator(const FBuildGraphParams& InParams, FBuildGraphResults& OutResults) const
 	{
 		using namespace OperatorBuilderPrivate;
 
 		FBuildStatus BuildStatus = FBuildStatus::NoError;
+		const FBuildStatus MaxBuildStatusErrorLevel = FBuildStatus::GetMaxErrorLevel(InParams.BuilderSettings);
 
 		// Validate that the sources and destinations declared in an edge actually
 		// exist in the node.
-		if (BuilderSettings.bValidateVerticesExist)
+		if (InParams.BuilderSettings.bValidateVerticesExist)
 		{
-			if (!FGraphLinter::ValidateVerticesExist(InParams.Graph, OutErrors))
+			if (!FGraphLinter::ValidateVerticesExist(InParams.Graph, OutResults.Errors))
 			{
 				BuildStatus |= FBuildStatus::FatalError;
 			}
 		}
 
 		// Validate that the data types for a source and destination match.
-		if (BuilderSettings.bValidateEdgeDataTypesMatch)
+		if (InParams.BuilderSettings.bValidateEdgeDataTypesMatch)
 		{
-			if (!FGraphLinter::ValidateEdgeDataTypesMatch(InParams.Graph, OutErrors))
+			if (!FGraphLinter::ValidateEdgeDataTypesMatch(InParams.Graph, OutResults.Errors))
 			{
 				BuildStatus |= FBuildStatus::FatalError;
 			}
 		}
 		
 		// Validate that node inputs only have one source
-		if (BuilderSettings.bValidateNoDuplicateInputs)
+		if (InParams.BuilderSettings.bValidateNoDuplicateInputs)
 		{
-			if (!FGraphLinter::ValidateNoDuplicateInputs(InParams.Graph, OutErrors))
+			if (!FGraphLinter::ValidateNoDuplicateInputs(InParams.Graph, OutResults.Errors))
 			{
 				BuildStatus |= FBuildStatus::FatalError;
 			}
@@ -160,17 +77,16 @@ namespace Metasound
 			return TUniquePtr<IOperator>(nullptr);
 		}
 
-
 		// Create algo adapter view of graph to cache graph operations.
 		TPimplPtr<FDirectedGraphAlgoAdapter> AlgoAdapter = FDirectedGraphAlgo::CreateDirectedGraphAlgoAdapter(InParams.Graph);
 		
 		if (!AlgoAdapter.IsValid())
 		{
-			AddBuildError<FInternalError>(OutErrors, __FILE__, __LINE__);
+			AddBuildError<FInternalError>(OutResults.Errors, __FILE__, __LINE__);
 			return TUniquePtr<IOperator>(nullptr);
 		}
 
-		FBuildContext BuildContext(InParams.Graph, *AlgoAdapter, InParams.OperatorSettings, InParams.Environment, OutErrors);
+		FBuildContext BuildContext(InParams.Graph, *AlgoAdapter, InParams.OperatorSettings, InParams.Environment, InParams.BuilderSettings, OutResults);
 
 		TArray<const INode*> SortedNodes;
 
@@ -187,11 +103,8 @@ namespace Metasound
 			return TUniquePtr<IOperator>(nullptr);
 		}
 
-		TArray<TUniquePtr<IOperator>> Operators;
-		FNodeDataReferenceMap NodeData;
-
-		// Create node operators from factories. 
-		BuildStatus |= CreateOperators(BuildContext, SortedNodes, InParams.InputDataReferences, Operators, NodeData);
+		// Create node operators from factories.
+		BuildStatus |= CreateOperators(BuildContext, SortedNodes, InParams.InputDataReferences);
 
 		if (BuildStatus > MaxBuildStatusErrorLevel)
 		{
@@ -199,19 +112,19 @@ namespace Metasound
 		}
 
 		// Create graph operator from collection of node operators.
-		return CreateGraphOperator(BuildContext, Operators, NodeData);
+		return CreateGraphOperator(BuildContext, InParams);
 	}
 
-	FOperatorBuilder::FBuildStatus FOperatorBuilder::DepthFirstTopologicalSort(FBuildContext& InContext, TArray<const INode*>& OutNodes) const
+	FOperatorBuilder::FBuildStatus FOperatorBuilder::DepthFirstTopologicalSort(FBuildContext& InOutContext, TArray<const INode*>& OutNodes) const
 	{
 		using namespace OperatorBuilderPrivate;
 
-		bool bSuccess = FDirectedGraphAlgo::DepthFirstTopologicalSort(InContext.AlgoAdapter, OutNodes);
+		bool bSuccess = FDirectedGraphAlgo::DepthFirstTopologicalSort(InOutContext.AlgoAdapter, OutNodes);
 
 		if (!bSuccess)
 		{
 			// If there was an error, there is likely a cycle in the graph.
-			AddBuildErrorsForCycles(InContext.AlgoAdapter, InContext.Errors);
+			AddBuildErrorsForCycles(InOutContext.AlgoAdapter, InOutContext.Results.Errors);
 
 			return FBuildStatus::FatalError;
 		}
@@ -219,16 +132,16 @@ namespace Metasound
 		return FBuildStatus::NoError;
 	}
 
-	FOperatorBuilder::FBuildStatus FOperatorBuilder::KahnsTopologicalSort(FBuildContext& InContext, TArray<const INode*>& OutNodes) const
+	FOperatorBuilder::FBuildStatus FOperatorBuilder::KahnsTopologicalSort(FBuildContext& InOutContext, TArray<const INode*>& OutNodes) const
 	{
 		using namespace OperatorBuilderPrivate;
 
-		bool bSuccess = FDirectedGraphAlgo::KahnTopologicalSort(InContext.AlgoAdapter, OutNodes);
+		bool bSuccess = FDirectedGraphAlgo::KahnTopologicalSort(InOutContext.AlgoAdapter, OutNodes);
 
 		if (!bSuccess)
 		{
 			// If there was an error, there is likely a cycle in the graph.
-			AddBuildErrorsForCycles(InContext.AlgoAdapter, InContext.Errors);
+			AddBuildErrorsForCycles(InOutContext.AlgoAdapter, InOutContext.Results.Errors);
 
 			return FBuildStatus::FatalError;
 		}
@@ -236,7 +149,7 @@ namespace Metasound
 		return FBuildStatus::NoError;
 	}
 
-	FOperatorBuilder::FBuildStatus FOperatorBuilder::PruneNodes(FBuildContext& InContext, TArray<const INode*>& InOutNodes) const
+	FOperatorBuilder::FBuildStatus FOperatorBuilder::PruneNodes(FBuildContext& InOutContext, TArray<const INode*>& InOutNodes) const
 	{
 		using namespace OperatorBuilderPrivate;
 
@@ -244,18 +157,18 @@ namespace Metasound
 
 		TSet<const INode*> ReachableNodes;
 
-		switch (BuilderSettings.PruningMode)
+		switch (InOutContext.BuilderSettings.PruningMode)
 		{
 			case EOperatorBuilderNodePruning::PruneNodesWithoutExternalDependency:
-				FDirectedGraphAlgo::FindReachableNodes(InContext.AlgoAdapter, ReachableNodes);
+				FDirectedGraphAlgo::FindReachableNodes(InOutContext.AlgoAdapter, ReachableNodes);
 				break;
 
 			case EOperatorBuilderNodePruning::PruneNodesWithoutOutputDependency:
-				FDirectedGraphAlgo::FindReachableNodesFromOutput(InContext.AlgoAdapter, ReachableNodes);
+				FDirectedGraphAlgo::FindReachableNodesFromOutput(InOutContext.AlgoAdapter, ReachableNodes);
 				break;
 
 			case EOperatorBuilderNodePruning::PruneNodesWithoutInputDependency:
-				FDirectedGraphAlgo::FindReachableNodesFromInput(InContext.AlgoAdapter, ReachableNodes);
+				FDirectedGraphAlgo::FindReachableNodesFromInput(InOutContext.AlgoAdapter, ReachableNodes);
 				break;
 
 			case EOperatorBuilderNodePruning::None:
@@ -275,7 +188,7 @@ namespace Metasound
 			// Pruning all nodes. 
 			for (const INode* Node : InOutNodes)
 			{
-				AddBuildError<FNodePrunedError>(InContext.Errors, Node);
+				AddBuildError<FNodePrunedError>(InOutContext.Results.Errors, Node);
 			}
 
 			InOutNodes.Reset();
@@ -298,7 +211,7 @@ namespace Metasound
 			}
 			else
 			{
-				AddBuildError<FNodePrunedError>(InContext.Errors, Node);
+				AddBuildError<FNodePrunedError>(InOutContext.Results.Errors, Node);
 
 				// Denote a pruned node as a non-fatal error. In the future this
 				// may be simply a warning as some nodes are required to conform
@@ -312,7 +225,7 @@ namespace Metasound
 		return BuildStatus;
 	}
 
-	FOperatorBuilder::FBuildStatus FOperatorBuilder::GatherInputDataReferences(FBuildContext& InContext, const INode* InNode, const FNodeEdgeMultiMap& InEdgeMap, const FNodeDataReferenceMap& InDataReferenceMap, FDataReferenceCollection& OutCollection) const
+	FOperatorBuilder::FBuildStatus FOperatorBuilder::GatherInputDataReferences(FBuildContext& InOutContext, const INode* InNode, const FNodeEdgeMultiMap& InEdgeMap, FDataReferenceCollection& OutInputCollection) const
 	{
 		using namespace OperatorBuilderPrivate;
 
@@ -325,28 +238,28 @@ namespace Metasound
 		// Add parameters to collection based on edge description
 		for (const FDataEdge* Edge : Edges)
 		{
-			if (!InDataReferenceMap.Contains(Edge->From.Node))
+			if (!InOutContext.DataReferences.Contains(Edge->From.Node))
 			{
 				// This is likely due to a failed topological sort and is more of an internal error
 				// than a user error.
-				AddBuildError<FInternalError>(InContext.Errors, TEXT(__FILE__), __LINE__);
+				AddBuildError<FInternalError>(InOutContext.Results.Errors, TEXT(__FILE__), __LINE__);
 
 				return FBuildStatus::NonFatalError;
 			}
 
-			const FDataReferenceCollection& FromDataReferenceCollection = InDataReferenceMap[Edge->From.Node].Outputs;
+			const FDataReferenceCollection& FromDataReferenceCollection = InOutContext.DataReferences[Edge->From.Node].Outputs;
 
 			bool bSuccess = false;
 
 			if (FromDataReferenceCollection.ContainsDataWriteReference(Edge->From.Vertex.VertexName, Edge->From.Vertex.DataTypeName))
 			{
 				// Contains data write reference
-				bSuccess = OutCollection.AddDataWriteReferenceFrom(Edge->To.Vertex.VertexName, FromDataReferenceCollection, Edge->From.Vertex.VertexName, Edge->From.Vertex.DataTypeName);
+				bSuccess = OutInputCollection.AddDataWriteReferenceFrom(Edge->To.Vertex.VertexName, FromDataReferenceCollection, Edge->From.Vertex.VertexName, Edge->From.Vertex.DataTypeName);
 			}
 			else if (FromDataReferenceCollection.ContainsDataReadReference(Edge->From.Vertex.VertexName, Edge->From.Vertex.DataTypeName))
 			{
 				// Contains data read reference
-				bSuccess = OutCollection.AddDataReadReferenceFrom(Edge->To.Vertex.VertexName, FromDataReferenceCollection, Edge->From.Vertex.VertexName, Edge->From.Vertex.DataTypeName);
+				bSuccess = OutInputCollection.AddDataReadReferenceFrom(Edge->To.Vertex.VertexName, FromDataReferenceCollection, Edge->From.Vertex.VertexName, Edge->From.Vertex.DataTypeName);
 			}
 
 			if (!bSuccess)
@@ -356,7 +269,7 @@ namespace Metasound
 				// did not match the readable parameter refs created by the operators outputs. Or, the edge description is invalid.
 
 				// TODO: consider checking that outputs of node operators match output descriptions of nodes. 
-				AddBuildError<FMissingOutputDataReferenceError>(InContext.Errors, Edge->From);
+				AddBuildError<FMissingOutputDataReferenceError>(InOutContext.Results.Errors, Edge->From);
 
 				BuildStatus |= FBuildStatus::NonFatalError;
 				continue;
@@ -366,7 +279,7 @@ namespace Metasound
 		return BuildStatus;
 	}
 
-	FOperatorBuilder::FBuildStatus FOperatorBuilder::GatherExternalInputDataReferences(FBuildContext& InContext, const INode* InNode, const FNodeDestinationMap& InNodeDestinationMap, const FDataReferenceCollection& InExternalCollection, FDataReferenceCollection& OutDataReferences) const
+	FOperatorBuilder::FBuildStatus FOperatorBuilder::GatherExternalInputDataReferences(FBuildContext& InOutContext, const INode* InNode, const FNodeDestinationMap& InNodeDestinationMap, const FDataReferenceCollection& InExternalCollection, FDataReferenceCollection& OutDataReferences) const
 	{
 		using namespace OperatorBuilderPrivate;
 
@@ -383,7 +296,7 @@ namespace Metasound
 
 					if (!bSuccess)
 					{
-						AddBuildError<FInternalError>(InContext.Errors, __FILE__, __LINE__);
+						AddBuildError<FInternalError>(InOutContext.Results.Errors, __FILE__, __LINE__);
 
 						BuildStatus |= FBuildStatus::NonFatalError;
 					}
@@ -394,21 +307,21 @@ namespace Metasound
 		return BuildStatus;
 	}
 
-	FOperatorBuilder::FBuildStatus FOperatorBuilder::CreateOperators(FBuildContext& InContext, const TArray<const INode*>& InSortedNodes, const FDataReferenceCollection& InExternalInputDataReferences, TArray<TUniquePtr<IOperator>>& OutOperators, FNodeDataReferenceMap& OutDataReferences) const
+	FOperatorBuilder::FBuildStatus FOperatorBuilder::CreateOperators(FBuildContext& InOutContext, const TArray<const INode*>& InSortedNodes, const FDataReferenceCollection& InExternalInputDataReferences) const
 	{
 		FBuildStatus BuildStatus;
 
 		FNodeEdgeMultiMap NodeInputEdges;
 
 		// Gather input edges for each node.
-		for (const FDataEdge& Edge : InContext.Graph.GetDataEdges())
+		for (const FDataEdge& Edge : InOutContext.Graph.GetDataEdges())
 		{
 			NodeInputEdges.Add(Edge.To.Node, &Edge);
 		}
 
 		// Map input nodes to graph destinations
 		FNodeDestinationMap NodeDestinations;
-		for (const auto& InputDestinationKV : InContext.Graph.GetInputDataDestinations())
+		for (const auto& InputDestinationKV : InOutContext.Graph.GetInputDataDestinations())
 		{
 			const FInputDataDestination& Destination = InputDestinationKV.Value;
 			NodeDestinations.Add(Destination.Node, &Destination);
@@ -417,12 +330,12 @@ namespace Metasound
 		// Call operator factory for each node.
 		for (const INode* Node : InSortedNodes)
 		{
+			// Gather the input parameters for this IOperator from the output parameters of already created IOperators. 
 			FDataReferenceCollection InputCollection;
+			BuildStatus |= GatherInputDataReferences(InOutContext, Node, NodeInputEdges, InputCollection);
 
-			// Gather the input parameters for this IOperator from the output parameters of already created IOperators. 
-			BuildStatus |= GatherInputDataReferences(InContext, Node, NodeInputEdges, OutDataReferences, InputCollection);
 			// Gather the input parameters for this IOperator from the graph inputs.
-			BuildStatus |= GatherExternalInputDataReferences(InContext, Node, NodeDestinations, InExternalInputDataReferences, InputCollection);
+			BuildStatus |= GatherExternalInputDataReferences(InOutContext, Node, NodeDestinations, InExternalInputDataReferences, InputCollection);
 
 			if (BuildStatus >= FBuildStatus::FatalError)
 			{
@@ -431,9 +344,9 @@ namespace Metasound
 
 			FOperatorFactorySharedRef Factory = Node->GetDefaultOperatorFactory();
 
-			FCreateOperatorParams CreateParams{*Node, InContext.Settings, InputCollection, InContext.Environment, this};
+			FCreateOperatorParams CreateParams { *Node, InOutContext.Settings, InputCollection, InOutContext.Environment, InOutContext.BuilderSettings };
 
-			FOperatorPtr Operator = Factory->CreateOperator(CreateParams, InContext.Errors);
+			FOperatorPtr Operator = Factory->CreateOperator(CreateParams, InOutContext.Results);
 
 			if (!Operator.IsValid())
 			{
@@ -441,16 +354,15 @@ namespace Metasound
 			}
 
 			// Save input and output collection for future use 
-			OutDataReferences.Emplace(Node, FOperatorDataReferences(Operator->GetInputs(), Operator->GetOutputs()));
+			InOutContext.DataReferences.Emplace(Node, FOperatorDataReferences(Operator->GetInputs(), Operator->GetOutputs()));
 
 			// Add operator to operator array
-			OutOperators.Add(MoveTemp(Operator));
+			InOutContext.Operators.Add(MoveTemp(Operator));
 		}
-
 		return BuildStatus;
 	}
 
-	FOperatorBuilder::FBuildStatus FOperatorBuilder::GatherGraphDataReferences(FBuildContext& InContext, FNodeDataReferenceMap& InNodeDataReferences, FDataReferenceCollection& OutGraphInputs, FDataReferenceCollection& OutGraphOutputs) const
+	FOperatorBuilder::FBuildStatus FOperatorBuilder::GatherGraphDataReferences(FBuildContext& InOutContext, FDataReferenceCollection& OutGraphInputs, FDataReferenceCollection& OutGraphOutputs) const
 	{
 		using namespace OperatorBuilderPrivate;
 		using FDestinationElement = FInputDataDestinationCollection::ElementType;
@@ -459,20 +371,20 @@ namespace Metasound
 		FBuildStatus BuildStatus;
 
 		// Gather graph inputs
-		for (const FDestinationElement& Element : InContext.Graph.GetInputDataDestinations())
+		for (const FDestinationElement& Element : InOutContext.Graph.GetInputDataDestinations())
 		{
 			const FInputDataDestination& InputDestination = Element.Value;
 
-			if (!InNodeDataReferences.Contains(InputDestination.Node))
+			if (!InOutContext.DataReferences.Contains(InputDestination.Node))
 			{
 				// An input node was likely pruned.
-				AddBuildError<FMissingInputDataReferenceError>(InContext.Errors, InputDestination);
+				AddBuildError<FMissingInputDataReferenceError>(InOutContext.Results.Errors, InputDestination);
 
 				BuildStatus |= FBuildStatus::NonFatalError;
 				continue;
 			}
 
-			FDataReferenceCollection& Collection = InNodeDataReferences[InputDestination.Node].Inputs;
+			const FDataReferenceCollection& Collection = InOutContext.DataReferences[InputDestination.Node].Inputs;
 
 			// Get either readable or writable values for inputs. 
 			// Readable inputs can occur when the data reference is provided outside
@@ -484,7 +396,7 @@ namespace Metasound
 
 			if (!bContainsInput)
 			{
-				AddBuildError<FMissingInputDataReferenceError>(InContext.Errors, InputDestination);
+				AddBuildError<FMissingInputDataReferenceError>(InOutContext.Results.Errors, InputDestination);
 
 				BuildStatus |= FBuildStatus::NonFatalError;
 				continue;
@@ -502,7 +414,7 @@ namespace Metasound
 
 			if (!bSuccess)
 			{
-				AddBuildError<FMissingInputDataReferenceError>(InContext.Errors, InputDestination);
+				AddBuildError<FMissingInputDataReferenceError>(InOutContext.Results.Errors, InputDestination);
 
 				BuildStatus |= FBuildStatus::NonFatalError;
 				continue;
@@ -510,27 +422,27 @@ namespace Metasound
 		}
 
 		// Gather graph outputs.
-		for (const FSourceElement& Element : InContext.Graph.GetOutputDataSources())
+		for (const FSourceElement& Element : InOutContext.Graph.GetOutputDataSources())
 		{
 			const FOutputDataSource& OutputSource = Element.Value;
 
-			if (!InNodeDataReferences.Contains(OutputSource.Node))
+			if (!InOutContext.DataReferences.Contains(OutputSource.Node))
 			{
 				// An output node was likely pruned.
-				AddBuildError<FMissingOutputDataReferenceError>(InContext.Errors, OutputSource);
+				AddBuildError<FMissingOutputDataReferenceError>(InOutContext.Results.Errors, OutputSource);
 
 				BuildStatus |= FBuildStatus::NonFatalError;
 
 				continue;
 			}
 
-			FDataReferenceCollection& Collection = InNodeDataReferences[OutputSource.Node].Outputs;
+			const FDataReferenceCollection& Collection = InOutContext.DataReferences[OutputSource.Node].Outputs;
 
 			if (!Collection.ContainsDataReadReference(OutputSource.Vertex.VertexName, OutputSource.Vertex.DataTypeName))
 			{
 				// This will likely produce an IOperator which does not work as
 				// expected.
-				AddBuildError<FMissingOutputDataReferenceError>(InContext.Errors, OutputSource);
+				AddBuildError<FMissingOutputDataReferenceError>(InOutContext.Results.Errors, OutputSource);
 
 				BuildStatus |= FBuildStatus::NonFatalError;
 
@@ -543,7 +455,7 @@ namespace Metasound
 			{
 				// This will likely produce an IOperator which does not work as
 				// expected.
-				AddBuildError<FMissingOutputDataReferenceError>(InContext.Errors, OutputSource);
+				AddBuildError<FMissingOutputDataReferenceError>(InOutContext.Results.Errors, OutputSource);
 
 				BuildStatus |= FBuildStatus::NonFatalError;
 
@@ -554,14 +466,14 @@ namespace Metasound
 		return BuildStatus;
 	}
 
-	TUniquePtr<IOperator> FOperatorBuilder::CreateGraphOperator(FBuildContext& InContext, TArray<FOperatorPtr>& InOperators, FNodeDataReferenceMap& InNodeDataReferences) const
+	TUniquePtr<IOperator> FOperatorBuilder::CreateGraphOperator(FBuildContext& InOutContext, const FBuildGraphParams& InParams) const
 	{
 		FDataReferenceCollection GraphInputs;
 		FDataReferenceCollection GraphOutputs;
 
-		FBuildStatus BuildStatus = GatherGraphDataReferences(InContext, InNodeDataReferences, GraphInputs, GraphOutputs);
+		FBuildStatus BuildStatus = GatherGraphDataReferences(InOutContext, GraphInputs, GraphOutputs);
 
-		if (BuildStatus > MaxBuildStatusErrorLevel)
+		if (BuildStatus > FBuildStatus::GetMaxErrorLevel(InParams.BuilderSettings))
 		{
 			return TUniquePtr<IOperator>(nullptr);
 		}
@@ -571,12 +483,26 @@ namespace Metasound
 		GraphOperator->SetInputs(GraphInputs);
 		GraphOperator->SetOutputs(GraphOutputs);
 
-		for (FOperatorPtr& Ptr : InOperators)
+		if (InOutContext.BuilderSettings.bPopulateInternalDataReferences)
 		{
-			GraphOperator->AppendOperator(MoveTemp(Ptr));
+			TMap<FGuid, FDataReferenceCollection> References;
+			for (TPair<const INode*, FOperatorDataReferences>& ReferencePair : InOutContext.DataReferences)
+			{
+				const INode* Node = ReferencePair.Key;
+				check(Node);
+				References.Emplace(Node->GetInstanceID(), MoveTemp(ReferencePair.Value.Outputs));
+			}
+
+			if (InParams.BuilderSettings.bPopulateInternalDataReferences)
+			{
+				InOutContext.Results.InternalDataReferences.Append(MoveTemp(References));
+			}
 		}
 
-		InOperators.Reset();
+		for (FOperatorPtr& Ptr : InOutContext.Operators)
+		{
+			GraphOperator->AppendOperator(MoveTemp(Ptr));
+		}
 
 		return MoveTemp(GraphOperator);
 	}
diff --git a/Engine/Plugins/Runtime/Metasound/Source/MetasoundGraphCore/Private/MetasoundOperatorBuilderSettings.cpp b/Engine/Plugins/Runtime/Metasound/Source/MetasoundGraphCore/Private/MetasoundOperatorBuilderSettings.cpp
new file mode 100644
--- /dev/null
+++ b/Engine/Plugins/Runtime/Metasound/Source/MetasoundGraphCore/Private/MetasoundOperatorBuilderSettings.cpp
@@ -0,0 +1,96 @@
+// Copyright Epic Games, Inc. All Rights Reserved.
+
+#include "MetasoundOperatorBuilderSettings.h"
+
+#include "Core/Public/Misc/Build.h"
+
+
+namespace Metasound
+{
+	const FOperatorBuilderSettings& FOperatorBuilderSettings::GetDefaultSettings()
+	{
+#if UE_BUILD_SHIPPING
+		return GetDefaultShippingSettings();
+#elif UE_BUILD_TEST
+		return GetDefaultTestSettings();
+#elif UE_BUILD_DEVELOPMENT
+		return GetDefaultDevelopementSettings();
+#elif UE_BUILD_DEBUG
+		return GetDefaultDebugSettings();
+#else
+		return GetDefaultShippingSettings();
+#endif
+	}
+
+	const FOperatorBuilderSettings& FOperatorBuilderSettings::GetDefaultDebugSettings()
+	{
+		auto CreateSettings = [&]()
+		{
+			FOperatorBuilderSettings Settings;
+			Settings.PruningMode = EOperatorBuilderNodePruning::None;
+			Settings.bValidateNoCyclesInGraph = true;
+			Settings.bValidateNoDuplicateInputs = true;
+			Settings.bValidateVerticesExist = true;
+			Settings.bValidateEdgeDataTypesMatch = true;
+			Settings.bFailOnAnyError = false;
+			return Settings;
+		};
+
+		static const FOperatorBuilderSettings DebugSettings = CreateSettings();
+		return DebugSettings;
+	}
+
+	const FOperatorBuilderSettings& FOperatorBuilderSettings::GetDefaultDevelopementSettings()
+	{
+		auto CreateSettings = [&]()
+		{
+			FOperatorBuilderSettings Settings;
+			Settings.PruningMode = EOperatorBuilderNodePruning::PruneNodesWithoutExternalDependency;
+			Settings.bValidateNoCyclesInGraph = true;
+			Settings.bValidateNoDuplicateInputs = true;
+			Settings.bValidateVerticesExist = true;
+			Settings.bValidateEdgeDataTypesMatch = true;
+			Settings.bFailOnAnyError = false;
+			return Settings;
+		};
+
+		static const FOperatorBuilderSettings DevSettings = CreateSettings();
+		return DevSettings;
+	}
+
+	const FOperatorBuilderSettings& FOperatorBuilderSettings::GetDefaultTestSettings()
+	{
+		auto CreateSettings = [&]()
+		{
+			FOperatorBuilderSettings Settings;
+			Settings.PruningMode = EOperatorBuilderNodePruning::PruneNodesWithoutExternalDependency;
+			Settings.bValidateNoCyclesInGraph = false;
+			Settings.bValidateNoDuplicateInputs = false;
+			Settings.bValidateVerticesExist = false;
+			Settings.bValidateEdgeDataTypesMatch = false;
+			Settings.bFailOnAnyError = false;
+			return Settings;
+		};
+
+		static const FOperatorBuilderSettings TestSettings = CreateSettings();
+		return TestSettings;
+	}
+
+	const FOperatorBuilderSettings& FOperatorBuilderSettings::GetDefaultShippingSettings()
+	{
+		auto CreateSettings = [&]()
+		{
+			FOperatorBuilderSettings Settings;
+			Settings.PruningMode = EOperatorBuilderNodePruning::PruneNodesWithoutExternalDependency;
+			Settings.bValidateNoCyclesInGraph = false;
+			Settings.bValidateNoDuplicateInputs = false;
+			Settings.bValidateVerticesExist = false;
+			Settings.bValidateEdgeDataTypesMatch = false;
+			Settings.bFailOnAnyError = false;
+			return Settings;
+		};
+
+		static const FOperatorBuilderSettings ShippingSettings = CreateSettings();
+		return ShippingSettings;
+	}
+} // namespace Metasound
\ No newline at end of file
diff --git a/Engine/Plugins/Runtime/Metasound/Source/MetasoundGraphCore/Public/MetasoundBuilderInterface.h b/Engine/Plugins/Runtime/Metasound/Source/MetasoundGraphCore/Public/MetasoundBuilderInterface.h
--- a/Engine/Plugins/Runtime/Metasound/Source/MetasoundGraphCore/Public/MetasoundBuilderInterface.h
+++ b/Engine/Plugins/Runtime/Metasound/Source/MetasoundGraphCore/Public/MetasoundBuilderInterface.h
@@ -2,10 +2,14 @@
 
 #pragma once
 
-#include "CoreMinimal.h"
+#include "Containers/Array.h"
 #include "MetasoundEnvironment.h"
 #include "MetasoundNodeInterface.h"
+#include "MetasoundOperatorBuilderSettings.h"
 #include "MetasoundOperatorInterface.h"
+#include "Templates/UniquePtr.h"
+#include "UObject/NameTypes.h"
+
 
 namespace Metasound
 {
@@ -40,12 +44,28 @@ namespace Metasound
 			virtual const TArray<FDataEdge>& GetDataEdges() const = 0;
 	};
 
+	/** Map of NodeID (Guid) to DataReferenceCollection. */
+	// TODO: Move to FAnyDataReference
+	using FNodeVertexDataMap = TMap<FGuid, FDataReferenceCollection>;
+
+	/** Array of build errors. */
+	using FBuildErrorArray = TArray<TUniquePtr<IOperatorBuildError>>;
+
+	/** Structure of all resulting data generated during graph operator build. */
+	struct FBuildGraphResults
+	{
+		/** An array of errors. Errors can be added if issues occur while creating an IOperator. */
+		FBuildErrorArray Errors;
+
+		/** References used by graph analyzer, if enabled (not collected if analysis is disabled). */
+		FNodeVertexDataMap InternalDataReferences;
+	};
+
 	/** FCreateOperatorParams holds the parameters provided to operator factories
 	 * during the creation of an IOperator
 	 */
 	struct FCreateOperatorParams
 	{
-
 		/** The node associated with this factory and the desired IOperator. */
 		const INode& Node;
 
@@ -58,13 +78,13 @@ namespace Metasound
 		/** Environment settings available. */
 		const FMetasoundEnvironment& Environment;
 
+		/** Settings for the builder requesting creation of the IOperator. */
+		const FOperatorBuilderSettings& BuilderSettings;
+
 		/** Pointer to builder actively building graph. */
 		const IOperatorBuilder* Builder = nullptr;
 	};
 
-	/** Array of build errors. */
-	typedef TArray<TUniquePtr<IOperatorBuildError>> FBuildErrorArray;
-
 	/** Convenience template for adding build errors.
 	 *
 	 * The function can be used in the following way:
@@ -92,14 +112,16 @@ namespace Metasound
 		public:
 			virtual ~IOperatorFactory() = default;
 
-			/** Create a new IOperator.
-			 *
-			 * @param InParams - The parameters available for building an IOperator.
-			 * @param OutErrors - An array of errors. Errors can be added if issues occur while creating the IOperator.
-			 *
-			 * @return A unique pointer to an IOperator. 
-			 */
-			virtual TUniquePtr<IOperator> CreateOperator(const FCreateOperatorParams& InParams, FBuildErrorArray& OutErrors) = 0;
+			virtual TUniquePtr<IOperator> CreateOperator(const FCreateOperatorParams& InParams, FBuildGraphResults& OutResults) = 0;
+
+			UE_DEPRECATED(5.1, "Use CreateOperator overload providing the FBuildGraphResults struct.")
+			virtual TUniquePtr<IOperator> CreateOperator(const FCreateOperatorParams& InParams, FBuildErrorArray& OutErrors)
+			{
+				FBuildGraphResults Results { MoveTemp(OutErrors) };
+				TUniquePtr<IOperator> Operator = CreateOperator(InParams, Results);
+				OutErrors = MoveTemp(Results.Errors);
+				return Operator;
+			}
 	};
 
 	struct FBuildGraphParams
@@ -110,11 +132,14 @@ namespace Metasound
 		/** General operator settings for the graph. */
 		const FOperatorSettings& OperatorSettings;
 
-		/** Collection of input parameters available for to an IOperator. */
+		/** Collection of input parameters available for an IOperator. */
 		const FDataReferenceCollection& InputDataReferences;
 
 		/** Environment settings available. */
 		const FMetasoundEnvironment& Environment;
+
+		/** Settings for graph's operator builder. */
+		const FOperatorBuilderSettings& BuilderSettings;
 	};
 
 	/** IOperatorBuilder
@@ -129,14 +154,16 @@ namespace Metasound
 
 			virtual ~IOperatorBuilder() = default;
 
-			/** Build a graph operator from a graph. 
+			UE_DEPRECATED(5.1, "Use BuildGraphOperator overload providing the FBuildGraphResults struct.")
+			virtual TUniquePtr<IOperator> BuildGraphOperator(const FBuildGraphParams& InParams, FBuildErrorArray& OutErrors) const { return { }; }
+
+			/** Build a graph operator from a graph.
 			 *
 			 * @params InParams - Input parameters for building a graph.
-			 * @param OutErrors - An array of errors. Errors can be added if issues occur while creating the IOperator.
+			 * @param OutResults - Results data pertaining to the given build operator result.
 			 *
-			 * @return A unique pointer to an IOperator. 
+			 * @return A unique pointer to the built IOperator. Null if build failed.
 			 */
-			virtual TUniquePtr<IOperator> BuildGraphOperator(const FBuildGraphParams& InParams, FBuildErrorArray& OutErrors) const = 0;
+			virtual TUniquePtr<IOperator> BuildGraphOperator(const FBuildGraphParams& InParams, FBuildGraphResults& OutResults) const = 0;
 	};
 }
-
diff --git a/Engine/Plugins/Runtime/Metasound/Source/MetasoundGraphCore/Public/MetasoundFacade.h b/Engine/Plugins/Runtime/Metasound/Source/MetasoundGraphCore/Public/MetasoundFacade.h
--- a/Engine/Plugins/Runtime/Metasound/Source/MetasoundGraphCore/Public/MetasoundFacade.h
+++ b/Engine/Plugins/Runtime/Metasound/Source/MetasoundGraphCore/Public/MetasoundFacade.h
@@ -99,7 +99,7 @@ namespace Metasound
 				public:
 					FFactory(FCreateOperatorFunction InCreateFunc);
 
-					virtual TUniquePtr<IOperator> CreateOperator(const FCreateOperatorParams& InParams, FBuildErrorArray& OutErrors) override;
+					virtual TUniquePtr<IOperator> CreateOperator(const FCreateOperatorParams& InParams, FBuildGraphResults& OutResults) override;
 
 				private:
 					FCreateOperatorFunction CreateFunc;
diff --git a/Engine/Plugins/Runtime/Metasound/Source/MetasoundGraphCore/Public/MetasoundGraph.h b/Engine/Plugins/Runtime/Metasound/Source/MetasoundGraphCore/Public/MetasoundGraph.h
--- a/Engine/Plugins/Runtime/Metasound/Source/MetasoundGraphCore/Public/MetasoundGraph.h
+++ b/Engine/Plugins/Runtime/Metasound/Source/MetasoundGraphCore/Public/MetasoundGraph.h
@@ -1,11 +1,11 @@
 // Copyright Epic Games, Inc. All Rights Reserved.
 #pragma once
 
-#include "CoreMinimal.h"
+#include "Containers/Array.h"
 #include "MetasoundBuilderInterface.h"
 #include "MetasoundNode.h"
-
-// Forward Declarations
+#include "MetasoundNodeInterface.h"
+#include "MetasoundVertex.h"
 
 
 namespace Metasound
@@ -133,7 +133,7 @@ namespace Metasound
 			public:
 				virtual ~FFactory() = default;
 
-				virtual TUniquePtr<IOperator> CreateOperator(const FCreateOperatorParams& InParams, FBuildErrorArray& OutErrors) override;
+				virtual TUniquePtr<IOperator> CreateOperator(const FCreateOperatorParams& InParams, FBuildGraphResults& OutResults) override;
 			};
 
 			FVertexName InstanceName;
diff --git a/Engine/Plugins/Runtime/Metasound/Source/MetasoundGraphCore/Public/MetasoundLiteralNode.h b/Engine/Plugins/Runtime/Metasound/Source/MetasoundGraphCore/Public/MetasoundLiteralNode.h
--- a/Engine/Plugins/Runtime/Metasound/Source/MetasoundGraphCore/Public/MetasoundLiteralNode.h
+++ b/Engine/Plugins/Runtime/Metasound/Source/MetasoundGraphCore/Public/MetasoundLiteralNode.h
@@ -90,7 +90,7 @@ namespace Metasound
 		{
 		}
 
-		virtual TUniquePtr<IOperator> CreateOperator(const FCreateOperatorParams& InParams, FBuildErrorArray& OutErrors) override;
+		virtual TUniquePtr<IOperator> CreateOperator(const FCreateOperatorParams& InParams, FBuildGraphResults& OutResults) override;
 
 	private:
 		FLiteral InitParam;
@@ -171,7 +171,7 @@ namespace Metasound
 	};
 
 	template<typename DataType>
-	TUniquePtr<IOperator> TLiteralOperatorLiteralFactory<DataType>::CreateOperator(const FCreateOperatorParams& InParams, FBuildErrorArray& OutErrors)
+	TUniquePtr<IOperator> TLiteralOperatorLiteralFactory<DataType>::CreateOperator(const FCreateOperatorParams& InParams, FBuildGraphResults& OutResults)
 	{
 		using FLiteralNodeType = TLiteralNode<DataType>;
 
diff --git a/Engine/Plugins/Runtime/Metasound/Source/MetasoundGraphCore/Public/MetasoundOperatorBuilder.h b/Engine/Plugins/Runtime/Metasound/Source/MetasoundGraphCore/Public/MetasoundOperatorBuilder.h
--- a/Engine/Plugins/Runtime/Metasound/Source/MetasoundGraphCore/Public/MetasoundOperatorBuilder.h
+++ b/Engine/Plugins/Runtime/Metasound/Source/MetasoundGraphCore/Public/MetasoundOperatorBuilder.h
@@ -2,91 +2,15 @@
 
 #pragma once
 
-#include "CoreMinimal.h"
+#include "MetasoundBuilderInterface.h"
 #include "MetasoundEnvironment.h"
+#include "MetasoundGraphOperator.h"
 #include "MetasoundNodeInterface.h"
+#include "MetasoundOperatorBuilderSettings.h"
 #include "MetasoundOperatorInterface.h"
-#include "MetasoundBuilderInterface.h"
 
 namespace Metasound
 {
-	/** EOperatorBuildNodePruning expresses the desired pruning behavior during
-	 * the node pruning step.  
-	 *
-	 * Some nodes are unreachable in the graph either by traversing from the
-	 * input nodes to from the output nodes. Because they are not dependent, 
-	 * they have no impact on the produced output and can be pruned without 
-	 * causing any change to the declared behavior of the graph.
-	 */
-	enum class EOperatorBuilderNodePruning : uint8
-	{
-		/** Do not prune any nodes. */
-		None, 
-
-		/** Prune nodes which cannot be reached from the output nodes. */
-		PruneNodesWithoutOutputDependency, 
-
-		/** Prune nodes which cannot be reached from the input nodes. */
-		PruneNodesWithoutInputDependency, 
-
-		/** Prune nodes which cannot be reached from the input nodes or output nodes. */
-		PruneNodesWithoutExternalDependency, 
-	};
-
-	/** FOperatorBuilderSettings
-	 *
-	 * Settings for building IGraphs into IOperators.
-	 */
-	struct METASOUNDGRAPHCORE_API FOperatorBuilderSettings
-	{
-		/** Desired node pruning behavior. */
-		EOperatorBuilderNodePruning PruningMode = EOperatorBuilderNodePruning::None;
-
-		/** If true, the IGraph will be analyzed to detect cycles. Errors will be
-		 * generated if a cycle is detected in the graph.
-		 */
-		bool bValidateNoCyclesInGraph = true;
-
-		/** If true, the inputs to each node in the IGraph will be analyzed to
-		 * detect duplicate inputs connected to an individual vertex on a given 
-		 * node.  Errors will be generated if duplicates are detected. */
-		bool bValidateNoDuplicateInputs = true;
-
-		/** If true, each FDataEdge in the IGraph will be validated by checking
-		 * that the corresponding INodes contain matching FDataVertex information
-		 * as described by the FDataEdge. Errors will be generated if 
-		 * inconsistencies are detected.
-		 */
-		bool bValidateVerticesExist = true;
-
-		/** If true, each FDataEdge in the IGraph will be validated by checking
-		 * that the FInputDataSource and FOutputDataDestination data types are
-		 * equal. Errors will be generated if unequal data types are detected.
-		 */
-		bool bValidateEdgeDataTypesMatch = true;
-
-		/** If true, the builder will return an invalid IOperator if any errors
-		 * are detected. If false, the builder will return an invalid IOperator
-		 * only if fatal errors are detected.
-		 */
-		bool bFailOnAnyError = false;
-
-		/** Return the default settings for the current build environment. */
-		static FOperatorBuilderSettings GetDefaultSettings();
-
-		/** Return the default settings for a debug build environment. */
-		static FOperatorBuilderSettings GetDefaultDebugSettings();
-
-		/** Return the default settings for a development build environment. */
-		static FOperatorBuilderSettings GetDefaultDevelopementSettings();
-
-		/** Return the default settings for a test build environment. */
-		static FOperatorBuilderSettings GetDefaultTestSettings();
-
-		/** Return the default settings for a shipping build environment. */
-		static FOperatorBuilderSettings GetDefaultShippingSettings();
-	};
-
 	// Forward declare.
 	class FDirectedGraphAlgoAdapter;
 
@@ -99,22 +23,19 @@ namespace Metasound
 			 *
 			 * @param InBuilderSettings  - Settings to configure builder options.
 			 */
-			FOperatorBuilder(const FOperatorBuilderSettings& InBuilderSettings);
+			FOperatorBuilder() = default;
 
 			virtual ~FOperatorBuilder();
 
 			/** Create an IOperator from an IGraph.
 			 *
-			 * @param InGraph            - The graph containing input, output and edge information.
-			 * @param InOperatorSettings - Settings to be passed to all operators on creation.
-			 * @param InEnvironment      - The environment variables to use during construction. 
-			 * @param OutErrors          - An array of build errors that will be populated with any issues encountered during the build process.
+			 * @param InParams   - Params of the current build
+			 * @param OutResults - Results data pertaining to the given build operator result.
 			 *
 			 * @return A TUniquePtr to an IOperator. If the processes was unsuccessful, 
 			 *         the returned pointer will contain a nullptr and be invalid.
 			 */
-			virtual TUniquePtr<IOperator> BuildGraphOperator(const FBuildGraphParams& InParams, FBuildErrorArray& OutErrors) const override;
-			
+			virtual TUniquePtr<IOperator> BuildGraphOperator(const FBuildGraphParams& InParams, FBuildGraphResults& OutResults) const override;
 
 		private:
 			// Collection of existing inputs and outputs associated with a given
@@ -150,7 +71,7 @@ namespace Metasound
 				// build status int32 values increase as the build status deteriorates.
 				// Build statuses are merged by taking the maximum int32 value of
 				// the EStatus. 
-				enum EStatus 
+				enum EStatus
 				{
 					// No error has been encountered.
 					NoError = 0,
@@ -162,6 +83,11 @@ namespace Metasound
 					FatalError = 2
 				};
 
+				static FBuildStatus GetMaxErrorLevel(const FOperatorBuilderSettings& InBuilderSettings)
+				{
+					return InBuilderSettings.bFailOnAnyError ? EStatus::NoError : EStatus::NonFatalError;
+				}
+
 				FBuildStatus() = default;
 
 				FBuildStatus(FBuildStatus::EStatus InStatus)
@@ -181,8 +107,7 @@ namespace Metasound
 					return Value;
 				}
 
-				private:
-
+			private:
 				EStatus Value = NoError;
 			};
 
@@ -192,46 +117,53 @@ namespace Metasound
 				const FDirectedGraphAlgoAdapter& AlgoAdapter;
 				const FOperatorSettings& Settings;
 				const FMetasoundEnvironment& Environment;
+				const FOperatorBuilderSettings& BuilderSettings;
+
+				FBuildGraphResults& Results;
 
-				TArray<FBuildErrorPtr>& Errors;
+				TArray<FOperatorPtr> Operators;
+				FNodeDataReferenceMap DataReferences;
 
-				FBuildContext(const IGraph& InGraph, FDirectedGraphAlgoAdapter& InAlgoAdapter, const FOperatorSettings& InSettings, const FMetasoundEnvironment& InEnvironment, TArray<FBuildErrorPtr>& InOutErrors)
+				FBuildContext(
+					const IGraph& InGraph,
+					const FDirectedGraphAlgoAdapter& InAlgoAdapter,
+					const FOperatorSettings& InSettings,
+					const FMetasoundEnvironment& InEnvironment,
+					const FOperatorBuilderSettings& bInBuilderSettings,
+					FBuildGraphResults& OutResults)
 				:	Graph(InGraph)
 				,	AlgoAdapter(InAlgoAdapter)
 				,	Settings(InSettings)
 				,	Environment(InEnvironment)
-				,	Errors(InOutErrors)
+				,	BuilderSettings(bInBuilderSettings)
+				,	Results(OutResults)
 				{
 				}
 			};
 
 			// Perform topological sort using depth first algorithm.
-			FBuildStatus DepthFirstTopologicalSort(FBuildContext& InContext, TArray<const INode*>& OutNodes) const;
+			FBuildStatus DepthFirstTopologicalSort(FBuildContext& InOutContext, TArray<const INode*>& OutNodes) const;
 
 			// Perform topological sort using kahns algorithm.
-			FBuildStatus KahnsTopologicalSort(FBuildContext& InContext, TArray<const INode*>& OutNodes) const;
+			FBuildStatus KahnsTopologicalSort(FBuildContext& InOutContext, TArray<const INode*>& OutNodes) const;
 
 			// Prune unreachable nodes from InOutNodes
-			FBuildStatus PruneNodes(FBuildContext& InContext, TArray<const INode*>& InOutNodes) const;
+			FBuildStatus PruneNodes(FBuildContext& InOutContext, TArray<const INode*>& InOutNodes) const;
 
 			// Get all input data references for a given node for inputs provided internally to the graph.
-			FBuildStatus GatherInputDataReferences(FBuildContext& InContext, const INode* InNode, const FNodeEdgeMultiMap& InEdgeMap, const FNodeDataReferenceMap& InDataReferenceMap, FDataReferenceCollection& OutCollection) const;
+			FBuildStatus GatherInputDataReferences(FBuildContext& InOutContext, const INode* InNode, const FNodeEdgeMultiMap& InEdgeMap, FDataReferenceCollection& OutInputCollection) const;
 
 			// Get all input data references for a given node for inputs provided externally to the graph.
-			FBuildStatus GatherExternalInputDataReferences(FBuildContext& InContext, const INode* InNode, const FNodeDestinationMap& InNodeDestinationMap, const FDataReferenceCollection& InExternalCollection, FDataReferenceCollection& OutDataReferences) const;
+			FBuildStatus GatherExternalInputDataReferences(FBuildContext& InOutContext, const INode* InNode, const FNodeDestinationMap& InNodeDestinationMap, const FDataReferenceCollection& InExternalCollection, FDataReferenceCollection& OutDataReferences) const;
 
 			// Get all input/output data references for a given graph.
-			FBuildStatus GatherGraphDataReferences(FBuildContext& InContext, FNodeDataReferenceMap& InNodeDataReferences, FDataReferenceCollection& OutGraphInputs, FDataReferenceCollection& OutGraphOutputs) const;
+			FBuildStatus GatherGraphDataReferences(FBuildContext& InOutContext, FDataReferenceCollection& OutGraphInputs, FDataReferenceCollection& OutGraphOutputs) const;
 
 			// Call the operator factories for the nodes
-			FBuildStatus CreateOperators(FBuildContext& InContext, const TArray<const INode*>& InSortedNodes, const FDataReferenceCollection& InGraphInputs, TArray<FOperatorPtr>& OutOperators, FNodeDataReferenceMap& OutDataReferences) const;
-
-			// Create the final graph operator from the individual operators.
-			TUniquePtr<IOperator> CreateGraphOperator(FBuildContext& InContext, TArray<FOperatorPtr>& InOperators, FNodeDataReferenceMap& InNodeDataReferences) const;
-
-			FOperatorBuilderSettings BuilderSettings;
+			FBuildStatus CreateOperators(FBuildContext& InOutContext, const TArray<const INode*>& InSortedNodes, const FDataReferenceCollection& InGraphInputs) const;
 
-			FBuildStatus MaxBuildStatusErrorLevel;
+			// Create the final graph operator from the provided build context.
+			TUniquePtr<IOperator> CreateGraphOperator(FBuildContext& InOutContext, const FBuildGraphParams& InParams) const;
 	};
 }
 
diff --git a/Engine/Plugins/Runtime/Metasound/Source/MetasoundGraphCore/Public/MetasoundOperatorBuilderSettings.h b/Engine/Plugins/Runtime/Metasound/Source/MetasoundGraphCore/Public/MetasoundOperatorBuilderSettings.h
new file mode 100644
--- /dev/null
+++ b/Engine/Plugins/Runtime/Metasound/Source/MetasoundGraphCore/Public/MetasoundOperatorBuilderSettings.h
@@ -0,0 +1,86 @@
+// Copyright Epic Games, Inc. All Rights Reserved.
+#pragma once
+
+namespace Metasound
+{
+	/** EOperatorBuildNodePruning expresses the desired pruning behavior during
+	 * the node pruning step.  
+	 *
+	 * Some nodes are unreachable in the graph either by traversing from the
+	 * input nodes to from the output nodes. Because they are not dependent, 
+	 * they have no impact on the produced output and can be pruned without 
+	 * causing any change to the declared behavior of the graph.
+	 */
+	enum class EOperatorBuilderNodePruning : uint8
+	{
+		/** Do not prune any nodes. */
+		None, 
+
+		/** Prune nodes which cannot be reached from the output nodes. */
+		PruneNodesWithoutOutputDependency, 
+
+		/** Prune nodes which cannot be reached from the input nodes. */
+		PruneNodesWithoutInputDependency, 
+
+		/** Prune nodes which cannot be reached from the input nodes or output nodes. */
+		PruneNodesWithoutExternalDependency, 
+	};
+
+	/** FOperatorBuilderSettings
+	 *
+	 * Settings for building IGraphs into IOperators.
+	 */
+	struct METASOUNDGRAPHCORE_API FOperatorBuilderSettings
+	{
+		/** Desired node pruning behavior. */
+		EOperatorBuilderNodePruning PruningMode = EOperatorBuilderNodePruning::None;
+
+		/** If true, the IGraph will be analyzed to detect cycles. Errors will be
+		 * generated if a cycle is detected in the graph.
+		 */
+		bool bValidateNoCyclesInGraph = true;
+
+		/** If true, the inputs to each node in the IGraph will be analyzed to
+		 * detect duplicate inputs connected to an individual vertex on a given 
+		 * node.  Errors will be generated if duplicates are detected. */
+		bool bValidateNoDuplicateInputs = true;
+
+		/** If true, each FDataEdge in the IGraph will be validated by checking
+		 * that the corresponding INodes contain matching FDataVertex information
+		 * as described by the FDataEdge. Errors will be generated if 
+		 * inconsistencies are detected.
+		 */
+		bool bValidateVerticesExist = true;
+
+		/** If true, each FDataEdge in the IGraph will be validated by checking
+		 * that the FInputDataSource and FOutputDataDestination data types are
+		 * equal. Errors will be generated if unequal data types are detected.
+		 */
+		bool bValidateEdgeDataTypesMatch = true;
+
+		/** If true, the builder will return an invalid IOperator if any errors
+		 * are detected. If false, the builder will return an invalid IOperator
+		 * only if fatal errors are detected.
+		 */
+		bool bFailOnAnyError = false;
+
+		/** If true, enables tracking all internal data references (can be used
+		  * for analyzing internal graph state by managing sound generator). */
+		bool bPopulateInternalDataReferences = false;
+
+		/** Return the default settings for the current build environment. */
+		static const FOperatorBuilderSettings& GetDefaultSettings();
+
+		/** Return the default settings for a debug build environment. */
+		static const FOperatorBuilderSettings& GetDefaultDebugSettings();
+
+		/** Return the default settings for a development build environment. */
+		static const FOperatorBuilderSettings& GetDefaultDevelopementSettings();
+
+		/** Return the default settings for a test build environment. */
+		static const FOperatorBuilderSettings& GetDefaultTestSettings();
+
+		/** Return the default settings for a shipping build environment. */
+		static const FOperatorBuilderSettings& GetDefaultShippingSettings();
+	};
+} // namespace Metasound
\ No newline at end of file
diff --git a/Engine/Plugins/Runtime/Metasound/Source/MetasoundGraphCore/Public/MetasoundVariableNodes.h b/Engine/Plugins/Runtime/Metasound/Source/MetasoundGraphCore/Public/MetasoundVariableNodes.h
--- a/Engine/Plugins/Runtime/Metasound/Source/MetasoundGraphCore/Public/MetasoundVariableNodes.h
+++ b/Engine/Plugins/Runtime/Metasound/Source/MetasoundGraphCore/Public/MetasoundVariableNodes.h
@@ -173,7 +173,7 @@ namespace Metasound
 
 			virtual ~FFactory() = default;
 
-			virtual TUniquePtr<IOperator> CreateOperator(const FCreateOperatorParams& InParams, FBuildErrorArray& OutErrors) override
+			virtual TUniquePtr<IOperator> CreateOperator(const FCreateOperatorParams& InParams, FBuildGraphResults& OutResults) override
 			{
 				TDataWriteReference<DataType> Data = TDataWriteReferenceLiteralFactory<DataType>::CreateExplicitArgs(InParams.OperatorSettings, Literal);
 				return MakeUnique<FOperator>(TDataWriteReference<FVariable>::CreateNew(Data));
diff --git a/Engine/Plugins/Runtime/Metasound/Source/MetasoundStandardNodes/Private/MetasoundAudioFormats.cpp b/Engine/Plugins/Runtime/Metasound/Source/MetasoundStandardNodes/Private/MetasoundAudioFormats.cpp
--- a/Engine/Plugins/Runtime/Metasound/Source/MetasoundStandardNodes/Private/MetasoundAudioFormats.cpp
+++ b/Engine/Plugins/Runtime/Metasound/Source/MetasoundStandardNodes/Private/MetasoundAudioFormats.cpp
@@ -153,7 +153,7 @@ namespace Metasound
 			{
 			}
 
-			TUniquePtr<IOperator> CreateOperator(const FCreateOperatorParams& InParams, FBuildErrorArray& OutErrors) override
+			TUniquePtr<IOperator> CreateOperator(const FCreateOperatorParams& InParams, FBuildGraphResults& OutResults) override
 			{
 				using namespace StereoAudioFormatVertexKeys;
 
@@ -298,7 +298,7 @@ namespace Metasound
 			{
 			}
 
-			TUniquePtr<IOperator> CreateOperator(const FCreateOperatorParams& InParams, FBuildErrorArray& OutErrors) override
+			TUniquePtr<IOperator> CreateOperator(const FCreateOperatorParams& InParams, FBuildGraphResults& OutResults) override
 			{
 				// Split a stereo signal into left/right.
 				TDataReadReference<FStereoAudioFormat> Stereo = InParams.InputDataReferences.GetDataReadReferenceOrConstruct<FStereoAudioFormat>(InputName, InParams.OperatorSettings);
@@ -436,7 +436,7 @@ namespace Metasound
 			{
 			}
 
-			TUniquePtr<IOperator> CreateOperator(const FCreateOperatorParams& InParams, FBuildErrorArray& OutErrors) override
+			TUniquePtr<IOperator> CreateOperator(const FCreateOperatorParams& InParams, FBuildGraphResults& OutResults) override
 			{
 				TDataReadReference<FAudioBuffer> Center = InParams.InputDataReferences.GetDataReadReferenceOrConstruct<FAudioBuffer>(OutputName, InParams.OperatorSettings);
 
@@ -570,7 +570,7 @@ namespace Metasound
 			{
 			}
 
-			TUniquePtr<IOperator> CreateOperator(const FCreateOperatorParams& InParams, FBuildErrorArray& OutErrors) override
+			TUniquePtr<IOperator> CreateOperator(const FCreateOperatorParams& InParams, FBuildGraphResults& OutResults) override
 			{
 				TDataReadReference<FMonoAudioFormat> Mono = InParams.InputDataReferences.GetDataReadReferenceOrConstruct<FMonoAudioFormat>(InputName, InParams.OperatorSettings);
 				TDataReadReference<FAudioBuffer> Center = Mono->GetCenter();
diff --git a/Engine/Plugins/Runtime/Metasound/Source/MetasoundStandardNodes/Private/MetasoundOscillatorNodes.cpp b/Engine/Plugins/Runtime/Metasound/Source/MetasoundStandardNodes/Private/MetasoundOscillatorNodes.cpp
--- a/Engine/Plugins/Runtime/Metasound/Source/MetasoundStandardNodes/Private/MetasoundOscillatorNodes.cpp
+++ b/Engine/Plugins/Runtime/Metasound/Source/MetasoundStandardNodes/Private/MetasoundOscillatorNodes.cpp
@@ -805,7 +805,7 @@ namespace Metasound
 			return Interface;
 		}
 		
-		TUniquePtr<IOperator> CreateOperator(const FCreateOperatorParams& InParams, FBuildErrorArray& OutErrors) override
+		TUniquePtr<IOperator> CreateOperator(const FCreateOperatorParams& InParams, FBuildGraphResults& OutResults) override
 		{
 			const FSineOscilatorNode& SineNode = static_cast<const FSineOscilatorNode&>(InParams.Node);
 			const FDataReferenceCollection& InputCol = InParams.InputDataReferences;
@@ -856,7 +856,6 @@ namespace Metasound
 			} // HasFM
 			return nullptr;
 		}
-	private:
 	};
 
 	FSineOscilatorNode::FSineOscilatorNode(const FVertexName& InInstanceName, const FGuid& InInstanceID, float InDefaultFrequency, float InDefautlGlideFactor, bool bInDefaultEnablement)
@@ -896,10 +895,10 @@ namespace Metasound
 
 	class FSawOscilatorNode::FFactory : public FOscilatorFactoryBase
 	{
-	public:		
+	public:
 		FFactory() = default;
 
-		TUniquePtr<IOperator> CreateOperator(const FCreateOperatorParams& InParams, FBuildErrorArray& OutErrors) override;
+		TUniquePtr<IOperator> CreateOperator(const FCreateOperatorParams& InParams, FBuildGraphResults& OutResults) override;
 
 		static const FNodeClassMetadata& GetNodeInfo()
 		{
@@ -920,6 +919,7 @@ namespace Metasound
 			static const FNodeClassMetadata Info = InitNodeInfo();
 			return Info;
 		}
+
 		static const FVertexInterface& GetVertexInterface()
 		{
 			using namespace SawOscilatorVertexNames;
@@ -935,18 +935,16 @@ namespace Metasound
 			static const FVertexInterface Interface = MakeInterface();
 			return Interface;
 		}
-
-	private:
 	};
 
-	TUniquePtr<Metasound::IOperator> FSawOscilatorNode::FFactory::CreateOperator(const FCreateOperatorParams& InParams, FBuildErrorArray& OutErrors)
+	TUniquePtr<Metasound::IOperator> FSawOscilatorNode::FFactory::CreateOperator(const FCreateOperatorParams& InParams, FBuildGraphResults& OutResults)
 	{
 		const FSawOscilatorNode& Node = static_cast<const FSawOscilatorNode&>(InParams.Node);
 		const FDataReferenceCollection& InputCol = InParams.InputDataReferences;
 		const FOperatorSettings& Settings = InParams.OperatorSettings;
-		using namespace Generators;		
-		using namespace OscillatorCommonVertexNames; 
-		using namespace SawOscilatorVertexNames; 
+		using namespace Generators;
+		using namespace OscillatorCommonVertexNames;
+		using namespace SawOscilatorVertexNames;
 
 		FSawGenerationTypeReadRef Type = InputCol.GetDataReadReferenceOrConstruct<FEnumSawGenerationType>(METASOUND_GET_PARAM_NAME(SawType));
 
@@ -1094,9 +1092,9 @@ namespace Metasound
 	public:
 		FFactory() = default;
 
-		TUniquePtr<IOperator> CreateOperator(const FCreateOperatorParams& InParams, FBuildErrorArray& OutErrors) override
+		TUniquePtr<IOperator> CreateOperator(const FCreateOperatorParams& InParams, FBuildGraphResults& OutResults) override
 		{
-			using namespace SquareOscillatorVertexNames; 
+			using namespace SquareOscillatorVertexNames;
 
 			const FSquareOscilatorNode& Node = static_cast<const FSquareOscilatorNode&>(InParams.Node);
 			const FDataReferenceCollection& InputCol = InParams.InputDataReferences;
@@ -1217,7 +1215,7 @@ namespace Metasound
 	public:
 		FFactory() = default;
 
-		TUniquePtr<IOperator> CreateOperator(const FCreateOperatorParams& InParams, FBuildErrorArray& OutErrors) override
+		TUniquePtr<IOperator> CreateOperator(const FCreateOperatorParams& InParams, FBuildGraphResults& OutResults) override
 		{
 			const FTriangleOscilatorNode& Node = static_cast<const FTriangleOscilatorNode&>(InParams.Node);
 			const FDataReferenceCollection& InputCol = InParams.InputDataReferences;
@@ -1282,6 +1280,7 @@ namespace Metasound
 			static const FNodeClassMetadata Info = InitNodeInfo();
 			return Info;
 		}
+
 		static const FVertexInterface& GetVertexInterface()
 		{
 			using namespace TriangleOscilatorVertexNames; 
@@ -1297,7 +1296,6 @@ namespace Metasound
 			static const FVertexInterface Interface = MakeInterface();
 			return Interface;
 		}
-	private:
 	};
 	
 	FTriangleOscilatorNode::FTriangleOscilatorNode(const FVertexName& InInstanceName, const FGuid& InInstanceID, float InDefaultFrequency, float InDefaultGlideFactor, bool bInDefaultEnablement)
