RingMod MetaSound Node
- Adding a node for Ring Modulation

diff --git a/Engine/Plugins/Runtime/Metasound/Source/MetasoundStandardNodes/Private/MetasoundRingModNode.cpp b/Engine/Plugins/Runtime/Metasound/Source/MetasoundStandardNodes/Private/MetasoundRingModNode.cpp
new file mode 100644
--- /dev/null
+++ b/Engine/Plugins/Runtime/Metasound/Source/MetasoundStandardNodes/Private/MetasoundRingModNode.cpp
@@ -0,0 +1,151 @@
+// Copyright Epic Games, Inc. All Rights Reserved.
+
+#include "Internationalization/Text.h"
+#include "MetasoundExecutableOperator.h"
+#include "MetasoundEnumRegistrationMacro.h"
+#include "MetasoundNodeRegistrationMacro.h"
+#include "MetasoundStandardNodesNames.h"
+#include "MetasoundAudioBuffer.h"
+#include "MetasoundStandardNodesCategories.h"
+#include "MetasoundFacade.h"
+#include "MetasoundParamHelper.h"
+
+#define LOCTEXT_NAMESPACE "MetasoundStandardNodes_MetasoundRingModNodes"
+
+namespace Metasound
+{
+	/* RingMod Processor */
+	namespace RingModVertexNames
+	{
+		METASOUND_PARAM(InputAudioCarrier, "In Carrier", "The carrier audio signal.")
+		METASOUND_PARAM(InputAudioModulator, "In Modulator", "The modulator audio signal.")
+
+		METASOUND_PARAM(OutputAudio, "Out Audio", "The modulated audio signal.")
+	}
+
+	// Operator Class
+	class FRingModOperator : public TExecutableOperator<FRingModOperator>
+	{
+	public:
+
+		FRingModOperator(const FOperatorSettings& InSettings,
+			const FAudioBufferReadRef& InputAudioCarrier,
+			const FAudioBufferReadRef& InputAudioModulator)
+			: AudioInputCarrier(InputAudioCarrier)
+			, AudioInputModulator(InputAudioModulator)
+			, AudioOutput(FAudioBufferWriteRef::CreateNew(InSettings))
+		{
+		}
+
+		static const FNodeClassMetadata& GetNodeInfo()
+		{
+			auto CreateNodeClassMetadata = []() -> FNodeClassMetadata
+			{
+				FVertexInterface NodeInterface = DeclareVertexInterface();
+
+				FNodeClassMetadata Metadata
+				{
+					FNodeClassName { StandardNodes::Namespace, "RingMod", StandardNodes::AudioVariant },
+					1, // Major Version
+					0, // Minor Version
+					METASOUND_LOCTEXT("RingModDisplayName", "RingMod"),
+					METASOUND_LOCTEXT("RingModDesc", "Modulates a carrier signal."),
+					PluginAuthor,
+					PluginNodeMissingPrompt,
+					NodeInterface,
+					{ NodeCategories::Functions},
+					{ },
+					FNodeDisplayStyle{}
+				};
+
+				return Metadata;
+			};
+
+			static const FNodeClassMetadata Metadata = CreateNodeClassMetadata();
+			return Metadata;
+		}
+
+		static const FVertexInterface& DeclareVertexInterface()
+		{
+			using namespace RingModVertexNames;
+
+			static const FVertexInterface Interface(
+				FInputVertexInterface(
+					TInputDataVertex<FAudioBuffer>(METASOUND_GET_PARAM_NAME_AND_METADATA(InputAudioCarrier)),
+					TInputDataVertex<FAudioBuffer>(METASOUND_GET_PARAM_NAME_AND_METADATA(InputAudioModulator))
+				),
+				FOutputVertexInterface(
+					TOutputDataVertex<FAudioBuffer>(METASOUND_GET_PARAM_NAME_AND_METADATA(OutputAudio))
+				)
+			);
+
+			return Interface;
+		}
+
+		virtual FDataReferenceCollection GetInputs() const override
+		{
+			using namespace RingModVertexNames;
+
+			FDataReferenceCollection InputDataReferences;
+			
+			InputDataReferences.AddDataReadReference(METASOUND_GET_PARAM_NAME(InputAudioCarrier), AudioInputCarrier);
+			InputDataReferences.AddDataReadReference(METASOUND_GET_PARAM_NAME(InputAudioModulator), AudioInputModulator);
+			
+			return InputDataReferences;
+		}
+
+		virtual FDataReferenceCollection GetOutputs() const override
+		{
+			using namespace RingModVertexNames;
+			
+			FDataReferenceCollection OutputDataReferences;
+
+			OutputDataReferences.AddDataReadReference(METASOUND_GET_PARAM_NAME(OutputAudio), AudioOutput);
+			
+			return OutputDataReferences;
+		}
+		
+		static TUniquePtr<IOperator> CreateOperator(const FCreateOperatorParams& InParams, FBuildErrorArray& OutErrors)
+		{
+			using namespace RingModVertexNames;
+			const FDataReferenceCollection& Inputs = InParams.InputDataReferences;
+			
+			FAudioBufferReadRef CarrierAudioIn = Inputs.GetDataReadReferenceOrConstruct<FAudioBuffer>(METASOUND_GET_PARAM_NAME(InputAudioCarrier), InParams.OperatorSettings);
+			FAudioBufferReadRef ModulatorAudioIn = Inputs.GetDataReadReferenceOrConstruct<FAudioBuffer>(METASOUND_GET_PARAM_NAME(InputAudioModulator), InParams.OperatorSettings);
+
+			return MakeUnique<FRingModOperator>(InParams.OperatorSettings, CarrierAudioIn, ModulatorAudioIn);
+		}
+
+		void Execute()
+		{
+			*AudioOutput = *AudioInputCarrier;
+			
+			// @TODO Add Array Multiply Ain, Bin, Cout
+			Audio::ArrayMultiplyInPlace(*AudioInputModulator, *AudioOutput);
+		}
+
+	private:
+
+		// The input audio buffer
+		FAudioBufferReadRef AudioInputCarrier;
+		FAudioBufferReadRef AudioInputModulator;
+		
+		// Output audio buffer
+		FAudioBufferWriteRef AudioOutput;
+	};
+
+	// Node Class
+	class FRingModNode : public FNodeFacade
+	{
+	public:
+		FRingModNode(const FNodeInitData& InitData)
+			: FNodeFacade(InitData.InstanceName, InitData.InstanceID, TFacadeOperatorClass<FRingModOperator>())
+		{
+		}
+	};
+
+	// Register node
+	METASOUND_REGISTER_NODE(FRingModNode)
+}
+
+#undef LOCTEXT_NAMESPACE
