- Narrow/privitize template implementation scope for MetaSound Analyzers
- Implement/encapsulate binding for analyzer output/sender logic
- Move to using FAnyDataReferences where possible over collections for analyzers
- Misc clean-up

diff --git a/Engine/Plugins/Runtime/Metasound/Source/MetasoundFrontend/Private/Analysis/MetasoundFrontendAnalyzerRegistry.cpp b/Engine/Plugins/Runtime/Metasound/Source/MetasoundFrontend/Private/Analysis/MetasoundFrontendAnalyzerRegistry.cpp
--- a/Engine/Plugins/Runtime/Metasound/Source/MetasoundFrontend/Private/Analysis/MetasoundFrontendAnalyzerRegistry.cpp
+++ b/Engine/Plugins/Runtime/Metasound/Source/MetasoundFrontend/Private/Analysis/MetasoundFrontendAnalyzerRegistry.cpp
@@ -3,8 +3,8 @@
 
 #include "Analysis/MetasoundFrontendAnalyzerFactory.h"
 #include "Analysis/MetasoundFrontendVertexAnalyzerEnvelopeFollower.h"
+#include "Analysis/MetasoundFrontendVertexAnalyzerForwardValue.h"
 #include "Analysis/MetasoundFrontendVertexAnalyzerTriggerDensity.h"
-#include "Analysis/MetasoundFrontendVertexAnalyzerValue.h"
 #include "MetasoundPrimitives.h"
 
 
@@ -27,16 +27,16 @@ namespace Metasound
 			FVertexAnalyzerRegistry() = default;
 			virtual ~FVertexAnalyzerRegistry() = default;
 
-			virtual const TUniquePtr<IVertexAnalyzerFactory>& FindAnalyzerFactory(FName InAnalyzerName) const override
+			virtual const IVertexAnalyzerFactory* FindAnalyzerFactory(FName InAnalyzerName) const override
 			{
 				const TUniquePtr<IVertexAnalyzerFactory>* Factory = AnalyzerFactoryRegistry.Find(InAnalyzerName);
 				if (ensureMsgf(Factory, TEXT("Failed to find registered MetaSound Analyzer Factory with name '%s'"), *InAnalyzerName.ToString()))
 				{
-					return *Factory;
+					check(Factory->IsValid());
+					return Factory->Get();
 				}
 
-				static const TUniquePtr<IVertexAnalyzerFactory> InvalidFactory;
-				return InvalidFactory;
+				return nullptr;
 			}
 
 			virtual void RegisterAnalyzerFactories() override
@@ -45,10 +45,10 @@ namespace Metasound
 				RegisterAnalyzerFactory<FVertexAnalyzerTriggerDensity::FFactory>();
 
 				// Primitives
-				RegisterAnalyzerFactory<FVertexAnalyzerBool::FFactory>();
-				RegisterAnalyzerFactory<FVertexAnalyzerFloat::FFactory>();
-				RegisterAnalyzerFactory<FVertexAnalyzerInt::FFactory>();
-				RegisterAnalyzerFactory<FVertexAnalyzerString::FFactory>();
+				RegisterAnalyzerFactory<FVertexAnalyzerForwardBool::FFactory>();
+				RegisterAnalyzerFactory<FVertexAnalyzerForwardFloat::FFactory>();
+				RegisterAnalyzerFactory<FVertexAnalyzerForwardInt::FFactory>();
+				RegisterAnalyzerFactory<FVertexAnalyzerForwardString::FFactory>();
 			}
 		};
 
diff --git a/Engine/Plugins/Runtime/Metasound/Source/MetasoundFrontend/Private/Analysis/MetasoundFrontendAnalyzerView.cpp b/Engine/Plugins/Runtime/Metasound/Source/MetasoundFrontend/Private/Analysis/MetasoundFrontendAnalyzerView.cpp
--- a/Engine/Plugins/Runtime/Metasound/Source/MetasoundFrontend/Private/Analysis/MetasoundFrontendAnalyzerView.cpp
+++ b/Engine/Plugins/Runtime/Metasound/Source/MetasoundFrontend/Private/Analysis/MetasoundFrontendAnalyzerView.cpp
@@ -1,6 +1,8 @@
 // Copyright Epic Games, Inc. All Rights Reserved.
 #include "Analysis/MetasoundFrontendAnalyzerView.h"
 
+#include "Analysis/MetasoundFrontendAnalyzerFactory.h"
+#include "Analysis/MetasoundFrontendAnalyzerRegistry.h"
 #include "Analysis/MetasoundFrontendVertexAnalyzer.h"
 #include "Analysis/MetasoundFrontendVertexAnalyzerEnvelopeFollower.h"
 #include "Analysis/MetasoundFrontendGraphAnalyzer.h"
@@ -15,24 +17,32 @@ namespace Metasound
 {
 	namespace Frontend
 	{
-		FMetasoundAnalyzerView::FMetasoundAnalyzerView(const FOperatorSettings& InOperatorSettings, FAnalyzerAddress&& InAnalyzerAddress)
+		FMetasoundAnalyzerView::FMetasoundAnalyzerView(FAnalyzerAddress&& InAnalyzerAddress)
 			: AnalyzerAddress(MoveTemp(InAnalyzerAddress))
 		{
-			const FSendAddress SendAddress = AnalyzerAddress.ToSendAddress();
-			const FReceiverInitParams ReceiverParams { InOperatorSettings };
-			Receiver = FDataTransmissionCenter::Get().RegisterNewReceiver(SendAddress, ReceiverParams);
 		}
 
-		IReceiver& FMetasoundAnalyzerView::GetReceiverChecked()
+		void FMetasoundAnalyzerView::BindToAllOutputs(const FOperatorSettings& InOperatorSettings)
 		{
-			check(Receiver.IsValid());
-			return *Receiver;
+			const IVertexAnalyzerFactory* Factory = IVertexAnalyzerRegistry::Get().FindAnalyzerFactory(AnalyzerAddress.AnalyzerName);
+			if (ensureMsgf(Factory, TEXT("Failed to bind AnalyzerView to all Analyzer outputs: Missing factory definition for analyzer with name '%s'"), *AnalyzerAddress.AnalyzerName.ToString()))
+			{
+				for (const FAnalyzerOutput& Output : Factory->GetAnalyzerOutputs())
+				{
+					FAnalyzerAddress OutputAddress = AnalyzerAddress;
+					OutputAddress.AnalyzerMemberName = Output.Name;
+					OutputAddress.DataType = Output.DataType;
+					const FSendAddress SendAddress = OutputAddress.ToSendAddress();
+					const FReceiverInitParams ReceiverParams { InOperatorSettings };
+					IReceiver* Receiver = FDataTransmissionCenter::Get().RegisterNewReceiver(SendAddress, ReceiverParams).Release();
+					OutputReceivers.Add({ Output.Name, TSharedPtr<IReceiver>(Receiver) });
+				}
+			}
 		}
 
-		const IReceiver& FMetasoundAnalyzerView::GetReceiverChecked() const
+		bool FMetasoundAnalyzerView::UnbindOutput(FName InOutputName)
 		{
-			check(Receiver.IsValid());
-			return *Receiver;
+			return OutputReceivers.Remove(InOutputName) > 0;
 		}
 	} // namespace Frontend
 } // namespace Metasound
diff --git a/Engine/Plugins/Runtime/Metasound/Source/MetasoundFrontend/Private/Analysis/MetasoundFrontendGraphAnalyzer.cpp b/Engine/Plugins/Runtime/Metasound/Source/MetasoundFrontend/Private/Analysis/MetasoundFrontendGraphAnalyzer.cpp
--- a/Engine/Plugins/Runtime/Metasound/Source/MetasoundFrontend/Private/Analysis/MetasoundFrontendGraphAnalyzer.cpp
+++ b/Engine/Plugins/Runtime/Metasound/Source/MetasoundFrontend/Private/Analysis/MetasoundFrontendGraphAnalyzer.cpp
@@ -12,7 +12,7 @@ namespace Metasound
 {
 	namespace Frontend
 	{
-		FGraphAnalyzer::FGraphAnalyzer(const FOperatorSettings& InSettings, uint64 InInstanceID, FNodeVertexDataMap&& InReferences)
+		FGraphAnalyzer::FGraphAnalyzer(const FOperatorSettings& InSettings, uint64 InInstanceID, FNodeVertexDataReferenceMap&& InReferences)
 			: OperatorSettings(InSettings)
 			, InstanceID(InInstanceID)
 			, InternalDataReferences(MoveTemp(InReferences))
@@ -72,12 +72,14 @@ namespace Metasound
 // 						if (ensureMsgf(Collection != nullptr, TEXT("Failed to create MetaSoundAnalyzer: DataReferenceCollection for node analyzer at address '%s' not found."), *AnalyzerKey))
 						if (Collection)
 						{
-							const TUniquePtr<IVertexAnalyzerFactory>& Factory = IVertexAnalyzerRegistry::Get().FindAnalyzerFactory(AnalyzerAddress.AnalyzerName);
-							if (Factory.IsValid())
+							if (const IVertexAnalyzerFactory* Factory = IVertexAnalyzerRegistry::Get().FindAnalyzerFactory(AnalyzerAddress.AnalyzerName))
 							{
-								FCreateAnalyzerParams Params { AnalyzerAddress, OperatorSettings, *Collection };
-								TUniquePtr<IVertexAnalyzer> NewAnalyzer = Factory->CreateAnalyzer(Params);
-								Analyzers.Add(MoveTemp(NewAnalyzer));
+								if (const FAnyDataReference* DataRef = Collection->FindDataReference(AnalyzerAddress.OutputName))
+								{
+									FCreateAnalyzerParams Params{ AnalyzerAddress, OperatorSettings, *DataRef };
+									TUniquePtr<IVertexAnalyzer> NewAnalyzer = Factory->CreateAnalyzer(Params);
+									Analyzers.Add(MoveTemp(NewAnalyzer));
+								}
 							}
 						}
 					}
diff --git a/Engine/Plugins/Runtime/Metasound/Source/MetasoundFrontend/Private/Analysis/MetasoundFrontendGraphAnalyzerView.cpp b/Engine/Plugins/Runtime/Metasound/Source/MetasoundFrontend/Private/Analysis/MetasoundFrontendGraphAnalyzerView.cpp
--- a/Engine/Plugins/Runtime/Metasound/Source/MetasoundFrontend/Private/Analysis/MetasoundFrontendGraphAnalyzerView.cpp
+++ b/Engine/Plugins/Runtime/Metasound/Source/MetasoundFrontend/Private/Analysis/MetasoundFrontendGraphAnalyzerView.cpp
@@ -21,8 +21,8 @@ namespace Metasound
 			using FIterateOutputFunctionRef = TFunctionRef<void(const FConstOutputHandle&, const FAnalyzerAddress& AnalyzerAddress)>;
 			void IterateOutputsSupportingAnalyzer(const FConstGraphHandle& GraphHandle, uint64 InInstanceID, FName InAnalyzerName, bool bInRequiresConnection, FIterateOutputFunctionRef InFunc)
 			{
-				const TUniquePtr<IVertexAnalyzerFactory>& Factory = IVertexAnalyzerRegistry::Get().FindAnalyzerFactory(InAnalyzerName);
-				if (!Factory.IsValid())
+				const IVertexAnalyzerFactory* Factory = IVertexAnalyzerRegistry::Get().FindAnalyzerFactory(InAnalyzerName);
+				if (!Factory)
 				{
 					return;
 				}
@@ -71,18 +71,18 @@ namespace Metasound
 			using namespace GraphAnalyzerViewPrivate;
 
 			FConstGraphHandle GraphHandle = GetMetaSoundAssetChecked().GetRootGraphHandle();
-			const TUniquePtr<IVertexAnalyzerFactory>& Factory = IVertexAnalyzerRegistry::Get().FindAnalyzerFactory(InAnalyzerName);
-			if (!Factory.IsValid())
+			const IVertexAnalyzerFactory* Factory = IVertexAnalyzerRegistry::Get().FindAnalyzerFactory(InAnalyzerName);
+			if (!Factory)
 			{
 				return;
 			}
 
-			IterateOutputsSupportingAnalyzer(GraphHandle, InstanceID, InAnalyzerName, bInRequiresConnection, [this, FactoryPtr = Factory.Get(), &InAnalyzerName](FConstOutputHandle OutputHandle, const FAnalyzerAddress& AnalyzerAddress)
+			IterateOutputsSupportingAnalyzer(GraphHandle, InstanceID, InAnalyzerName, bInRequiresConnection, [this, &Factory, &InAnalyzerName](FConstOutputHandle OutputHandle, const FAnalyzerAddress& AnalyzerAddress)
 			{
 				FString AnalyzerKey = AnalyzerAddress.ToString();
 				ActiveAnalyzerKeys.Add(MoveTemp(AnalyzerKey));
 
-				const TArray<FAnalyzerOutput>& AnalyzerOutputs = FactoryPtr->GetAnalyzerOutputs();
+				const TArray<FAnalyzerOutput>& AnalyzerOutputs = Factory->GetAnalyzerOutputs();
 				for (const FAnalyzerOutput& AnalyzerOutput : AnalyzerOutputs)
 				{
 					FAnalyzerAddress OutputReceiverAddress = AnalyzerAddress;
@@ -90,7 +90,8 @@ namespace Metasound
 					OutputReceiverAddress.DataType = AnalyzerOutput.DataType;
 
 					FMetasoundGraphAnalyzerOutputKey OutputKey { AnalyzerAddress.NodeID, AnalyzerAddress.OutputName };
-					FMetasoundAnalyzerView NewView(OperatorSettings, MoveTemp(OutputReceiverAddress));
+					FMetasoundAnalyzerView NewView(MoveTemp(OutputReceiverAddress));
+					NewView.BindToAllOutputs(OperatorSettings);
 					AnalyzerViews.FindOrAdd(OutputKey).Add(MoveTemp(NewView));
 				}
 			});
diff --git a/Engine/Plugins/Runtime/Metasound/Source/MetasoundFrontend/Private/Analysis/MetasoundFrontendVertexAnalyzerEnvelopeFollower.cpp b/Engine/Plugins/Runtime/Metasound/Source/MetasoundFrontend/Private/Analysis/MetasoundFrontendVertexAnalyzerEnvelopeFollower.cpp
--- a/Engine/Plugins/Runtime/Metasound/Source/MetasoundFrontend/Private/Analysis/MetasoundFrontendVertexAnalyzerEnvelopeFollower.cpp
+++ b/Engine/Plugins/Runtime/Metasound/Source/MetasoundFrontend/Private/Analysis/MetasoundFrontendVertexAnalyzerEnvelopeFollower.cpp
@@ -2,6 +2,7 @@
 
 #include "Analysis/MetasoundFrontendVertexAnalyzerEnvelopeFollower.h"
 
+#include "MetasoundAudioBuffer.h"
 #include "MetasoundDataReference.h"
 #include "MetasoundPrimitives.h"
 #include "MetasoundRouter.h"
@@ -11,39 +12,46 @@ namespace Metasound
 {
 	namespace Frontend
 	{
-		const FAnalyzerOutput FVertexAnalyzerEnvelopeFollower::FOutputs::Value = { "EnvelopeValue", GetMetasoundDataTypeName<float>() };
+		const FAnalyzerOutput& FVertexAnalyzerEnvelopeFollower::FOutputs::GetValue()
+		{
+			static FAnalyzerOutput Value = { "EnvelopeValue", GetMetasoundDataTypeName<float>() };
+			return Value;
+		}
 
-		FVertexAnalyzerEnvelopeFollower::FVertexAnalyzerEnvelopeFollower(const FCreateAnalyzerParams& InParams)
-			: TVertexAnalyzer(InParams)
+		const FName& FVertexAnalyzerEnvelopeFollower::GetAnalyzerName()
 		{
-			FAnalyzerAddress OutputAddress = InParams.AnalyzerAddress;
-			OutputAddress.DataType = FOutputs::Value.DataType;
-			OutputAddress.AnalyzerMemberName = FOutputs::Value.Name;
-			FSendAddress SendAddress = OutputAddress.ToSendAddress();
+			static const FName AnalyzerName = "UE.Audio.EnvelopeFollower"; return AnalyzerName;
+		}
 
-			const FSenderInitParams InitParams { InParams.OperatorSettings, 0.0f };
-			Sender = FDataTransmissionCenter::Get().RegisterNewSender(MoveTemp(SendAddress), InitParams);
+		const FName& FVertexAnalyzerEnvelopeFollower::GetDataType()
+		{
+			return GetMetasoundDataTypeName<FAudioBuffer>();
+		}
 
+		FVertexAnalyzerEnvelopeFollower::FVertexAnalyzerEnvelopeFollower(const FCreateAnalyzerParams& InParams)
+		: FVertexAnalyzerBase(InParams.AnalyzerAddress, InParams.VertexDataReference)
+		, EnvelopeValue(TDataWriteReference<float>::CreateNew())
+		{
 			Audio::FEnvelopeFollowerInitParams Params;
 			Params.Mode = Audio::EPeakMode::RootMeanSquared;
-			Params.SampleRate = OperatorSettings.GetSampleRate();
+			Params.SampleRate = InParams.OperatorSettings.GetSampleRate();
 			Params.NumChannels = 1;
 			Params.AttackTimeMsec = 10;
 			Params.ReleaseTimeMsec = 10;
 			EnvelopeFollower.Init(Params);
+
+			FVertexAnalyzerBase::BindOutputData<float>(FOutputs::GetValue().Name, InParams.OperatorSettings, TDataReadReference<float>(EnvelopeValue));
 		}
 
 		void FVertexAnalyzerEnvelopeFollower::Execute()
 		{
-			const FAudioBuffer& InputData = GetAnalysisData();
-			EnvelopeFollower.ProcessAudio(InputData.GetData(), InputData.Num());
+			const FAudioBuffer& AudioBuffer = GetVertexData<FAudioBuffer>();
+			EnvelopeFollower.ProcessAudio(AudioBuffer.GetData(), AudioBuffer.Num());
 
 			check(EnvelopeFollower.GetEnvelopeValues().Num() == 1);
-			const float EnvelopeValue = EnvelopeFollower.GetEnvelopeValues().Last();
+			*EnvelopeValue = EnvelopeFollower.GetEnvelopeValues().Last();
 
-			FLiteral Literal;
-			Literal.Set(EnvelopeValue);
-			Sender->PushLiteral(Literal);
+			MarkOutputDirty();
 		}
 	} // namespace Frontend
 } // namespace Metasound
diff --git a/Engine/Plugins/Runtime/Metasound/Source/MetasoundFrontend/Private/Analysis/MetasoundFrontendVertexAnalyzerTriggerDensity.cpp b/Engine/Plugins/Runtime/Metasound/Source/MetasoundFrontend/Private/Analysis/MetasoundFrontendVertexAnalyzerTriggerDensity.cpp
--- a/Engine/Plugins/Runtime/Metasound/Source/MetasoundFrontend/Private/Analysis/MetasoundFrontendVertexAnalyzerTriggerDensity.cpp
+++ b/Engine/Plugins/Runtime/Metasound/Source/MetasoundFrontend/Private/Analysis/MetasoundFrontendVertexAnalyzerTriggerDensity.cpp
@@ -12,32 +12,42 @@ namespace Metasound
 {
 	namespace Frontend
 	{
-		const FAnalyzerOutput FVertexAnalyzerTriggerDensity::FOutputs::Value = { "TriggerDensity", GetMetasoundDataTypeName<float>() };
+		const FAnalyzerOutput& FVertexAnalyzerTriggerDensity::FOutputs::GetValue()
+		{
+			static const FAnalyzerOutput Value = { "TriggerDensity", GetMetasoundDataTypeName<float>() };
+			return Value;
+		}
+
+		const FName& FVertexAnalyzerTriggerDensity::GetAnalyzerName()
+		{
+			static const FName AnalyzerName = "UE.Trigger.Density";
+			return AnalyzerName;
+		}
+
+		const FName& FVertexAnalyzerTriggerDensity::GetDataType()
+		{
+			return GetMetasoundDataTypeName<FTrigger>();
+		}
 
 		FVertexAnalyzerTriggerDensity::FVertexAnalyzerTriggerDensity(const FCreateAnalyzerParams& InParams)
-			: TVertexAnalyzer(InParams)
+			: FVertexAnalyzerBase(InParams.AnalyzerAddress, InParams.VertexDataReference)
+			, EnvelopeValue(TDataWriteReference<float>::CreateNew())
 			, NumFramesPerBlock(InParams.OperatorSettings.GetNumFramesPerBlock())
 		{
 			Audio::FEnvelopeFollowerInitParams Params;
 			Params.Mode = Audio::EPeakMode::Peak;
-			Params.SampleRate = OperatorSettings.GetSampleRate();
+			Params.SampleRate = InParams.OperatorSettings.GetSampleRate();
 			Params.NumChannels = 1;
 			Params.AttackTimeMsec = 0;
 			Params.ReleaseTimeMsec = 120;
 			EnvelopeFollower.Init(Params);
 
-			FAnalyzerAddress OutputAddress = InParams.AnalyzerAddress;
-			OutputAddress.DataType = FOutputs::Value.DataType;
-			OutputAddress.AnalyzerMemberName = FOutputs::Value.Name;
-			FSendAddress SendAddress = OutputAddress.ToSendAddress();
-
-			const FSenderInitParams InitParams { InParams.OperatorSettings, 0.0f };
-			Sender = FDataTransmissionCenter::Get().RegisterNewSender(MoveTemp(SendAddress), InitParams);
+			FVertexAnalyzerBase::BindOutputData<float>(FOutputs::GetValue().Name, InParams.OperatorSettings, TDataReadReference<float>(EnvelopeValue));
 		}
 
 		void FVertexAnalyzerTriggerDensity::Execute()
 		{
-			const FTrigger& Trigger = GetAnalysisData();
+			const FTrigger& Trigger = GetVertexData<FTrigger>();
 
 			ScratchBuffer.Reset();
 			ScratchBuffer.AddZeroed(NumFramesPerBlock);
@@ -54,10 +64,9 @@ namespace Metasound
 			EnvelopeFollower.ProcessAudio(ScratchBuffer.GetData(), ScratchBuffer.Num());
 
 			check(EnvelopeFollower.GetEnvelopeValues().Num() == 1);
-			const float EnvelopeValue = EnvelopeFollower.GetEnvelopeValues().Last();
-			FLiteral Literal;
-			Literal.Set(EnvelopeValue);
-			Sender->PushLiteral(Literal);
+			*EnvelopeValue = EnvelopeFollower.GetEnvelopeValues().Last();
+
+			MarkOutputDirty();
 		}
 	} // namespace Frontend
 } // namespace Metasound
diff --git a/Engine/Plugins/Runtime/Metasound/Source/MetasoundFrontend/Private/Analysis/MetasoundFrontendVertexAnalyzerValue.cpp b/Engine/Plugins/Runtime/Metasound/Source/MetasoundFrontend/Private/Analysis/MetasoundFrontendVertexAnalyzerValue.cpp
deleted file mode 100644
diff --git a/Engine/Plugins/Runtime/Metasound/Source/MetasoundFrontend/Public/Analysis/MetasoundFrontendAnalyzerFactory.h b/Engine/Plugins/Runtime/Metasound/Source/MetasoundFrontend/Public/Analysis/MetasoundFrontendAnalyzerFactory.h
--- a/Engine/Plugins/Runtime/Metasound/Source/MetasoundFrontend/Public/Analysis/MetasoundFrontendAnalyzerFactory.h
+++ b/Engine/Plugins/Runtime/Metasound/Source/MetasoundFrontend/Public/Analysis/MetasoundFrontendAnalyzerFactory.h
@@ -16,7 +16,7 @@ namespace Metasound
 
 			virtual const TArray<FAnalyzerOutput>& GetAnalyzerOutputs() const = 0;
 			virtual FName GetDataType() const = 0;
-			virtual TUniquePtr<IVertexAnalyzer> CreateAnalyzer(const FCreateAnalyzerParams& InParams) = 0;
+			virtual TUniquePtr<IVertexAnalyzer> CreateAnalyzer(const FCreateAnalyzerParams& InParams) const = 0;
 		};
 
 		template<typename TAnalyzerClass>
@@ -28,7 +28,7 @@ namespace Metasound
 			virtual ~TVertexAnalyzerFactory() = default;
 
 			virtual FName GetDataType() const override { return TAnalyzerClass::GetDataType(); }
-			virtual TUniquePtr<IVertexAnalyzer> CreateAnalyzer(const FCreateAnalyzerParams& InParams) override
+			virtual TUniquePtr<IVertexAnalyzer> CreateAnalyzer(const FCreateAnalyzerParams& InParams) const override
 			{
 				return TUniquePtr<IVertexAnalyzer>(new TAnalyzerClass(InParams));
 			}
diff --git a/Engine/Plugins/Runtime/Metasound/Source/MetasoundFrontend/Public/Analysis/MetasoundFrontendAnalyzerRegistry.h b/Engine/Plugins/Runtime/Metasound/Source/MetasoundFrontend/Public/Analysis/MetasoundFrontendAnalyzerRegistry.h
--- a/Engine/Plugins/Runtime/Metasound/Source/MetasoundFrontend/Public/Analysis/MetasoundFrontendAnalyzerRegistry.h
+++ b/Engine/Plugins/Runtime/Metasound/Source/MetasoundFrontend/Public/Analysis/MetasoundFrontendAnalyzerRegistry.h
@@ -19,7 +19,7 @@ namespace Metasound
 		public:
 			static IVertexAnalyzerRegistry& Get();
 
-			virtual const TUniquePtr<IVertexAnalyzerFactory>& FindAnalyzerFactory(FName InAnalyzerName) const = 0;
+			virtual const IVertexAnalyzerFactory* FindAnalyzerFactory(FName InAnalyzerName) const = 0;
 			virtual void RegisterAnalyzerFactories() = 0;
 		};
 	} // namespace Frontend
diff --git a/Engine/Plugins/Runtime/Metasound/Source/MetasoundFrontend/Public/Analysis/MetasoundFrontendAnalyzerView.h b/Engine/Plugins/Runtime/Metasound/Source/MetasoundFrontend/Public/Analysis/MetasoundFrontendAnalyzerView.h
--- a/Engine/Plugins/Runtime/Metasound/Source/MetasoundFrontend/Public/Analysis/MetasoundFrontendAnalyzerView.h
+++ b/Engine/Plugins/Runtime/Metasound/Source/MetasoundFrontend/Public/Analysis/MetasoundFrontendAnalyzerView.h
@@ -2,7 +2,7 @@
 #pragma once
 
 #include "Analysis/MetasoundFrontendVertexAnalyzer.h"
-#include "Containers/Map.h"
+#include "Containers/Array.h"
 #include "HAL/Platform.h"
 #include "MetasoundAssetBase.h"
 #include "MetasoundRouter.h"
@@ -15,22 +15,42 @@ class UAudioComponent;
 
 namespace Metasound
 {
+	// Forward Declarations
+	class FOperatorSettings;
+
 	namespace Frontend
 	{
 		// Pairs an IReceiver with a given AnalyzerAddress, which enables
 		// watching a particular analyzer result on any given thread.
 		class METASOUNDFRONTEND_API FMetasoundAnalyzerView
 		{
-			TUniquePtr<IReceiver> Receiver;
+			TMap<FName, TSharedPtr<IReceiver>> OutputReceivers;
 
 		public:
 			const FAnalyzerAddress AnalyzerAddress = { };
 
 			FMetasoundAnalyzerView() = default;
-			FMetasoundAnalyzerView(const FOperatorSettings& InOperatorSettings, FAnalyzerAddress&& InAnalyzerAddress);
+			FMetasoundAnalyzerView(FAnalyzerAddress&& InAnalyzerAddress);
+
+			void BindToAllOutputs(const FOperatorSettings& InOperatorSettings);
+			bool UnbindOutput(FName InOutputName);
+
+			template <typename DataType>
+			bool TryGetOutputData(FName InOutputName, DataType& OutValue)
+			{
+				TSharedPtr<IReceiver>* Receiver = OutputReceivers.Find(InOutputName);
+				if (Receiver && Receiver->IsValid())
+				{
+					TReceiver<DataType>& TypedReceiver = (*Receiver)->GetAs<TReceiver<DataType>>();
+					if (TypedReceiver.CanPop())
+					{
+						TypedReceiver.Pop(OutValue);
+						return true;
+					}
+				}
 
-			IReceiver& GetReceiverChecked();
-			const IReceiver& GetReceiverChecked() const;
+				return false;
+			}
 		};
 	}
 }
diff --git a/Engine/Plugins/Runtime/Metasound/Source/MetasoundFrontend/Public/Analysis/MetasoundFrontendGraphAnalyzer.h b/Engine/Plugins/Runtime/Metasound/Source/MetasoundFrontend/Public/Analysis/MetasoundFrontendGraphAnalyzer.h
--- a/Engine/Plugins/Runtime/Metasound/Source/MetasoundFrontend/Public/Analysis/MetasoundFrontendGraphAnalyzer.h
+++ b/Engine/Plugins/Runtime/Metasound/Source/MetasoundFrontend/Public/Analysis/MetasoundFrontendGraphAnalyzer.h
@@ -5,6 +5,7 @@
 #include "MetasoundBuilderInterface.h"
 #include "MetasoundOperatorSettings.h"
 #include "MetasoundRouter.h"
+#include "MetasoundVertex.h"
 #include "Templates/UniquePtr.h"
 
 
@@ -12,6 +13,8 @@ namespace Metasound
 {
 	namespace Frontend
 	{
+		using FNodeVertexDataReferenceMap = TMap<FGuid, FDataReferenceCollection>;
+
 		// Handles intrinsic analysis operations within a given graph
 		// should the graph's operator be enabled for analysis.
 		class METASOUNDFRONTEND_API FGraphAnalyzer
@@ -21,10 +24,10 @@ namespace Metasound
 
 			TUniquePtr<IReceiver> ActiveAnalyzerReceiver;
 			TArray<TUniquePtr<Frontend::IVertexAnalyzer>> Analyzers;
-			FNodeVertexDataMap InternalDataReferences;
+			FNodeVertexDataReferenceMap InternalDataReferences;
 
 		public:
-			FGraphAnalyzer(const FOperatorSettings& InSettings, uint64 InInstanceID, FNodeVertexDataMap&& InGraphReferences);
+			FGraphAnalyzer(const FOperatorSettings& InSettings, uint64 InInstanceID, FNodeVertexDataReferenceMap&& InGraphReferences);
 			~FGraphAnalyzer() = default;
 
 			// Creates a send channel name unique for the given sound instance used to send array of analyzer
diff --git a/Engine/Plugins/Runtime/Metasound/Source/MetasoundFrontend/Public/Analysis/MetasoundFrontendVertexAnalyzer.h b/Engine/Plugins/Runtime/Metasound/Source/MetasoundFrontend/Public/Analysis/MetasoundFrontendVertexAnalyzer.h
--- a/Engine/Plugins/Runtime/Metasound/Source/MetasoundFrontend/Public/Analysis/MetasoundFrontendVertexAnalyzer.h
+++ b/Engine/Plugins/Runtime/Metasound/Source/MetasoundFrontend/Public/Analysis/MetasoundFrontendVertexAnalyzer.h
@@ -2,9 +2,11 @@
 #pragma once
 
 #include "Analysis/MetasoundFrontendAnalyzerAddress.h"
+#include "Containers/Map.h"
 #include "HAL/Platform.h"
 #include "MetasoundDataReference.h"
 #include "MetasoundDataReferenceCollection.h"
+#include "MetasoundOperatorSettings.h"
 #include "Misc/Guid.h"
 
 
@@ -12,61 +14,150 @@ namespace Metasound
 {
 	namespace Frontend
 	{
+		// Output of analyzer (not to be confused with the node vertex output which an
+		// analyzer may be processing).  Used to signify by an analyzer what information
+		// is sent to other threads using an analyzer view for logging, displaying, processing,
+		// etc. resulting analyzer data.
 		struct METASOUNDFRONTEND_API FAnalyzerOutput
 		{
 			FName Name;
 			FName DataType;
 		};
 
+		// Parameters used to generate an analyzer instance
 		struct METASOUNDFRONTEND_API FCreateAnalyzerParams
 		{
+			// Address of analyzer
 			const FAnalyzerAddress& AnalyzerAddress;
+
+			// OperatorSettings used for analyzer execution
 			const FOperatorSettings& OperatorSettings;
-			const FDataReferenceCollection& Collection;
+
+			// Data reference to vertex (currently only output vertices
+			// support analysis) writing data to be analyzed.
+			const FAnyDataReference& VertexDataReference;
 		};
 
+		// Analyzer that watches a data reference associated with a particular
+		// node vertex. Currently, only output node vertex analysis is supported.
 		class METASOUNDFRONTEND_API IVertexAnalyzer
 		{
 		public:
 			virtual ~IVertexAnalyzer() = default;
 
+			// Returns analyzer address
 			virtual const FAnalyzerAddress& GetAnalyzerAddress() const = 0;
+
+			// Executes analysis
 			virtual void Execute() = 0;
 		};
 
-		template<typename TChildClass, typename TDataType>
-		class TVertexAnalyzer : public IVertexAnalyzer
+		// Bound output from an analyzer (not to be confused with an output vertex)
+		// sent internally using the Transmission System. Each bound output corresponds to
+		// a DataChannel with an analyzer view or views potentially receiving analysis results.
+		class METASOUNDFRONTEND_API IBoundAnalyzerOutput
 		{
-		protected:
-			FAnalyzerAddress Address;
-			FOperatorSettings OperatorSettings { 48000, 100 };
-			TDataReadReference<TDataType> Reference;
+		public:
+			virtual ~IBoundAnalyzerOutput() = default;
+
+			// Returns the data reference associated with the analyzer's output
+			virtual FAnyDataReference GetDataReference() const = 0;
+
+			// Pushes data to the sender to be forwarded to all actively associated analyser views.
+			virtual void PushData() = 0;
+		};
 
+		// Templatized implementation of a bound analyzer output (see IBoundAnalyzerOutput)
+		template<typename DataType>
+		class TBoundAnalyzerOutput final : public IBoundAnalyzerOutput
+		{
 		public:
-			static FName GetDataType()
+			TBoundAnalyzerOutput(const FAnalyzerAddress& InAnalyzerOutputAddress, const FOperatorSettings& InOperatorSettings, TDataReadReference<DataType>&& InData)
+				: DataRef(MoveTemp(InData))
+			{
+				Sender = FDataTransmissionCenter::Get().RegisterNewSender<DataType>(InAnalyzerOutputAddress.ToSendAddress(), FSenderInitParams { InOperatorSettings, 0 });
+				ensure(Sender);
+			}
+
+			virtual ~TBoundAnalyzerOutput() = default;
+
+			FAnyDataReference GetDataReference() const override
 			{
-				return GetMetasoundDataTypeName<TDataType>();
+				return DataRef;
 			}
 
-			TVertexAnalyzer(const FCreateAnalyzerParams& InParams)
-				: Address(InParams.AnalyzerAddress)
-				, OperatorSettings(InParams.OperatorSettings)
-				, Reference(InParams.Collection.GetDataReadReference<TDataType>(InParams.AnalyzerAddress.OutputName))
+			void PushData() override
 			{
-				check(Address.AnalyzerName == TChildClass::GetAnalyzerName());
+				if (Sender.IsValid())
+				{
+					Sender->Push(*DataRef);
+				}
 			}
 
-			virtual ~TVertexAnalyzer() = default;
+		private:
+			TDataReadReference<DataType> DataRef;
+			TUniquePtr<TSender<DataType>> Sender;
+		};
 
-			const TDataType& GetAnalysisData() const
+		// Base implementation of a vertex analyzer.
+		class METASOUNDFRONTEND_API FVertexAnalyzerBase : public IVertexAnalyzer
+		{
+		public:
+			FVertexAnalyzerBase(const FAnalyzerAddress& InAnalyzerAddress, const FAnyDataReference& InDataReference)
+				: AnalyzerAddress(InAnalyzerAddress)
+				, VertexDataReference(InDataReference)
 			{
-				return *Reference;
 			}
 
+			virtual ~FVertexAnalyzerBase() = default;
+
 			virtual const FAnalyzerAddress& GetAnalyzerAddress() const override
 			{
-				return Address;
+				return AnalyzerAddress;
+			}
+
+		protected:
+			using FBoundOutputDataPtr = TSharedPtr<IBoundAnalyzerOutput>;
+
+			// Marks all output channels as dirty, refreshing all bound data.
+			// Likely best called after all computation is complete within an
+			// execution call, to be implemented by the child class.
+			void MarkOutputDirty()
+			{
+				for (TPair<FName, FBoundOutputDataPtr>& Pair : BoundOutputData)
+				{
+					Pair.Value->PushData();
+				}
+			}
+
+			// Unbinds a particular named analyzer output
+			bool UnbindOutputData(FName InAnalyzerOutputName)
+			{
+				return BoundOutputData.Remove(InAnalyzerOutputName) > 0;
+			}
+
+			// Binds a particular named analyzer output to be updated when MarkOutputDirty is called.
+			template<typename DataType>
+			void BindOutputData(FName InAnalyzerOutputName, const FOperatorSettings& InOperatorSettings, TDataReadReference<DataType>&& InData)
+			{
+				FAnalyzerAddress OutputAddress = AnalyzerAddress;
+				OutputAddress.AnalyzerMemberName = InAnalyzerOutputName;
+				OutputAddress.DataType = GetMetasoundDataTypeName<DataType>();
+				FBoundOutputDataPtr BoundOutputDataPtr = FBoundOutputDataPtr(new TBoundAnalyzerOutput<DataType>(OutputAddress, InOperatorSettings, MoveTemp(InData)));
+				BoundOutputData.Emplace(InAnalyzerOutputName, MoveTemp(BoundOutputDataPtr));
 			}
+
+			// Returns the most recent vertex data.
+			template<typename DataType>
+			const DataType& GetVertexData() const
+			{
+				return *VertexDataReference.GetDataReadReference<DataType>();
+			}
+
+		private:
+			FAnalyzerAddress AnalyzerAddress;
+			TMap<FName, FBoundOutputDataPtr> BoundOutputData;
+			FAnyDataReference VertexDataReference;
 		};
 	} // namespace Frontend
 } // namespace Metasound
diff --git a/Engine/Plugins/Runtime/Metasound/Source/MetasoundFrontend/Public/Analysis/MetasoundFrontendVertexAnalyzerEnvelopeFollower.h b/Engine/Plugins/Runtime/Metasound/Source/MetasoundFrontend/Public/Analysis/MetasoundFrontendVertexAnalyzerEnvelopeFollower.h
--- a/Engine/Plugins/Runtime/Metasound/Source/MetasoundFrontend/Public/Analysis/MetasoundFrontendVertexAnalyzerEnvelopeFollower.h
+++ b/Engine/Plugins/Runtime/Metasound/Source/MetasoundFrontend/Public/Analysis/MetasoundFrontendVertexAnalyzerEnvelopeFollower.h
@@ -13,21 +13,15 @@ namespace Metasound
 {
 	namespace Frontend
 	{
-		class METASOUNDFRONTEND_API FVertexAnalyzerEnvelopeFollower : public TVertexAnalyzer<FVertexAnalyzerEnvelopeFollower, FAudioBuffer>
+		class METASOUNDFRONTEND_API FVertexAnalyzerEnvelopeFollower : public FVertexAnalyzerBase
 		{
-			TUniquePtr<ISender> Sender;
-			Audio::FEnvelopeFollower EnvelopeFollower;
-
 		public:
-			static const FName& GetAnalyzerName()
-			{
-				static const FName AnalyzerName = "UE.Audio.EnvelopeFollower";
-				return AnalyzerName;
-			}
+			static const FName& GetAnalyzerName();
+			static const FName& GetDataType();
 
 			struct METASOUNDFRONTEND_API FOutputs
 			{
-				static const FAnalyzerOutput Value;
+				static const FAnalyzerOutput& GetValue();
 			};
 
 			class METASOUNDFRONTEND_API FFactory : public TVertexAnalyzerFactory<FVertexAnalyzerEnvelopeFollower>
@@ -35,7 +29,7 @@ namespace Metasound
 			public:
 				virtual const TArray<FAnalyzerOutput>& GetAnalyzerOutputs() const override
 				{
-					static const TArray<FAnalyzerOutput> Outputs { FOutputs::Value };
+					static const TArray<FAnalyzerOutput> Outputs { FOutputs::GetValue() };
 					return Outputs;
 				}
 			};
@@ -44,6 +38,10 @@ namespace Metasound
 			virtual ~FVertexAnalyzerEnvelopeFollower() = default;
 
 			virtual void Execute() override;
+
+		private:
+			Audio::FEnvelopeFollower EnvelopeFollower;
+			TDataWriteReference<float> EnvelopeValue;
 		};
 	} // namespace Frontend
 } // namespace Metasound
diff --git a/Engine/Plugins/Runtime/Metasound/Source/MetasoundFrontend/Public/Analysis/MetasoundFrontendVertexAnalyzerForwardValue.h b/Engine/Plugins/Runtime/Metasound/Source/MetasoundFrontend/Public/Analysis/MetasoundFrontendVertexAnalyzerForwardValue.h
new file mode 100644
--- /dev/null
+++ b/Engine/Plugins/Runtime/Metasound/Source/MetasoundFrontend/Public/Analysis/MetasoundFrontendVertexAnalyzerForwardValue.h
@@ -0,0 +1,96 @@
+// Copyright Epic Games, Inc. All Rights Reserved.
+#pragma once
+
+#include "Analysis/MetasoundFrontendAnalyzerAddress.h"
+#include "Analysis/MetasoundFrontendAnalyzerFactory.h"
+#include "Analysis/MetasoundFrontendVertexAnalyzer.h"
+#include "Containers/Array.h"
+#include "MetasoundDataReference.h"
+#include "MetasoundPrimitives.h"
+#include "MetasoundRouter.h"
+#include "MetasoundTrigger.h"
+
+
+namespace Metasound
+{
+	namespace Frontend
+	{
+		// Templatized implementation of the simplest/cheapest form of an analyzer,
+		// which just forwards a given value to an associated view or views. Can be
+		// used in conjunction with an additional process, task, etc. handled by a
+		// view to avoid further processing/expense within the graph execution thread
+		// context.
+		template <typename ChildClass, typename DataType>
+		class TVertexAnalyzerForwardValue : public FVertexAnalyzerBase
+		{
+			DataType LastValue;
+
+		public:
+			static const FName& GetDataType()
+			{
+				return GetMetasoundDataTypeName<DataType>();
+			}
+
+			struct FOutputs
+			{
+				static const FAnalyzerOutput& GetValue()
+				{
+					auto MakeValue = []()
+					{
+						return FAnalyzerOutput { "Value", GetMetasoundDataTypeName<DataType>() };
+					};
+
+					static FAnalyzerOutput Value = MakeValue();
+					return Value;
+				}
+			};
+
+			class FFactory : public TVertexAnalyzerFactory<ChildClass>
+			{
+			public:
+				virtual const TArray<FAnalyzerOutput>& GetAnalyzerOutputs() const override
+				{
+					static const TArray<FAnalyzerOutput> Outputs { ChildClass::FOutputs::GetValue() };
+					return Outputs;
+				}
+			};
+
+			TVertexAnalyzerForwardValue(const FCreateAnalyzerParams& InParams)
+ 				: FVertexAnalyzerBase(InParams.AnalyzerAddress, InParams.VertexDataReference)
+			{
+				TDataReadReference<DataType> DataRef = InParams.VertexDataReference.GetDataReadReference<DataType>();
+				FVertexAnalyzerBase::BindOutputData<DataType>(ChildClass::FOutputs::GetValue().Name, InParams.OperatorSettings, MoveTemp(DataRef));
+			}
+
+			virtual ~TVertexAnalyzerForwardValue() = default;
+
+			virtual void Execute() override
+			{
+				const DataType& Value = GetVertexData<DataType>();
+				if (LastValue != Value)
+				{
+					MarkOutputDirty();
+				}
+			}
+		};
+
+		#define METASOUND_DECLARE_FORWARD_VALUE_ANALYZER(ClassType, AnalyzerFName, DataType) \
+			class METASOUNDFRONTEND_API ClassType : public TVertexAnalyzerForwardValue<ClassType, DataType> \
+			{ \
+			public: \
+				static const FName& GetAnalyzerName() \
+				{ \
+					static const FName AnalyzerName = AnalyzerFName; return AnalyzerName; \
+				} \
+				ClassType(const FCreateAnalyzerParams& InParams) \
+					: TVertexAnalyzerForwardValue(InParams) { } \
+				virtual ~ClassType() = default; \
+			};
+
+		// Primitive forward value types
+		METASOUND_DECLARE_FORWARD_VALUE_ANALYZER(FVertexAnalyzerForwardBool, "UE.Forward.Bool", bool)
+		METASOUND_DECLARE_FORWARD_VALUE_ANALYZER(FVertexAnalyzerForwardFloat, "UE.Forward.Float", float)
+		METASOUND_DECLARE_FORWARD_VALUE_ANALYZER(FVertexAnalyzerForwardInt, "UE.Forward.Int32", int32)
+		METASOUND_DECLARE_FORWARD_VALUE_ANALYZER(FVertexAnalyzerForwardString, "UE.Forward.String", FString)
+	} // namespace Frontend
+} // namespace Metasound
diff --git a/Engine/Plugins/Runtime/Metasound/Source/MetasoundFrontend/Public/Analysis/MetasoundFrontendVertexAnalyzerTriggerDensity.h b/Engine/Plugins/Runtime/Metasound/Source/MetasoundFrontend/Public/Analysis/MetasoundFrontendVertexAnalyzerTriggerDensity.h
--- a/Engine/Plugins/Runtime/Metasound/Source/MetasoundFrontend/Public/Analysis/MetasoundFrontendVertexAnalyzerTriggerDensity.h
+++ b/Engine/Plugins/Runtime/Metasound/Source/MetasoundFrontend/Public/Analysis/MetasoundFrontendVertexAnalyzerTriggerDensity.h
@@ -12,24 +12,15 @@ namespace Metasound
 {
 	namespace Frontend
 	{
-		class METASOUNDFRONTEND_API FVertexAnalyzerTriggerDensity : public TVertexAnalyzer<FVertexAnalyzerTriggerDensity, FTrigger>
+		class METASOUNDFRONTEND_API FVertexAnalyzerTriggerDensity : public FVertexAnalyzerBase
 		{
-			TUniquePtr<ISender> Sender;
-			int32 NumFramesPerBlock = 0;
-
-			Audio::FEnvelopeFollower EnvelopeFollower;
-			Audio::FAlignedFloatBuffer ScratchBuffer;
-
 		public:
-			static const FName& GetAnalyzerName()
-			{
-				static const FName AnalyzerName = "UE.Trigger.Density";
-				return AnalyzerName;
-			}
+			static const FName& GetAnalyzerName();
+			static const FName& GetDataType();
 
 			struct METASOUNDFRONTEND_API FOutputs
 			{
-				static const FAnalyzerOutput Value;
+				static const FAnalyzerOutput& GetValue();
 
 			};
 
@@ -38,7 +29,7 @@ namespace Metasound
 			public:
 				virtual const TArray<FAnalyzerOutput>& GetAnalyzerOutputs() const override
 				{
-					static const TArray<FAnalyzerOutput> Outputs { FOutputs::Value };
+					static const TArray<FAnalyzerOutput> Outputs { FOutputs::GetValue() };
 					return Outputs;
 				}
 			};
@@ -47,6 +38,13 @@ namespace Metasound
 			virtual ~FVertexAnalyzerTriggerDensity() = default;
 
 			virtual void Execute() override;
+
+		private:
+
+			Audio::FEnvelopeFollower EnvelopeFollower;
+			TDataWriteReference<float> EnvelopeValue;
+			int32 NumFramesPerBlock = 0;
+			Audio::FAlignedFloatBuffer ScratchBuffer;
 		};
 	} // namespace Frontend
 } // namespace Metasound
diff --git a/Engine/Plugins/Runtime/Metasound/Source/MetasoundFrontend/Public/Analysis/MetasoundFrontendVertexAnalyzerValue.h b/Engine/Plugins/Runtime/Metasound/Source/MetasoundFrontend/Public/Analysis/MetasoundFrontendVertexAnalyzerValue.h
deleted file mode 100644
diff --git a/Engine/Plugins/Runtime/Metasound/Source/MetasoundGraphCore/Private/MetasoundOperatorBuilder.cpp b/Engine/Plugins/Runtime/Metasound/Source/MetasoundGraphCore/Private/MetasoundOperatorBuilder.cpp
--- a/Engine/Plugins/Runtime/Metasound/Source/MetasoundGraphCore/Private/MetasoundOperatorBuilder.cpp
+++ b/Engine/Plugins/Runtime/Metasound/Source/MetasoundGraphCore/Private/MetasoundOperatorBuilder.cpp
@@ -112,7 +112,7 @@ namespace Metasound
 		}
 
 		// Create graph operator from collection of node operators.
-		return CreateGraphOperator(BuildContext, InParams);
+		return CreateGraphOperator(BuildContext);
 	}
 
 	FOperatorBuilder::FBuildStatus FOperatorBuilder::DepthFirstTopologicalSort(FBuildContext& InOutContext, TArray<const INode*>& OutNodes) const
@@ -466,14 +466,14 @@ namespace Metasound
 		return BuildStatus;
 	}
 
-	TUniquePtr<IOperator> FOperatorBuilder::CreateGraphOperator(FBuildContext& InOutContext, const FBuildGraphParams& InParams) const
+	TUniquePtr<IOperator> FOperatorBuilder::CreateGraphOperator(FBuildContext& InOutContext) const
 	{
 		FDataReferenceCollection GraphInputs;
 		FDataReferenceCollection GraphOutputs;
 
 		FBuildStatus BuildStatus = GatherGraphDataReferences(InOutContext, GraphInputs, GraphOutputs);
 
-		if (BuildStatus > FBuildStatus::GetMaxErrorLevel(InParams.BuilderSettings))
+		if (BuildStatus > FBuildStatus::GetMaxErrorLevel(InOutContext.BuilderSettings))
 		{
 			return TUniquePtr<IOperator>(nullptr);
 		}
@@ -493,10 +493,7 @@ namespace Metasound
 				References.Emplace(Node->GetInstanceID(), MoveTemp(ReferencePair.Value.Outputs));
 			}
 
-			if (InParams.BuilderSettings.bPopulateInternalDataReferences)
-			{
-				InOutContext.Results.InternalDataReferences.Append(MoveTemp(References));
-			}
+			InOutContext.Results.InternalDataReferences.Append(MoveTemp(References));
 		}
 
 		for (FOperatorPtr& Ptr : InOutContext.Operators)
diff --git a/Engine/Plugins/Runtime/Metasound/Source/MetasoundGraphCore/Public/MetasoundBuilderInterface.h b/Engine/Plugins/Runtime/Metasound/Source/MetasoundGraphCore/Public/MetasoundBuilderInterface.h
--- a/Engine/Plugins/Runtime/Metasound/Source/MetasoundGraphCore/Public/MetasoundBuilderInterface.h
+++ b/Engine/Plugins/Runtime/Metasound/Source/MetasoundGraphCore/Public/MetasoundBuilderInterface.h
@@ -3,6 +3,7 @@
 #pragma once
 
 #include "Containers/Array.h"
+#include "MetasoundDataReference.h"
 #include "MetasoundEnvironment.h"
 #include "MetasoundNodeInterface.h"
 #include "MetasoundOperatorBuilderSettings.h"
@@ -44,10 +45,6 @@ namespace Metasound
 			virtual const TArray<FDataEdge>& GetDataEdges() const = 0;
 	};
 
-	/** Map of NodeID (Guid) to DataReferenceCollection. */
-	// TODO: Move to FAnyDataReference
-	using FNodeVertexDataMap = TMap<FGuid, FDataReferenceCollection>;
-
 	/** Array of build errors. */
 	using FBuildErrorArray = TArray<TUniquePtr<IOperatorBuildError>>;
 
@@ -57,8 +54,8 @@ namespace Metasound
 		/** An array of errors. Errors can be added if issues occur while creating an IOperator. */
 		FBuildErrorArray Errors;
 
-		/** References used by graph analyzer, if enabled (not collected if analysis is disabled). */
-		FNodeVertexDataMap InternalDataReferences;
+		/** Internal data references if enabled by build settings (not populated if disabled). */
+		TMap<FGuid, FDataReferenceCollection> InternalDataReferences;
 	};
 
 	/** FCreateOperatorParams holds the parameters provided to operator factories
diff --git a/Engine/Plugins/Runtime/Metasound/Source/MetasoundGraphCore/Public/MetasoundOperatorBuilder.h b/Engine/Plugins/Runtime/Metasound/Source/MetasoundGraphCore/Public/MetasoundOperatorBuilder.h
--- a/Engine/Plugins/Runtime/Metasound/Source/MetasoundGraphCore/Public/MetasoundOperatorBuilder.h
+++ b/Engine/Plugins/Runtime/Metasound/Source/MetasoundGraphCore/Public/MetasoundOperatorBuilder.h
@@ -163,7 +163,7 @@ namespace Metasound
 			FBuildStatus CreateOperators(FBuildContext& InOutContext, const TArray<const INode*>& InSortedNodes, const FDataReferenceCollection& InGraphInputs) const;
 
 			// Create the final graph operator from the provided build context.
-			TUniquePtr<IOperator> CreateGraphOperator(FBuildContext& InOutContext, const FBuildGraphParams& InParams) const;
+			TUniquePtr<IOperator> CreateGraphOperator(FBuildContext& InOutContext) const;
 	};
 }
 
