Fix for generator nodes behaving incorrectly when given negative frequencies

diff --git a/Engine/Plugins/Runtime/Metasound/Source/MetasoundStandardNodes/Private/MetasoundOscillatorNodes.cpp b/Engine/Plugins/Runtime/Metasound/Source/MetasoundStandardNodes/Private/MetasoundOscillatorNodes.cpp
--- a/Engine/Plugins/Runtime/Metasound/Source/MetasoundStandardNodes/Private/MetasoundOscillatorNodes.cpp
+++ b/Engine/Plugins/Runtime/Metasound/Source/MetasoundStandardNodes/Private/MetasoundOscillatorNodes.cpp
@@ -52,7 +52,7 @@ namespace Metasound
 					InPhase -= 1.f;
 				}
 
-				while (InPhase <= -1.f)
+				while (InPhase < 0.0f)
 				{
 					InPhase += 1.f;
 				}
@@ -71,18 +71,19 @@ namespace Metasound
 		FORCEINLINE float PolySmoothSaw(const float InPhase, const float InPhaseDelta)
 		{	
 			float Output = 0.0f;
+			float AbsolutePhaseDelta = FMath::Abs(InPhaseDelta);
 
 			// The current phase is on the left side of discontinuity
-			if (InPhase > 1.0f - InPhaseDelta)
+			if (InPhase > 1.0f - AbsolutePhaseDelta)
 			{
-				const float Dist = (InPhase - 1.0f) / InPhaseDelta;
+				const float Dist = (InPhase - 1.0f) / AbsolutePhaseDelta;
 				Output = -Dist * Dist - 2.0f * Dist - 1.0f;
 			}
 			// The current phase is on the right side of the discontinuity
-			else if (InPhase < InPhaseDelta)
+			else if (InPhase < AbsolutePhaseDelta)
 			{
 				// Distance into polynomial
-				const float Dist = InPhase / InPhaseDelta;
+				const float Dist = InPhase / AbsolutePhaseDelta;
 				Output = Dist * Dist - 2.0f * Dist + 1.0f;
 			}
 			return Output;
@@ -441,10 +442,11 @@ namespace Metasound
 		// Square.
 		struct FSquareGenerator
 		{
-			FORCEINLINE float operator()(float InPhase, float, const FGeneratorArgs&)
+			FORCEINLINE float operator()(float InPhase, float InPhaseDelta, const FGeneratorArgs& InArgs)
 			{
-				// Does not obey pulse width.
-				return InPhase >= 0.5f ? 1.f : -1.f;
+				float PulseWidthToUse = (InPhaseDelta >= 0) ? InArgs.PulseWidth : 1 - InArgs.PulseWidth;
+
+				return InPhase >= PulseWidthToUse ? 1.f : -1.f;
 			}
 		};
 		struct FSquarePolysmoothGenerator
@@ -484,7 +486,15 @@ namespace Metasound
 				// Simplified version of 
 				// float Output = 0.5f * SquareSaw1 - 0.5f * SquareSaw2;
 				// Output = 2.0f * (Output + InArgs.PulseWidth) - 1.0f;
-				return SquareSaw1 - SquareSaw2 + 2.0f * (InArgs.PulseWidth - 0.5f);
+				if (InPhaseDelta > 0.0)
+				{
+					return SquareSaw1 - SquareSaw2 + 2.0f * (InArgs.PulseWidth - 0.5f);
+				}
+				else
+				{
+					return SquareSaw1 - SquareSaw2 + 2.0f * (0.5f - InArgs.PulseWidth);
+				}
+				
 			}
 		};
 
