Add metasound constructor vertices
Integrate vertex data into graph builder.

diff --git a/Engine/Plugins/Runtime/Metasound/Source/MetasoundEngine/Private/MetasoundWavePlayerNode.cpp b/Engine/Plugins/Runtime/Metasound/Source/MetasoundEngine/Private/MetasoundWavePlayerNode.cpp
--- a/Engine/Plugins/Runtime/Metasound/Source/MetasoundEngine/Private/MetasoundWavePlayerNode.cpp
+++ b/Engine/Plugins/Runtime/Metasound/Source/MetasoundEngine/Private/MetasoundWavePlayerNode.cpp
@@ -52,7 +52,7 @@ namespace Metasound
 	{
 		class FOperatorFactory : public IOperatorFactory
 		{
-			virtual TUniquePtr<IOperator> CreateOperator(const FCreateOperatorParams& InParams, FBuildGraphResults& OutResults) override;
+			virtual TUniquePtr<IOperator> CreateOperator(const FBuildOperatorParams& InParams, FBuildResults& OutResults) override;
 		};
 
 	public:
@@ -983,26 +983,26 @@ namespace Metasound
 		
 	};
 
-	TUniquePtr<IOperator> FWavePlayerNode::FOperatorFactory::CreateOperator(const FCreateOperatorParams& InParams, FBuildGraphResults& OutResults)
+	TUniquePtr<IOperator> FWavePlayerNode::FOperatorFactory::CreateOperator(const FBuildOperatorParams& InParams, FBuildResults& OutResults)
 	{
 		using namespace Audio;
 		using namespace WavePlayerVertexNames;
 
 		const FWavePlayerNode& WaveNode = static_cast<const FWavePlayerNode&>(InParams.Node);
 
-		const FDataReferenceCollection& Inputs = InParams.InputDataReferences;
+		const FInputVertexInterfaceData& Inputs = InParams.InputData;
 
 		FWavePlayerOpArgs Args =
 		{
 			InParams.OperatorSettings,
-			Inputs.GetDataReadReferenceOrConstruct<FTrigger>(METASOUND_GET_PARAM_NAME(InputTriggerPlay), InParams.OperatorSettings),
-			Inputs.GetDataReadReferenceOrConstruct<FTrigger>(METASOUND_GET_PARAM_NAME(InputTriggerStop), InParams.OperatorSettings),
-			Inputs.GetDataReadReferenceOrConstruct<FWaveAsset>(METASOUND_GET_PARAM_NAME(InputWaveAsset)),
-			Inputs.GetDataReadReferenceOrConstruct<FTime>(METASOUND_GET_PARAM_NAME(InputStartTime)),
-			Inputs.GetDataReadReferenceOrConstruct<float>(METASOUND_GET_PARAM_NAME(InputPitchShift)),
-			Inputs.GetDataReadReferenceOrConstruct<bool>(METASOUND_GET_PARAM_NAME(InputLoop)),
-			Inputs.GetDataReadReferenceOrConstruct<FTime>(METASOUND_GET_PARAM_NAME(InputLoopStart)),
-			Inputs.GetDataReadReferenceOrConstruct<FTime>(METASOUND_GET_PARAM_NAME(InputLoopDuration))
+			Inputs.GetOrConstructDataReadReference<FTrigger>(METASOUND_GET_PARAM_NAME(InputTriggerPlay), InParams.OperatorSettings),
+			Inputs.GetOrConstructDataReadReference<FTrigger>(METASOUND_GET_PARAM_NAME(InputTriggerStop), InParams.OperatorSettings),
+			Inputs.GetOrConstructDataReadReference<FWaveAsset>(METASOUND_GET_PARAM_NAME(InputWaveAsset)),
+			Inputs.GetOrCreateDefaultDataReadReference<FTime>(METASOUND_GET_PARAM_NAME(InputStartTime), InParams.OperatorSettings),
+			Inputs.GetOrCreateDefaultDataReadReference<float>(METASOUND_GET_PARAM_NAME(InputPitchShift), InParams.OperatorSettings),
+			Inputs.GetOrCreateDefaultDataReadReference<bool>(METASOUND_GET_PARAM_NAME(InputLoop), InParams.OperatorSettings),
+			Inputs.GetOrCreateDefaultDataReadReference<FTime>(METASOUND_GET_PARAM_NAME(InputLoopStart), InParams.OperatorSettings),
+			Inputs.GetOrCreateDefaultDataReadReference<FTime>(METASOUND_GET_PARAM_NAME(InputLoopDuration), InParams.OperatorSettings)
 		};
 
 		return MakeUnique<FWavePlayerOperator>(Args);
diff --git a/Engine/Plugins/Runtime/Metasound/Source/MetasoundFrontend/Private/MetasoundFrontendGraphController.cpp b/Engine/Plugins/Runtime/Metasound/Source/MetasoundFrontend/Private/MetasoundFrontendGraphController.cpp
--- a/Engine/Plugins/Runtime/Metasound/Source/MetasoundFrontend/Private/MetasoundFrontendGraphController.cpp
+++ b/Engine/Plugins/Runtime/Metasound/Source/MetasoundFrontend/Private/MetasoundFrontendGraphController.cpp
@@ -1508,13 +1508,10 @@ namespace Metasound
 			return INodeController::GetInvalidHandle();
 		}
 
-		TUniquePtr<IOperator> FGraphController::BuildOperator(const FOperatorSettings& InSettings, const FMetasoundEnvironment& InEnvironment, FBuildGraphResults& OutResults) const
+		TUniquePtr<IOperator> FGraphController::BuildOperator(const FOperatorSettings& InSettings, const FMetasoundEnvironment& InEnvironment, FBuildResults& OutResults) const
 		{
 			if (const FMetasoundFrontendGraphClass* GraphClass = GraphClassPtr.Get())
 			{
-
-				// TODO: Implement subgraph inflation step here.
-
 				// TODO: bubble up errors. 
 				const TArray<FMetasoundFrontendGraphClass>& Subgraphs = OwningDocument->GetSubgraphs();
 				const TArray<FMetasoundFrontendClass>& Dependencies = OwningDocument->GetDependencies();
@@ -1527,8 +1524,8 @@ namespace Metasound
 					return TUniquePtr<IOperator>(nullptr);
 				}
 
-				FBuildGraphParams BuildParams { *Graph, InSettings, FDataReferenceCollection { }, InEnvironment, FOperatorBuilderSettings::GetDefaultSettings() };
-				return FOperatorBuilder().BuildGraphOperator(BuildParams, OutResults);
+				FBuildGraphOperatorParams BuildParams { *Graph, InSettings, FInputVertexInterfaceData{}, InEnvironment };
+				return FOperatorBuilder(FOperatorBuilderSettings::GetDefaultSettings()).BuildGraphOperator(BuildParams, OutResults);
 			}
 			else
 			{
diff --git a/Engine/Plugins/Runtime/Metasound/Source/MetasoundFrontend/Private/MetasoundFrontendGraphController.h b/Engine/Plugins/Runtime/Metasound/Source/MetasoundFrontend/Private/MetasoundFrontendGraphController.h
--- a/Engine/Plugins/Runtime/Metasound/Source/MetasoundFrontend/Private/MetasoundFrontendGraphController.h
+++ b/Engine/Plugins/Runtime/Metasound/Source/MetasoundFrontend/Private/MetasoundFrontendGraphController.h
@@ -185,7 +185,7 @@ namespace Metasound
 
 			FNodeHandle CreateEmptySubgraph(const FMetasoundFrontendClassMetadata& InInfo) override;
 
-			TUniquePtr<IOperator> BuildOperator(const FOperatorSettings& InSettings, const FMetasoundEnvironment& InEnvironment, FBuildGraphResults& OutBuildErrors) const override;
+			TUniquePtr<IOperator> BuildOperator(const FOperatorSettings& InSettings, const FMetasoundEnvironment& InEnvironment, FBuildResults& OutBuildErrors) const override;
 
 			FDocumentHandle GetOwningDocument() override;
 			FConstDocumentHandle GetOwningDocument() const override;
diff --git a/Engine/Plugins/Runtime/Metasound/Source/MetasoundFrontend/Private/MetasoundFrontendInvalidController.h b/Engine/Plugins/Runtime/Metasound/Source/MetasoundFrontend/Private/MetasoundFrontendInvalidController.h
--- a/Engine/Plugins/Runtime/Metasound/Source/MetasoundFrontend/Private/MetasoundFrontendInvalidController.h
+++ b/Engine/Plugins/Runtime/Metasound/Source/MetasoundFrontend/Private/MetasoundFrontendInvalidController.h
@@ -419,7 +419,7 @@ namespace Metasound
 
 			virtual FNodeHandle CreateEmptySubgraph(const FMetasoundFrontendClassMetadata& InInfo) override { return INodeController::GetInvalidHandle(); }
 
-			virtual TUniquePtr<IOperator> BuildOperator(const FOperatorSettings& InSettings, const FMetasoundEnvironment& InEnvironment, FBuildGraphResults& OutBuildErrors) const override { return { }; }
+			virtual TUniquePtr<IOperator> BuildOperator(const FOperatorSettings& InSettings, const FMetasoundEnvironment& InEnvironment, FBuildResults& OutBuildErrors) const override { return { }; }
 
 			virtual FDocumentHandle GetOwningDocument() override;
 			virtual FConstDocumentHandle GetOwningDocument() const override;
diff --git a/Engine/Plugins/Runtime/Metasound/Source/MetasoundFrontend/Public/MetasoundAutoConverterNode.h b/Engine/Plugins/Runtime/Metasound/Source/MetasoundFrontend/Public/MetasoundAutoConverterNode.h
--- a/Engine/Plugins/Runtime/Metasound/Source/MetasoundFrontend/Public/MetasoundAutoConverterNode.h
+++ b/Engine/Plugins/Runtime/Metasound/Source/MetasoundFrontend/Public/MetasoundAutoConverterNode.h
@@ -212,15 +212,15 @@ namespace Metasound
 			public:
 				FCoverterOperatorFactory() = default;
 
-				virtual TUniquePtr<IOperator> CreateOperator(const FCreateOperatorParams& InParams, FBuildGraphResults& OutResults) override
+				virtual TUniquePtr<IOperator> CreateOperator(const FBuildOperatorParams& InParams, FBuildResults& OutResults) override
 				{
 					TDataWriteReference<ToDataType> WriteReference = TDataWriteReferenceFactory<ToDataType>::CreateAny(InParams.OperatorSettings);
 
 					const FVertexName& InputName = GetInputName();
-					const bool bContainsRef = InParams.InputDataReferences.ContainsDataReadReference<FromDataType>(InputName);
+					const bool bContainsRef = InParams.InputData.IsVertexBound(InputName);
 					if (bContainsRef)
 					{
-						TDataReadReference<FromDataType> ReadReference = InParams.InputDataReferences.GetDataReadReference<FromDataType>(InputName);
+						TDataReadReference<FromDataType> ReadReference = InParams.InputData.GetDataReadReference<FromDataType>(InputName);
 						return MakeUnique<FConverterOperator>(ReadReference, WriteReference);
 					}
 
diff --git a/Engine/Plugins/Runtime/Metasound/Source/MetasoundFrontend/Public/MetasoundFrontendController.h b/Engine/Plugins/Runtime/Metasound/Source/MetasoundFrontend/Public/MetasoundFrontendController.h
--- a/Engine/Plugins/Runtime/Metasound/Source/MetasoundFrontend/Public/MetasoundFrontendController.h
+++ b/Engine/Plugins/Runtime/Metasound/Source/MetasoundFrontend/Public/MetasoundFrontendController.h
@@ -962,7 +962,7 @@ namespace Metasound
 			 *
 			 * @return On success, a valid pointer to a Metasound operator. An invalid pointer on failure.
 			 */
-			virtual TUniquePtr<IOperator> BuildOperator(const FOperatorSettings& InSettings, const FMetasoundEnvironment& InEnvironment, FBuildGraphResults& OutResults) const = 0;
+			virtual TUniquePtr<IOperator> BuildOperator(const FOperatorSettings& InSettings, const FMetasoundEnvironment& InEnvironment, FBuildResults& OutResults) const = 0;
 
 			/** Returns a handle to the document owning this graph. */
 			virtual FDocumentHandle GetOwningDocument() = 0;
diff --git a/Engine/Plugins/Runtime/Metasound/Source/MetasoundFrontend/Public/MetasoundInputNode.h b/Engine/Plugins/Runtime/Metasound/Source/MetasoundFrontend/Public/MetasoundInputNode.h
--- a/Engine/Plugins/Runtime/Metasound/Source/MetasoundFrontend/Public/MetasoundInputNode.h
+++ b/Engine/Plugins/Runtime/Metasound/Source/MetasoundFrontend/Public/MetasoundInputNode.h
@@ -246,7 +246,7 @@ namespace Metasound
 			{
 			}
 
-			virtual TUniquePtr<IOperator> CreateOperator(const FCreateOperatorParams& InParams, FBuildGraphResults& OutResults) override;
+			virtual TUniquePtr<IOperator> CreateOperator(const FBuildOperatorParams& InParams, FBuildResults& OutResults) override;
 
 		private:
 			ReferenceCreatorType ReferenceCreator;
@@ -361,23 +361,17 @@ namespace Metasound
 
 
 	template<typename DataType, typename ReferenceCreatorType>
-	TUniquePtr<IOperator> TInputOperatorFactory<DataType, ReferenceCreatorType>::CreateOperator(const FCreateOperatorParams& InParams, FBuildGraphResults& OutResults)
+	TUniquePtr<IOperator> TInputOperatorFactory<DataType, ReferenceCreatorType>::CreateOperator(const FBuildOperatorParams& InParams, FBuildResults& OutResults)
 	{
 		using FInputNodeType = TInputNode<DataType>;
 
 		const FInputNodeType& InputNode = static_cast<const FInputNodeType&>(InParams.Node);
 		const FVertexName& VertexKey = InputNode.GetVertexName();
 
-		if (InParams.InputDataReferences.ContainsDataWriteReference<DataType>(VertexKey))
+		if (InParams.InputData.IsVertexBound(VertexKey))
 		{
 			// Data is externally owned. Use pass through operator
-			FDataWriteReference DataRef = InParams.InputDataReferences.GetDataWriteReference<DataType>(VertexKey);
-			return MakeUnique<TPassThroughOperator<DataType>>(InputNode.GetVertexName(), DataRef);
-		}
-		else if (InParams.InputDataReferences.ContainsDataReadReference<DataType>(VertexKey))
-		{
-			// Data is externally owned. Use pass through operator
-			FDataReadReference DataRef = InParams.InputDataReferences.GetDataReadReference<DataType>(VertexKey);
+			FDataReadReference DataRef = InParams.InputData.GetDataReadReference<DataType>(VertexKey);
 			return MakeUnique<TPassThroughOperator<DataType>>(InputNode.GetVertexName(), DataRef);
 		}
 		else
diff --git a/Engine/Plugins/Runtime/Metasound/Source/MetasoundFrontend/Public/MetasoundOutputNode.h b/Engine/Plugins/Runtime/Metasound/Source/MetasoundFrontend/Public/MetasoundOutputNode.h
--- a/Engine/Plugins/Runtime/Metasound/Source/MetasoundFrontend/Public/MetasoundOutputNode.h
+++ b/Engine/Plugins/Runtime/Metasound/Source/MetasoundFrontend/Public/MetasoundOutputNode.h
@@ -75,12 +75,12 @@ namespace Metasound
 				{
 				}
 
-				virtual TUniquePtr<IOperator> CreateOperator(const FCreateOperatorParams& InParams, FBuildGraphResults& OutResults) override
+				virtual TUniquePtr<IOperator> CreateOperator(const FBuildOperatorParams& InParams, FBuildResults& OutResults) override
 				{
-					const bool bContainsRef = InParams.InputDataReferences.ContainsDataReadReference<DataType>(DataReferenceName);
+					const bool bContainsRef = InParams.InputData.IsVertexBound(DataReferenceName);
 					if (bContainsRef)
 					{
-						TDataReadReference<DataType> InputReadRef = InParams.InputDataReferences.GetDataReadReference<DataType>(DataReferenceName);
+						TDataReadReference<DataType> InputReadRef = InParams.InputData.GetDataReadReference<DataType>(DataReferenceName);
 						return MakeUnique<FOutputOperator>(DataReferenceName, InputReadRef);
 					}
 
diff --git a/Engine/Plugins/Runtime/Metasound/Source/MetasoundFrontend/Public/MetasoundReceiveNode.h b/Engine/Plugins/Runtime/Metasound/Source/MetasoundFrontend/Public/MetasoundReceiveNode.h
--- a/Engine/Plugins/Runtime/Metasound/Source/MetasoundFrontend/Public/MetasoundReceiveNode.h
+++ b/Engine/Plugins/Runtime/Metasound/Source/MetasoundFrontend/Public/MetasoundReceiveNode.h
@@ -212,21 +212,21 @@ namespace Metasound
 			public:
 				FReceiverOperatorFactory() = default;
 
-				virtual TUniquePtr<IOperator> CreateOperator(const FCreateOperatorParams& InParams, FBuildGraphResults& OutResults) override
+				virtual TUniquePtr<IOperator> CreateOperator(const FBuildOperatorParams& InParams, FBuildResults& OutResults) override
 				{
 					using namespace ReceiveNodeInfo; 
 
 					TDataReadReference<TDataType> DefaultReadRef = TDataReadReferenceFactory<TDataType>::CreateAny(InParams.OperatorSettings);
 
-					if (InParams.InputDataReferences.ContainsDataReadReference<TDataType>(METASOUND_GET_PARAM_NAME(DefaultDataInput)))
+					if (InParams.InputData.IsVertexBound(METASOUND_GET_PARAM_NAME(DefaultDataInput)))
 					{
-						DefaultReadRef = InParams.InputDataReferences.GetDataReadReference<TDataType>(METASOUND_GET_PARAM_NAME(DefaultDataInput));
+						DefaultReadRef = InParams.InputData.GetDataReadReference<TDataType>(METASOUND_GET_PARAM_NAME(DefaultDataInput));
 					}
 
 					return MakeUnique<TReceiverOperator>(
 						DefaultReadRef,
 						TDataWriteReferenceFactory<TDataType>::CreateAny(InParams.OperatorSettings, *DefaultReadRef),
-						InParams.InputDataReferences.GetDataReadReferenceOrConstruct<FSendAddress>(METASOUND_GET_PARAM_NAME(AddressInput)),
+						InParams.InputData.GetOrConstructDataReadReference<FSendAddress>(METASOUND_GET_PARAM_NAME(AddressInput)),
 						InParams.OperatorSettings
 						);
 				}
diff --git a/Engine/Plugins/Runtime/Metasound/Source/MetasoundFrontend/Public/MetasoundSendNode.h b/Engine/Plugins/Runtime/Metasound/Source/MetasoundFrontend/Public/MetasoundSendNode.h
--- a/Engine/Plugins/Runtime/Metasound/Source/MetasoundFrontend/Public/MetasoundSendNode.h
+++ b/Engine/Plugins/Runtime/Metasound/Source/MetasoundFrontend/Public/MetasoundSendNode.h
@@ -169,14 +169,14 @@ namespace Metasound
 			public:
 				FSendOperatorFactory() = default;
 
-				virtual TUniquePtr<IOperator> CreateOperator(const FCreateOperatorParams& InParams, FBuildGraphResults& OutResults) override
+				virtual TUniquePtr<IOperator> CreateOperator(const FBuildOperatorParams& InParams, FBuildResults& OutResults) override
 				{
 					using namespace SendVertexNames;
 
-					if (InParams.InputDataReferences.ContainsDataReadReference<TDataType>(GetSendInputName()))
+					if (InParams.InputData.IsVertexBound(GetSendInputName()))
 					{
-						return MakeUnique<TSendOperator>(InParams.InputDataReferences.GetDataReadReference<TDataType>(GetSendInputName()),
-							InParams.InputDataReferences.GetDataReadReferenceOrConstruct<FSendAddress>(METASOUND_GET_PARAM_NAME(AddressInput)),
+						return MakeUnique<TSendOperator>(InParams.InputData.GetDataReadReference<TDataType>(GetSendInputName()),
+							InParams.InputData.GetOrConstructDataReadReference<FSendAddress>(METASOUND_GET_PARAM_NAME(AddressInput)),
 							InParams.OperatorSettings
 						);
 					}
diff --git a/Engine/Plugins/Runtime/Metasound/Source/MetasoundGenerator/Private/MetasoundGenerator.cpp b/Engine/Plugins/Runtime/Metasound/Source/MetasoundGenerator/Private/MetasoundGenerator.cpp
--- a/Engine/Plugins/Runtime/Metasound/Source/MetasoundGenerator/Private/MetasoundGenerator.cpp
+++ b/Engine/Plugins/Runtime/Metasound/Source/MetasoundGenerator/Private/MetasoundGenerator.cpp
@@ -57,10 +57,11 @@ namespace Metasound
 
 		// Create an instance of the new graph
 		FDataReferenceCollection DataReferenceCollection;
-		FBuildGraphParams BuildParams { *InitParams.Graph, InitParams.OperatorSettings, DataReferenceCollection, InitParams.Environment, InitParams.BuilderSettings };
-		FBuildGraphResults BuildResults;
+		FInputVertexInterfaceData InputVertexData;
+		FBuildGraphOperatorParams BuildParams { *InitParams.Graph, InitParams.OperatorSettings, InputVertexData, InitParams.Environment };
+		FBuildResults BuildResults;
 
-		TUniquePtr<IOperator> GraphOperator = FOperatorBuilder().BuildGraphOperator(BuildParams, BuildResults);
+		TUniquePtr<IOperator> GraphOperator = FOperatorBuilder(InitParams.BuilderSettings).BuildGraphOperator(BuildParams, BuildResults);
 
 		TUniquePtr<FGraphAnalyzer> GraphAnalyzer;
 		if (InitParams.BuilderSettings.bPopulateInternalDataReferences)
diff --git a/Engine/Plugins/Runtime/Metasound/Source/MetasoundGraphCore/Private/MetasoundDataReference.cpp b/Engine/Plugins/Runtime/Metasound/Source/MetasoundGraphCore/Private/MetasoundDataReference.cpp
--- a/Engine/Plugins/Runtime/Metasound/Source/MetasoundGraphCore/Private/MetasoundDataReference.cpp
+++ b/Engine/Plugins/Runtime/Metasound/Source/MetasoundGraphCore/Private/MetasoundDataReference.cpp
@@ -27,3 +27,30 @@ namespace Metasound
 		return *InArrayTypeName.ToString().LeftChop(SuffixLength);
 	}
 }
+
+FString LexToString(Metasound::EDataReferenceAccessType InAccessType)
+{
+	using namespace Metasound;
+
+	switch (InAccessType)
+	{
+		case EDataReferenceAccessType::None:
+			return FString(TEXT("None"));
+
+		case EDataReferenceAccessType::Read:
+			return FString(TEXT("Read"));
+
+		case EDataReferenceAccessType::Write:
+			return FString(TEXT("Write"));
+
+		case EDataReferenceAccessType::Value:
+			return FString(TEXT("Value"));
+
+		default:
+			{
+				checkNoEntry();
+			}
+	}
+	return FString(TEXT(""));
+}
+
diff --git a/Engine/Plugins/Runtime/Metasound/Source/MetasoundGraphCore/Private/MetasoundFacade.cpp b/Engine/Plugins/Runtime/Metasound/Source/MetasoundGraphCore/Private/MetasoundFacade.cpp
--- a/Engine/Plugins/Runtime/Metasound/Source/MetasoundGraphCore/Private/MetasoundFacade.cpp
+++ b/Engine/Plugins/Runtime/Metasound/Source/MetasoundGraphCore/Private/MetasoundFacade.cpp
@@ -9,9 +9,13 @@ namespace Metasound
 	{
 	}
 
-	TUniquePtr<IOperator> FNodeFacade::FFactory::CreateOperator(const FCreateOperatorParams& InParams, FBuildGraphResults& OutResults)
+	TUniquePtr<IOperator> FNodeFacade::FFactory::CreateOperator(const FBuildOperatorParams& InParams, FBuildResults& OutResults)
 	{
-		return CreateFunc(InParams, OutResults.Errors);
+		// TODO: update facade to deprecate old CreateOperator functionality.
+		FDataReferenceCollection InputDataReferences = InParams.InputData.ToDataReferenceCollection();
+		FCreateOperatorParams DeprecatedParams{ InParams.Node, InParams.OperatorSettings, InputDataReferences, InParams.Environment, InParams.Builder };
+
+		return CreateFunc(DeprecatedParams, OutResults.Errors);
 	}
 
 	const FVertexInterface& FNodeFacade::GetVertexInterface() const
diff --git a/Engine/Plugins/Runtime/Metasound/Source/MetasoundGraphCore/Private/MetasoundGraph.cpp b/Engine/Plugins/Runtime/Metasound/Source/MetasoundGraphCore/Private/MetasoundGraph.cpp
--- a/Engine/Plugins/Runtime/Metasound/Source/MetasoundGraphCore/Private/MetasoundGraph.cpp
+++ b/Engine/Plugins/Runtime/Metasound/Source/MetasoundGraphCore/Private/MetasoundGraph.cpp
@@ -139,11 +139,11 @@ namespace Metasound
 		return MakeShared<FGraph::FFactory>();
 	}
 
-	TUniquePtr<IOperator> FGraph::FFactory::CreateOperator(const FCreateOperatorParams& InParams, FBuildGraphResults& OutResults)
+	TUniquePtr<IOperator> FGraph::FFactory::CreateOperator(const FBuildOperatorParams& InParams, FBuildResults& OutResults)
 	{
 		const FGraph& Graph = static_cast<const FGraph&>(InParams.Node);
 
-		FBuildGraphParams BuildParams { Graph, InParams.OperatorSettings, InParams.InputDataReferences, InParams.Environment, InParams.BuilderSettings };
+		FBuildGraphOperatorParams BuildParams { Graph, InParams.OperatorSettings, InParams.InputData, InParams.Environment};
 
 		if (nullptr != InParams.Builder)
 		{
@@ -152,7 +152,7 @@ namespace Metasound
 		}
 		else
 		{
-			return FOperatorBuilder().BuildGraphOperator(BuildParams, OutResults);
+			return FOperatorBuilder(FOperatorBuilderSettings::GetDefaultSettings()).BuildGraphOperator(BuildParams, OutResults);
 		}
 	}
 }
diff --git a/Engine/Plugins/Runtime/Metasound/Source/MetasoundGraphCore/Private/MetasoundGraphOperator.cpp b/Engine/Plugins/Runtime/Metasound/Source/MetasoundGraphCore/Private/MetasoundGraphOperator.cpp
--- a/Engine/Plugins/Runtime/Metasound/Source/MetasoundGraphCore/Private/MetasoundGraphOperator.cpp
+++ b/Engine/Plugins/Runtime/Metasound/Source/MetasoundGraphCore/Private/MetasoundGraphOperator.cpp
@@ -28,22 +28,32 @@ namespace Metasound
 
 	void FGraphOperator::SetInputs(const FDataReferenceCollection& InCollection)
 	{
-		Inputs = InCollection;
+		VertexData.GetInputs().Bind(InCollection);
 	}
 
 	void FGraphOperator::SetOutputs(const FDataReferenceCollection& InCollection)
 	{
-		Outputs = InCollection;
+		VertexData.GetOutputs().Bind(InCollection);
+	}
+
+	void FGraphOperator::SetVertexInterfaceData(FVertexInterfaceData&& InVertexData)
+	{
+		VertexData = InVertexData;
 	}
 
 	FDataReferenceCollection FGraphOperator::GetInputs() const
 	{
-		return Inputs;
+		return VertexData.GetInputs().ToDataReferenceCollection();
 	}
 
 	FDataReferenceCollection FGraphOperator::GetOutputs() const
 	{
-		return Outputs;
+		return VertexData.GetOutputs().ToDataReferenceCollection();
+	}
+
+	void FGraphOperator::Bind(FVertexInterfaceData& InVertexData) const
+	{
+		InVertexData.Bind(VertexData);
 	}
 
 	IOperator::FExecuteFunction FGraphOperator::GetExecuteFunction()
diff --git a/Engine/Plugins/Runtime/Metasound/Source/MetasoundGraphCore/Private/MetasoundOperatorBuilder.cpp b/Engine/Plugins/Runtime/Metasound/Source/MetasoundGraphCore/Private/MetasoundOperatorBuilder.cpp
--- a/Engine/Plugins/Runtime/Metasound/Source/MetasoundGraphCore/Private/MetasoundOperatorBuilder.cpp
+++ b/Engine/Plugins/Runtime/Metasound/Source/MetasoundGraphCore/Private/MetasoundOperatorBuilder.cpp
@@ -4,15 +4,14 @@
 
 #include "Algo/Count.h"
 #include "Algo/Transform.h"
-
 #include "MetasoundBuildError.h"
 #include "MetasoundBuilderInterface.h"
 #include "MetasoundBuildError.h"
 #include "MetasoundGraphAlgo.h"
+#include "MetasoundGraphLinter.h"
 #include "MetasoundGraphOperator.h"
 #include "MetasoundOperatorInterface.h"
-#include "MetasoundGraphAlgo.h"
-#include "MetasoundGraphLinter.h"
+#include "MetasoundThreadLocalDebug.h"
 #include "Templates/PimplPtr.h"
 
 
@@ -30,22 +29,27 @@ namespace Metasound
 				AddBuildError<FInternalError>(OutErrors, __FILE__, __LINE__);
 			}
 		}
+
+	}
+
+	FOperatorBuilder::FOperatorBuilder(const FOperatorBuilderSettings& InBuilderSettings)
+	: BuilderSettings(InBuilderSettings)
+	{
 	}
 
 	FOperatorBuilder::~FOperatorBuilder()
 	{
 	}
 
-	TUniquePtr<IOperator> FOperatorBuilder::BuildGraphOperator(const FBuildGraphParams& InParams, FBuildGraphResults& OutResults) const
+	TUniquePtr<IOperator> FOperatorBuilder::BuildGraphOperator(const FBuildGraphOperatorParams& InParams, FBuildResults& OutResults) const
 	{
 		using namespace OperatorBuilderPrivate;
 
 		FBuildStatus BuildStatus = FBuildStatus::NoError;
-		const FBuildStatus MaxBuildStatusErrorLevel = FBuildStatus::GetMaxErrorLevel(InParams.BuilderSettings);
 
 		// Validate that the sources and destinations declared in an edge actually
 		// exist in the node.
-		if (InParams.BuilderSettings.bValidateVerticesExist)
+		if (BuilderSettings.bValidateVerticesExist)
 		{
 			if (!FGraphLinter::ValidateVerticesExist(InParams.Graph, OutResults.Errors))
 			{
@@ -54,7 +58,7 @@ namespace Metasound
 		}
 
 		// Validate that the data types for a source and destination match.
-		if (InParams.BuilderSettings.bValidateEdgeDataTypesMatch)
+		if (BuilderSettings.bValidateEdgeDataTypesMatch)
 		{
 			if (!FGraphLinter::ValidateEdgeDataTypesMatch(InParams.Graph, OutResults.Errors))
 			{
@@ -63,7 +67,7 @@ namespace Metasound
 		}
 		
 		// Validate that node inputs only have one source
-		if (InParams.BuilderSettings.bValidateNoDuplicateInputs)
+		if (BuilderSettings.bValidateNoDuplicateInputs)
 		{
 			if (!FGraphLinter::ValidateNoDuplicateInputs(InParams.Graph, OutResults.Errors))
 			{
@@ -72,7 +76,7 @@ namespace Metasound
 		}
 
 		// Possible early exit if edge validation fails.
-		if (BuildStatus > MaxBuildStatusErrorLevel)
+		if (BuildStatus > GetMaxErrorLevel())
 		{
 			return TUniquePtr<IOperator>(nullptr);
 		}
@@ -86,7 +90,7 @@ namespace Metasound
 			return TUniquePtr<IOperator>(nullptr);
 		}
 
-		FBuildContext BuildContext(InParams.Graph, *AlgoAdapter, InParams.OperatorSettings, InParams.Environment, InParams.BuilderSettings, OutResults);
+		FBuildContext BuildContext(InParams.Graph, *AlgoAdapter, InParams.OperatorSettings, InParams.Environment, BuilderSettings, OutResults);
 
 		TArray<const INode*> SortedNodes;
 
@@ -98,15 +102,15 @@ namespace Metasound
 		// BuildStatus |= PruneNodges(BuildContext, SortedNodes);
 
 		// Check build status in case build routine should be exited early.
-		if (BuildStatus > MaxBuildStatusErrorLevel)
+		if (BuildStatus > GetMaxErrorLevel())
 		{
 			return TUniquePtr<IOperator>(nullptr);
 		}
 
 		// Create node operators from factories.
-		BuildStatus |= CreateOperators(BuildContext, SortedNodes, InParams.InputDataReferences);
+		BuildStatus |= CreateOperators(BuildContext, SortedNodes, InParams.InputData);
 
-		if (BuildStatus > MaxBuildStatusErrorLevel)
+		if (BuildStatus > GetMaxErrorLevel())
 		{
 			return TUniquePtr<IOperator>(nullptr);
 		}
@@ -157,7 +161,7 @@ namespace Metasound
 
 		TSet<const INode*> ReachableNodes;
 
-		switch (InOutContext.BuilderSettings.PruningMode)
+		switch (BuilderSettings.PruningMode)
 		{
 			case EOperatorBuilderNodePruning::PruneNodesWithoutExternalDependency:
 				FDirectedGraphAlgo::FindReachableNodes(InOutContext.AlgoAdapter, ReachableNodes);
@@ -225,7 +229,7 @@ namespace Metasound
 		return BuildStatus;
 	}
 
-	FOperatorBuilder::FBuildStatus FOperatorBuilder::GatherInputDataReferences(FBuildContext& InOutContext, const INode* InNode, const FNodeEdgeMultiMap& InEdgeMap, FDataReferenceCollection& OutInputCollection) const
+	FOperatorBuilder::FBuildStatus FOperatorBuilder::GatherInputDataReferences(FBuildContext& InOutContext, const INode* InNode, const FNodeEdgeMultiMap& InEdgeMap, FInputVertexInterfaceData& OutVertexData) const
 	{
 		using namespace OperatorBuilderPrivate;
 
@@ -247,28 +251,20 @@ namespace Metasound
 				return FBuildStatus::NonFatalError;
 			}
 
-			const FDataReferenceCollection& FromDataReferenceCollection = InOutContext.DataReferences[Edge->From.Node].Outputs;
-
-			bool bSuccess = false;
+			const FOutputVertexInterfaceData& FromData = InOutContext.DataReferences[Edge->From.Node].GetOutputs();
 
-			if (FromDataReferenceCollection.ContainsDataWriteReference(Edge->From.Vertex.VertexName, Edge->From.Vertex.DataTypeName))
-			{
-				// Contains data write reference
-				bSuccess = OutInputCollection.AddDataWriteReferenceFrom(Edge->To.Vertex.VertexName, FromDataReferenceCollection, Edge->From.Vertex.VertexName, Edge->From.Vertex.DataTypeName);
-			}
-			else if (FromDataReferenceCollection.ContainsDataReadReference(Edge->From.Vertex.VertexName, Edge->From.Vertex.DataTypeName))
+			if (const FAnyDataReference* DataReference = FromData.FindDataReference(Edge->From.Vertex.VertexName))
 			{
-				// Contains data read reference
-				bSuccess = OutInputCollection.AddDataReadReferenceFrom(Edge->To.Vertex.VertexName, FromDataReferenceCollection, Edge->From.Vertex.VertexName, Edge->From.Vertex.DataTypeName);
+				check(DataReference->GetDataTypeName() == Edge->From.Vertex.DataTypeName);
+				check(DataReference->GetDataTypeName() == Edge->To.Vertex.DataTypeName);
+				OutVertexData.BindVertex(Edge->To.Vertex.VertexName, FAnyDataReference{*DataReference});
 			}
-
-			if (!bSuccess)
+			else 
 			{
 				// Does not contain any reference
 				// This is likely a node programming error where the edges reported by the INode interface
 				// did not match the readable parameter refs created by the operators outputs. Or, the edge description is invalid.
 
-				// TODO: consider checking that outputs of node operators match output descriptions of nodes. 
 				AddBuildError<FMissingOutputDataReferenceError>(InOutContext.Results.Errors, Edge->From);
 
 				BuildStatus |= FBuildStatus::NonFatalError;
@@ -279,26 +275,40 @@ namespace Metasound
 		return BuildStatus;
 	}
 
-	FOperatorBuilder::FBuildStatus FOperatorBuilder::GatherExternalInputDataReferences(FBuildContext& InOutContext, const INode* InNode, const FNodeDestinationMap& InNodeDestinationMap, const FDataReferenceCollection& InExternalCollection, FDataReferenceCollection& OutDataReferences) const
+	FOperatorBuilder::FBuildStatus FOperatorBuilder::GatherExternalInputDataReferences(FBuildContext& InOutContext, const INode* InNode, const FNodeDestinationMap& InNodeDestinationMap, const FInputVertexInterfaceData& InExternalInputData, FInputVertexInterfaceData& OutVertexData) const
 	{
 		using namespace OperatorBuilderPrivate;
 
 		FBuildStatus BuildStatus;
 
+		// Check if there are any graph inputs connected to the node's inputs.
 		if (const FInputDataDestination* const* DestinationPtr = InNodeDestinationMap.Find(InNode))
 		{
 			if (const FInputDataDestination* Destination = *DestinationPtr)
 			{
-				const bool bReferenceExists = InExternalCollection.ContainsDataReadReference(Destination->Vertex.VertexName, Destination->Vertex.DataTypeName);
-				if (bReferenceExists)
+				if (const FAnyDataReference* DataReference = InExternalInputData.FindDataReference(Destination->Vertex.VertexName))
 				{
-					const bool bSuccess = OutDataReferences.AddDataReadReferenceFrom(Destination->Vertex.VertexName, InExternalCollection, Destination->Vertex.VertexName, Destination->Vertex.DataTypeName);
-
-					if (!bSuccess)
+					if (DataReference->GetDataTypeName() == Destination->Vertex.DataTypeName)
 					{
-						AddBuildError<FInternalError>(InOutContext.Results.Errors, __FILE__, __LINE__);
-
-						BuildStatus |= FBuildStatus::NonFatalError;
+						OutVertexData.BindVertex(Destination->Vertex.VertexName, FAnyDataReference{*DataReference});
+					}
+					else
+					{
+						// Mismatch in datatypes. This likely corresponds to a corrupt 
+						// Metasound Graph. The graph's inputs should route directly
+						// to TInputNode<>s with matching data types. 
+						
+						// Create source for reporting connection since external inputs do not have nodes. 
+						const FInputDataVertex& GraphVertex = InExternalInputData.GetVertex(Destination->Vertex.VertexName);
+
+						FOutputDataSource Source;
+						Source.Vertex.VertexName = GraphVertex.VertexName; 
+						Source.Vertex.DataTypeName = GraphVertex.DataTypeName;
+#if WITH_EDITORONLY_DATA
+						Source.Vertex.Metadata = GraphVertex.Metadata;
+#endif
+
+						AddBuildError<FInvalidConnectionDataTypeError>(InOutContext.Results.Errors, FDataEdge{Source, *Destination});
 					}
 				}
 			}
@@ -307,7 +317,41 @@ namespace Metasound
 		return BuildStatus;
 	}
 
-	FOperatorBuilder::FBuildStatus FOperatorBuilder::CreateOperators(FBuildContext& InOutContext, const TArray<const INode*>& InSortedNodes, const FDataReferenceCollection& InExternalInputDataReferences) const
+	void FOperatorBuilder::GatherInternalGraphDataReferences(FBuildContext& InOutContext, TMap<FGuid, FDataReferenceCollection>& OutNodeVertexData) const
+	{
+		for (TPair<const INode*, FVertexInterfaceData>& ReferencePair : InOutContext.DataReferences)
+		{
+			const INode* Node = ReferencePair.Key;
+			check(Node);
+			OutNodeVertexData.Emplace(Node->GetInstanceID(), ReferencePair.Value.GetOutputs().ToDataReferenceCollection());
+		}
+	}
+
+	FOperatorBuilder::FBuildStatus FOperatorBuilder::ValidateOperatorOutputsAreBound(const INode& InNode, const FOutputVertexInterfaceData& InVertexData) const
+	{
+		using FOutputVertexBinding = MetasoundVertexDataPrivate::TBinding<FOutputDataVertex>;
+
+		FBuildStatus BuildStatus = FBuildStatus::NoError;
+		bool bFoundUnboundVertex = false;
+		for (const FOutputVertexBinding& OutputBinding : InVertexData)
+		{
+			if (!OutputBinding.IsBound())
+			{
+				bFoundUnboundVertex = true;
+				const FNodeClassMetadata& Metadata = InNode.GetMetadata();
+				UE_LOG(LogMetaSound, Warning, TEXT("Operator for node %s v%d.%d contains unbound output vertex %s"), *Metadata.ClassName.GetFullName().ToString(), Metadata.MajorVersion, Metadata.MinorVersion, *OutputBinding.GetVertex().VertexName.ToString());
+			}
+		}
+
+		if (bFoundUnboundVertex)
+		{
+			BuildStatus |= FBuildStatus::NonFatalError;
+		}
+
+		return BuildStatus;
+	}
+
+	FOperatorBuilder::FBuildStatus FOperatorBuilder::CreateOperators(FBuildContext& InOutContext, const TArray<const INode*>& InSortedNodes, const FInputVertexInterfaceData& InExternalInputData) const
 	{
 		FBuildStatus BuildStatus;
 
@@ -330,12 +374,15 @@ namespace Metasound
 		// Call operator factory for each node.
 		for (const INode* Node : InSortedNodes)
 		{
-			// Gather the input parameters for this IOperator from the output parameters of already created IOperators. 
-			FDataReferenceCollection InputCollection;
-			BuildStatus |= GatherInputDataReferences(InOutContext, Node, NodeInputEdges, InputCollection);
+			ThreadLocalDebug::SetActiveNodeClass(Node->GetMetadata());
+
+			const FVertexInterface& NodeInterface = Node->GetVertexInterface();
+			FInputVertexInterfaceData InputData(NodeInterface.GetInputInterface());
 
+			// Gather the input parameters for this IOperator from the output parameters of already created IOperators. 
+			BuildStatus |= GatherInputDataReferences(InOutContext, Node, NodeInputEdges, InputData);
 			// Gather the input parameters for this IOperator from the graph inputs.
-			BuildStatus |= GatherExternalInputDataReferences(InOutContext, Node, NodeDestinations, InExternalInputDataReferences, InputCollection);
+			BuildStatus |= GatherExternalInputDataReferences(InOutContext, Node, NodeDestinations, InExternalInputData, InputData);
 
 			if (BuildStatus >= FBuildStatus::FatalError)
 			{
@@ -344,7 +391,11 @@ namespace Metasound
 
 			FOperatorFactorySharedRef Factory = Node->GetDefaultOperatorFactory();
 
-			FCreateOperatorParams CreateParams { *Node, InOutContext.Settings, InputCollection, InOutContext.Environment, InOutContext.BuilderSettings };
+			// This is here to handle the deprecated field FCreateOperator::InputDataReferences.
+			// It can be removed once the deprecated member is removed.
+			FDataReferenceCollection InputCollection = InputData.ToDataReferenceCollection();
+
+			FBuildOperatorParams CreateParams{*Node, InOutContext.Settings, InputData, InOutContext.Environment, this};
 
 			FOperatorPtr Operator = Factory->CreateOperator(CreateParams, InOutContext.Results);
 
@@ -353,16 +404,29 @@ namespace Metasound
 				return FBuildStatus::FatalError;
 			}
 
-			// Save input and output collection for future use 
-			InOutContext.DataReferences.Emplace(Node, FOperatorDataReferences(Operator->GetInputs(), Operator->GetOutputs()));
+			// Bind vertex to operator data
+			FVertexInterfaceData BoundOperatorData(NodeInterface);
+			Operator->Bind(BoundOperatorData);
+			
+			// Check if outputs are bound correctly.
+			if (BuilderSettings.bValidateOperatorOutputsAreBound)
+			{
+				BuildStatus |= ValidateOperatorOutputsAreBound(*Node, BoundOperatorData.GetOutputs());
+			}
+
+			// Store bound interface for later use.
+			InOutContext.DataReferences.Emplace(Node, MoveTemp(BoundOperatorData));
 
 			// Add operator to operator array
 			InOutContext.Operators.Add(MoveTemp(Operator));
+
+			ThreadLocalDebug::ResetActiveNodeClass();
 		}
+
 		return BuildStatus;
 	}
 
-	FOperatorBuilder::FBuildStatus FOperatorBuilder::GatherGraphDataReferences(FBuildContext& InOutContext, FDataReferenceCollection& OutGraphInputs, FDataReferenceCollection& OutGraphOutputs) const
+	FOperatorBuilder::FBuildStatus FOperatorBuilder::GatherGraphDataReferences(FBuildContext& InOutContext, FVertexInterfaceData& OutVertexData) const
 	{
 		using namespace OperatorBuilderPrivate;
 		using FDestinationElement = FInputDataDestinationCollection::ElementType;
@@ -373,93 +437,52 @@ namespace Metasound
 		// Gather graph inputs
 		for (const FDestinationElement& Element : InOutContext.Graph.GetInputDataDestinations())
 		{
+			bool bFoundDataReference = false;
 			const FInputDataDestination& InputDestination = Element.Value;
 
-			if (!InOutContext.DataReferences.Contains(InputDestination.Node))
-			{
-				// An input node was likely pruned.
-				AddBuildError<FMissingInputDataReferenceError>(InOutContext.Results.Errors, InputDestination);
-
-				BuildStatus |= FBuildStatus::NonFatalError;
-				continue;
-			}
-
-			const FDataReferenceCollection& Collection = InOutContext.DataReferences[InputDestination.Node].Inputs;
-
-			// Get either readable or writable values for inputs. 
-			// Readable inputs can occur when the data reference is provided outside
-			// of the graph being built. 
-			const bool bContainsWritableInput = Collection.ContainsDataWriteReference(InputDestination.Vertex.VertexName, InputDestination.Vertex.DataTypeName);
-			const bool bContainsReadableInput = Collection.ContainsDataReadReference(InputDestination.Vertex.VertexName, InputDestination.Vertex.DataTypeName);
-
-			const bool bContainsInput = bContainsWritableInput || bContainsReadableInput;
-
-			if (!bContainsInput)
+			if (const FVertexInterfaceData* VertexData = InOutContext.DataReferences.Find(InputDestination.Node))
 			{
-				AddBuildError<FMissingInputDataReferenceError>(InOutContext.Results.Errors, InputDestination);
-
-				BuildStatus |= FBuildStatus::NonFatalError;
-				continue;
-			}
-
-			bool bSuccess = true;
-			if (bContainsWritableInput)
-			{
-				bSuccess = OutGraphInputs.AddDataWriteReferenceFrom(InputDestination.Vertex.VertexName, Collection, InputDestination.Vertex.VertexName, InputDestination.Vertex.DataTypeName);
-			}
-			else
-			{
-				bSuccess = OutGraphInputs.AddDataReadReferenceFrom(InputDestination.Vertex.VertexName, Collection, InputDestination.Vertex.VertexName, InputDestination.Vertex.DataTypeName);
+				const FInputVertexInterfaceData& NodeInputData = VertexData->GetInputs();
+				if (const FAnyDataReference* DataReference = NodeInputData.FindDataReference(InputDestination.Vertex.VertexName))
+				{
+					if (DataReference->GetDataTypeName() == InputDestination.Vertex.DataTypeName)
+					{
+						bFoundDataReference = true;
+						OutVertexData.GetInputs().BindVertex(InputDestination.Vertex.VertexName, *DataReference);
+					}
+				}
 			}
 
-			if (!bSuccess)
+			if (!bFoundDataReference)
 			{
 				AddBuildError<FMissingInputDataReferenceError>(InOutContext.Results.Errors, InputDestination);
-
 				BuildStatus |= FBuildStatus::NonFatalError;
-				continue;
 			}
 		}
 
 		// Gather graph outputs.
 		for (const FSourceElement& Element : InOutContext.Graph.GetOutputDataSources())
 		{
+			bool bFoundDataReference = false;
 			const FOutputDataSource& OutputSource = Element.Value;
 
-			if (!InOutContext.DataReferences.Contains(OutputSource.Node))
-			{
-				// An output node was likely pruned.
-				AddBuildError<FMissingOutputDataReferenceError>(InOutContext.Results.Errors, OutputSource);
-
-				BuildStatus |= FBuildStatus::NonFatalError;
-
-				continue;
-			}
-
-			const FDataReferenceCollection& Collection = InOutContext.DataReferences[OutputSource.Node].Outputs;
-
-			if (!Collection.ContainsDataReadReference(OutputSource.Vertex.VertexName, OutputSource.Vertex.DataTypeName))
+			if (const FVertexInterfaceData* VertexData = InOutContext.DataReferences.Find(OutputSource.Node))
 			{
-				// This will likely produce an IOperator which does not work as
-				// expected.
-				AddBuildError<FMissingOutputDataReferenceError>(InOutContext.Results.Errors, OutputSource);
-
-				BuildStatus |= FBuildStatus::NonFatalError;
-
-				continue;
+				const FOutputVertexInterfaceData& NodeOutputData = VertexData->GetOutputs();
+				if (const FAnyDataReference* DataReference = NodeOutputData.FindDataReference(OutputSource.Vertex.VertexName))
+				{
+					if (DataReference->GetDataTypeName() == OutputSource.Vertex.DataTypeName)
+					{
+						bFoundDataReference = true;
+						OutVertexData.GetOutputs().BindVertex(OutputSource.Vertex.VertexName, FAnyDataReference{*DataReference});
+					}
+				}
 			}
 
-			bool bSuccess = OutGraphOutputs.AddDataReadReferenceFrom(OutputSource.Vertex.VertexName, Collection, OutputSource.Vertex.VertexName, OutputSource.Vertex.DataTypeName);
-
-			if (!bSuccess)
+			if (!bFoundDataReference)
 			{
-				// This will likely produce an IOperator which does not work as
-				// expected.
 				AddBuildError<FMissingOutputDataReferenceError>(InOutContext.Results.Errors, OutputSource);
-
 				BuildStatus |= FBuildStatus::NonFatalError;
-
-				continue;
 			}
 		}
 	
@@ -468,32 +491,22 @@ namespace Metasound
 
 	TUniquePtr<IOperator> FOperatorBuilder::CreateGraphOperator(FBuildContext& InOutContext) const
 	{
-		FDataReferenceCollection GraphInputs;
-		FDataReferenceCollection GraphOutputs;
+		FVertexInterfaceData BoundGraphData(InOutContext.Graph.GetVertexInterface());
 
-		FBuildStatus BuildStatus = GatherGraphDataReferences(InOutContext, GraphInputs, GraphOutputs);
+		FBuildStatus BuildStatus = GatherGraphDataReferences(InOutContext, BoundGraphData);
 
-		if (BuildStatus > FBuildStatus::GetMaxErrorLevel(InOutContext.BuilderSettings))
+		if (BuildStatus > GetMaxErrorLevel())
 		{
 			return TUniquePtr<IOperator>(nullptr);
 		}
 
 		TUniquePtr<FGraphOperator> GraphOperator = MakeUnique<FGraphOperator>();
 
-		GraphOperator->SetInputs(GraphInputs);
-		GraphOperator->SetOutputs(GraphOutputs);
-
-		if (InOutContext.BuilderSettings.bPopulateInternalDataReferences)
+		GraphOperator->SetVertexInterfaceData(MoveTemp(BoundGraphData));
+		
+		if (BuilderSettings.bPopulateInternalDataReferences)
 		{
-			TMap<FGuid, FDataReferenceCollection> References;
-			for (TPair<const INode*, FOperatorDataReferences>& ReferencePair : InOutContext.DataReferences)
-			{
-				const INode* Node = ReferencePair.Key;
-				check(Node);
-				References.Emplace(Node->GetInstanceID(), MoveTemp(ReferencePair.Value.Outputs));
-			}
-
-			InOutContext.Results.InternalDataReferences.Append(MoveTemp(References));
+			GatherInternalGraphDataReferences(InOutContext, InOutContext.Results.InternalDataReferences);
 		}
 
 		for (FOperatorPtr& Ptr : InOutContext.Operators)
@@ -503,4 +516,9 @@ namespace Metasound
 
 		return MoveTemp(GraphOperator);
 	}
+
+	FOperatorBuilder::FBuildStatus::EStatus FOperatorBuilder::GetMaxErrorLevel() const
+	{
+		return BuilderSettings.bFailOnAnyError ? FBuildStatus::NoError : FBuildStatus::NonFatalError;
+	}
 }
diff --git a/Engine/Plugins/Runtime/Metasound/Source/MetasoundGraphCore/Private/MetasoundOperatorBuilderSettings.cpp b/Engine/Plugins/Runtime/Metasound/Source/MetasoundGraphCore/Private/MetasoundOperatorBuilderSettings.cpp
--- a/Engine/Plugins/Runtime/Metasound/Source/MetasoundGraphCore/Private/MetasoundOperatorBuilderSettings.cpp
+++ b/Engine/Plugins/Runtime/Metasound/Source/MetasoundGraphCore/Private/MetasoundOperatorBuilderSettings.cpp
@@ -32,6 +32,7 @@ namespace Metasound
 			Settings.bValidateNoDuplicateInputs = true;
 			Settings.bValidateVerticesExist = true;
 			Settings.bValidateEdgeDataTypesMatch = true;
+			Settings.bValidateOperatorOutputsAreBound = true;
 			Settings.bFailOnAnyError = false;
 			return Settings;
 		};
@@ -50,6 +51,7 @@ namespace Metasound
 			Settings.bValidateNoDuplicateInputs = true;
 			Settings.bValidateVerticesExist = true;
 			Settings.bValidateEdgeDataTypesMatch = true;
+			Settings.bValidateOperatorOutputsAreBound = true;
 			Settings.bFailOnAnyError = false;
 			return Settings;
 		};
@@ -68,6 +70,7 @@ namespace Metasound
 			Settings.bValidateNoDuplicateInputs = false;
 			Settings.bValidateVerticesExist = false;
 			Settings.bValidateEdgeDataTypesMatch = false;
+			Settings.bValidateOperatorOutputsAreBound = false;
 			Settings.bFailOnAnyError = false;
 			return Settings;
 		};
@@ -86,6 +89,7 @@ namespace Metasound
 			Settings.bValidateNoDuplicateInputs = false;
 			Settings.bValidateVerticesExist = false;
 			Settings.bValidateEdgeDataTypesMatch = false;
+			Settings.bValidateOperatorOutputsAreBound = false;
 			Settings.bFailOnAnyError = false;
 			return Settings;
 		};
@@ -93,4 +97,4 @@ namespace Metasound
 		static const FOperatorBuilderSettings ShippingSettings = CreateSettings();
 		return ShippingSettings;
 	}
-} // namespace Metasound
\ No newline at end of file
+} // namespace Metasound
diff --git a/Engine/Plugins/Runtime/Metasound/Source/MetasoundGraphCore/Private/MetasoundThreadLocalDebug.cpp b/Engine/Plugins/Runtime/Metasound/Source/MetasoundGraphCore/Private/MetasoundThreadLocalDebug.cpp
new file mode 100644
--- /dev/null
+++ b/Engine/Plugins/Runtime/Metasound/Source/MetasoundGraphCore/Private/MetasoundThreadLocalDebug.cpp
@@ -0,0 +1,44 @@
+// Copyright Epic Games, Inc. All Rights Reserved.
+
+#include "MetasoundThreadLocalDebug.h"
+
+#include "MetasoundNodeInterface.h"
+#include "Containers/UnrealString.h"
+
+#define ENABLE_METASOUND_THREAD_LOCAL_DEBUG (DO_CHECK || DO_ENSURE)
+
+namespace Metasound
+{
+	namespace ThreadLocalDebug
+	{
+		namespace ThreadLocalDebugPrivate
+		{
+#if ENABLE_METASOUND_THREAD_LOCAL_DEBUG
+			static thread_local FString ActiveNodeClassNameAndVersion;
+#endif
+		}
+
+		void SetActiveNodeClass(const FNodeClassMetadata& InMetadata)
+		{
+#if ENABLE_METASOUND_THREAD_LOCAL_DEBUG
+			ThreadLocalDebugPrivate::ActiveNodeClassNameAndVersion = FString::Format(TEXT("{0} v{1}.{2}"), {InMetadata.ClassName.GetFullName().ToString(), InMetadata.MajorVersion, InMetadata.MinorVersion});
+#endif
+		}
+
+		void ResetActiveNodeClass()
+		{
+#if ENABLE_METASOUND_THREAD_LOCAL_DEBUG
+			ThreadLocalDebugPrivate::ActiveNodeClassNameAndVersion = TEXT("");
+#endif
+		}
+
+		const TCHAR* GetActiveNodeClassNameAndVersion()
+		{
+#if ENABLE_METASOUND_THREAD_LOCAL_DEBUG
+			return *ThreadLocalDebugPrivate::ActiveNodeClassNameAndVersion;	
+#else
+			return TEXT("");
+#endif
+		}
+	}
+}
diff --git a/Engine/Plugins/Runtime/Metasound/Source/MetasoundGraphCore/Private/MetasoundThreadLocalDebug.h b/Engine/Plugins/Runtime/Metasound/Source/MetasoundGraphCore/Private/MetasoundThreadLocalDebug.h
new file mode 100644
--- /dev/null
+++ b/Engine/Plugins/Runtime/Metasound/Source/MetasoundGraphCore/Private/MetasoundThreadLocalDebug.h
@@ -0,0 +1,23 @@
+// Copyright Epic Games, Inc. All Rights Reserved.
+
+#pragma once
+
+#include "HAL/Platform.h"
+
+namespace Metasound
+{
+	struct FNodeClassMetadata;
+
+	// Debug util for setting stack scope debug variables
+	namespace ThreadLocalDebug
+	{
+		// Sets the current active node class for the current thread.
+		void SetActiveNodeClass(const FNodeClassMetadata& InMetadata);
+
+		// Resets the current active node class for the current thread.
+		void ResetActiveNodeClass();
+
+		// Returns the class name and version string for the active node in the current thread.
+		const TCHAR* GetActiveNodeClassNameAndVersion();
+	}
+}
diff --git a/Engine/Plugins/Runtime/Metasound/Source/MetasoundGraphCore/Private/MetasoundVertexData.cpp b/Engine/Plugins/Runtime/Metasound/Source/MetasoundGraphCore/Private/MetasoundVertexData.cpp
--- a/Engine/Plugins/Runtime/Metasound/Source/MetasoundGraphCore/Private/MetasoundVertexData.cpp
+++ b/Engine/Plugins/Runtime/Metasound/Source/MetasoundGraphCore/Private/MetasoundVertexData.cpp
@@ -2,10 +2,51 @@
 
 #include "MetasoundVertexData.h"
 
+#include "MetasoundDataReference.h"
+#include "MetasoundLog.h"
+#include "MetasoundThreadLocalDebug.h"
+
 namespace Metasound
 {
 	namespace MetasoundVertexDataPrivate
 	{
+#if ENABLE_METASOUND_ACCESS_TYPE_COMPATIBILITY_TEST
+		void CheckAccessTypeCompatibility(const FDataVertex& InDataVertex, const FAnyDataReference& InDataReference)
+		{
+			bool bIsCompatible = false;
+			const EDataReferenceAccessType ReferenceAccessType = InDataReference.GetAccessType();
+
+			switch (InDataVertex.AccessType)
+			{
+				case EDataReferenceAccessType::None:
+					bIsCompatible = true;
+					break;
+
+				case EDataReferenceAccessType::Read:
+					// Read vertices can be a read, write or value data reference type
+					bIsCompatible = EDataReferenceAccessType::None != ReferenceAccessType;
+					break;
+
+				case EDataReferenceAccessType::Write:
+					// Write vertices require a write data reference type
+					bIsCompatible = EDataReferenceAccessType::Write == ReferenceAccessType;
+					break;
+
+				case EDataReferenceAccessType::Value:
+					// value vertices require a value data reference type
+					bIsCompatible = EDataReferenceAccessType::Value == ReferenceAccessType;
+					break;
+
+				default:
+					{
+						checkNoEntry();
+					}
+			}
+
+			checkf(bIsCompatible, TEXT("Vertex access type \"%s\" is incompatible with data access type \"%s\" on vertex \"%s\" on node \"%s\""), *LexToString(InDataVertex.AccessType), *LexToString(ReferenceAccessType), *InDataVertex.VertexName.ToString(), ThreadLocalDebug::GetActiveNodeClassNameAndVersion());
+		}
+#endif // #if ENABLE_METASOUND_ACCESS_TYPE_COMPATIBILITY_TEST
+
 		template<typename VertexType>
 		void EmplaceBindings(TArray<TBinding<VertexType>>& InArray, const TVertexInterfaceGroup<VertexType>& InVertexInterface)
 		{
@@ -55,6 +96,30 @@ namespace Metasound
 			}
 		}
 
+		template<typename BindingType>
+		void BindVertex(TArray<BindingType>& InBindings, const FVertexName& InVertexName, const FAnyDataReference& InDataReference)
+		{
+			BindVertex<BindingType>(InBindings, InVertexName, FAnyDataReference{InDataReference});
+		}
+
+		template<typename BindingType>
+		void Bind(TArray<BindingType>& InTargetBindings, const TArray<BindingType>& InSourceBindings)
+		{
+			for (const BindingType& SourceBinding : InSourceBindings)
+			{
+				if (const FAnyDataReference* DataReference = SourceBinding.GetDataReference())
+				{
+					if (BindingType* TargetBinding = Find(InTargetBindings, SourceBinding.GetVertex().VertexName))
+					{
+						if (DataReference->GetDataTypeName() == TargetBinding->GetVertex().DataTypeName)
+						{
+							TargetBinding->Bind(*DataReference);
+						}
+					}
+				}
+			}
+		}
+
 		template<typename BindingType>
 		void Bind(TArray<BindingType>& InBindings, const FDataReferenceCollection& InCollection)
 		{
@@ -85,6 +150,7 @@ namespace Metasound
 			return Algo::AllOf(InBindings, [](const BindingType& Binding) { return Binding.IsBound(); });
 		}
 
+
 		template<typename BindingType>
 		EDataReferenceAccessType GetVertexDataAccessType(const TArray<BindingType>& InBindings, const FVertexName& InVertexName)
 		{
@@ -94,6 +160,30 @@ namespace Metasound
 			}
 			return EDataReferenceAccessType::None;
 		}
+
+		template<typename BindingType>
+		FDataReferenceCollection ToDataReferenceCollection(const TArray<BindingType>& InBindings)
+		{
+			FDataReferenceCollection Collection;
+			for (const BindingType& Binding : InBindings)
+			{
+				if (const FAnyDataReference* Ref = Binding.GetDataReference())
+				{
+					Collection.AddDataReference(Binding.GetVertex().VertexName, FAnyDataReference{*Ref});
+				}
+			}
+			return Collection;
+		}
+
+		template<typename BindingType>
+		const FAnyDataReference* FindDataReference(const TArray<BindingType>& InBindings, const FVertexName& InVertexName)
+		{
+			if (const BindingType* Binding = Find<BindingType>(InBindings, InVertexName))
+			{
+				return Binding->GetDataReference();
+			}
+			return nullptr;
+		}
 	}
 
 	FInputVertexInterfaceData::FInputVertexInterfaceData(const FInputVertexInterface& InVertexInterface)
@@ -107,16 +197,36 @@ namespace Metasound
 		MetasoundVertexDataPrivate::BindVertex(Bindings, InVertexName, MoveTemp(InDataReference));
 	}
 
+	void FInputVertexInterfaceData::BindVertex(const FVertexName& InVertexName, const FAnyDataReference& InDataReference)
+	{
+		MetasoundVertexDataPrivate::BindVertex(Bindings, InVertexName, InDataReference);
+	}
+
+	void FInputVertexInterfaceData::Bind(const FInputVertexInterfaceData& InVertexData)
+	{
+		MetasoundVertexDataPrivate::Bind(Bindings, InVertexData.Bindings);
+	}
+
 	void FInputVertexInterfaceData::Bind(const FDataReferenceCollection& InCollection)
 	{
 		MetasoundVertexDataPrivate::Bind(Bindings, InCollection);
 	}
 
+	FDataReferenceCollection FInputVertexInterfaceData::ToDataReferenceCollection() const
+	{
+		return MetasoundVertexDataPrivate::ToDataReferenceCollection(Bindings);
+	}
+
 	bool FInputVertexInterfaceData::IsVertexBound(const FVertexName& InVertexName) const
 	{
 		return MetasoundVertexDataPrivate::IsVertexBound(Bindings, InVertexName);
 	}
 
+	const FInputDataVertex& FInputVertexInterfaceData::GetVertex(const FVertexName& InVertexName) const
+	{
+		return FindChecked(InVertexName)->GetVertex();
+	}
+
 	EDataReferenceAccessType FInputVertexInterfaceData::GetVertexDataAccessType(const FVertexName& InVertexName) const
 	{
 		return MetasoundVertexDataPrivate::GetVertexDataAccessType(Bindings, InVertexName);
@@ -127,6 +237,11 @@ namespace Metasound
 		return MetasoundVertexDataPrivate::AreAllVerticesBound(Bindings);
 	}
 
+	const FAnyDataReference* FInputVertexInterfaceData::FindDataReference(const FVertexName& InVertexName) const
+	{
+		return MetasoundVertexDataPrivate::FindDataReference(Bindings, InVertexName);
+	}
+
 	FInputVertexInterfaceData::FBinding* FInputVertexInterfaceData::Find(const FVertexName& InVertexName)
 	{
 		return MetasoundVertexDataPrivate::Find(Bindings, InVertexName);
@@ -140,14 +255,14 @@ namespace Metasound
 	FInputVertexInterfaceData::FBinding* FInputVertexInterfaceData::FindChecked(const FVertexName& InVertexName)
 	{
 		FBinding* Binding = Find(InVertexName);
-		check(nullptr != Binding);
+		checkf(nullptr != Binding, TEXT("Attempt to access vertex \"%s\" which does not exist on interface."), *InVertexName.ToString());
 		return Binding;
 	}
 
 	const FInputVertexInterfaceData::FBinding* FInputVertexInterfaceData::FindChecked(const FVertexName& InVertexName) const
 	{
 		const FBinding* Binding = Find(InVertexName);
-		check(nullptr != Binding);
+		checkf(nullptr != Binding, TEXT("Attempt to access vertex \"%s\" which does not exist on interface."), *InVertexName.ToString());
 		return Binding;
 	}
 
@@ -163,16 +278,36 @@ namespace Metasound
 		MetasoundVertexDataPrivate::BindVertex(Bindings, InVertexName, MoveTemp(InDataReference));
 	}
 
+	void FOutputVertexInterfaceData::BindVertex(const FVertexName& InVertexName, const FAnyDataReference& InDataReference)
+	{
+		MetasoundVertexDataPrivate::BindVertex(Bindings, InVertexName, InDataReference);
+	}
+
+	void FOutputVertexInterfaceData::Bind(const FOutputVertexInterfaceData& InVertexData)
+	{
+		MetasoundVertexDataPrivate::Bind(Bindings, InVertexData.Bindings);
+	}
+
 	void FOutputVertexInterfaceData::Bind(const FDataReferenceCollection& InCollection)
 	{
 		MetasoundVertexDataPrivate::Bind(Bindings, InCollection);
 	}
 
+	FDataReferenceCollection FOutputVertexInterfaceData::ToDataReferenceCollection() const
+	{
+		return MetasoundVertexDataPrivate::ToDataReferenceCollection(Bindings);
+	}
+
 	bool FOutputVertexInterfaceData::IsVertexBound(const FVertexName& InVertexName) const
 	{
 		return MetasoundVertexDataPrivate::IsVertexBound(Bindings, InVertexName);
 	}
 
+	const FOutputDataVertex& FOutputVertexInterfaceData::GetVertex(const FVertexName& InVertexName) const
+	{
+		return FindChecked(InVertexName)->GetVertex();
+	}
+
 	EDataReferenceAccessType FOutputVertexInterfaceData::GetVertexDataAccessType(const FVertexName& InVertexName) const
 	{
 		return MetasoundVertexDataPrivate::GetVertexDataAccessType(Bindings, InVertexName);
@@ -183,6 +318,12 @@ namespace Metasound
 		return MetasoundVertexDataPrivate::AreAllVerticesBound(Bindings);
 	}
 
+	const FAnyDataReference* FOutputVertexInterfaceData::FindDataReference(const FVertexName& InVertexName) const
+	{
+		return MetasoundVertexDataPrivate::FindDataReference(Bindings, InVertexName);
+	}
+
+
 	FOutputVertexInterfaceData::FBinding* FOutputVertexInterfaceData::Find(const FVertexName& InVertexName)
 	{
 		return MetasoundVertexDataPrivate::Find(Bindings, InVertexName);
@@ -196,14 +337,20 @@ namespace Metasound
 	FOutputVertexInterfaceData::FBinding* FOutputVertexInterfaceData::FindChecked(const FVertexName& InVertexName)
 	{
 		FBinding* Binding = Find(InVertexName);
-		check(nullptr != Binding);
+		checkf(nullptr != Binding, TEXT("Attempt to access vertex \"%s\" which does not exist on interface."), *InVertexName.ToString());
 		return Binding;
 	}
 
 	const FOutputVertexInterfaceData::FBinding* FOutputVertexInterfaceData::FindChecked(const FVertexName& InVertexName) const
 	{
 		const FBinding* Binding = Find(InVertexName);
-		check(nullptr != Binding);
+		checkf(nullptr != Binding, TEXT("Attempt to access vertex \"%s\" which does not exist on interface."), *InVertexName.ToString());
 		return Binding;
 	}
+
+	void FVertexInterfaceData::Bind(const FVertexInterfaceData& InVertexData)
+	{
+		InputVertexInterfaceData.Bind(InVertexData.GetInputs());
+		OutputVertexInterfaceData.Bind(InVertexData.GetOutputs());
+	}
 }
diff --git a/Engine/Plugins/Runtime/Metasound/Source/MetasoundGraphCore/Public/MetasoundBuilderInterface.h b/Engine/Plugins/Runtime/Metasound/Source/MetasoundGraphCore/Public/MetasoundBuilderInterface.h
--- a/Engine/Plugins/Runtime/Metasound/Source/MetasoundGraphCore/Public/MetasoundBuilderInterface.h
+++ b/Engine/Plugins/Runtime/Metasound/Source/MetasoundGraphCore/Public/MetasoundBuilderInterface.h
@@ -6,7 +6,6 @@
 #include "MetasoundDataReference.h"
 #include "MetasoundEnvironment.h"
 #include "MetasoundNodeInterface.h"
-#include "MetasoundOperatorBuilderSettings.h"
 #include "MetasoundOperatorInterface.h"
 #include "Templates/UniquePtr.h"
 #include "UObject/NameTypes.h"
@@ -49,7 +48,7 @@ namespace Metasound
 	using FBuildErrorArray = TArray<TUniquePtr<IOperatorBuildError>>;
 
 	/** Structure of all resulting data generated during graph operator build. */
-	struct FBuildGraphResults
+	struct FBuildResults
 	{
 		/** An array of errors. Errors can be added if issues occur while creating an IOperator. */
 		FBuildErrorArray Errors;
@@ -75,13 +74,66 @@ namespace Metasound
 		/** Environment settings available. */
 		const FMetasoundEnvironment& Environment;
 
-		/** Settings for the builder requesting creation of the IOperator. */
-		const FOperatorBuilderSettings& BuilderSettings;
+		/** Pointer to builder actively building graph. */
+		const IOperatorBuilder* Builder = nullptr;
+	};
+
+	/** FBuildOperatorParams holds the parameters provided to operator factories
+	 * during the creation of an IOperator
+	 */
+	struct FBuildOperatorParams
+	{
+		/** The node associated with this factory and the desired IOperator. */
+		const INode& Node;
+
+		/** General operator settings for the graph. */
+		const FOperatorSettings& OperatorSettings;
+
+		/** Input data references for an IOperator */
+		const FInputVertexInterfaceData& InputData;
+
+		/** Environment settings available. */
+		const FMetasoundEnvironment& Environment;
 
 		/** Pointer to builder actively building graph. */
 		const IOperatorBuilder* Builder = nullptr;
 	};
 
+	/** Parameters for building an operator from a graph. 
+	  *
+	  * This object is slated to be deprecated in UE 5.2
+	  */
+	struct FBuildGraphParams
+	{
+		/** Reference to graph being built. */
+		const IGraph& Graph;
+
+		/** General operator settings for the graph. */
+		const FOperatorSettings& OperatorSettings;
+
+		/** Collection of input parameters available for an IOperator. */
+		const FDataReferenceCollection& InputDataReferences;
+
+		/** Environment settings available. */
+		const FMetasoundEnvironment& Environment;
+	};
+
+	/** Parameters for building an operator from a graph. */
+	struct FBuildGraphOperatorParams
+	{
+		/** Reference to graph being built. */
+		const IGraph& Graph;
+
+		/** General operator settings for the graph. */
+		const FOperatorSettings& OperatorSettings;
+
+		/** Bound input data available for an IOperator. */
+		const FInputVertexInterfaceData& InputData;
+
+		/** Environment settings available. */
+		const FMetasoundEnvironment& Environment;
+	};
+
 	/** Convenience template for adding build errors.
 	 *
 	 * The function can be used in the following way:
@@ -98,47 +150,41 @@ namespace Metasound
 		OutErrors.Add(MakeUnique<ErrorType>(Forward<ArgTypes>(Args)...));
 	}
 
+
 	/** IOperatorFactory
 	 *
 	 * IOperatorFactory defines an interface for building an IOperator from an INode.  In practice,
 	 * each INode returns its own IOperatorFactory through the INode::GetDefaultOperatorFactory() 
 	 * member function.
 	 */
-	class IOperatorFactory
+	class METASOUNDGRAPHCORE_API IOperatorFactory
 	{
 		public:
 			virtual ~IOperatorFactory() = default;
 
-			virtual TUniquePtr<IOperator> CreateOperator(const FCreateOperatorParams& InParams, FBuildGraphResults& OutResults) = 0;
+			virtual TUniquePtr<IOperator> CreateOperator(const FBuildOperatorParams& InParams, FBuildResults& OutResults) = 0;
 
-			UE_DEPRECATED(5.1, "Use CreateOperator overload providing the FBuildGraphResults struct.")
+			UE_DEPRECATED(5.1, "Use CreateOperator overload providing the FBuildResults struct.")
 			virtual TUniquePtr<IOperator> CreateOperator(const FCreateOperatorParams& InParams, FBuildErrorArray& OutErrors)
 			{
-				FBuildGraphResults Results { MoveTemp(OutErrors) };
-				TUniquePtr<IOperator> Operator = CreateOperator(InParams, Results);
-				OutErrors = MoveTemp(Results.Errors);
-				return Operator;
-			}
-	};
-
-	struct FBuildGraphParams
-	{
-		/** Reference to graph being built. */
-		const IGraph& Graph;
+				// Map inputs to new structures
+				FInputVertexInterfaceData InputData(InParams.Node.GetVertexInterface().GetInputInterface());
+				InputData.Bind(InParams.InputDataReferences);
 
-		/** General operator settings for the graph. */
-		const FOperatorSettings& OperatorSettings;
+				FBuildOperatorParams Params2 { InParams.Node, InParams.OperatorSettings, InputData, InParams.Environment, InParams.Builder };
 
-		/** Collection of input parameters available for an IOperator. */
-		const FDataReferenceCollection& InputDataReferences;
+				// Create operator
+				FBuildResults Results;
+				TUniquePtr<IOperator> Operator = CreateOperator(Params2, Results);
 
-		/** Environment settings available. */
-		const FMetasoundEnvironment& Environment;
+				// Map new output to old output
+				OutErrors = MoveTemp(Results.Errors);
 
-		/** Settings for graph's operator builder. */
-		const FOperatorBuilderSettings& BuilderSettings;
+				return MoveTemp(Operator);
+			}
 	};
 
+
 	/** IOperatorBuilder
 	 *
 	 * Defines an interface for building a graph of operators from a graph of nodes. 
@@ -151,8 +197,27 @@ namespace Metasound
 
 			virtual ~IOperatorBuilder() = default;
 
-			UE_DEPRECATED(5.1, "Use BuildGraphOperator overload providing the FBuildGraphResults struct.")
-			virtual TUniquePtr<IOperator> BuildGraphOperator(const FBuildGraphParams& InParams, FBuildErrorArray& OutErrors) const { return { }; }
+			/** Build a graph operator from a graph.
+			 *
+			 * @params InParams - Input parameters for building a graph.
+			 * @param OutErrors - An array of errors which occurred during while building the operator.
+			 *
+			 * @return A unique pointer to the built IOperator. Null if build failed.
+			 */
+			UE_DEPRECATED(5.1, "Use BuildGraphOperator overload providing the FBuildResults struct.")
+			virtual TUniquePtr<IOperator> BuildGraphOperator(const FBuildGraphParams& InParams, FBuildErrorArray& OutErrors) const 
+			{ 
+				FInputVertexInterfaceData InputData(InParams.Graph.GetVertexInterface().GetInputInterface());
+				InputData.Bind(InParams.InputDataReferences);
+				FBuildGraphOperatorParams NewParams{InParams.Graph, InParams.OperatorSettings, InputData, InParams.Environment};
+				FBuildResults Results;
+
+				TUniquePtr<IOperator> GraphOperator = BuildGraphOperator(NewParams, Results);
+
+				OutErrors.Append(MoveTemp(Results.Errors));
+				
+				return MoveTemp(GraphOperator);
+			}
 
 			/** Build a graph operator from a graph.
 			 *
@@ -161,6 +226,6 @@ namespace Metasound
 			 *
 			 * @return A unique pointer to the built IOperator. Null if build failed.
 			 */
-			virtual TUniquePtr<IOperator> BuildGraphOperator(const FBuildGraphParams& InParams, FBuildGraphResults& OutResults) const = 0;
+			virtual TUniquePtr<IOperator> BuildGraphOperator(const FBuildGraphOperatorParams& InParams, FBuildResults& OutResults) const = 0;
 	};
 }
diff --git a/Engine/Plugins/Runtime/Metasound/Source/MetasoundGraphCore/Public/MetasoundDataReference.h b/Engine/Plugins/Runtime/Metasound/Source/MetasoundGraphCore/Public/MetasoundDataReference.h
--- a/Engine/Plugins/Runtime/Metasound/Source/MetasoundGraphCore/Public/MetasoundDataReference.h
+++ b/Engine/Plugins/Runtime/Metasound/Source/MetasoundGraphCore/Public/MetasoundDataReference.h
@@ -18,6 +18,15 @@ namespace Audio
 	class IProxyData;
 }
 
+namespace Metasound
+{
+	// Forward declare for LexToString
+	enum class EDataReferenceAccessType : uint8;
+}
+
+/** Convert a EDataReferenceAccessType to FString. */
+FString METASOUNDGRAPHCORE_API LexToString(Metasound::EDataReferenceAccessType InAccessType);
+
 using FMetasoundDataTypeId = void const*;
 
 namespace Metasound
@@ -159,58 +168,124 @@ namespace Metasound
 		{
 		}
 
-		public:
-			typedef TSharedRef<DataType, ESPMode::NotThreadSafe> FRefType;
+		typedef TSharedRef<DataType, ESPMode::NotThreadSafe> FRefType;
 
-			typedef TDataReferenceTypeInfo<DataType> FInfoType;
+	public:
 
-			/** This should be used to construct a new DataType object and return this TDataReference as a wrapper around it.
-			 */
-			template <typename... ArgTypes>
-			static TDataReference<DataType> CreateNew(ArgTypes&&... Args)
-			{
-				static_assert(std::is_constructible<DataType, ArgTypes...>::value, "Tried to call TDataReference::CreateNew with args that don't match any constructor for an underlying type!");
-				return TDataReference<DataType>(EDataRefShouldConstruct::NewObject, Forward<ArgTypes>(Args)...);
-			}
+		typedef TDataReferenceTypeInfo<DataType> FInfoType;
 
-			/** Enable copy constructor */
-			TDataReference(const TDataReference<DataType>& Other) = default;
+		/** This should be used to construct a new DataType object and return this TDataReference as a wrapper around it.
+		 */
+		template <typename... ArgTypes>
+		static TDataReference<DataType> CreateNew(ArgTypes&&... Args)
+		{
+			static_assert(std::is_constructible<DataType, ArgTypes...>::value, "Tried to call TDataReference::CreateNew with args that don't match any constructor for an underlying type!");
+			return TDataReference<DataType>(EDataRefShouldConstruct::NewObject, Forward<ArgTypes>(Args)...);
+		}
 
-			/** Enable move constructor */
-			TDataReference(TDataReference<DataType>&& Other) = default;
+		/** Enable copy constructor */
+		TDataReference(const TDataReference<DataType>& Other) = default;
 
-			/** Enable copy operator */
-			TDataReference<DataType>& operator=(const TDataReference<DataType>& Other) = default;
+		/** Enable move constructor */
+		TDataReference(TDataReference<DataType>&& Other) = default;
 
-			/** Enable move operator */
-			TDataReference<DataType>& operator=(TDataReference<DataType>&& Other) = default;
+		/** Enable copy operator */
+		TDataReference<DataType>& operator=(const TDataReference<DataType>& Other) = default;
 
-			/** Return the name of the underlying type. */
-			virtual const FName& GetDataTypeName() const override
-			{
-				static const FName Name = GetMetasoundDataTypeName<DataType>();
+		/** Enable move operator */
+		TDataReference<DataType>& operator=(TDataReference<DataType>&& Other) = default;
 
-				return Name;
-			}
+		/** Return the name of the underlying type. */
+		virtual const FName& GetDataTypeName() const override
+		{
+			static const FName Name = GetMetasoundDataTypeName<DataType>();
+			return Name;
+		}
 
-			/** Return the ID of the underlying type. */
-			virtual const void* const GetDataTypeId() const override
-			{
-				return GetMetasoundDataTypeId<DataType>();
-			}
+		/** Return the ID of the underlying type. */
+		virtual const void* const GetDataTypeId() const override
+		{
+			return GetMetasoundDataTypeId<DataType>();
+		}
 
-		protected:
+	protected:
 
-			// Protected object reference is utilized by subclasses which define what
-			// access is provided to the ObjectReference. 
-			FRefType ObjectReference;
+		// Protected object reference is utilized by subclasses which define what
+		// access is provided to the ObjectReference. 
+		FRefType ObjectReference;
 	};
 
-
 	// Forward declare
 	template <typename DataType>
 	class TDataReadReference;
 
+	/** TDataValueReference represents a constant value and provides read only access. 
+	 * A TDataValueReference can never change value. */
+	template<typename DataType>
+	class TDataValueReference : public TDataReference<DataType>
+	{
+		// Construct operator with no arguments if the DataType has a default constructor.
+		template <typename... ArgTypes>
+		TDataValueReference(EDataRefShouldConstruct InToken, ArgTypes&&... Args)
+			: FDataReference(InToken, Forward<ArgTypes>(Args)...)
+		{
+		}
+
+	public:
+		typedef TDataReference<DataType> FDataReference;
+
+		/** This should be used to construct a new DataType object and return this TDataValueReference as a wrapper around it. */
+		template <typename... ArgTypes>
+		static TDataValueReference<DataType> CreateNew(ArgTypes&&... Args)
+		{
+			static_assert(std::is_constructible<DataType, ArgTypes...>::value, "TDataValueReference::CreateNew underlying type is not constructible with provided arguments.");
+			return TDataValueReference<DataType>(EDataRefShouldConstruct::NewObject, Forward<ArgTypes>(Args)...);
+		}
+
+		/** Enable copy constructor */
+		TDataValueReference(const TDataValueReference<DataType>& Other) = default;
+
+		/** Enable move constructor */
+		TDataValueReference(TDataValueReference<DataType>&& Other) = default;
+
+		/** Enable assignment operator. */
+		TDataValueReference<DataType>& operator=(const TDataValueReference<DataType>& Other) = default;
+
+		/** Enable move operator. */
+		TDataValueReference<DataType>& operator=(TDataValueReference<DataType>&& Other) = default;
+
+		/** Implicit conversion to a readable parameter. */
+		operator TDataReadReference<DataType>() const
+		{
+			return TDataReadReference<DataType>(*this);
+		}
+
+		/** Const access to the underlying parameter object. */
+		FORCEINLINE const DataType& operator*() const
+		{
+			return *TDataReference<DataType>::ObjectReference;
+		}
+
+		/** Const access to the underlying parameter object. */
+		FORCEINLINE const DataType* operator->() const
+		{
+			return TDataReference<DataType>::ObjectReference.operator->();
+		}
+
+		FORCEINLINE const DataType* Get() const
+		{
+			return TDataReference<DataType>::ObjectReference.operator->();
+		}
+
+		/** Create a clone of this parameter reference. */
+		virtual TUniquePtr<IDataReference> Clone() const override
+		{
+			typedef TDataValueReference<DataType> FValueDataReference;
+
+			return MakeUnique< FValueDataReference >(*this);
+		}
+	};
+
 	/** TDataWriteReference provides write access to a shared parameter reference. */
 	template <typename DataType>
 	class TDataWriteReference : public TDataReference<DataType>
@@ -263,6 +338,12 @@ namespace Metasound
 				return TDataReference<DataType>::ObjectReference.operator->();
 			}
 
+			/** Non-const access to the underlying parameter object. */
+			FORCEINLINE DataType* Get() const
+			{
+				return TDataReference<DataType>::ObjectReference.operator->();
+			}
+
 			/** Create a clone of this parameter reference. */
 			virtual TUniquePtr<IDataReference> Clone() const override
 			{
@@ -337,6 +418,12 @@ namespace Metasound
 		{
 		}
 
+		/** Construct a readable reference from a value reference. */
+		explicit TDataReadReference(const TDataValueReference<DataType>& ValueRef)
+		:	FDataReference(ValueRef)
+		{
+		}
+
 		/** Assign a readable parameter ref from a writable parameter ref. */
 		TDataReadReference<DataType>& operator=(const TDataWriteReference<DataType>& Other)
 		{
@@ -362,6 +449,12 @@ namespace Metasound
 			return TDataReference<DataType>::ObjectReference.operator->();
 		}
 
+		/** Non-const access to the underlying parameter object. */
+		FORCEINLINE const DataType* Get() const
+		{
+			return TDataReference<DataType>::ObjectReference.operator->();
+		}
+
 		/** Create a clone of this parameter reference. */
 		virtual TUniquePtr<IDataReference> Clone() const override
 		{
@@ -379,11 +472,14 @@ namespace Metasound
 	 */
 	enum class EDataReferenceAccessType : uint8
 	{
-		None, //< The data is inaccessible, or the data reference does not exist.
-		Read, //< The data is accessible through a TDataReadReference.
-		Write //< The data is accessible through a TDataWriteReference.
+		None = 0x00, 	//< The data is inaccessible, or the data reference does not exist.
+		Read = 0x01, 	//< The data is accessible through a TDataReadReference.
+		Write = 0x02,	//< The data is accessible through a TDataWriteReference.
+		Value = 0x04 	//< The data is accessible by value.
 	};
 
+
+
 	/** Container for any data reference. 
 	 *
 	 * This container maintains the underlying containers access type (Read or Write)
@@ -416,6 +512,13 @@ namespace Metasound
 		{
 		}
 
+		/** Construct with a TDataValueReference. */
+		template<typename DataType>
+		FAnyDataReference(const TDataValueReference<DataType>& InDataRef)
+		: FAnyDataReference(EDataReferenceAccessType::Value, InDataRef)
+		{
+		}
+
 		/** Copy construct with a FAnyDataReference. */
 		FAnyDataReference(const FAnyDataReference& InOther)
 		: FAnyDataReference(InOther.AccessType, *InOther.DataRefPtr)
@@ -461,12 +564,46 @@ namespace Metasound
 			return DataRefPtr->Clone();
 		}
 
+		/** Returns the current value of a reference. 
+		 *
+		 * This method's behavior is undefined and an assert will be called if
+		 * the DataType differs from the underlying data reference's DataType. 
+		 */
+		template<typename DataType>
+		const DataType* GetValue() const
+		{
+			if (DataRefPtr.IsValid())
+			{
+				check(IsDataReferenceOfType<DataType>(*DataRefPtr));
+				switch (AccessType)
+				{
+					case EDataReferenceAccessType::Read:
+					{
+						return static_cast<const TDataReadReference<DataType>*>(DataRefPtr.Get())->Get();
+					}
+					case EDataReferenceAccessType::Write:
+					{
+						return static_cast<const TDataWriteReference<DataType>*>(DataRefPtr.Get())->Get();
+					}
+					case EDataReferenceAccessType::Value:
+					{
+						return static_cast<const TDataValueReference<DataType>*>(DataRefPtr.Get())->Get();
+					}
+					default:
+					{
+						checkNoEntry();
+					}
+				}
+			}
+			return nullptr;
+		}
+
 		/** Get access to a TDataReadReference. 
 		 *
 		 * This method will return a valid TDataReadReference of the templated data
 		 * type. The returned object is only valid if:
-		 *     1. The template paramter DataType matches that of the underlying data reference.
-		 *     2. The underlying data reference is has either Read or Write access.
+		 *     1. The template parameter DataType matches that of the underlying data reference.
+		 *     2. The underlying data reference is has Read, Write or Value access.
 		 *
 		 * If this method's behavior is undefined and will assert if it is called 
 		 * with a mismatched data type or unsupported access type.
@@ -477,14 +614,25 @@ namespace Metasound
 			check(DataRefPtr.IsValid());
 			check(IsDataReferenceOfType<DataType>(*DataRefPtr));
 
-			if (EDataReferenceAccessType::Read == AccessType)
+			switch (AccessType)
 			{
-				return *static_cast<const TDataReadReference<DataType>*>(DataRefPtr.Get());
-			}
-			else 
-			{
-				check(EDataReferenceAccessType::Write == AccessType);
-				return *static_cast<const TDataWriteReference<DataType>*>(DataRefPtr.Get());
+				case EDataReferenceAccessType::Read:
+				{
+					return *static_cast<const TDataReadReference<DataType>*>(DataRefPtr.Get());
+				}
+				case EDataReferenceAccessType::Write:
+				{
+					return *static_cast<const TDataWriteReference<DataType>*>(DataRefPtr.Get());
+				}
+				case EDataReferenceAccessType::Value:
+				{
+					return *static_cast<const TDataValueReference<DataType>*>(DataRefPtr.Get());
+				}
+				default:
+				{
+					checkNoEntry();
+					return *static_cast<const TDataReadReference<DataType>*>(DataRefPtr.Get());
+				}
 			}
 		}
 
@@ -501,9 +649,12 @@ namespace Metasound
 		template<typename DataType>
 		TDataWriteReference<DataType> GetDataWriteReference() const
 		{
-			check(DataRefPtr.IsValid());
-			check(IsDataReferenceOfType<DataType>(*DataRefPtr));
-			check(EDataReferenceAccessType::Write == AccessType);
+			checkf(DataRefPtr.IsValid(), TEXT("FAnyDataReference is in an invalid state"));
+
+			checkf(IsDataReferenceOfType<DataType>(*DataRefPtr), TEXT("Attempt to get data reference with underlying type \"%s\" when actual underlying type is \"%s\"."), *GetMetasoundDataTypeString<DataType>(), *GetDataTypeName().ToString());
+
+			checkf(EDataReferenceAccessType::Write == AccessType, TEXT("Invalid attempt to convert a data ref with \"%s\" access to \"%s\" access"), *LexToString(GetAccessType()), *LexToString(EDataReferenceAccessType::Write));
+
 			return *static_cast<const TDataWriteReference<DataType>*>(DataRefPtr.Get());
 		}
 
diff --git a/Engine/Plugins/Runtime/Metasound/Source/MetasoundGraphCore/Public/MetasoundDataReferenceCollection.h b/Engine/Plugins/Runtime/Metasound/Source/MetasoundGraphCore/Public/MetasoundDataReferenceCollection.h
--- a/Engine/Plugins/Runtime/Metasound/Source/MetasoundGraphCore/Public/MetasoundDataReferenceCollection.h
+++ b/Engine/Plugins/Runtime/Metasound/Source/MetasoundGraphCore/Public/MetasoundDataReferenceCollection.h
@@ -57,6 +57,9 @@ namespace Metasound
 				AddDataReference(InName, MoveTemp(Wrapper));
 			}
 
+			/** Add a data reference to this collection. */
+			void AddDataReference(const FVertexName& InName, FAnyDataReference&& InDataReference);
+
 			/** Add a readable data reference from another collection.
 			 *
 			 * @param InLocalName - Name to give data added to this collection.
@@ -255,9 +258,6 @@ namespace Metasound
 
 		private:
 
-
-			void AddDataReference(const FVertexName& InName, FAnyDataReference&& InDataReference);
-
 			FDataReferenceMap DataRefMap;
 	};
 }
diff --git a/Engine/Plugins/Runtime/Metasound/Source/MetasoundGraphCore/Public/MetasoundFacade.h b/Engine/Plugins/Runtime/Metasound/Source/MetasoundGraphCore/Public/MetasoundFacade.h
--- a/Engine/Plugins/Runtime/Metasound/Source/MetasoundGraphCore/Public/MetasoundFacade.h
+++ b/Engine/Plugins/Runtime/Metasound/Source/MetasoundGraphCore/Public/MetasoundFacade.h
@@ -99,7 +99,7 @@ namespace Metasound
 				public:
 					FFactory(FCreateOperatorFunction InCreateFunc);
 
-					virtual TUniquePtr<IOperator> CreateOperator(const FCreateOperatorParams& InParams, FBuildGraphResults& OutResults) override;
+					virtual TUniquePtr<IOperator> CreateOperator(const FBuildOperatorParams& InParams, FBuildResults& OutResults) override;
 
 				private:
 					FCreateOperatorFunction CreateFunc;
diff --git a/Engine/Plugins/Runtime/Metasound/Source/MetasoundGraphCore/Public/MetasoundGraph.h b/Engine/Plugins/Runtime/Metasound/Source/MetasoundGraphCore/Public/MetasoundGraph.h
--- a/Engine/Plugins/Runtime/Metasound/Source/MetasoundGraphCore/Public/MetasoundGraph.h
+++ b/Engine/Plugins/Runtime/Metasound/Source/MetasoundGraphCore/Public/MetasoundGraph.h
@@ -133,7 +133,7 @@ namespace Metasound
 			public:
 				virtual ~FFactory() = default;
 
-				virtual TUniquePtr<IOperator> CreateOperator(const FCreateOperatorParams& InParams, FBuildGraphResults& OutResults) override;
+				virtual TUniquePtr<IOperator> CreateOperator(const FBuildOperatorParams& InParams, FBuildResults& OutResults) override;
 			};
 
 			FVertexName InstanceName;
diff --git a/Engine/Plugins/Runtime/Metasound/Source/MetasoundGraphCore/Public/MetasoundGraphOperator.h b/Engine/Plugins/Runtime/Metasound/Source/MetasoundGraphCore/Public/MetasoundGraphOperator.h
--- a/Engine/Plugins/Runtime/Metasound/Source/MetasoundGraphCore/Public/MetasoundGraphOperator.h
+++ b/Engine/Plugins/Runtime/Metasound/Source/MetasoundGraphCore/Public/MetasoundGraphOperator.h
@@ -20,16 +20,26 @@ namespace Metasound
 
 			virtual ~FGraphOperator();
 
+			// Add an operator to the end of the executation stack.
 			void AppendOperator(FOperatorPtr InOperator);
 
+			UE_DEPRECATED(5.1, "Use FGraphOperator::SetVertexInterfaceData instead.")
 			void SetInputs(const FDataReferenceCollection& InCollection);
 
+			UE_DEPRECATED(5.1, "Use FGraphOperator::SetVertexInterfaceData instead.")
 			void SetOutputs(const FDataReferenceCollection& InCollection);
 
+			// Set the vertex interface data. This data will be copied to output 
+			// during calls to Bind(InOutVertexData).
+			void SetVertexInterfaceData(FVertexInterfaceData&& InVertexData);
+
 			virtual FDataReferenceCollection GetInputs() const override;
 
 			virtual FDataReferenceCollection GetOutputs() const override;
 
+			// Bind the graph's interface data references to FVertexInterfaceData.
+			virtual void Bind(FVertexInterfaceData& InOutVertexData) const override;
+
 			virtual FExecuteFunction GetExecuteFunction() override;
 
 			void Execute();
@@ -38,13 +48,10 @@ namespace Metasound
 			// Delete copy operator because underlying types cannot be copied. 
 			FGraphOperator& operator=(const FGraphOperator&) = delete;
 			FGraphOperator(const FGraphOperator&) = delete;
-			
 
 			static void ExecuteFunction(IOperator* InOperator);
 
 			TArray<FExecuter> OperatorStack;
-
-			FDataReferenceCollection Inputs;
-			FDataReferenceCollection Outputs;
+			FVertexInterfaceData VertexData;
 	};
 }
diff --git a/Engine/Plugins/Runtime/Metasound/Source/MetasoundGraphCore/Public/MetasoundLiteralNode.h b/Engine/Plugins/Runtime/Metasound/Source/MetasoundGraphCore/Public/MetasoundLiteralNode.h
--- a/Engine/Plugins/Runtime/Metasound/Source/MetasoundGraphCore/Public/MetasoundLiteralNode.h
+++ b/Engine/Plugins/Runtime/Metasound/Source/MetasoundGraphCore/Public/MetasoundLiteralNode.h
@@ -90,7 +90,7 @@ namespace Metasound
 		{
 		}
 
-		virtual TUniquePtr<IOperator> CreateOperator(const FCreateOperatorParams& InParams, FBuildGraphResults& OutResults) override;
+		virtual TUniquePtr<IOperator> CreateOperator(const FBuildOperatorParams& InParams, FBuildResults& OutResults) override;
 
 	private:
 		FLiteral InitParam;
@@ -171,7 +171,7 @@ namespace Metasound
 	};
 
 	template<typename DataType>
-	TUniquePtr<IOperator> TLiteralOperatorLiteralFactory<DataType>::CreateOperator(const FCreateOperatorParams& InParams, FBuildGraphResults& OutResults)
+	TUniquePtr<IOperator> TLiteralOperatorLiteralFactory<DataType>::CreateOperator(const FBuildOperatorParams& InParams, FBuildResults& OutResults)
 	{
 		using FLiteralNodeType = TLiteralNode<DataType>;
 
diff --git a/Engine/Plugins/Runtime/Metasound/Source/MetasoundGraphCore/Public/MetasoundOperatorBuilder.h b/Engine/Plugins/Runtime/Metasound/Source/MetasoundGraphCore/Public/MetasoundOperatorBuilder.h
--- a/Engine/Plugins/Runtime/Metasound/Source/MetasoundGraphCore/Public/MetasoundOperatorBuilder.h
+++ b/Engine/Plugins/Runtime/Metasound/Source/MetasoundGraphCore/Public/MetasoundOperatorBuilder.h
@@ -2,17 +2,23 @@
 
 #pragma once
 
+#include "Containers/Array.h"
+#include "Containers/Map.h"
 #include "MetasoundBuilderInterface.h"
 #include "MetasoundEnvironment.h"
 #include "MetasoundGraphOperator.h"
 #include "MetasoundNodeInterface.h"
 #include "MetasoundOperatorBuilderSettings.h"
 #include "MetasoundOperatorInterface.h"
+#include "MetasoundVertexData.h"
+#include "Templates/UniquePtr.h"
 
 namespace Metasound
 {
 	// Forward declare.
 	class FDirectedGraphAlgoAdapter;
+	class FOperatorSettings;
+	class FMetasoundEnvironment;
 
 	/** FOperatorBuilder builds an IOperator from an IGraph. */
 	class METASOUNDGRAPHCORE_API FOperatorBuilder : public IOperatorBuilder
@@ -23,7 +29,7 @@ namespace Metasound
 			 *
 			 * @param InBuilderSettings  - Settings to configure builder options.
 			 */
-			FOperatorBuilder() = default;
+			FOperatorBuilder(const FOperatorBuilderSettings& InBuilderSettings);
 
 			virtual ~FOperatorBuilder();
 
@@ -35,30 +41,13 @@ namespace Metasound
 			 * @return A TUniquePtr to an IOperator. If the processes was unsuccessful, 
 			 *         the returned pointer will contain a nullptr and be invalid.
 			 */
-			virtual TUniquePtr<IOperator> BuildGraphOperator(const FBuildGraphParams& InParams, FBuildGraphResults& OutResults) const override;
+			virtual TUniquePtr<IOperator> BuildGraphOperator(const FBuildGraphOperatorParams& InParams, FBuildResults& OutResults) const override;
 
 		private:
-			// Collection of existing inputs and outputs associated with a given
-			// IOperator.
-			struct FOperatorDataReferences
-			{
-				FOperatorDataReferences()
-				{
-				}
-
-				FOperatorDataReferences(const FDataReferenceCollection& InInputs, const FDataReferenceCollection& InOutputs)
-				:	Inputs(InInputs)
-				,	Outputs(InOutputs)
-				{
-				}
-
-				FDataReferenceCollection Inputs;
-				FDataReferenceCollection Outputs;
-			};
 
 			using FNodeEdgeMultiMap = TMultiMap<const INode*, const FDataEdge*>;
 			using FNodeDestinationMap = TMap<const INode*, const FInputDataDestination*>;
-			using FNodeDataReferenceMap = TMap<const INode*, FOperatorDataReferences>;
+			using FNodeVertexInterfaceDataMap = TMap<const INode*, FVertexInterfaceData>;
 			using FOperatorPtr = TUniquePtr<IOperator>;
 
 			// Handles build status of current build operation.
@@ -83,10 +72,6 @@ namespace Metasound
 					FatalError = 2
 				};
 
-				static FBuildStatus GetMaxErrorLevel(const FOperatorBuilderSettings& InBuilderSettings)
-				{
-					return InBuilderSettings.bFailOnAnyError ? EStatus::NoError : EStatus::NonFatalError;
-				}
 
 				FBuildStatus() = default;
 
@@ -119,10 +104,10 @@ namespace Metasound
 				const FMetasoundEnvironment& Environment;
 				const FOperatorBuilderSettings& BuilderSettings;
 
-				FBuildGraphResults& Results;
+				FBuildResults& Results;
 
 				TArray<FOperatorPtr> Operators;
-				FNodeDataReferenceMap DataReferences;
+				FNodeVertexInterfaceDataMap DataReferences;
 
 				FBuildContext(
 					const IGraph& InGraph,
@@ -130,7 +115,7 @@ namespace Metasound
 					const FOperatorSettings& InSettings,
 					const FMetasoundEnvironment& InEnvironment,
 					const FOperatorBuilderSettings& bInBuilderSettings,
-					FBuildGraphResults& OutResults)
+					FBuildResults& OutResults)
 				:	Graph(InGraph)
 				,	AlgoAdapter(InAlgoAdapter)
 				,	Settings(InSettings)
@@ -151,19 +136,29 @@ namespace Metasound
 			FBuildStatus PruneNodes(FBuildContext& InOutContext, TArray<const INode*>& InOutNodes) const;
 
 			// Get all input data references for a given node for inputs provided internally to the graph.
-			FBuildStatus GatherInputDataReferences(FBuildContext& InOutContext, const INode* InNode, const FNodeEdgeMultiMap& InEdgeMap, FDataReferenceCollection& OutInputCollection) const;
+			FBuildStatus GatherInputDataReferences(FBuildContext& InOutContext, const INode* InNode, const FNodeEdgeMultiMap& InEdgeMap, FInputVertexInterfaceData& OutVertexData) const;
 
 			// Get all input data references for a given node for inputs provided externally to the graph.
-			FBuildStatus GatherExternalInputDataReferences(FBuildContext& InOutContext, const INode* InNode, const FNodeDestinationMap& InNodeDestinationMap, const FDataReferenceCollection& InExternalCollection, FDataReferenceCollection& OutDataReferences) const;
+			FBuildStatus GatherExternalInputDataReferences(FBuildContext& InContext, const INode* InNode, const FNodeDestinationMap& InNodeDestinationMap, const FInputVertexInterfaceData& InExternalCollection, FInputVertexInterfaceData& OutVertexData) const;
+
+			// Graphs all internal graph references and places them in output map.
+			void GatherInternalGraphDataReferences(FBuildContext& InOutContext, TMap<FGuid, FDataReferenceCollection>& OutNodeVertexData) const;
+
+			// Validates whether all operator outputs are bound to data references. 
+			FBuildStatus ValidateOperatorOutputsAreBound(const INode& InNode, const FOutputVertexInterfaceData& InVertexData) const;
 
 			// Get all input/output data references for a given graph.
-			FBuildStatus GatherGraphDataReferences(FBuildContext& InOutContext, FDataReferenceCollection& OutGraphInputs, FDataReferenceCollection& OutGraphOutputs) const;
+			FBuildStatus GatherGraphDataReferences(FBuildContext& InOutContext, FVertexInterfaceData& OutVertexData) const;
 
 			// Call the operator factories for the nodes
-			FBuildStatus CreateOperators(FBuildContext& InOutContext, const TArray<const INode*>& InSortedNodes, const FDataReferenceCollection& InGraphInputs) const;
+			FBuildStatus CreateOperators(FBuildContext& InOutContext, const TArray<const INode*>& InSortedNodes, const FInputVertexInterfaceData& InExternalInputData) const;
 
 			// Create the final graph operator from the provided build context.
 			TUniquePtr<IOperator> CreateGraphOperator(FBuildContext& InOutContext) const;
+
+			FBuildStatus::EStatus GetMaxErrorLevel() const;
+
+			FOperatorBuilderSettings BuilderSettings;
 	};
 }
 
diff --git a/Engine/Plugins/Runtime/Metasound/Source/MetasoundGraphCore/Public/MetasoundOperatorBuilderSettings.h b/Engine/Plugins/Runtime/Metasound/Source/MetasoundGraphCore/Public/MetasoundOperatorBuilderSettings.h
--- a/Engine/Plugins/Runtime/Metasound/Source/MetasoundGraphCore/Public/MetasoundOperatorBuilderSettings.h
+++ b/Engine/Plugins/Runtime/Metasound/Source/MetasoundGraphCore/Public/MetasoundOperatorBuilderSettings.h
@@ -61,6 +61,12 @@ namespace Metasound
 		 */
 		bool bValidateEdgeDataTypesMatch = true;
 
+		/** If true, each IOperator in the graph will be validated by checking
+		 * that each output of the FVertexInterface is bound to data. Errors will
+		 * be generated outputs are not bound.
+		 */
+		bool bValidateOperatorOutputsAreBound = true;
+
 		/** If true, the builder will return an invalid IOperator if any errors
 		 * are detected. If false, the builder will return an invalid IOperator
 		 * only if fatal errors are detected.
@@ -86,4 +92,4 @@ namespace Metasound
 		/** Return the default settings for a shipping build environment. */
 		static const FOperatorBuilderSettings& GetDefaultShippingSettings();
 	};
-} // namespace Metasound
\ No newline at end of file
+} // namespace Metasound
diff --git a/Engine/Plugins/Runtime/Metasound/Source/MetasoundGraphCore/Public/MetasoundVariableNodes.h b/Engine/Plugins/Runtime/Metasound/Source/MetasoundGraphCore/Public/MetasoundVariableNodes.h
--- a/Engine/Plugins/Runtime/Metasound/Source/MetasoundGraphCore/Public/MetasoundVariableNodes.h
+++ b/Engine/Plugins/Runtime/Metasound/Source/MetasoundGraphCore/Public/MetasoundVariableNodes.h
@@ -173,7 +173,7 @@ namespace Metasound
 
 			virtual ~FFactory() = default;
 
-			virtual TUniquePtr<IOperator> CreateOperator(const FCreateOperatorParams& InParams, FBuildGraphResults& OutResults) override
+			virtual TUniquePtr<IOperator> CreateOperator(const FBuildOperatorParams& InParams, FBuildResults& OutResults) override
 			{
 				TDataWriteReference<DataType> Data = TDataWriteReferenceLiteralFactory<DataType>::CreateExplicitArgs(InParams.OperatorSettings, Literal);
 				return MakeUnique<FOperator>(TDataWriteReference<FVariable>::CreateNew(Data));
diff --git a/Engine/Plugins/Runtime/Metasound/Source/MetasoundGraphCore/Public/MetasoundVertex.h b/Engine/Plugins/Runtime/Metasound/Source/MetasoundGraphCore/Public/MetasoundVertex.h
--- a/Engine/Plugins/Runtime/Metasound/Source/MetasoundGraphCore/Public/MetasoundVertex.h
+++ b/Engine/Plugins/Runtime/Metasound/Source/MetasoundGraphCore/Public/MetasoundVertex.h
@@ -23,6 +23,7 @@ namespace Metasound
 		bool bIsAdvancedDisplay = false;
 	};
 
+
 	/** FDataVertex
 	 *
 	 * An FDataVertex is a named vertex of a MetaSound node which can contain data.
@@ -39,13 +40,15 @@ namespace Metasound
 		 * @InVertexName - Name of vertex.
 		 * @InDataTypeName - Name of data type.
 		 * @InMetadata - Metadata pertaining to given vertex.
+		 * @InAccessType - The access type of the vertex.
 		 */
-		FDataVertex(const FVertexName& InVertexName, const FName& InDataTypeName, const FDataVertexMetadata& InMetadata)
-		:	VertexName(InVertexName)
-		,	DataTypeName(InDataTypeName)
+		FDataVertex(const FVertexName& InVertexName, const FName& InDataTypeName, const FDataVertexMetadata& InMetadata, EDataReferenceAccessType InAccessType)
+		: VertexName(InVertexName)
+		, DataTypeName(InDataTypeName)
 #if WITH_EDITORONLY_DATA
-		,	Metadata(InMetadata)
+		, Metadata(InMetadata)
 #endif // WITH_EDITORONLY_DATA
+		, AccessType(InAccessType)
 		{
 		}
 
@@ -63,6 +66,9 @@ namespace Metasound
 		FDataVertexMetadata Metadata;
 #endif // WITH_EDITORONLY_DATA
 
+		/** Access type of the vertex. */
+		EDataReferenceAccessType AccessType;
+
 		UE_DEPRECATED(5.1, "GetVertexName() is deprecated. Use VertexName")
 		const FVertexName& GetVertexName() const
 		{
@@ -174,16 +180,25 @@ namespace Metasound
 		FInputDataVertex() = default;
 
 		/** Construct an FInputDataVertex. */
-		FInputDataVertex(const FVertexName& InVertexName, const FName& InDataTypeName, const FDataVertexMetadata& InMetadata)
-			: FDataVertex(InVertexName, InDataTypeName, InMetadata)
+		FInputDataVertex(const FVertexName& InVertexName, const FName& InDataTypeName, const FDataVertexMetadata& InMetadata, EDataReferenceAccessType InAccessType=EDataReferenceAccessType::Read)
+			: FDataVertex(InVertexName, InDataTypeName, InMetadata, InAccessType)
 			, LiteralFactory(FLiteral::FNone{})
 		{
 		}
 
 		/** Construct an FInputDataVertex with a default literal. */
 		template<typename LiteralValueType>
+		UE_DEPRECATED(5.1, "Use constructor with EDataReferenceAccessType argument.")
 		FInputDataVertex(const FVertexName& InVertexName, const FName& InDataTypeName, const FDataVertexMetadata& InMetadata, const LiteralValueType& InLiteralValue)
-			: FDataVertex(InVertexName, InDataTypeName, InMetadata)
+			: FDataVertex(InVertexName, InDataTypeName, InMetadata, EDataReferenceAccessType::Read)
+			, LiteralFactory(InLiteralValue)
+		{
+		}
+
+		/** Construct an FInputDataVertex with a default literal. */
+		template<typename LiteralValueType>
+		FInputDataVertex(const FVertexName& InVertexName, const FName& InDataTypeName, const FDataVertexMetadata& InMetadata, EDataReferenceAccessType InAccessType, const LiteralValueType& InLiteralValue)
+			: FDataVertex(InVertexName, InDataTypeName, InMetadata, InAccessType)
 			, LiteralFactory(InLiteralValue)
 		{
 		}
@@ -212,20 +227,35 @@ namespace Metasound
 
 		template<typename... ArgTypes>
 		TInputDataVertex(const FVertexName& InVertexName, const FDataVertexMetadata& InMetadata, ArgTypes&&... Args)
-		: FInputDataVertex(InVertexName, GetMetasoundDataTypeName<DataType>(), InMetadata, Forward<ArgTypes>(Args)...)
+		: FInputDataVertex(InVertexName, GetMetasoundDataTypeName<DataType>(), InMetadata, EDataReferenceAccessType::Read, Forward<ArgTypes>(Args)...)
 		{
 		}
 
 		template<typename... ArgTypes>
  		UE_DEPRECATED(5.1, "TInputDataVertex constructor without explicit FDataVertexMatadata is deprecated. Use constructor which accepts TDataVertexMetadata")
 		TInputDataVertex(const FVertexName& InVertexName, const FText& InDescription, ArgTypes&&... Args)
-		: FInputDataVertex(InVertexName, GetMetasoundDataTypeName<DataType>(), { InDescription }, Forward<ArgTypes>(Args)...)
+		: FInputDataVertex(InVertexName, GetMetasoundDataTypeName<DataType>(), { InDescription }, EDataReferenceAccessType::Read, Forward<ArgTypes>(Args)...)
 		{
 		}
 	};
 
 	template<typename DataType>
 	using TInputDataVertexModel UE_DEPRECATED(5.1, "TInputDataVertexModel<DataType> is deprecated. Replace with TInputDataVertex<DataType>") = TInputDataVertex<DataType>;
+
+	/** Create a FInputDataVertex with a templated MetaSound data type which only
+	 * reads data at operator time. */
+	template<typename DataType>
+	class TInputConstructorVertex : public FInputDataVertex
+	{
+	public:
+		TInputConstructorVertex() = default;
+
+		template<typename... ArgTypes>
+		TInputConstructorVertex(const FVertexName& InVertexName, const FDataVertexMetadata& InMetadata, ArgTypes&&... Args)
+		: FInputDataVertex(InVertexName, GetMetasoundDataTypeName<DataType>(), InMetadata, EDataReferenceAccessType::Value, Forward<ArgTypes>(Args)...)
+		{
+		}
+	};
 	
 	/** FOutputDataVertex
 	 *
@@ -251,14 +281,14 @@ namespace Metasound
 
 		template<typename... ArgTypes>
 		TOutputDataVertex(const FVertexName& InVertexName, const FDataVertexMetadata& InMetadata, ArgTypes&&... Args)
-		: FOutputDataVertex(InVertexName, GetMetasoundDataTypeName<DataType>(), InMetadata, Forward<ArgTypes>(Args)...)
+		: FOutputDataVertex(InVertexName, GetMetasoundDataTypeName<DataType>(), InMetadata, EDataReferenceAccessType::Read, Forward<ArgTypes>(Args)...)
 		{
 		}
 
 		template<typename... ArgTypes>
  		UE_DEPRECATED(5.1, "TOutputDataVertex constructor without explicit FDataVertexMatadata is deprecated. Use constructor which accepts TDataVertexMetadata")
 		TOutputDataVertex(const FVertexName& InVertexName, const FText& InDescription, ArgTypes&&... Args)
-		: FOutputDataVertex(InVertexName, GetMetasoundDataTypeName<DataType>(), { InDescription }, Forward<ArgTypes>(Args)...)
+		: FOutputDataVertex(InVertexName, GetMetasoundDataTypeName<DataType>(), { InDescription }, EDataReferenceAccessType::Read, Forward<ArgTypes>(Args)...)
 		{
 		}
 	};
@@ -266,6 +296,22 @@ namespace Metasound
 	template<typename DataType>
 	using TOutputDataVertexModel UE_DEPRECATED(5.1, "TOutputDataVertexModel<DataType> is deprecated. Replace with TOutputDataVertex<DataType>") = TOutputDataVertex<DataType>;
 
+	/** Create a FOutputDataVertex with a templated MetaSound data type which is only
+	 * sets data at operator construction time. 
+	 */
+	template<typename DataType>
+	class TOutputConstructorVertex : public FOutputDataVertex
+	{
+	public:
+
+		TOutputConstructorVertex() = default;
+
+		template<typename... ArgTypes>
+		TOutputConstructorVertex(const FVertexName& InVertexName, const FDataVertexMetadata& InMetadata, ArgTypes&&... Args)
+		: FOutputDataVertex(InVertexName, GetMetasoundDataTypeName<DataType>(), InMetadata, EDataReferenceAccessType::Value, Forward<ArgTypes>(Args)...)
+		{
+		}
+	};
 
 	/** FEnvironmentVertex
 	 *
@@ -324,172 +370,210 @@ namespace Metasound
 	template<typename VertexType>
 	class TVertexInterfaceGroup
 	{
-			static constexpr bool bIsDerivedFromFInputDataVertex = std::is_base_of<FInputDataVertex, VertexType>::value;
-			static constexpr bool bIsDerivedFromFOutputDataVertex = std::is_base_of<FOutputDataVertex, VertexType>::value;
-			static constexpr bool bIsDerivedFromFEnvironmentVertex = std::is_base_of<FEnvironmentVertex, VertexType>::value;
+		// Determine if the vertex type is a supported vertex type.
+		template<typename T>
+		struct TIsSupportedVertexType
+		{
+		private:
+			static constexpr bool bIsDerivedFromFInputDataVertex = std::is_base_of<FInputDataVertex, T>::value;
+			static constexpr bool bIsDerivedFromFOutputDataVertex = std::is_base_of<FOutputDataVertex, T>::value;
+			static constexpr bool bIsDerivedFromFEnvironmentVertex = std::is_base_of<FEnvironmentVertex, T>::value;
 			static constexpr bool bIsSupportedVertexType = bIsDerivedFromFInputDataVertex || bIsDerivedFromFOutputDataVertex || bIsDerivedFromFEnvironmentVertex;
+		public:
+			static constexpr bool Value = bIsSupportedVertexType;
+		};
 
-			static_assert(bIsSupportedVertexType, "VertexType must be derived from FInputDataVertex, FOutputDataVertex, or FEnvironmentVertex");
-
-			using FContainerType = TArray<VertexType>;
-
-			struct FEqualVertexNamePredicate
-			{
-				const FVertexName& NameRef;
+		static_assert(TIsSupportedVertexType<VertexType>::Value, "VertexType must be derived from FInputDataVertex, FOutputDataVertex, or FEnvironmentVertex");
 
-				FEqualVertexNamePredicate(const FVertexName& InName)
-				: NameRef(InName)
-				{
-				}
 
-				FEqualVertexNamePredicate(const VertexType& InVertex)
-				: NameRef(InVertex.VertexName)
-				{
-				}
+		// Helper struct for determining whether the vertices in a parameter pack
+		// can be used in the TVertexInterfaceGroup constructor
+		template<typename ...>
+		struct TIsConstructibleFromParameterPack
+		{
+			static constexpr bool Value = false;
+		};
 
-				bool operator()(const VertexType& InOther)
-				{
-					return InOther.VertexName == NameRef;
-				}
-			};
+		// Helper struct specialization for no template argument. 
+		template<>
+		struct TIsConstructibleFromParameterPack<>
+		{
+			static constexpr bool Value = true;
+		};
 
-			void AddOrUpdateVertex(VertexType&& InVertex)
-			{
-				if (VertexType* Vertex = Find(InVertex.VertexName))
-				{
-					*Vertex = MoveTemp(InVertex);
-				}
-				else
-				{
-					Vertices.Add(MoveTemp(InVertex));
-				}
-			}
 
-			// Required for end of recursion.
-			void CopyInputs()
-			{
-			}
+	
 
-			// Recursive call used to unpack a template parameter pack. This stores
-			// each object in the template parameter pack into the groups internal
-			// storage.
-			//
-			// Assume that InputType is either a sublcass of FDataVertex
-			// and that VertexType is either FInputDataVertex or FOutputDataVertex
-			template<typename CurrentVertexType, typename... RemainingVertexTypes>
-			void CopyInputs(CurrentVertexType&& InInput, RemainingVertexTypes&&... InRemainingInputs)
-			{
-				// Create vertex out of vertex model
-				AddOrUpdateVertex(MoveTemp(InInput));
-				CopyInputs(InRemainingInputs...);
-			}
+		// Helper struct specialization for unwinding a parameter pack.
+		template<typename T, typename ... Ts>
+		struct TIsConstructibleFromParameterPack<T, Ts...>
+		{
+			static constexpr bool Value = std::is_constructible<VertexType, T>::value && TIsConstructibleFromParameterPack<Ts...>::Value;
+		};
 
 
-		public:
+		using FContainerType = TArray<VertexType>;
 
-			using RangedForConstIteratorType = typename FContainerType::RangedForConstIteratorType;
+		struct FEqualVertexNamePredicate
+		{
+			const FVertexName& NameRef;
 
-			/** TVertexInterfaceGroup constructor with variadic list of vertex
-			 * models.
-			 */
-			template<typename... VertexTypes>
-			TVertexInterfaceGroup(VertexTypes&&... InVertexs)
+			FEqualVertexNamePredicate(const FVertexName& InName)
+			: NameRef(InName)
 			{
-				// Reserve array to hold exact number of vertices to avoid
-				// over allocation.
-				Vertices.Reserve(sizeof...(VertexTypes));
-
-				CopyInputs(Forward<VertexTypes>(InVertexs)...);
 			}
 
-			/** Add a vertex to the group. */
-			void Add(const VertexType& InVertex)
+			FEqualVertexNamePredicate(const VertexType& InVertex)
+			: NameRef(InVertex.VertexName)
 			{
-				AddOrUpdateVertex(VertexType(InVertex));
 			}
 
-			void Add(VertexType&& InVertex)
+			bool operator()(const VertexType& InOther)
 			{
-				AddOrUpdateVertex(MoveTemp(InVertex));
+				return InOther.VertexName == NameRef;
 			}
+		};
 
-			/** Remove a vertex by key. */
-			bool Remove(const FVertexName& InKey)
+		void AddOrUpdateVertex(VertexType&& InVertex)
+		{
+			if (VertexType* Vertex = Find(InVertex.VertexName))
 			{
-				int32 NumRemoved = Vertices.RemoveAll(FEqualVertexNamePredicate(InKey));
-				return (NumRemoved > 0);
+				*Vertex = MoveTemp(InVertex);
 			}
-
-			/** Returns true if the group contains a vertex with a matching key. */
-			bool Contains(const FVertexName& InKey) const
+			else
 			{
-				return Vertices.ContainsByPredicate(FEqualVertexNamePredicate(InKey));
+				Vertices.Add(MoveTemp(InVertex));
 			}
+		}
 
-			/** Find a vertex with a given VertexName */
-			VertexType* Find(const FVertexName& InKey)
-			{
-				return Vertices.FindByPredicate(FEqualVertexNamePredicate(InKey));
-			}
+		// Required for end of recursion.
+		void CopyInputs()
+		{
+		}
 
-			/** Find a vertex with a given VertexName */
-			const VertexType* Find(const FVertexName& InKey) const
-			{
-				return Vertices.FindByPredicate(FEqualVertexNamePredicate(InKey));
-			}
+		// Recursive call used to unpack a template parameter pack. This stores
+		// each object in the template parameter pack into the groups internal
+		// storage.
+		//
+		// Assume that InputType is either a sublcass of FDataVertex
+		// and that VertexType is either FInputDataVertex or FOutputDataVertex
+		template<typename CurrentVertexType, typename... RemainingVertexTypes>
+		void CopyInputs(CurrentVertexType&& InInput, RemainingVertexTypes&&... InRemainingInputs)
+		{
+			// Create vertex out of vertex model
+			AddOrUpdateVertex(MoveTemp(InInput));
+			CopyInputs(InRemainingInputs...);
+		}
 
-			/** Return the sort order index of a vertex with the given name.
-			 *
-			 * @param InKey - FVertexName of vertex of interest.
-			 *
-			 * @return The index of the vertex. INDEX_NONE if the vertex does not exist. 
-			 */
-			int32 GetSortOrderIndex(const FVertexName& InKey) const
-			{
-				return Vertices.IndexOfByPredicate(FEqualVertexNamePredicate(InKey));
-			}
 
-			/** Return the vertex for a given vertex key. */
-			const VertexType& operator[](const FVertexName& InName) const
-			{
-				const VertexType* Vertex = Find(InName);
-				checkf(nullptr != Vertex, TEXT("Vertex with name '%s' does not exist"), *InName.ToString());
-				return *Vertex;
-			}
+	public:
 
-			/** Iterator for ranged for loops. */
-			RangedForConstIteratorType begin() const
-			{
-				return Vertices.begin();
-			}
+		using RangedForConstIteratorType = typename FContainerType::RangedForConstIteratorType;
 
-			/** Iterator for ranged for loops. */
-			RangedForConstIteratorType end() const
-			{
-				return Vertices.end();
-			}
+		/** TVertexInterfaceGroup constructor with variadic list of vertex
+		 * models.
+		 */
+		template<
+			typename... VertexTypes, 
+			typename std::enable_if< TIsConstructibleFromParameterPack<VertexTypes...>::Value, int >::type = 0
+		>
+		TVertexInterfaceGroup(VertexTypes&&... InVertexs)
+		{
+			// Reserve array to hold exact number of vertices to avoid
+			// over allocation.
+			Vertices.Reserve(sizeof...(VertexTypes));
 
-			/** Returns the number of vertices in the group. */
-			int32 Num() const
-			{
-				return Vertices.Num();
-			}
+			CopyInputs(Forward<VertexTypes>(InVertexs)...);
+		}
 
-			/** Compare whether two vertex groups are equal. */
-			friend bool operator==(const TVertexInterfaceGroup<VertexType>& InLHS, const TVertexInterfaceGroup<VertexType>& InRHS)
-			{
-				return InLHS.Vertices == InRHS.Vertices;
-			}
+		/** Add a vertex to the group. */
+		void Add(const VertexType& InVertex)
+		{
+			AddOrUpdateVertex(VertexType(InVertex));
+		}
 
-			/** Compare whether two vertex groups are unequal. */
-			friend bool operator!=(const TVertexInterfaceGroup<VertexType>& InLHS, const TVertexInterfaceGroup<VertexType>& InRHS)
-			{
-				return !(InLHS == InRHS);
-			}
+		void Add(VertexType&& InVertex)
+		{
+			AddOrUpdateVertex(MoveTemp(InVertex));
+		}
 
-		private:
+		/** Remove a vertex by key. */
+		bool Remove(const FVertexName& InKey)
+		{
+			int32 NumRemoved = Vertices.RemoveAll(FEqualVertexNamePredicate(InKey));
+			return (NumRemoved > 0);
+		}
+
+		/** Returns true if the group contains a vertex with a matching key. */
+		bool Contains(const FVertexName& InKey) const
+		{
+			return Vertices.ContainsByPredicate(FEqualVertexNamePredicate(InKey));
+		}
+
+		/** Find a vertex with a given VertexName */
+		VertexType* Find(const FVertexName& InKey)
+		{
+			return Vertices.FindByPredicate(FEqualVertexNamePredicate(InKey));
+		}
+
+		/** Find a vertex with a given VertexName */
+		const VertexType* Find(const FVertexName& InKey) const
+		{
+			return Vertices.FindByPredicate(FEqualVertexNamePredicate(InKey));
+		}
+
+		/** Return the sort order index of a vertex with the given name.
+		 *
+		 * @param InKey - FVertexName of vertex of interest.
+		 *
+		 * @return The index of the vertex. INDEX_NONE if the vertex does not exist. 
+		 */
+		int32 GetSortOrderIndex(const FVertexName& InKey) const
+		{
+			return Vertices.IndexOfByPredicate(FEqualVertexNamePredicate(InKey));
+		}
+
+		/** Return the vertex for a given vertex key. */
+		const VertexType& operator[](const FVertexName& InName) const
+		{
+			const VertexType* Vertex = Find(InName);
+			checkf(nullptr != Vertex, TEXT("Vertex with name '%s' does not exist"), *InName.ToString());
+			return *Vertex;
+		}
+
+		/** Iterator for ranged for loops. */
+		RangedForConstIteratorType begin() const
+		{
+			return Vertices.begin();
+		}
+
+		/** Iterator for ranged for loops. */
+		RangedForConstIteratorType end() const
+		{
+			return Vertices.end();
+		}
+
+		/** Returns the number of vertices in the group. */
+		int32 Num() const
+		{
+			return Vertices.Num();
+		}
+
+		/** Compare whether two vertex groups are equal. */
+		friend bool operator==(const TVertexInterfaceGroup<VertexType>& InLHS, const TVertexInterfaceGroup<VertexType>& InRHS)
+		{
+			return InLHS.Vertices == InRHS.Vertices;
+		}
+
+		/** Compare whether two vertex groups are unequal. */
+		friend bool operator!=(const TVertexInterfaceGroup<VertexType>& InLHS, const TVertexInterfaceGroup<VertexType>& InRHS)
+		{
+			return !(InLHS == InRHS);
+		}
+
+	private:
 
-			FContainerType Vertices;
+		FContainerType Vertices;
 	};
 
 	
diff --git a/Engine/Plugins/Runtime/Metasound/Source/MetasoundGraphCore/Public/MetasoundVertexData.h b/Engine/Plugins/Runtime/Metasound/Source/MetasoundGraphCore/Public/MetasoundVertexData.h
--- a/Engine/Plugins/Runtime/Metasound/Source/MetasoundGraphCore/Public/MetasoundVertexData.h
+++ b/Engine/Plugins/Runtime/Metasound/Source/MetasoundGraphCore/Public/MetasoundVertexData.h
@@ -13,6 +13,10 @@
 
 #include <type_traits>
 
+/** Enable runtime tests for compatible access types between vertex access types
+ * and data references bound to the vertex. */
+#define ENABLE_METASOUND_ACCESS_TYPE_COMPATIBILITY_TEST DO_CHECK
+
 namespace Metasound
 {
 	namespace MetasoundVertexDataPrivate
@@ -20,6 +24,14 @@ namespace Metasound
 		template<typename V>
 		using EnableIfIsInputDataVertex = std::enable_if_t<std::is_same_v<FInputDataVertex, V>>;
 
+		// Tests to see that the access type of the data reference is compatible 
+		// with the access type of the vertex.
+#if ENABLE_METASOUND_ACCESS_TYPE_COMPATIBILITY_TEST
+		void CheckAccessTypeCompatibility(const FDataVertex& InDataVertex, const FAnyDataReference& InDataReference);
+#else
+		void CheckAccessTypeCompatibility(const FDataVertex& InDataVertex, const FAnyDataReference& InDataReference) {}
+#endif // #if DO_CHECK
+
 		// Binds a vertex to a data reference.
 		template<typename VertexType>
 		class TBinding
@@ -41,12 +53,14 @@ namespace Metasound
 			{
 				check(Vertex.DataTypeName == InDataReference.GetDataTypeName());
 				Data.Emplace(InDataReference);
+				CheckAccessTypeCompatibility(Vertex, *Data);
 			}
 
 			void Bind(FAnyDataReference&& InAnyDataReference)
 			{
 				check(Vertex.DataTypeName == InAnyDataReference.GetDataTypeName());
 				Data.Emplace(MoveTemp(InAnyDataReference));
+				CheckAccessTypeCompatibility(Vertex, *Data);
 			}
 			
 			const VertexType& GetVertex() const
@@ -68,6 +82,12 @@ namespace Metasound
 				return EDataReferenceAccessType::None;
 			}
 
+			// Get data reference 
+			const FAnyDataReference* GetDataReference() const
+			{
+				return Data.GetPtrOrNull();
+			}
+
 			// Get data read reference assuming data is bound and read or write accessible.
 			template<typename DataType>
 			TDataReadReference<DataType> GetDataReadReference() const
@@ -148,6 +168,26 @@ namespace Metasound
 				return TDataWriteReference<DataType>::CreateNew(Forward<ConstructorArgTypes>(ConstructorArgs)...);
 			}
 
+			// Gets the value of the bound data reference if it exists. Returns
+			// nullptr if not bound.
+			template<typename DataType>
+			const DataType* GetValue() const
+			{
+				if (Data.IsSet())
+				{
+					return Data->GetValue<DataType>();
+				}
+				return nullptr;
+			}
+
+			// Set the value with a constant value reference.
+			template<typename DataType>
+			void SetValue(const DataType& InValue)
+			{
+				Data.Emplace(TDataValueReference<DataType>::CreateNew(InValue));
+				CheckAccessTypeCompatibility(Vertex, *Data);
+			}
+
 		private:
 
 			VertexType Vertex;
@@ -165,12 +205,42 @@ namespace Metasound
 		using FRangedForIteratorType = typename TArray<FBinding>::RangedForIteratorType;
 		using FRangedForConstIteratorType = typename TArray<FBinding>::RangedForConstIteratorType;
 
+		FInputVertexInterfaceData() = default;
+
 		/** Construct with an FInputVertexInterface. */
 		FInputVertexInterfaceData(const FInputVertexInterface& InVertexInterface);
 
-		/** Bind a data reference to a vertex. */
-		template<typename DataReferenceType>
-		void BindVertex(const FVertexName& InVertexName, const DataReferenceType& InDataReference)
+		/** Bind a value vertex from a value reference. */
+		template<typename DataType>
+		void BindValueVertex(const FVertexName& InVertexName, const TDataValueReference<DataType>& InDataReference)
+		{
+			BindVertex(InVertexName, FAnyDataReference{InDataReference});
+		}
+
+		/** Bind a read vertex from a value reference. */
+		template<typename DataType>
+		void BindReadVertex(const FVertexName& InVertexName, const TDataValueReference<DataType>& InDataReference)
+		{
+			BindVertex(InVertexName, FAnyDataReference{static_cast<TDataReadReference<DataType>>(InDataReference)});
+		}
+
+		/** Bind a read vertex from a read reference. */
+		template<typename DataType>
+		void BindReadVertex(const FVertexName& InVertexName, const TDataReadReference<DataType>& InDataReference)
+		{
+			BindVertex(InVertexName, FAnyDataReference{InDataReference});
+		}
+
+		/** Bind a read vertex from a write reference. */
+		template<typename DataType>
+		void BindReadVertex(const FVertexName& InVertexName, const TDataWriteReference<DataType>& InDataReference)
+		{
+			BindVertex(InVertexName, FAnyDataReference{static_cast<TDataReadReference<DataType>>(InDataReference)});
+		}
+
+		/** Bind a write vertex from a write reference. */
+		template<typename DataType>
+		void BindWriteVertex(const FVertexName& InVertexName, const TDataWriteReference<DataType>& InDataReference)
 		{
 			BindVertex(InVertexName, FAnyDataReference{InDataReference});
 		}
@@ -178,9 +248,21 @@ namespace Metasound
 		/** Bind a data reference to a vertex. */
 		void BindVertex(const FVertexName& InVertexName, FAnyDataReference&& InDataReference);
 
+		/** Bind a data reference to a vertex. */
+		void BindVertex(const FVertexName& InVertexName, const FAnyDataReference& InDataReference);
+
+		/** Bind vertex data using other vertex data. */
+		void Bind(const FInputVertexInterfaceData& InInputVertexData);
+
 		/** Bind a data references to vertices with matching vertex names. */
 		void Bind(const FDataReferenceCollection& InCollection);
 
+		/** Convert vertex data to a data reference collection. */
+		FDataReferenceCollection ToDataReferenceCollection() const;
+
+		/** Return the vertex associated with the vertex name. */
+		const FInputDataVertex& GetVertex(const FVertexName& InVertexName) const;
+
 		/** Returns true if a vertex with the given vertex name exists and is bound
 		 * to a data reference. */
 		bool IsVertexBound(const FVertexName& InVertexName) const;
@@ -214,6 +296,24 @@ namespace Metasound
 			return Bindings.end();
 		}
 
+		/** Find data reference bound to vertex. Returns a nullptr if no data reference is bound. */
+		const FAnyDataReference* FindDataReference(const FVertexName& InVertexName) const;
+
+		/** Returns the current value of a vertex. */
+		template<typename DataType>
+		const DataType* GetValue(const FVertexName& InVertexName) const
+		{
+			const FBinding* Binding = FindChecked(InVertexName);
+			return Binding->GetValue<DataType>();
+		}
+
+		/** Set the value of a vertex. */
+		template<typename DataType>
+		void SetValue(const FVertexName& InVertexName, const DataType& InValue)
+		{
+			FBinding* Binding = FindChecked(InVertexName);
+			return Binding->SetValue<DataType>(InValue);
+		}
 
 		/** Get data read reference assuming data is bound and read or write accessible. */
 		template<typename DataType> 
@@ -239,7 +339,7 @@ namespace Metasound
 				UE_LOG(LogMetaSound, Warning, TEXT("Failed find vertex with name '%s'. Cannot check for existing bound data."), *InVertexName.ToString());
 			}
 
-			return TDataReadReferenceLiteralFactory<DataType>::CreateAny(InSettings, FLiteral::CreateInvalid());
+			return TDataReadReferenceLiteralFactory<DataType>::CreateExplicitArgs(InSettings, FLiteral::CreateInvalid());
 		}
 
 		/** Get the bound data read reference if it exists. Otherwise create and 
@@ -329,12 +429,42 @@ namespace Metasound
 		using FRangedForIteratorType = typename TArray<FBinding>::RangedForIteratorType;
 		using FRangedForConstIteratorType = typename TArray<FBinding>::RangedForConstIteratorType;
 
+		FOutputVertexInterfaceData() = default;
+
 		/* Construct using an FOutputVertexInterface. */
 		FOutputVertexInterfaceData(const FOutputVertexInterface& InVertexInterface);
 
-		/** Bind a data reference to a vertex. */
-		template<typename DataReferenceType>
-		void BindVertex(const FVertexName& InVertexName, const DataReferenceType& InDataReference)
+		/** Bind a value vertex from a value reference. */
+		template<typename DataType>
+		void BindValueVertex(const FVertexName& InVertexName, const TDataValueReference<DataType>& InDataReference)
+		{
+			BindVertex(InVertexName, FAnyDataReference{InDataReference});
+		}
+
+		/** Bind a read vertex from a value reference. */
+		template<typename DataType>
+		void BindReadVertex(const FVertexName& InVertexName, const TDataValueReference<DataType>& InDataReference)
+		{
+			BindVertex(InVertexName, FAnyDataReference{static_cast<TDataReadReference<DataType>>(InDataReference)});
+		}
+
+		/** Bind a read vertex from a read reference. */
+		template<typename DataType>
+		void BindReadVertex(const FVertexName& InVertexName, const TDataReadReference<DataType>& InDataReference)
+		{
+			BindVertex(InVertexName, FAnyDataReference{InDataReference});
+		}
+
+		/** Bind a read vertex from a write reference. */
+		template<typename DataType>
+		void BindReadVertex(const FVertexName& InVertexName, const TDataWriteReference<DataType>& InDataReference)
+		{
+			BindVertex(InVertexName, FAnyDataReference{static_cast<TDataReadReference<DataType>>(InDataReference)});
+		}
+
+		/** Bind a write vertex from a write reference. */
+		template<typename DataType>
+		void BindWriteVertex(const FVertexName& InVertexName, const TDataWriteReference<DataType>& InDataReference)
 		{
 			BindVertex(InVertexName, FAnyDataReference{InDataReference});
 		}
@@ -342,13 +472,25 @@ namespace Metasound
 		/** Bind a data reference to a vertex. */
 		void BindVertex(const FVertexName& InVertexName, FAnyDataReference&& InDataReference);
 
+		/** Bind a data reference to a vertex. */
+		void BindVertex(const FVertexName& InVertexName, const FAnyDataReference& InDataReference);
+
+		/** Bind vertex data using other vertex data. */
+		void Bind(const FOutputVertexInterfaceData& InOutputVertexData);
+
 		/** Bind a data references to vertices with matching vertex names. */
 		void Bind(const FDataReferenceCollection& InCollection);
 
+		/** Converts the vertex data to a data reference collection. */
+		FDataReferenceCollection ToDataReferenceCollection() const;
+
 		/** Returns true if a vertex with the given vertex name exists and is bound
 		 * to a data reference. */
 		bool IsVertexBound(const FVertexName& InVertexName) const;
 
+		/** Return the vertex associated with the vertex name. */
+		const FOutputDataVertex& GetVertex(const FVertexName& InVertexName) const;
+
 		/** Returns the access type of a bound vertex. If the vertex does not exist
 		 * or if it is unbound, this will return EDataReferenceAccessType::None
 		 */
@@ -378,6 +520,24 @@ namespace Metasound
 			return Bindings.end();
 		}
 
+		/** Find data reference bound to vertex. Returns a nullptr if no data reference is bound. */
+		const FAnyDataReference* FindDataReference(const FVertexName& InVertexName) const;
+
+		/** Returns the current value of a vertex. */
+		template<typename DataType>
+		const DataType* GetValue(const FVertexName& InVertexName) const
+		{
+			const FBinding* Binding = FindChecked(InVertexName);
+			return Binding->GetValue<DataType>();
+		}
+
+		/** Set the value of a vertex. */
+		template<typename DataType>
+		void SetValue(const FVertexName& InVertexName, const DataType& InValue)
+		{
+			FBinding* Binding = FindChecked(InVertexName);
+			return Binding->SetValue<DataType>(InValue);
+		}
 
 		/** Get data read reference assuming data is bound and read or write accessible. */
 		template<typename DataType> 
@@ -451,6 +611,8 @@ namespace Metasound
 	{
 	public:
 
+		FVertexInterfaceData() = default;
+
 		/** Construct using an FVertexInterface. */
 		FVertexInterfaceData(const FVertexInterface& InVertexInterface)
 		: InputVertexInterfaceData(InVertexInterface.GetInputInterface())
@@ -458,6 +620,9 @@ namespace Metasound
 		{
 		}
 
+		/** Bind vertex data using other vertex data. */
+		void Bind(const FVertexInterfaceData& InVertexData);
+
 		/** Get input vertex interface data. */
 		FInputVertexInterfaceData& GetInputs()
 		{
diff --git a/Engine/Plugins/Runtime/Metasound/Source/MetasoundStandardNodes/Private/MetasoundAudioFormats.cpp b/Engine/Plugins/Runtime/Metasound/Source/MetasoundStandardNodes/Private/MetasoundAudioFormats.cpp
--- a/Engine/Plugins/Runtime/Metasound/Source/MetasoundStandardNodes/Private/MetasoundAudioFormats.cpp
+++ b/Engine/Plugins/Runtime/Metasound/Source/MetasoundStandardNodes/Private/MetasoundAudioFormats.cpp
@@ -153,13 +153,13 @@ namespace Metasound
 			{
 			}
 
-			TUniquePtr<IOperator> CreateOperator(const FCreateOperatorParams& InParams, FBuildGraphResults& OutResults) override
+			TUniquePtr<IOperator> CreateOperator(const FBuildOperatorParams& InParams, FBuildResults& OutResults) override
 			{
 				using namespace StereoAudioFormatVertexKeys;
 
 				// Construct stereo from left and right audio buffers. 
-				TDataReadReference<FAudioBuffer> Left = InParams.InputDataReferences.GetDataReadReferenceOrConstruct<FAudioBuffer>(METASOUND_GET_PARAM_NAME(LeftChannelVertex), InParams.OperatorSettings);
-				TDataReadReference<FAudioBuffer> Right = InParams.InputDataReferences.GetDataReadReferenceOrConstruct<FAudioBuffer>(METASOUND_GET_PARAM_NAME(RightChannelVertex), InParams.OperatorSettings);
+				TDataReadReference<FAudioBuffer> Left = InParams.InputData.GetOrConstructDataReadReference<FAudioBuffer>(METASOUND_GET_PARAM_NAME(LeftChannelVertex), InParams.OperatorSettings);
+				TDataReadReference<FAudioBuffer> Right = InParams.InputData.GetOrConstructDataReadReference<FAudioBuffer>(METASOUND_GET_PARAM_NAME(RightChannelVertex), InParams.OperatorSettings);
 
 				TDataReadReference<FStereoAudioFormat> Stereo = TDataReadReferenceFactory<FStereoAudioFormat>::CreateExplicitArgs(InParams.OperatorSettings, WriteCast(Left), WriteCast(Right));
 
@@ -298,10 +298,10 @@ namespace Metasound
 			{
 			}
 
-			TUniquePtr<IOperator> CreateOperator(const FCreateOperatorParams& InParams, FBuildGraphResults& OutResults) override
+			TUniquePtr<IOperator> CreateOperator(const FBuildOperatorParams& InParams, FBuildResults& OutResults) override
 			{
 				// Split a stereo signal into left/right.
-				TDataReadReference<FStereoAudioFormat> Stereo = InParams.InputDataReferences.GetDataReadReferenceOrConstruct<FStereoAudioFormat>(InputName, InParams.OperatorSettings);
+				TDataReadReference<FStereoAudioFormat> Stereo = InParams.InputData.GetOrConstructDataReadReference<FStereoAudioFormat>(InputName, InParams.OperatorSettings);
 				TDataReadReference<FAudioBuffer> Left = Stereo->GetLeft();
 				TDataReadReference<FAudioBuffer> Right = Stereo->GetRight();
 
@@ -436,9 +436,9 @@ namespace Metasound
 			{
 			}
 
-			TUniquePtr<IOperator> CreateOperator(const FCreateOperatorParams& InParams, FBuildGraphResults& OutResults) override
+			TUniquePtr<IOperator> CreateOperator(const FBuildOperatorParams& InParams, FBuildResults& OutResults) override
 			{
-				TDataReadReference<FAudioBuffer> Center = InParams.InputDataReferences.GetDataReadReferenceOrConstruct<FAudioBuffer>(OutputName, InParams.OperatorSettings);
+				TDataReadReference<FAudioBuffer> Center = InParams.InputData.GetOrConstructDataReadReference<FAudioBuffer>(OutputName, InParams.OperatorSettings);
 
 				TDataReadReference<FMonoAudioFormat> Mono = TDataReadReferenceFactory<FMonoAudioFormat>::CreateExplicitArgs(InParams.OperatorSettings, WriteCast(Center));
 
@@ -570,9 +570,9 @@ namespace Metasound
 			{
 			}
 
-			TUniquePtr<IOperator> CreateOperator(const FCreateOperatorParams& InParams, FBuildGraphResults& OutResults) override
+			TUniquePtr<IOperator> CreateOperator(const FBuildOperatorParams& InParams, FBuildResults& OutResults) override
 			{
-				TDataReadReference<FMonoAudioFormat> Mono = InParams.InputDataReferences.GetDataReadReferenceOrConstruct<FMonoAudioFormat>(InputName, InParams.OperatorSettings);
+				TDataReadReference<FMonoAudioFormat> Mono = InParams.InputData.GetOrConstructDataReadReference<FMonoAudioFormat>(InputName, InParams.OperatorSettings);
 				TDataReadReference<FAudioBuffer> Center = Mono->GetCenter();
 
 				return MakeUnique<FInputOperator>(InputName, Center, Mono);
diff --git a/Engine/Plugins/Runtime/Metasound/Source/MetasoundStandardNodes/Private/MetasoundOscillatorNodes.cpp b/Engine/Plugins/Runtime/Metasound/Source/MetasoundStandardNodes/Private/MetasoundOscillatorNodes.cpp
--- a/Engine/Plugins/Runtime/Metasound/Source/MetasoundStandardNodes/Private/MetasoundOscillatorNodes.cpp
+++ b/Engine/Plugins/Runtime/Metasound/Source/MetasoundStandardNodes/Private/MetasoundOscillatorNodes.cpp
@@ -815,10 +815,10 @@ namespace Metasound
 			return Interface;
 		}
 		
-		TUniquePtr<IOperator> CreateOperator(const FCreateOperatorParams& InParams, FBuildGraphResults& OutResults) override
+		TUniquePtr<IOperator> CreateOperator(const FBuildOperatorParams& InParams, FBuildResults& OutResults) override
 		{
 			const FSineOscilatorNode& SineNode = static_cast<const FSineOscilatorNode&>(InParams.Node);
-			const FDataReferenceCollection& InputCol = InParams.InputDataReferences;
+			const FInputVertexInterfaceData& InputData = InParams.InputData;
 			const FOperatorSettings& Settings = InParams.OperatorSettings;
 			using namespace Generators;
 			using namespace OscillatorCommonVertexNames; 
@@ -827,24 +827,24 @@ namespace Metasound
 			FOscillatorOperatorConstructParams OpParams
 			{
 				Settings,
-				InputCol.GetDataReadReferenceOrConstruct<bool>(METASOUND_GET_PARAM_NAME(EnabledPin), SineNode.GetDefaultEnablement()),
-				InputCol.GetDataReadReferenceOrConstruct<float>(METASOUND_GET_PARAM_NAME(OscBaseFrequencyPin), SineNode.GetDefaultFrequency()),
-				InputCol.GetDataReadReferenceOrConstruct<float>(METASOUND_GET_PARAM_NAME(PhaseOffsetPin), SineNode.GetDefaultPhaseOffset()),
-				InputCol.GetDataReadReferenceOrConstruct<FTrigger>(METASOUND_GET_PARAM_NAME(OscPhaseResetPin), Settings),
-				InputCol.GetDataReadReferenceOrConstruct<float>(METASOUND_GET_PARAM_NAME(GlideFactorPin), SineNode.GetDefaultGlideFactor()),
-				InputCol.GetDataReadReferenceOrConstruct<bool>(METASOUND_GET_PARAM_NAME(BiPolarPin), true)
+				InputData.GetOrConstructDataReadReference<bool>(METASOUND_GET_PARAM_NAME(EnabledPin), SineNode.GetDefaultEnablement()),
+				InputData.GetOrConstructDataReadReference<float>(METASOUND_GET_PARAM_NAME(OscBaseFrequencyPin), SineNode.GetDefaultFrequency()),
+				InputData.GetOrConstructDataReadReference<float>(METASOUND_GET_PARAM_NAME(PhaseOffsetPin), SineNode.GetDefaultPhaseOffset()),
+				InputData.GetOrConstructDataReadReference<FTrigger>(METASOUND_GET_PARAM_NAME(OscPhaseResetPin), Settings),
+				InputData.GetOrConstructDataReadReference<float>(METASOUND_GET_PARAM_NAME(GlideFactorPin), SineNode.GetDefaultGlideFactor()),
+				InputData.GetOrConstructDataReadReference<bool>(METASOUND_GET_PARAM_NAME(BiPolarPin), true)
 			};
 
 			// TODO: Make this a static prop. For now its a pin.
 			
 			// Check to see if we have an FM input connected.
-			bool bHasFM = InputCol.ContainsDataReadReference<FAudioBuffer>(METASOUND_GET_PARAM_NAME(FrequencyModPin));
-			FEnumSineGenerationTypeReadRef Type = InputCol.GetDataReadReferenceOrConstruct<FEnumSineGenerationType>(METASOUND_GET_PARAM_NAME(SineType), ESineGenerationType::Wavetable);
+			bool bHasFM = InputData.IsVertexBound(METASOUND_GET_PARAM_NAME(FrequencyModPin));
+			FEnumSineGenerationTypeReadRef Type = InputData.GetOrConstructDataReadReference<FEnumSineGenerationType>(METASOUND_GET_PARAM_NAME(SineType), ESineGenerationType::Wavetable);
 			if (bHasFM)
 			{
 
 				// FM Oscillators.
-				FAudioBufferReadRef FmBuffer = InputCol.GetDataReadReference<FAudioBuffer>(METASOUND_GET_PARAM_NAME(FrequencyModPin));
+				FAudioBufferReadRef FmBuffer = InputData.GetDataReadReference<FAudioBuffer>(METASOUND_GET_PARAM_NAME(FrequencyModPin));
 				switch (*Type)
 				{
 				default:
@@ -908,7 +908,7 @@ namespace Metasound
 	public:
 		FFactory() = default;
 
-		TUniquePtr<IOperator> CreateOperator(const FCreateOperatorParams& InParams, FBuildGraphResults& OutResults) override;
+		TUniquePtr<IOperator> CreateOperator(const FBuildOperatorParams& InParams, FBuildResults& OutResults) override;
 
 		static const FNodeClassMetadata& GetNodeInfo()
 		{
@@ -947,33 +947,33 @@ namespace Metasound
 		}
 	};
 
-	TUniquePtr<Metasound::IOperator> FSawOscilatorNode::FFactory::CreateOperator(const FCreateOperatorParams& InParams, FBuildGraphResults& OutResults)
+	TUniquePtr<Metasound::IOperator> FSawOscilatorNode::FFactory::CreateOperator(const FBuildOperatorParams& InParams, FBuildResults& OutResults)
 	{
 		const FSawOscilatorNode& Node = static_cast<const FSawOscilatorNode&>(InParams.Node);
-		const FDataReferenceCollection& InputCol = InParams.InputDataReferences;
+		const FInputVertexInterfaceData& InputData = InParams.InputData;
 		const FOperatorSettings& Settings = InParams.OperatorSettings;
 		using namespace Generators;
 		using namespace OscillatorCommonVertexNames;
 		using namespace SawOscilatorVertexNames;
 
-		FSawGenerationTypeReadRef Type = InputCol.GetDataReadReferenceOrConstruct<FEnumSawGenerationType>(METASOUND_GET_PARAM_NAME(SawType));
+		FSawGenerationTypeReadRef Type = InputData.GetOrConstructDataReadReference<FEnumSawGenerationType>(METASOUND_GET_PARAM_NAME(SawType));
 
 		FOscillatorOperatorConstructParams OpParams
 		{
 			Settings,
-			InputCol.GetDataReadReferenceOrConstruct<bool>(METASOUND_GET_PARAM_NAME(EnabledPin), Node.GetDefaultEnablement()),
-			InputCol.GetDataReadReferenceOrConstruct<float>(METASOUND_GET_PARAM_NAME(OscBaseFrequencyPin), Node.GetDefaultFrequency()),
-			InputCol.GetDataReadReferenceOrConstruct<float>(METASOUND_GET_PARAM_NAME(PhaseOffsetPin), Node.GetDefaultPhaseOffset()),
-			InputCol.GetDataReadReferenceOrConstruct<FTrigger>(METASOUND_GET_PARAM_NAME(OscPhaseResetPin), Settings),
-			InputCol.GetDataReadReferenceOrConstruct<float>(METASOUND_GET_PARAM_NAME(GlideFactorPin), Node.GetDefaultGlideFactor()),
-			InputCol.GetDataReadReferenceOrConstruct<bool>(METASOUND_GET_PARAM_NAME(BiPolarPin), true)
+			InputData.GetOrConstructDataReadReference<bool>(METASOUND_GET_PARAM_NAME(EnabledPin), Node.GetDefaultEnablement()),
+			InputData.GetOrConstructDataReadReference<float>(METASOUND_GET_PARAM_NAME(OscBaseFrequencyPin), Node.GetDefaultFrequency()),
+			InputData.GetOrConstructDataReadReference<float>(METASOUND_GET_PARAM_NAME(PhaseOffsetPin), Node.GetDefaultPhaseOffset()),
+			InputData.GetOrConstructDataReadReference<FTrigger>(METASOUND_GET_PARAM_NAME(OscPhaseResetPin), Settings),
+			InputData.GetOrConstructDataReadReference<float>(METASOUND_GET_PARAM_NAME(GlideFactorPin), Node.GetDefaultGlideFactor()),
+			InputData.GetOrConstructDataReadReference<bool>(METASOUND_GET_PARAM_NAME(BiPolarPin), true)
 		};
 
-		bool bHasFM = InputCol.ContainsDataReadReference<FAudioBuffer>(METASOUND_GET_PARAM_NAME(FrequencyModPin));
+		bool bHasFM = InputData.IsVertexBound(METASOUND_GET_PARAM_NAME(FrequencyModPin));
 
 		if (bHasFM)
 		{
-			FAudioBufferReadRef FmBuffer = InputCol.GetDataReadReference<FAudioBuffer>(METASOUND_GET_PARAM_NAME(FrequencyModPin));
+			FAudioBufferReadRef FmBuffer = InputData.GetDataReadReference<FAudioBuffer>(METASOUND_GET_PARAM_NAME(FrequencyModPin));
 			switch (*Type)
 			{
 			default:
@@ -1102,12 +1102,12 @@ namespace Metasound
 	public:
 		FFactory() = default;
 
-		TUniquePtr<IOperator> CreateOperator(const FCreateOperatorParams& InParams, FBuildGraphResults& OutResults) override
+		TUniquePtr<IOperator> CreateOperator(const FBuildOperatorParams& InParams, FBuildResults& OutResults) override
 		{
 			using namespace SquareOscillatorVertexNames;
 
 			const FSquareOscilatorNode& Node = static_cast<const FSquareOscilatorNode&>(InParams.Node);
-			const FDataReferenceCollection& InputCol = InParams.InputDataReferences;
+			const FInputVertexInterfaceData& InputData = InParams.InputData;
 			const FOperatorSettings& Settings = InParams.OperatorSettings;
 			using namespace Generators;
 			using namespace OscillatorCommonVertexNames; 
@@ -1115,22 +1115,22 @@ namespace Metasound
 			FOscillatorOperatorConstructParams OpParams
 			{
 				Settings,
-				InputCol.GetDataReadReferenceOrConstruct<bool>(METASOUND_GET_PARAM_NAME(EnabledPin), Node.GetDefaultEnablement()),
-				InputCol.GetDataReadReferenceOrConstruct<float>(METASOUND_GET_PARAM_NAME(OscBaseFrequencyPin), Node.GetDefaultFrequency()),
-				InputCol.GetDataReadReferenceOrConstruct<float>(METASOUND_GET_PARAM_NAME(PhaseOffsetPin), Node.GetDefaultPhaseOffset()),
-				InputCol.GetDataReadReferenceOrConstruct<FTrigger>(METASOUND_GET_PARAM_NAME(OscPhaseResetPin), Settings),
-				InputCol.GetDataReadReferenceOrConstruct<float>(METASOUND_GET_PARAM_NAME(GlideFactorPin), Node.GetDefaultGlideFactor()),
-				InputCol.GetDataReadReferenceOrConstruct<bool>(METASOUND_GET_PARAM_NAME(BiPolarPin), true)
+				InputData.GetOrConstructDataReadReference<bool>(METASOUND_GET_PARAM_NAME(EnabledPin), Node.GetDefaultEnablement()),
+				InputData.GetOrConstructDataReadReference<float>(METASOUND_GET_PARAM_NAME(OscBaseFrequencyPin), Node.GetDefaultFrequency()),
+				InputData.GetOrConstructDataReadReference<float>(METASOUND_GET_PARAM_NAME(PhaseOffsetPin), Node.GetDefaultPhaseOffset()),
+				InputData.GetOrConstructDataReadReference<FTrigger>(METASOUND_GET_PARAM_NAME(OscPhaseResetPin), Settings),
+				InputData.GetOrConstructDataReadReference<float>(METASOUND_GET_PARAM_NAME(GlideFactorPin), Node.GetDefaultGlideFactor()),
+				InputData.GetOrConstructDataReadReference<bool>(METASOUND_GET_PARAM_NAME(BiPolarPin), true)
 			};
 			
-			FSquareGenerationTypeReadRef Type = InputCol.GetDataReadReferenceOrConstruct<FEnumSquareGenerationType>(METASOUND_GET_PARAM_NAME(SquareTypePin));
+			FSquareGenerationTypeReadRef Type = InputData.GetOrConstructDataReadReference<FEnumSquareGenerationType>(METASOUND_GET_PARAM_NAME(SquareTypePin));
 
-			bool bHasFM = InputCol.ContainsDataReadReference<FAudioBuffer>(METASOUND_GET_PARAM_NAME(FrequencyModPin));
-			FFloatReadRef PulseWidth = InputCol.GetDataReadReferenceOrConstruct<float>(METASOUND_GET_PARAM_NAME(SquarePulseWidthPin), 0.5f);
+			bool bHasFM = InputData.IsVertexBound(METASOUND_GET_PARAM_NAME(FrequencyModPin));
+			FFloatReadRef PulseWidth = InputData.GetOrConstructDataReadReference<float>(METASOUND_GET_PARAM_NAME(SquarePulseWidthPin), 0.5f);
 
 			if (bHasFM)
 			{
-				FAudioBufferReadRef FmBuffer = InputCol.GetDataReadReference<FAudioBuffer>(METASOUND_GET_PARAM_NAME(FrequencyModPin));
+				FAudioBufferReadRef FmBuffer = InputData.GetDataReadReference<FAudioBuffer>(METASOUND_GET_PARAM_NAME(FrequencyModPin));
 				switch (*Type)
 				{
 				default:				
@@ -1225,33 +1225,33 @@ namespace Metasound
 	public:
 		FFactory() = default;
 
-		TUniquePtr<IOperator> CreateOperator(const FCreateOperatorParams& InParams, FBuildGraphResults& OutResults) override
+		TUniquePtr<IOperator> CreateOperator(const FBuildOperatorParams& InParams, FBuildResults& OutResults) override
 		{
 			const FTriangleOscilatorNode& Node = static_cast<const FTriangleOscilatorNode&>(InParams.Node);
-			const FDataReferenceCollection& InputCol = InParams.InputDataReferences;
+			const FInputVertexInterfaceData& InputData = InParams.InputData;
 			const FOperatorSettings& Settings = InParams.OperatorSettings;
 			using namespace Generators;
 			using namespace OscillatorCommonVertexNames;
 			using namespace TriangleOscilatorVertexNames;
 			
-			FTriangleGenerationTypeReadRef Type = InputCol.GetDataReadReferenceOrConstruct<FEnumTriangleGenerationType>(METASOUND_GET_PARAM_NAME(TriangeTypePin));
+			FTriangleGenerationTypeReadRef Type = InputData.GetOrConstructDataReadReference<FEnumTriangleGenerationType>(METASOUND_GET_PARAM_NAME(TriangeTypePin));
 		
 			FOscillatorOperatorConstructParams OpParams
 			{
 				Settings,
-				InputCol.GetDataReadReferenceOrConstruct<bool>(METASOUND_GET_PARAM_NAME(EnabledPin), Node.GetDefaultEnablement()),
-				InputCol.GetDataReadReferenceOrConstruct<float>(METASOUND_GET_PARAM_NAME(OscBaseFrequencyPin), Node.GetDefaultFrequency()),
-				InputCol.GetDataReadReferenceOrConstruct<float>(METASOUND_GET_PARAM_NAME(PhaseOffsetPin), Node.GetDefaultPhaseOffset()),
-				InputCol.GetDataReadReferenceOrConstruct<FTrigger>(METASOUND_GET_PARAM_NAME(OscPhaseResetPin), Settings),
-				InputCol.GetDataReadReferenceOrConstruct<float>(METASOUND_GET_PARAM_NAME(GlideFactorPin), Node.GetDefaultGlideFactor()),
-				InputCol.GetDataReadReferenceOrConstruct<bool>(METASOUND_GET_PARAM_NAME(BiPolarPin), true)
+				InputData.GetOrConstructDataReadReference<bool>(METASOUND_GET_PARAM_NAME(EnabledPin), Node.GetDefaultEnablement()),
+				InputData.GetOrConstructDataReadReference<float>(METASOUND_GET_PARAM_NAME(OscBaseFrequencyPin), Node.GetDefaultFrequency()),
+				InputData.GetOrConstructDataReadReference<float>(METASOUND_GET_PARAM_NAME(PhaseOffsetPin), Node.GetDefaultPhaseOffset()),
+				InputData.GetOrConstructDataReadReference<FTrigger>(METASOUND_GET_PARAM_NAME(OscPhaseResetPin), Settings),
+				InputData.GetOrConstructDataReadReference<float>(METASOUND_GET_PARAM_NAME(GlideFactorPin), Node.GetDefaultGlideFactor()),
+				InputData.GetOrConstructDataReadReference<bool>(METASOUND_GET_PARAM_NAME(BiPolarPin), true)
 			};
 
-			bool bHasFM = InputCol.ContainsDataReadReference<FAudioBuffer>(METASOUND_GET_PARAM_NAME(FrequencyModPin));
+			bool bHasFM = InputData.IsVertexBound(METASOUND_GET_PARAM_NAME(FrequencyModPin));
 			
 			if (bHasFM)
 			{
-				FAudioBufferReadRef FmBuffer = InputCol.GetDataReadReference<FAudioBuffer>(METASOUND_GET_PARAM_NAME(FrequencyModPin));
+				FAudioBufferReadRef FmBuffer = InputData.GetDataReadReference<FAudioBuffer>(METASOUND_GET_PARAM_NAME(FrequencyModPin));
 				switch (*Type)
 				{
 				default:
