MetaSound Editor Analyzer Checkpoint
#rnx

diff --git a/Engine/Plugins/Runtime/Metasound/Source/MetasoundEditor/MetasoundEditor.Build.cs b/Engine/Plugins/Runtime/Metasound/Source/MetasoundEditor/MetasoundEditor.Build.cs
--- a/Engine/Plugins/Runtime/Metasound/Source/MetasoundEditor/MetasoundEditor.Build.cs
+++ b/Engine/Plugins/Runtime/Metasound/Source/MetasoundEditor/MetasoundEditor.Build.cs
@@ -13,17 +13,21 @@ namespace UnrealBuildTool.Rules
 			(
 				new string[]
 				{
+					"AppFramework",
 					"AudioExtensions",
 					"AudioMixer",
+					"AudioWidgets",
+					"AudioSynesthesia",
+					"EditorWidgets",
+					"Kismet",
 					"KismetWidgets",
 					"MetasoundEngine",
 					"MetasoundFrontend",
+					"MetasoundGenerator",
 					"MetasoundGraphCore",
+					"SignalProcessing",
 					"ToolMenus",
-					"ToolWidgets",
-					"AudioWidgets",
-					"AudioSynesthesia",
-					"EditorWidgets"
+					"ToolWidgets"
 				}
 			);
 
@@ -32,6 +36,7 @@ namespace UnrealBuildTool.Rules
 				new string[]
 				{
 					"ApplicationCore",
+					"AudioAnalyzer",
 					"AudioMixer",
 					"AudioWidgets",
 					"ClassViewer",
@@ -39,7 +44,8 @@ namespace UnrealBuildTool.Rules
 					"Core",
 					"CoreUObject",
 					"DetailCustomizations",
-					
+					"EditorFramework",
+					"EditorStyle",
 					"Engine",
 					"GraphEditor",
 					"InputCore",
@@ -47,9 +53,7 @@ namespace UnrealBuildTool.Rules
 					"RenderCore",
 					"Slate",
 					"SlateCore",
-					"EditorFramework",
-					"UnrealEd",
-					"AudioAnalyzer"
+					"UnrealEd"
 				}
 			);
 		}
diff --git a/Engine/Plugins/Runtime/Metasound/Source/MetasoundEditor/Private/MetasoundEditor.cpp b/Engine/Plugins/Runtime/Metasound/Source/MetasoundEditor/Private/MetasoundEditor.cpp
--- a/Engine/Plugins/Runtime/Metasound/Source/MetasoundEditor/Private/MetasoundEditor.cpp
+++ b/Engine/Plugins/Runtime/Metasound/Source/MetasoundEditor/Private/MetasoundEditor.cpp
@@ -8,6 +8,7 @@
 #include "Components/AudioComponent.h"
 #include "DetailLayoutBuilder.h"
 #include "EdGraph/EdGraphNode.h"
+#include "EdGraph/EdGraphPin.h"
 #include "EdGraphUtilities.h"
 #include "Editor.h"
 #include "Styling/AppStyle.h"
@@ -506,17 +507,16 @@ namespace Metasound
 					}
 					FName DataTypeName = GraphMember->GetDataType(); 
 
-					IMetasoundEditorModule& EditorModule = FModuleManager::GetModuleChecked<IMetasoundEditorModule>("MetaSoundEditor");
-					if (const FEditorDataType* EditorDataType = EditorModule.FindDataType(DataTypeName))
+					const IMetasoundEditorModule& EditorModule = FModuleManager::GetModuleChecked<IMetasoundEditorModule>("MetaSoundEditor");
+					if (const FEdGraphPinType* PinType = EditorModule.FindPinType(DataTypeName))
 					{
-						FEdGraphPinType PinType = EditorDataType->PinType;
 						if (const UMetasoundEditorGraphSchema* Schema = GetDefault<UMetasoundEditorGraphSchema>())
 						{
-							IconColor = Schema->GetPinTypeColor(PinType);
+							IconColor = Schema->GetPinTypeColor(*PinType);
 						}
-
-						IconBrush = EditorDataType->GetIconBrush(bIsConstructorPin);
 					}
+
+					IconBrush = EditorModule.GetIconBrush(DataTypeName, bIsConstructorPin);
 				}
 
 				TSharedRef<SHorizontalBox> LayoutWidget = SNew(SHorizontalBox);
@@ -1420,6 +1420,16 @@ namespace Metasound
 					{
 						PreviewComp->SetAudioBusSendPostEffect(AudioBus, 1.0f);
 					}
+				
+					FMetasoundAssetBase* MetasoundAsset = IMetasoundUObjectRegistry::Get().GetObjectAsAssetBase(Metasound);
+					check(MetasoundAsset);
+
+					FAudioDevice* AudioDevice = PreviewComp->GetAudioDevice();
+					check(AudioDevice);
+
+					const FName& AudioBufferTypeName = GetMetasoundDataTypeName<FAudioBuffer>();
+					const FSampleRate SampleRate = static_cast<FSampleRate>(AudioDevice->GetSampleRate());
+					GraphConnectionManager = FGraphConnectionManager(*MetasoundAsset, *PreviewComp, SampleRate);
 				}
 
 				MetasoundGraphEditor->RegisterActiveTimer(0.0f,
@@ -1443,6 +1453,7 @@ namespace Metasound
 							SetPreviewID(INDEX_NONE);
 							PlayTime = 0.0;
 							PlayTimeWidget->SetText(FText::GetEmpty());
+							GraphConnectionManager = { };
 
 							return EActiveTimerReturnType::Stop;
 						}
@@ -2503,7 +2514,28 @@ namespace Metasound
 			return false;
 		}
 
-		void FEditor::RenameSelectedNode() 
+		FGraphConnectionManager& FEditor::GetConnectionManager()
+		{
+			return GraphConnectionManager;
+		}
+
+		const FGraphConnectionManager& FEditor::GetConnectionManager() const
+		{
+			return GraphConnectionManager;
+		}
+
+		UAudioComponent* FEditor::GetAudioComponent() const
+		{
+			// TODO: Instance for each editor
+			if (IsPlaying())
+			{
+				return GEditor->GetPreviewAudioComponent();
+			}
+
+			return nullptr;
+		}
+
+		void FEditor::RenameSelectedNode()
 		{
 			const FGraphPanelSelectionSet& SelectedNodes = MetasoundGraphEditor->GetSelectedNodes();
 			for (FGraphPanelSelectionSet::TConstIterator SelectedIter(SelectedNodes); SelectedIter; ++SelectedIter)
@@ -3136,6 +3168,8 @@ namespace Metasound
 					}
 				}
 			}
+
+			GraphConnectionManager.Update(DeltaTime);
 		}
 
 		TStatId FEditor::GetStatId() const
diff --git a/Engine/Plugins/Runtime/Metasound/Source/MetasoundEditor/Private/MetasoundEditor.h b/Engine/Plugins/Runtime/Metasound/Source/MetasoundEditor/Private/MetasoundEditor.h
--- a/Engine/Plugins/Runtime/Metasound/Source/MetasoundEditor/Private/MetasoundEditor.h
+++ b/Engine/Plugins/Runtime/Metasound/Source/MetasoundEditor/Private/MetasoundEditor.h
@@ -10,6 +10,7 @@
 #include "IDetailsView.h"
 #include "IMetasoundEditor.h"
 #include "Math/UnrealMathUtility.h"
+#include "MetasoundEditorGraphConnectionManager.h"
 #include "MetasoundEditorGraphNode.h"
 #include "MetasoundEditorMeter.h"
 #include "MetasoundFrontend.h"
@@ -42,6 +43,7 @@ class FSlateRect;
 class IDetailsView;
 class IToolkitHost;
 class SVerticalBox;
+class UAudioComponent;
 class UEdGraphNode;
 class UMetaSound;
 class UMetasoundEditorGraph;
@@ -100,6 +102,8 @@ namespace Metasound
 			/** Edits the specified Metasound object */
 			void InitMetasoundEditor(const EToolkitMode::Type Mode, const TSharedPtr<IToolkitHost>& InitToolkitHost, UObject* ObjectToEdit);
 
+			UAudioComponent* GetAudioComponent() const;
+
 			/** IMetasoundEditor interface */
 			virtual UObject* GetMetasoundObject() const override;
 			virtual void SetSelection(const TArray<UObject*>& SelectedObjects) override;
@@ -140,6 +144,8 @@ namespace Metasound
 			virtual void Play() override;
 			virtual void Stop() override;
 
+			virtual bool IsPlaying() const override;
+
 			/** Whether pasting the currently selected nodes is permissible */
 			bool CanPasteNodes();
 
@@ -153,6 +159,9 @@ namespace Metasound
 			void PasteNodes(const FVector2D* InLocation = nullptr);
 			void PasteNodes(const FVector2D* InLocation, const FText& InTransactionText);
 
+			/** Returns Graph Connection Manager associated with this editor */
+			FGraphConnectionManager& GetConnectionManager();
+			const FGraphConnectionManager& GetConnectionManager() const;
 
 			/** Forces all UX pertaining to the root graph's interface to be refreshed. */
 			void RefreshGraphMemberMenu();
@@ -373,8 +382,7 @@ namespace Metasound
 
 			/** Create new graph editor widget */
 			void CreateGraphEditorWidget();
-			
-		private:
+
 			TSharedPtr<SWidget> BuildAnalyzerWidget() const;
 
 			void EditObjectSettings();
@@ -385,8 +393,6 @@ namespace Metasound
 			void NotifyNodePasteFailure_MultipleVariableSetters();
 			void NotifyNodePasteFailure_ReferenceLoop();
 
-			bool IsPlaying() const;
-
 			/** List of open tool panels; used to ensure only one exists at any one time */
 			TMap<FName, TWeakPtr<SDockableTab>> SpawnedToolPanels;
 
@@ -413,6 +419,8 @@ namespace Metasound
 			TSharedPtr<STextBlock> PlayTimeWidget;
 			double PlayTime = 0.0;
 
+			FGraphConnectionManager GraphConnectionManager;
+
 			/** Command list for this editor */
 			TSharedPtr<FUICommandList> GraphEditorCommands;
 
diff --git a/Engine/Plugins/Runtime/Metasound/Source/MetasoundEditor/Private/MetasoundEditorGraph.cpp b/Engine/Plugins/Runtime/Metasound/Source/MetasoundEditor/Private/MetasoundEditorGraph.cpp
--- a/Engine/Plugins/Runtime/Metasound/Source/MetasoundEditor/Private/MetasoundEditorGraph.cpp
+++ b/Engine/Plugins/Runtime/Metasound/Source/MetasoundEditor/Private/MetasoundEditorGraph.cpp
@@ -94,9 +94,10 @@ void UMetasoundEditorGraphMember::ConformLiteralDataType()
 	using namespace Metasound::Editor;
 	using namespace Metasound::Frontend;
 
+	FDataTypeRegistryInfo DataTypeInfo;
 	IMetasoundEditorModule& EditorModule = FModuleManager::GetModuleChecked<IMetasoundEditorModule>("MetaSoundEditor");
-	const FEditorDataType& EditorDataType = EditorModule.FindDataTypeChecked(TypeName);
-	const EMetasoundFrontendLiteralType LiteralType = static_cast<EMetasoundFrontendLiteralType>(EditorDataType.RegistryInfo.PreferredLiteralType);
+	IDataTypeRegistry::Get().GetDataTypeInfo(TypeName, DataTypeInfo);
+	const EMetasoundFrontendLiteralType LiteralType = static_cast<EMetasoundFrontendLiteralType>(DataTypeInfo.PreferredLiteralType);
 
 	TSubclassOf<UMetasoundEditorGraphMemberDefaultLiteral> LiteralClass = EditorModule.FindDefaultLiteralClass(LiteralType);
 	if (!LiteralClass)
diff --git a/Engine/Plugins/Runtime/Metasound/Source/MetasoundEditor/Private/MetasoundEditorGraphBuilder.cpp b/Engine/Plugins/Runtime/Metasound/Source/MetasoundEditor/Private/MetasoundEditorGraphBuilder.cpp
--- a/Engine/Plugins/Runtime/Metasound/Source/MetasoundEditor/Private/MetasoundEditorGraphBuilder.cpp
+++ b/Engine/Plugins/Runtime/Metasound/Source/MetasoundEditor/Private/MetasoundEditorGraphBuilder.cpp
@@ -87,6 +87,69 @@ namespace Metasound
 			}
 		} // namespace GraphBuilderPrivate
 
+		bool FGraphBuilder::CanInspectPin(const UEdGraphPin* InPin)
+		{
+			// Can't inspect the value on an invalid pin object.
+			if (!InPin || InPin->IsPendingKill())
+			{
+				return false;
+			}
+
+			// Can't inspect the value on an orphaned pin object.
+			if (InPin->bOrphanedPin)
+			{
+				return false;
+			}
+
+			// Currently only inspection of connected pins is supported.
+			if (InPin->LinkedTo.IsEmpty())
+			{
+				return false;
+			}
+
+			// Can't inspect the value on an unknown pin object or if the owning node is disabled.
+			const UEdGraphNode* OwningNode = InPin->GetOwningNodeUnchecked();
+			if (!OwningNode || !OwningNode->IsNodeEnabled())
+			{
+				return false;
+			}
+
+			TSharedPtr<FEditor> Editor = GetEditorForPin(*InPin);
+			if (!Editor.IsValid())
+			{
+				return false;
+			}
+
+			if (!Editor->IsPlaying())
+			{
+				return false;
+			}
+
+			FName DataType;
+			if (InPin->Direction == EGPD_Input)
+			{
+				Frontend::FConstInputHandle InputHandle = GetConstInputHandleFromPin(InPin);
+				DataType = InputHandle->GetDataType();
+			}
+			else
+			{
+				Frontend::FConstOutputHandle OutputHandle = GetConstOutputHandleFromPin(InPin);
+				DataType = OutputHandle->GetDataType();
+			}
+
+			const bool bIsSupportedType = DataType == GetMetasoundDataTypeName<float>()
+				|| DataType == GetMetasoundDataTypeName<int32>()
+				|| DataType == GetMetasoundDataTypeName<FString>()
+				|| DataType == GetMetasoundDataTypeName<bool>();
+
+			if (!bIsSupportedType)
+			{
+				return false;
+			}
+
+			return true;
+		}
+
 		FText FGraphBuilder::GetDisplayName(const FMetasoundFrontendClassMetadata& InClassMetadata, FName InNodeName, bool bInIncludeNamespace)
 		{
 			using namespace Frontend;
@@ -514,13 +577,37 @@ namespace Metasound
 				}
 			}
 
-			return TSharedPtr<FEditor>(nullptr);
+			return { };
 		}
 
 		TSharedPtr<FEditor> FGraphBuilder::GetEditorForGraph(const UEdGraph& EdGraph)
 		{
-			const UMetasoundEditorGraph* MetasoundGraph = CastChecked<const UMetasoundEditorGraph>(&EdGraph);
-			return GetEditorForMetasound(MetasoundGraph->GetMetasoundChecked());
+			if (const UMetasoundEditorGraph* MetasoundGraph = Cast<const UMetasoundEditorGraph>(&EdGraph))
+			{
+				return GetEditorForMetasound(MetasoundGraph->GetMetasoundChecked());
+			}
+
+			return { };
+		}
+
+		TSharedPtr<FEditor> FGraphBuilder::GetEditorForNode(const UEdGraphNode& InEdNode)
+		{
+			if (const UMetasoundEditorGraph* Graph = Cast<UMetasoundEditorGraph>(InEdNode.GetGraph()))
+			{
+				return FGraphBuilder::GetEditorForGraph(*Graph);
+			}
+
+			return { };
+		}
+
+		TSharedPtr<FEditor> FGraphBuilder::GetEditorForPin(const UEdGraphPin& InEdPin)
+		{
+			if (const UMetasoundEditorGraphNode* Node = Cast<UMetasoundEditorGraphNode>(InEdPin.GetOwningNode()))
+			{
+				return GetEditorForNode(*Node);
+			}
+
+			return { };
 		}
 
 		FLinearColor FGraphBuilder::GetPinCategoryColor(const FEdGraphPinType& PinType)
@@ -633,6 +720,33 @@ namespace Metasound
 			return IOutputController::GetInvalidHandle();
 		}
 
+		const FMetasoundFrontendEdgeStyle* FGraphBuilder::GetOutputEdgeStyle(Frontend::FConstOutputHandle InOutputHandle)
+		{
+			using namespace Frontend;
+
+			if (InOutputHandle->IsValid())
+			{
+				FConstNodeHandle NodeHandle = InOutputHandle->GetOwningNode();
+				FConstGraphHandle GraphHandle = NodeHandle->GetOwningGraph();
+
+				const TArray<FMetasoundFrontendEdgeStyle>& EdgeStyles = GraphHandle->GetGraphStyle().EdgeStyles;
+				return EdgeStyles.FindByPredicate([&InOutputHandle](const FMetasoundFrontendEdgeStyle& InCandidate)
+				{
+					return InCandidate.NodeID == InOutputHandle->GetOwningNodeID() && InCandidate.OutputName == InOutputHandle->GetName();
+				});
+			}
+
+			return nullptr;
+		}
+
+		const FMetasoundFrontendEdgeStyle* FGraphBuilder::GetOutputEdgeStyle(const UEdGraphPin* InGraphPin)
+		{
+			using namespace Frontend;
+
+			FConstOutputHandle OutputHandle = GetConstOutputHandleFromPin(InGraphPin);
+			return GetOutputEdgeStyle(OutputHandle);
+		}
+
 		Frontend::FConstOutputHandle FGraphBuilder::GetConstOutputHandleFromPin(const UEdGraphPin* InPin)
 		{
 			return GetOutputHandleFromPin(InPin);
@@ -731,8 +845,10 @@ namespace Metasound
 
 			const FString& InStringValue = InInputPin.DefaultValue;
 			const FName TypeName = InputHandle->GetDataType();
-			const FEditorDataType DataType = EditorModule.FindDataTypeChecked(TypeName);
-			switch (DataType.RegistryInfo.PreferredLiteralType)
+
+			FDataTypeRegistryInfo DataTypeInfo;
+			IDataTypeRegistry::Get().GetDataTypeInfo(TypeName, DataTypeInfo);
+			switch (DataTypeInfo.PreferredLiteralType)
 			{
 				case ELiteralType::Boolean:
 				{
@@ -1564,10 +1680,10 @@ namespace Metasound
 			FEdGraphPinType PinType;
 			FName DataTypeName = InInputHandle->GetDataType();
 
-			IMetasoundEditorModule& EditorModule = FModuleManager::GetModuleChecked<IMetasoundEditorModule>("MetaSoundEditor");
-			if (const FEditorDataType* EditorDataType = EditorModule.FindDataType(DataTypeName))
+			const IMetasoundEditorModule& EditorModule = FModuleManager::GetModuleChecked<IMetasoundEditorModule>("MetaSoundEditor");
+			if (const FEdGraphPinType* RegisteredPinType = EditorModule.FindPinType(DataTypeName))
 			{
-				PinType = EditorDataType->PinType;
+				PinType = *RegisteredPinType;
 			}
 
 			FName PinName = FGraphBuilder::GetPinName(*InInputHandle);
@@ -1587,9 +1703,9 @@ namespace Metasound
 			FName DataTypeName = InOutputHandle->GetDataType();
 
 			IMetasoundEditorModule& EditorModule = FModuleManager::GetModuleChecked<IMetasoundEditorModule>("MetaSoundEditor");
-			if (const FEditorDataType* EditorDataType = EditorModule.FindDataType(DataTypeName))
+			if (const FEdGraphPinType* RegisteredPinType = EditorModule.FindPinType(DataTypeName))
 			{
-				PinType = EditorDataType->PinType;
+				PinType = *RegisteredPinType;
 			}
 
 			FName PinName = FGraphBuilder::GetPinName(*InOutputHandle);
@@ -1605,9 +1721,9 @@ namespace Metasound
 		bool FGraphBuilder::SynchronizePinType(const IMetasoundEditorModule& InEditorModule, UEdGraphPin& InPin, const FName InDataType)
 		{
 			FEdGraphPinType PinType;
-			if (const FEditorDataType* EditorDataType = InEditorModule.FindDataType(InDataType))
+			if (const FEdGraphPinType* RegisteredPinType = InEditorModule.FindPinType(InDataType))
 			{
-				PinType = EditorDataType->PinType;
+				PinType = *RegisteredPinType;
 			}
 
 			if (InPin.PinType != PinType)
diff --git a/Engine/Plugins/Runtime/Metasound/Source/MetasoundEditor/Private/MetasoundEditorGraphBuilder.h b/Engine/Plugins/Runtime/Metasound/Source/MetasoundEditor/Private/MetasoundEditorGraphBuilder.h
--- a/Engine/Plugins/Runtime/Metasound/Source/MetasoundEditor/Private/MetasoundEditorGraphBuilder.h
+++ b/Engine/Plugins/Runtime/Metasound/Source/MetasoundEditor/Private/MetasoundEditorGraphBuilder.h
@@ -65,9 +65,11 @@ namespace Metasound
 			// Adds an EdGraph node to mirror the provided FNodeHandle.
 			static UMetasoundEditorGraphNode* AddNode(UObject& InMetaSound, Frontend::FNodeHandle InNodeHandle, FVector2D InLocation, bool bInSelectNewNode = true);
 
-			// Convenience functions for retrieving the editor for the given MetaSound/EdGraph
+			// Convenience functions for retrieving the editor for the given UObject
 			static TSharedPtr<FEditor> GetEditorForMetasound(const UObject& InMetaSound);
 			static TSharedPtr<FEditor> GetEditorForGraph(const UEdGraph& InEdGraph);
+			static TSharedPtr<FEditor> GetEditorForNode(const UEdGraphNode& InEdNode);
+			static TSharedPtr<FEditor> GetEditorForPin(const UEdGraphPin& InEdPin);
 
 			// Wraps RegisterGraphWithFrontend logic in Frontend with any additional logic required to refresh editor & respective editor object state.
 			// @param InMetaSound - MetaSound to register
@@ -78,6 +80,9 @@ namespace Metasound
 			// Wraps RegisterGraphWithFrontend logic in Frontend with any additional logic required to refresh editor & respective editor object state.
 			static void UnregisterGraphWithFrontend(UObject& InMetaSound);
 
+			// Determines if pin supports inspection/probe view.
+			static bool CanInspectPin(const UEdGraphPin* InPin);
+
 			// Returns a display name for a node. If the node has an empty or whitespace
 			// only DisplayName, first attempts to use the asset name if class is defined
 			// in an asset, and finally the NodeName is used.
@@ -180,6 +185,9 @@ namespace Metasound
 			static Frontend::FInputHandle GetInputHandleFromPin(const UEdGraphPin* InPin);
 			static Frontend::FConstInputHandle GetConstInputHandleFromPin(const UEdGraphPin* InPin);
 
+			static const FMetasoundFrontendEdgeStyle* GetOutputEdgeStyle(Frontend::FConstOutputHandle InOutputHandle);
+			static const FMetasoundFrontendEdgeStyle* GetOutputEdgeStyle(const UEdGraphPin* InPin);
+
 			// Get the output handle from an output pin.  Ensures pin is an output pin.
 			// TODO: use IDs to connect rather than names. Likely need an UMetasoundEditorGraphPin
 			static Frontend::FOutputHandle GetOutputHandleFromPin(const UEdGraphPin* InPin);
diff --git a/Engine/Plugins/Runtime/Metasound/Source/MetasoundEditor/Private/MetasoundEditorGraphConnectionDrawingPolicy.cpp b/Engine/Plugins/Runtime/Metasound/Source/MetasoundEditor/Private/MetasoundEditorGraphConnectionDrawingPolicy.cpp
--- a/Engine/Plugins/Runtime/Metasound/Source/MetasoundEditor/Private/MetasoundEditorGraphConnectionDrawingPolicy.cpp
+++ b/Engine/Plugins/Runtime/Metasound/Source/MetasoundEditor/Private/MetasoundEditorGraphConnectionDrawingPolicy.cpp
@@ -1,19 +1,237 @@
 // Copyright Epic Games, Inc. All Rights Reserved.
 #include "MetasoundEditorGraphConnectionDrawingPolicy.h"
 
+#include "Analysis/MetasoundFrontendVertexAnalyzerEnvelopeFollower.h"
+#include "Analysis/MetasoundFrontendVertexAnalyzerTriggerDensity.h"
 #include "Components/AudioComponent.h"
 #include "Editor.h"
 #include "MetasoundEditor.h"
 #include "MetasoundEditorGraph.h"
 #include "MetasoundEditorGraphBuilder.h"
 #include "MetasoundEditorGraphSchema.h"
+#include "MetasoundEditorSettings.h"
 #include "Misc/App.h"
+#include "Templates/Function.h"
 
 
+#define METASOUND_DRAW_CONNECTION_DEBUG_DATA 0
+
 namespace Metasound
 {
 	namespace Editor
 	{
+		namespace DrawingPolicyPrivate
+		{
+#if METASOUND_DRAW_CONNECTION_DEBUG_DATA
+			// Draw the bounding box for debugging
+			void DrawDebugConnection(FSlateWindowElementList& InDrawElementsList, uint32 InLayerId, const FDrawConnectionData& InDrawConnectionData)
+			{
+				const FVector2D& TL = InDrawConnectionData.Bounds.Min;
+				const FVector2D& BR = InDrawConnectionData.Bounds.Max;
+				const FVector2D TR = FVector2D(InDrawConnectionData.Bounds.Max.X, InDrawConnectionData.Bounds.Min.Y);
+				const FVector2D BL = FVector2D(InDrawConnectionData.Bounds.Min.X, InDrawConnectionData.Bounds.Max.Y);
+
+				auto DrawSpaceLine = [&](const FVector2D& Point1, const FVector2D& Point2, const FLinearColor& InColor)
+				{
+					const FVector2D FakeTangent = (Point2 - Point1).GetSafeNormal();
+					FSlateDrawElement::MakeDrawSpaceSpline(InDrawElementsList, InLayerId, Point1, FakeTangent, Point2, FakeTangent, 1.0f, ESlateDrawEffect::None);
+				};
+
+				const FLinearColor BoundsWireColor = InDrawConnectionData.bCloseToSpline ? FLinearColor::Green : FLinearColor::White;
+				DrawSpaceLine(TL, TR, BoundsWireColor);
+				DrawSpaceLine(TR, BR, BoundsWireColor);
+				DrawSpaceLine(BR, BL, BoundsWireColor);
+				DrawSpaceLine(BL, TL, BoundsWireColor);
+			}
+#endif // METASOUND_DRAW_CONNECTION_DEBUG_DATA
+
+			template<typename TAnalyzerType>
+			bool DetermineWiringStyle_Envelope(TSharedPtr<FEditor> MetasoundEditor, UEdGraphPin* OutputPin, UEdGraphPin* InputPin, FConnectionParams& OutParams)
+			{
+				using namespace Frontend;
+
+				OutParams.bDrawBubbles = false;
+
+				if (MetasoundEditor.IsValid())
+				{
+					FConstOutputHandle OutputHandle = FGraphBuilder::GetConstOutputHandleFromPin(OutputPin);
+					const FGuid NodeId = OutputHandle->GetOwningNodeID();
+					FName OutputName = OutputHandle->GetName();
+
+					OutParams.bDrawBubbles = MetasoundEditor->GetConnectionManager().IsTracked(NodeId, OutputName, TAnalyzerType::GetAnalyzerName());
+					if (OutParams.bDrawBubbles)
+					{
+						float WindowValue = MetasoundEditor->GetConnectionManager().UpdateValueWindow<TAnalyzerType>(NodeId, OutputName);
+						OutParams.WireThickness = FMath::Clamp(WindowValue, 0.0f, 1.0f);
+						FLinearColor HSV = OutParams.WireColor.LinearRGBToHSV();
+						HSV.G = FMath::Lerp(HSV.G, 1.0f, WindowValue);
+						OutParams.WireColor = HSV.HSVToLinearRGB();
+					}
+
+					return true;
+				}
+
+				return false;
+			}
+
+			template<typename TNumericType>
+			void DrawConnectionSpline_Numeric(FGraphConnectionDrawingPolicy& InDrawingPolicy, int32 InLayerId, const FDrawConnectionData& InData, TNumericType InDefaultValue)
+			{
+				if (!InData.OutputPin)
+				{
+					return;
+				}
+
+				const UMetasoundEditorGraphNode* Node = Cast<UMetasoundEditorGraphNode>(InData.OutputPin->GetOwningNode());
+				if (!ensureMsgf(Node, TEXT("Expected MetaSound pin to be member of MetaSound node")))
+				{
+					return;
+				}
+
+				const UMetasoundEditorGraph* Graph = Cast<UMetasoundEditorGraph>(Node->GetGraph());
+				if (!ensureMsgf(Graph, TEXT("Expected MetaSound node to be member of MetaSound graph")))
+				{
+					return;
+				}
+
+				const TSharedPtr<const FEditor> Editor = static_cast<const FGraphConnectionDrawingPolicy&>(InDrawingPolicy).GetEditor();
+				if (!Editor)
+				{
+					return;
+				}
+
+				const FGuid NodeID = Node->GetNodeID();
+				TNumericType Value = InDefaultValue;
+				FName OutputName = InData.OutputHandle->GetName();
+				
+				Editor->GetConnectionManager().GetValue(NodeID, OutputName, Value);
+
+				FLinearColor InnerColor = Frontend::DisplayStyle::EdgeAnimation::DefaultColor;
+
+				if (!InData.EdgeStyle || InData.EdgeStyle->LiteralColorPairs.IsEmpty())
+				{
+					InDrawingPolicy.DrawConnectionSpline(InLayerId, InData);
+					return;
+				}
+
+				for (int32 i = 0; i < InData.EdgeStyle->LiteralColorPairs.Num() - 1; ++i)
+				{
+					const FMetasoundFrontendEdgeStyleLiteralColorPair& PairA = InData.EdgeStyle->LiteralColorPairs[i];
+
+					TNumericType ValueA = InDefaultValue;
+					if (PairA.Value.TryGet(ValueA))
+					{
+						if (Value > ValueA)
+						{
+							const FMetasoundFrontendEdgeStyleLiteralColorPair& PairB = InData.EdgeStyle->LiteralColorPairs[i + 1];
+							TNumericType ValueB = InDefaultValue;
+							if (PairB.Value.TryGet(ValueB))
+							{
+								const float Alpha = (static_cast<float>(Value) - ValueA) / (static_cast<float>(ValueB) - ValueA);
+								InnerColor = PairA.Color + FMath::Clamp(Alpha, 0.0f, 1.0f) * (PairB.Color - PairA.Color);
+							}
+							break;
+						}
+					}
+
+					InnerColor = PairA.Color;
+				}
+
+				// Draw the outer spline, which uses standardized color syntax
+				InDrawingPolicy.DrawSpline(InData);
+
+				// Draw the inner spline, which uses interpolated custom color syntax.
+				FDrawConnectionData InnerSplineData = InData;
+				InnerSplineData.Params.WireThickness *= 0.6666f;
+				InnerSplineData.Params.WireColor = InnerColor;
+				InDrawingPolicy.DrawAnimatedSpline(InnerSplineData);
+			}
+		} // namespace DrawingPolicyPrivate
+
+
+		FDrawConnectionData::FDrawConnectionData(const FVector2D& InStart, const FVector2D& InEnd, const FVector2D& InSplineTangent, const FConnectionParams& InParams, const UGraphEditorSettings& InSettings, const FVector2D& InMousePosition)
+			: P0(InStart)
+			, P0Tangent(InParams.StartDirection == EGPD_Output ? InSplineTangent : -InSplineTangent)
+			, P1(InEnd)
+			, P1Tangent(InParams.EndDirection == EGPD_Input ? InSplineTangent : -InSplineTangent)
+			, Params(InParams)
+			, OutputHandle(Frontend::IOutputController::GetInvalidHandle())
+		{
+			if (Params.AssociatedPin1)
+			{
+				OutputPin = Params.AssociatedPin1->Direction == EGPD_Output ? Params.AssociatedPin1 : Params.AssociatedPin2;
+				OutputHandle = FGraphBuilder::GetConstOutputHandleFromPin(OutputPin);
+			}
+
+			// The curve will include the endpoints but can extend out of a tight bounds because of the tangents
+			// P0Tangent coefficient maximizes to 4/27 at a=1/3, and P1Tangent minimizes to -4/27 at a=2/3.
+			constexpr float MaximumTangentContribution = 4.0f / 27.0f;
+			Bounds = FBox2D(ForceInit);
+
+			Bounds += FVector2D(P0);
+			Bounds += FVector2D(P0 + MaximumTangentContribution * P0Tangent);
+			Bounds += FVector2D(P1);
+			Bounds += FVector2D(P1 - MaximumTangentContribution * P1Tangent);
+
+			if (InSettings.bTreatSplinesLikePins)
+			{
+				QueryDistanceTriggerThresholdSquared = FMath::Square(InSettings.SplineHoverTolerance + Params.WireThickness * 0.5f);
+				QueryDistanceForCloseSquared = FMath::Square(FMath::Sqrt(QueryDistanceTriggerThresholdSquared) + InSettings.SplineCloseTolerance);
+				bCloseToSpline = Bounds.ComputeSquaredDistanceToPoint(InMousePosition) < QueryDistanceForCloseSquared;
+			}
+
+			ClosestPoint = FVector2D(ForceInit);
+			if (InSettings.bTreatSplinesLikePins)
+			{
+				if (bCloseToSpline)
+				{
+					const int32 NumStepsToTest = 16;
+					const float StepInterval = 1.0f / (float)NumStepsToTest;
+					FVector2D Point1 = FMath::CubicInterp(P0, P0Tangent, P1, P1Tangent, 0.0f);
+					for (float TestAlpha = 0.0f; TestAlpha < 1.0f; TestAlpha += StepInterval)
+					{
+						const FVector2D Point2 = FMath::CubicInterp(P0, P0Tangent, P1, P1Tangent, TestAlpha + StepInterval);
+
+						const FVector2D ClosestPointToSegment = FMath::ClosestPointOnSegment2D(InMousePosition, Point1, Point2);
+						const float DistanceSquared = (InMousePosition - ClosestPointToSegment).SizeSquared();
+
+						if (DistanceSquared < ClosestDistanceSquared)
+						{
+							ClosestDistanceSquared = DistanceSquared;
+							ClosestPoint = ClosestPointToSegment;
+						}
+
+						Point1 = Point2;
+					}
+				}
+			}
+
+			EdgeStyle = FGraphBuilder::GetOutputEdgeStyle(OutputHandle);
+		}
+
+		void FDrawConnectionData::UpdateSplineOverlap(FGraphSplineOverlapResult& OutResult) const
+		{
+			if (!bCloseToSpline)
+			{
+				return;
+			}
+
+			if (ClosestDistanceSquared < QueryDistanceTriggerThresholdSquared)
+			{
+				if (ClosestDistanceSquared < OutResult.GetDistanceSquared())
+				{
+					const float SquaredDistToPin1 = Params.AssociatedPin1 ? (P0 - ClosestPoint).SizeSquared() : TNumericLimits<float>::Max();
+					const float SquaredDistToPin2 = Params.AssociatedPin2 ? (P1 - ClosestPoint).SizeSquared() : TNumericLimits<float>::Max();
+
+					OutResult = FGraphSplineOverlapResult(Params.AssociatedPin1, Params.AssociatedPin2, ClosestDistanceSquared, SquaredDistToPin1, SquaredDistToPin2, bCloseToSpline);
+				}
+			}
+			else if (ClosestDistanceSquared < QueryDistanceForCloseSquared)
+			{
+				OutResult.SetCloseToSpline(bCloseToSpline);
+			}
+		}
+
 		FConnectionDrawingPolicy* FGraphConnectionDrawingPolicyFactory::CreateConnectionPolicy(
 			const UEdGraphSchema* Schema,
 			int32 InBackLayerID,
@@ -27,86 +245,402 @@ namespace Metasound
 			{
 				return new FGraphConnectionDrawingPolicy(InBackLayerID, InFrontLayerID, ZoomFactor, InClippingRect, InDrawElements, InGraphObj);
 			}
+
 			return nullptr;
 		}
 
-		void FGraphConnectionDrawingPolicy::Draw(TMap<TSharedRef<SWidget>, FArrangedWidget>& InPinGeometries, FArrangedChildren& ArrangedNodes)
+		FGraphConnectionDrawingPolicy::FGraphConnectionDrawingPolicy(int32 InBackLayerID, int32 InFrontLayerID, float InZoomFactor, const FSlateRect& InClippingRect, FSlateWindowElementList& InDrawElements, UEdGraph* InGraphObj)
+			: FConnectionDrawingPolicy(InBackLayerID, InFrontLayerID, InZoomFactor, InClippingRect, InDrawElements)
+			, GraphObj(InGraphObj)
+		{
+			WireAnimationLayerID = (InFrontLayerID + InBackLayerID) / 2;
+			ActiveWireThickness = Settings->TraceAttackWireThickness;
+			InactiveWireThickness = Settings->TraceReleaseWireThickness;
+
+			// Numeric wires are inflated up to account for interior animation
+			if (const UMetasoundEditorSettings* MetasoundSettings = GetDefault<UMetasoundEditorSettings>())
+			{
+				bAnimateConnections = MetasoundSettings->bAnimateConnections;
+				ActiveAnalyzerWireThickness = MetasoundSettings->ActiveAnalyzerWireThickness;
+				ActiveAnalyzerWireSignalScalarMin = MetasoundSettings->ActiveAnalyzerWireSignalScalarMin;
+				ActiveAnalyzerWireSignalScalarMax = MetasoundSettings->ActiveAnalyzerWireSignalScalarMax;
+			}
+
+			if (GraphObj)
+			{
+				MetasoundEditor = FGraphBuilder::GetEditorForGraph(*GraphObj);
+			}
+		}
+
+		void FGraphConnectionDrawingPolicy::ScaleEnvelopeWireThickness(FConnectionParams& OutParams) const
+		{
+			if (OutParams.bDrawBubbles)
+			{
+				OutParams.WireThickness = FMath::Lerp(InactiveWireThickness * ActiveAnalyzerWireSignalScalarMin, ActiveAnalyzerWireThickness * ActiveAnalyzerWireSignalScalarMax, OutParams.WireThickness);
+			}
+			else
+			{
+				OutParams.WireThickness = InactiveWireThickness;
+			}
+		}
+
+		const TSharedPtr<const FEditor> FGraphConnectionDrawingPolicy::GetEditor() const
 		{
-			// Draw everything
-			FConnectionDrawingPolicy::Draw(InPinGeometries, ArrangedNodes);
+			return MetasoundEditor.Pin();
 		}
 
-		void FGraphConnectionDrawingPolicy::DetermineWiringStyle(UEdGraphPin* OutputPin, UEdGraphPin* InputPin, FConnectionParams& OutParams)
+		TSharedPtr<FEditor> FGraphConnectionDrawingPolicy::GetEditor()
+		{
+			return MetasoundEditor.Pin();
+		}
+
+		void FGraphConnectionDrawingPolicy::DetermineWiringStyle(UEdGraphPin* InPinA, UEdGraphPin* InPinB, FConnectionParams& OutParams)
 		{
 			using namespace Frontend;
 
-			if (!OutputPin || !InputPin || !GraphObj)
+			TSharedPtr<FEditor> EditorPtr = MetasoundEditor.Pin();
+			if (!EditorPtr)
 			{
+				FConnectionDrawingPolicy::DetermineWiringStyle(InPinA, InPinB, OutParams);
 				return;
 			}
 
+			UEdGraphPin* InputPin = nullptr;
+			UEdGraphPin* OutputPin = nullptr;
+			if (InPinA && InPinA->Direction == EGPD_Input)
+			{
+				InputPin = InPinA;
+			}
+			else if (InPinB && InPinB->Direction == EGPD_Input)
+			{
+				InputPin = InPinB;
+			}
+
+			if (InPinA && InPinA->Direction == EGPD_Output)
+			{
+				OutputPin = InPinA;
+			}
+			else if (InPinB && InPinB->Direction == EGPD_Output)
+			{
+				OutputPin = InPinB;
+			}
+
 			OutParams.AssociatedPin1 = InputPin;
 			OutParams.AssociatedPin2 = OutputPin;
 
-			if (OutputPin->bOrphanedPin || InputPin->bOrphanedPin)
+			const UEdGraphPin* AnyPin = InputPin ? InputPin : OutputPin;
+
+			const bool bInputOrphaned = InputPin && InputPin->bOrphanedPin;
+			const bool bOutputOrphaned = OutputPin && OutputPin->bOrphanedPin;
+			if (bInputOrphaned || bOutputOrphaned)
 			{
 				OutParams.WireColor = FLinearColor::Red;
 			}
+			else if (AnyPin)
+			{
+				OutParams.WireColor = FGraphBuilder::GetPinCategoryColor(AnyPin->PinType);
+			}
+
+			if (AnyPin && bAnimateConnections)
+			{
+				if (AnyPin->PinType.PinCategory == FGraphBuilder::PinCategoryTrigger)
+				{
+					DrawingPolicyPrivate::DetermineWiringStyle_Envelope<FVertexAnalyzerTriggerDensity>(EditorPtr, OutputPin, InputPin, OutParams);
+					ScaleEnvelopeWireThickness(OutParams);
+				}
+				else if (AnyPin->PinType.PinCategory == FGraphBuilder::PinCategoryAudio)
+				{
+					DrawingPolicyPrivate::DetermineWiringStyle_Envelope<FVertexAnalyzerEnvelopeFollower>(EditorPtr, OutputPin, InputPin, OutParams);
+					ScaleEnvelopeWireThickness(OutParams);
+				}
+				else if (FGraphBuilder::CanInspectPin(AnyPin))
+				{
+					bool bEdgeStyleValid = false;
+					if (const FMetasoundFrontendEdgeStyle* Style = FGraphBuilder::GetOutputEdgeStyle(OutputPin))
+					{
+						bEdgeStyleValid = !Style->LiteralColorPairs.IsEmpty();
+					}
+
+					if (bEdgeStyleValid && EditorPtr->IsPlaying())
+					{
+						OutParams.WireThickness = ActiveAnalyzerWireThickness;
+					}
+					else
+					{
+						OutParams.WireThickness = InactiveWireThickness;
+					}
+				}
+			}
 			else
 			{
-				OutParams.WireColor = FGraphBuilder::GetPinCategoryColor(OutputPin->PinType);
+				OutParams.WireThickness = InactiveWireThickness;
+			}
+		}
+
+		void FGraphConnectionDrawingPolicy::DrawConnection(int32 LayerId, const FVector2D& Start, const FVector2D& End, const FConnectionParams& Params)
+		{
+			using namespace Frontend;
+
+			TSharedPtr<FEditor> EditorPtr = MetasoundEditor.Pin();
+			if (!EditorPtr)
+			{
+				FConnectionDrawingPolicy::DrawConnection(LayerId, Start, End, Params);
+				return;
+			}
+
+			if (!GraphObj || !EditorPtr->IsPlaying() || !bAnimateConnections)
+			{
+				FConnectionDrawingPolicy::DrawConnection(LayerId, Start, End, Params);
+				return;
+			}
+
+			const FDrawConnectionData ConnectionData(Start, End, FGraphConnectionDrawingPolicy::ComputeSplineTangent(Start, End), Params, *Settings, LocalMousePosition);
+			if (Settings->bTreatSplinesLikePins)
+			{
+#if METASOUND_DRAW_CONNECTION_DEBUG_DATA
+				DrawingPolicyPrivate::DrawDebugConnection(DrawElementsList, ArrowLayerID, ConnectionData);
+#endif // METASOUND_DRAW_CONNECTION_DEBUG_DATA
+
+				ConnectionData.UpdateSplineOverlap(SplineOverlapResult);
 			}
 
-			bool bExecuted = false;
+			if (ConnectionData.OutputHandle->IsValid() && MetasoundEditor.IsValid())
+			{
+				if (ConnectionData.OutputHandle->GetDataType() == GetMetasoundDataTypeName<float>())
+				{
+					DrawingPolicyPrivate::DrawConnectionSpline_Numeric<float>(*this, LayerId, ConnectionData, 0.0f);
+				}
+				else if (ConnectionData.OutputHandle->GetDataType() == GetMetasoundDataTypeName<int32>())
+				{
+					DrawingPolicyPrivate::DrawConnectionSpline_Numeric<int32>(*this, LayerId, ConnectionData, 0);
+				}
+				else if (ConnectionData.OutputHandle->GetDataType() == GetMetasoundDataTypeName<bool>())
+				{
+					DrawingPolicyPrivate::DrawConnectionSpline_Numeric<bool>(*this, LayerId, ConnectionData, 0);
+				}
+				else
+				{
+					DrawConnectionSpline(LayerId, ConnectionData);
+				}
+			}
+			else
+			{
+				DrawConnectionSpline(LayerId, ConnectionData);
+			}
 
-			// Run through the predecessors, and on
-			if (FExecPairingMap* PredecessorMap = PredecessorNodes.Find(OutputPin->GetOwningNode()))
+			if (Params.bDrawBubbles || MidpointImage)
 			{
-				if (FTimePair* Times = PredecessorMap->Find(InputPin->GetOwningNode()))
+				// Maps distance along curve to alpha
+				FInterpCurve<float> SplineReparamTable;
+				const float SplineLength = MakeSplineReparamTable(ConnectionData.P0, ConnectionData.P0Tangent, ConnectionData.P1, ConnectionData.P1Tangent, SplineReparamTable);
+
+				// Draw bubbles on the spline
+				if (Params.bDrawBubbles)
 				{
-					bExecuted = true;
+					const float AppTime = FPlatformTime::Seconds() - GStartTime;
+					FDrawConnectionSignalData SignalParams = { AppTime, LayerId, SplineReparamTable, SplineLength, ConnectionData };
+					if (SignalParams.ConnectionData.OutputHandle->GetDataType() == GetMetasoundDataTypeName<FAudioBuffer>())
+					{
+						SignalParams.SpacingFactor = 2.f;
+						SignalParams.SpeedFactor = 256.f;
+						const FName AnalyzerName = Frontend::FVertexAnalyzerEnvelopeFollower::GetAnalyzerName();
+						DrawConnectionSignal_Envelope(SignalParams, AnalyzerName);
+					}
+					else if (SignalParams.ConnectionData.OutputHandle->GetDataType() == GetMetasoundDataTypeName<FTrigger>())
+					{
+						SignalParams.SpacingFactor = 2.f;
+						SignalParams.SpeedFactor = 64.f;
+						const FName AnalyzerName = Frontend::FVertexAnalyzerTriggerDensity::GetAnalyzerName();
+						DrawConnectionSignal_Envelope(SignalParams, AnalyzerName);
+					}
+					else
+					{
+						DrawConnectionSignal(SignalParams);
+					}
+				}
 
-					OutParams.WireThickness = ActiveWireThickness;
-					OutParams.bDrawBubbles = true;
+				if (MidpointImage)
+				{
+					DrawConnection_MidpointImage(LayerId, SplineReparamTable, SplineLength, ConnectionData);
 				}
 			}
+		}
+
+		void FGraphConnectionDrawingPolicy::DrawConnection_MidpointImage(int32 InLayerId, const FInterpCurve<float>& InSplineReparamTable, float InSplineLength, const FDrawConnectionData& InData)
+		{
+			check(MidpointImage);
 
-			if (!bExecuted)
+			// Determine the spline position for the midpoint
+			const float MidpointAlpha = InSplineReparamTable.Eval(InSplineLength * 0.5f, 0.f);
+			const FVector2D Midpoint = FMath::CubicInterp(InData.P0, InData.P0Tangent, InData.P1, InData.P1Tangent, MidpointAlpha);
+
+			// Approximate the slope at the midpoint (to orient the midpoint image to the spline)
+			const FVector2D MidpointPlusE = FMath::CubicInterp(InData.P0, InData.P0Tangent, InData.P1, InData.P1Tangent, MidpointAlpha + KINDA_SMALL_NUMBER);
+			const FVector2D MidpointMinusE = FMath::CubicInterp(InData.P0, InData.P0Tangent, InData.P1, InData.P1Tangent, MidpointAlpha - KINDA_SMALL_NUMBER);
+			const FVector2D SlopeUnnormalized = MidpointPlusE - MidpointMinusE;
+
+			// Draw the arrow
+			const FVector2D MidpointDrawPos = Midpoint - MidpointRadius;
+			const float AngleInRadians = SlopeUnnormalized.IsNearlyZero() ? 0.0f : FMath::Atan2(SlopeUnnormalized.Y, SlopeUnnormalized.X);
+
+			FSlateDrawElement::MakeRotatedBox(
+				DrawElementsList,
+				InLayerId,
+				FPaintGeometry(MidpointDrawPos, MidpointImage->ImageSize * ZoomFactor, ZoomFactor),
+				MidpointImage,
+				ESlateDrawEffect::None,
+				AngleInRadians,
+				TOptional<FVector2D>(),
+				FSlateDrawElement::RelativeToElement,
+				InData.Params.WireColor
+			);
+		}
+
+		void FGraphConnectionDrawingPolicy::DrawConnectionSpline(int32 InLayerId, const FDrawConnectionData& InData)
+		{
+			FSlateDrawElement::MakeDrawSpaceSpline(
+				DrawElementsList,
+				InLayerId,
+				InData.P0, InData.P0Tangent,
+				InData.P1, InData.P1Tangent,
+				InData.Params.WireThickness,
+				ESlateDrawEffect::None,
+				InData.Params.WireColor
+			);
+		}
+
+		void FGraphConnectionDrawingPolicy::DrawConnectionSignal(const FDrawConnectionSignalData& InParams)
+		{
+			const float BubbleSpacing = InParams.SpacingFactor * ZoomFactor;
+			float BubbleSpeed = InParams.SpeedFactor * ZoomFactor;
+
+			const FVector2D BubbleSize = BubbleImage->ImageSize * ZoomFactor * 0.2f * InParams.ConnectionData.Params.WireThickness;
+
+			const float BubbleOffset = FMath::Fmod(InParams.AppTime * BubbleSpeed, BubbleSpacing);
+			const int32 NumBubbles = FMath::CeilToInt(InParams.SplineLength / BubbleSpacing);
+
+			for (int32 i = 0; i < NumBubbles; ++i)
 			{
-				if (InputPin->PinType.PinCategory == FGraphBuilder::PinCategoryTrigger)
+				const float Distance = (i * BubbleSpacing) + BubbleOffset;
+				if (Distance < InParams.SplineLength)
 				{
-					OutParams.WireThickness = Settings->DefaultExecutionWireThickness;
+					const float Alpha = InParams.SplineReparamTable.Eval(Distance, 0.f);
+					FVector2D BubblePos = FMath::CubicInterp(InParams.ConnectionData.P0, InParams.ConnectionData.P0Tangent, InParams.ConnectionData.P1, InParams.ConnectionData.P1Tangent, Alpha);
+					BubblePos -= BubbleSize * 0.5f;
+
+					FSlateDrawElement::MakeBox(
+						DrawElementsList,
+						InParams.LayerId,
+						FPaintGeometry(BubblePos, BubbleSize, ZoomFactor),
+						BubbleImage,
+						ESlateDrawEffect::None,
+						InParams.ConnectionData.Params.WireColor
+					);
 				}
-				else if (InputPin->PinType.PinCategory == FGraphBuilder::PinCategoryAudio)
+			}
+		}
+
+		void FGraphConnectionDrawingPolicy::DrawConnectionSignal_Envelope(const FDrawConnectionSignalData& InParams, FName InAnalyzerName)
+		{
+			using namespace Frontend;
+
+			TSharedPtr<FEditor> Editor = GetEditor();
+			if (!Editor)
+			{
+				return;
+			}
+
+			const float BubbleSpacing = InParams.SpacingFactor * ZoomFactor;
+			float BubbleSpeed = InParams.SpeedFactor * ZoomFactor;
+			const FVector2D BubbleSize = BubbleImage->ImageSize * ZoomFactor * 0.2f * InParams.ConnectionData.Params.WireThickness;
+			const float BubbleOffset = FMath::Fmod(InParams.AppTime * BubbleSpeed, BubbleSpacing);
+			const int32 NumBubbles = FMath::CeilToInt(InParams.SplineLength / BubbleSpacing);
+
+			TArray<FVector2D> ScaledBubbles;
+			ScaledBubbles.Init(BubbleSize, NumBubbles);
+			if (const UMetasoundEditorGraphNode* Node = Cast<UMetasoundEditorGraphNode>(InParams.ConnectionData.OutputPin->GetOwningNode()))
+			{
+				if (const UEdGraph* Graph = Node->GetGraph())
 				{
-					bool bIsActive = false;
-					if (UEdGraphNode* Node = InputPin->GetOwningNode())
+					if (MetasoundEditor.IsValid())
 					{
-						if (UEdGraph* Graph = Node->GetGraph())
+						FConstOutputHandle OutputHandle = FGraphBuilder::GetConstOutputHandleFromPin(InParams.ConnectionData.OutputPin);
+						const FGuid NodeId = OutputHandle->GetOwningNodeID();
+						FName OutputName = OutputHandle->GetName();
+
+						TArray<float> Values;
 						{
-							TSharedPtr<FEditor> MetasoundEditor = FGraphBuilder::GetEditorForGraph(*Graph);
-							if (MetasoundEditor.IsValid())
+							// Only set to max as multiple edges can exist and each may have a different window size.
+							FGraphConnectionManager& ConnectionManager = Editor->GetConnectionManager();
+							if (const FFloatMovingWindow* Window = ConnectionManager.GetValueWindow(NodeId, OutputName, InAnalyzerName))
+							{
+								if (NumBubbles > Window->Num())
+								{
+									ConnectionManager.TrackValue(NodeId, OutputName, InAnalyzerName, NumBubbles);
+								}
+								Values = Window->ToArray();
+								Values.SetNum(NumBubbles);
+								for (int32 i = 0; i < Values.Num(); ++i)
+								{
+									const float ClampedScalar = FMath::Clamp(Values[i], 0.0f, 1.0f) * ActiveAnalyzerWireSignalScalarMax;
+									ScaledBubbles[i].X *= ClampedScalar;
+									ScaledBubbles[i].Y *= ClampedScalar;
+								}
+							}
+							else
 							{
-								bIsActive = MetasoundEditor->GetPlayTime() > 0.0;
+								ConnectionManager.TrackValue(NodeId, OutputName, InAnalyzerName, NumBubbles);
 							}
 						}
 					}
-
-					OutParams.bDrawBubbles = bIsActive;
 				}
-				else
+			}
+
+			for (int32 i = 0; i < ScaledBubbles.Num(); ++i)
+			{
+				const float Distance = (i * BubbleSpacing) + BubbleOffset;
+				if (Distance < InParams.SplineLength)
 				{
-					OutParams.WireThickness = InactiveWireThickness;
+					const float Alpha = InParams.SplineReparamTable.Eval(Distance, 0.f);
+					FVector2D BubblePos = FMath::CubicInterp(InParams.ConnectionData.P0, InParams.ConnectionData.P0Tangent, InParams.ConnectionData.P1, InParams.ConnectionData.P1Tangent, Alpha);
+					BubblePos -= (ScaledBubbles[i] * 0.5f);
+
+					FSlateDrawElement::MakeBox(
+						DrawElementsList,
+						InParams.LayerId,
+						FPaintGeometry(BubblePos, ScaledBubbles[i], ZoomFactor),
+						BubbleImage,
+						ESlateDrawEffect::None,
+						InParams.ConnectionData.Params.WireColor
+					);
 				}
 			}
 		}
 
-		FGraphConnectionDrawingPolicy::FGraphConnectionDrawingPolicy(int32 InBackLayerID, int32 InFrontLayerID, float InZoomFactor, const FSlateRect& InClippingRect, FSlateWindowElementList& InDrawElements, UEdGraph* InGraphObj)
-			: FConnectionDrawingPolicy(InBackLayerID, InFrontLayerID, InZoomFactor, InClippingRect, InDrawElements)
-			, GraphObj(InGraphObj)
+		void FGraphConnectionDrawingPolicy::DrawSpline(const FDrawConnectionData& InData)
 		{
-			ActiveWireThickness = Settings->TraceAttackWireThickness;
-			InactiveWireThickness = Settings->TraceReleaseWireThickness;
+			FSlateDrawElement::MakeDrawSpaceSpline(
+				DrawElementsList,
+				WireLayerID,
+				InData.P0, InData.P0Tangent,
+				InData.P1, InData.P1Tangent,
+				InData.Params.WireThickness,
+				ESlateDrawEffect::None,
+				InData.Params.WireColor
+			);
+		}
+
+		void FGraphConnectionDrawingPolicy::DrawAnimatedSpline(const FDrawConnectionData& InData)
+		{
+			FSlateDrawElement::MakeDrawSpaceSpline(
+				DrawElementsList,
+				WireAnimationLayerID,
+				InData.P0, InData.P0Tangent,
+				InData.P1, InData.P1Tangent,
+				InData.Params.WireThickness,
+				ESlateDrawEffect::None,
+				InData.Params.WireColor
+			);
 		}
 	} // namespace Editor
 } // namespace Metasound
diff --git a/Engine/Plugins/Runtime/Metasound/Source/MetasoundEditor/Private/MetasoundEditorGraphConnectionDrawingPolicy.h b/Engine/Plugins/Runtime/Metasound/Source/MetasoundEditor/Private/MetasoundEditorGraphConnectionDrawingPolicy.h
--- a/Engine/Plugins/Runtime/Metasound/Source/MetasoundEditor/Private/MetasoundEditorGraphConnectionDrawingPolicy.h
+++ b/Engine/Plugins/Runtime/Metasound/Source/MetasoundEditor/Private/MetasoundEditorGraphConnectionDrawingPolicy.h
@@ -1,11 +1,16 @@
 // Copyright Epic Games, Inc. All Rights Reserved.
 #pragma once
 
-#include "CoreMinimal.h"
+#include "ConnectionDrawingPolicy.h"
+#include "EdGraph/EdGraphPin.h"
+#include "EdGraphUtilities.h"
 #include "Layout/ArrangedWidget.h"
+#include "MetasoundEditorGraphBuilder.h"
+#include "MetasoundFrontendController.h"
+#include "Misc/CoreMiscDefines.h"
+#include "Rendering/DrawElements.h"
+#include "Templates/Function.h"
 #include "Widgets/SWidget.h"
-#include "EdGraphUtilities.h"
-#include "ConnectionDrawingPolicy.h"
 
 
 // Forward Declarations
@@ -14,10 +19,48 @@ class FSlateWindowElementList;
 class UEdGraph;
 class UEdGraphSchema;
 
+
 namespace Metasound
 {
 	namespace Editor
 	{
+		// Forward Declarations
+		class FEditor;
+		class FGraphConnectionManager;
+
+		struct FDrawConnectionData
+		{
+			const FVector2D P0;
+			const FVector2D P0Tangent;
+			const FVector2D P1;
+			const FVector2D P1Tangent;
+
+			FConnectionParams Params;
+
+			FBox2D Bounds;
+
+			const UEdGraphPin* OutputPin = nullptr;
+			Frontend::FConstOutputHandle OutputHandle;
+
+			bool bCloseToSpline = false;
+			float ClosestDistanceSquared = TNumericLimits<float>::Max();
+
+			// Distance to consider as an overlap
+			float QueryDistanceTriggerThresholdSquared = 0.0f;
+
+			// Distance to pass the bounding box cull test. This is used for the bCloseToSpline output that can be used as a
+			// dead zone to avoid mistakes caused by missing a double-click on a connection.
+			float QueryDistanceForCloseSquared = 0.0f;
+
+			FVector2D ClosestPoint;
+
+			const FMetasoundFrontendEdgeStyle* EdgeStyle = nullptr;
+
+			FDrawConnectionData(const FVector2D& InStart, const FVector2D& InEnd, const FVector2D& InSplineTangent, const FConnectionParams& InParams, const UGraphEditorSettings& InSettings, const FVector2D& InMousePosition);
+
+			void UpdateSplineOverlap(FGraphSplineOverlapResult& OutResult) const;
+		};
+
 		struct FGraphConnectionDrawingPolicyFactory : public FGraphPanelPinConnectionFactory
 		{
 		public:
@@ -35,7 +78,7 @@ namespace Metasound
 			// ~FGraphPanelPinConnectionFactory
 		};
 
-		// This class draws the connections for an UEdGraph using a SoundCue schema
+		// Draws the connections for a UMetasoundEditorGraph using a UMetasoundEditorGraphSchema
 		class FGraphConnectionDrawingPolicy : public FConnectionDrawingPolicy
 		{
 		protected:
@@ -57,17 +100,52 @@ namespace Metasound
 			FLinearColor ActiveColor;
 			FLinearColor InactiveColor;
 
-			float ActiveWireThickness;
-			float InactiveWireThickness;
+			int32 WireAnimationLayerID = 0;
+
+			bool bAnimateConnections = false;
+			float ActiveWireThickness = 1.0f;
+			float ActiveAnalyzerWireThickness = 1.0f;
+			float ActiveAnalyzerWireSignalScalarMin = 1.0f;
+			float ActiveAnalyzerWireSignalScalarMax = 1.0f;
 
-			void Draw(TMap<TSharedRef<SWidget>, FArrangedWidget>& InPinGeometries, FArrangedChildren& ArrangedNodes);
+			float InactiveWireThickness = 1.0f;
+
+			TWeakPtr<FEditor> MetasoundEditor;
 
 		public:
 			FGraphConnectionDrawingPolicy(int32 InBackLayerID, int32 InFrontLayerID, float ZoomFactor, const FSlateRect& InClippingRect, FSlateWindowElementList& InDrawElements, UEdGraph* InGraphObj);
 
 			// FConnectionDrawingPolicy interface
-			virtual void DetermineWiringStyle(UEdGraphPin* OutputPin, UEdGraphPin* InputPin, /*inout*/ FConnectionParams& Params) override;
+			virtual void DrawConnection(int32 LayerId, const FVector2D& Start, const FVector2D& End, const FConnectionParams& InParams) override;
+			virtual void DetermineWiringStyle(UEdGraphPin* InPinA, UEdGraphPin* InPinB, FConnectionParams& OutParams) override;
 			// End of FConnectionDrawingPolicy interface
+
+			virtual void DrawAnimatedSpline(const FDrawConnectionData& InData);
+			virtual void DrawSpline(const FDrawConnectionData& InData);
+			virtual void DrawConnectionSpline(int32 InLayerId, const FDrawConnectionData& InData);
+
+			const TSharedPtr<const FEditor> GetEditor() const;
+			TSharedPtr<FEditor> GetEditor();
+
+		protected:
+
+			struct FDrawConnectionSignalData
+			{
+				const float AppTime = 0.0f;
+				const int32 LayerId = 0;
+				const FInterpCurve<float>& SplineReparamTable;
+				const float SplineLength = 0.0f;
+				const FDrawConnectionData& ConnectionData;
+
+				float SpacingFactor = 64.0f;
+				float SpeedFactor = 192.0f;
+			};
+
+			void ScaleEnvelopeWireThickness(FConnectionParams& OutParams) const;
+
+			virtual void DrawConnectionSignal(const FDrawConnectionSignalData& InParams);
+			virtual void DrawConnectionSignal_Envelope(const FDrawConnectionSignalData& InParams, FName InAnalyzerName);
+			virtual void DrawConnection_MidpointImage(int32 InLayerId, const FInterpCurve<float>& SplineReparamTable, float InSplineLength, const FDrawConnectionData& InData);
 		};
 	} // namespace Editor
 } // namespace Metasound
diff --git a/Engine/Plugins/Runtime/Metasound/Source/MetasoundEditor/Private/MetasoundEditorGraphConnectionManager.cpp b/Engine/Plugins/Runtime/Metasound/Source/MetasoundEditor/Private/MetasoundEditorGraphConnectionManager.cpp
new file mode 100644
--- /dev/null
+++ b/Engine/Plugins/Runtime/Metasound/Source/MetasoundEditor/Private/MetasoundEditorGraphConnectionManager.cpp
@@ -0,0 +1,157 @@
+// Copyright Epic Games, Inc. All Rights Reserved.
+#include "MetasoundEditorGraphConnectionManager.h"
+
+#include "Analysis/MetasoundFrontendVertexAnalyzerEnvelopeFollower.h"
+#include "Analysis/MetasoundFrontendVertexAnalyzerForwardValue.h"
+#include "Analysis/MetasoundFrontendVertexAnalyzerTriggerDensity.h"
+#include "MetasoundEditorGraphBuilder.h"
+#include "MetasoundVertex.h"
+
+
+namespace Metasound
+{
+	namespace Editor
+	{
+		FGraphConnectionManager::FGraphConnectionManager(const FMetasoundAssetBase& InAssetBase, const UAudioComponent& InAudioComponent, FSampleRate InSampleRate)
+			: AudioComponent(&InAudioComponent)
+			, GraphAnalyzerView({ InAssetBase, InAudioComponent.GetInstanceOwnerID(), InSampleRate })
+		{
+			using namespace Frontend;
+
+			GraphAnalyzerView.AddAnalyzerForAllSupportedOutputs(FVertexAnalyzerForwardBool::GetAnalyzerName());
+			GraphAnalyzerView.AddAnalyzerForAllSupportedOutputs(FVertexAnalyzerForwardFloat::GetAnalyzerName());
+			GraphAnalyzerView.AddAnalyzerForAllSupportedOutputs(FVertexAnalyzerForwardInt::GetAnalyzerName());
+			GraphAnalyzerView.AddAnalyzerForAllSupportedOutputs(FVertexAnalyzerForwardString::GetAnalyzerName());
+
+			GraphAnalyzerView.AddAnalyzerForAllSupportedOutputs(FVertexAnalyzerEnvelopeFollower::GetAnalyzerName());
+			GraphAnalyzerView.AddAnalyzerForAllSupportedOutputs(FVertexAnalyzerTriggerDensity::GetAnalyzerName());
+		}
+
+		bool FGraphConnectionManager::GetValue(const FGuid& InNodeID, FVertexName InOutputName, float& OutValue) const
+		{
+			const FString AnalyzerKey = GetAnalyzerKey<Frontend::FVertexAnalyzerForwardFloat>(InNodeID, InOutputName);
+			if (!AnalyzerKey.IsEmpty())
+			{
+				if (const Audio::FVolumeFader* Fader = ConnectionFaders.Find(AnalyzerKey))
+				{
+					OutValue = Fader->GetVolume();
+					return true;
+				}
+			}
+
+			return false;
+		}
+
+		bool FGraphConnectionManager::GetValue(const FGuid& InNodeID, FVertexName InOutputName, bool& OutValue) const
+		{
+			return GetPrimativeValue<Frontend::FVertexAnalyzerForwardBool, bool>(InNodeID, InOutputName, BoolConnectionValues, OutValue);
+		}
+
+		bool FGraphConnectionManager::GetValue(const FGuid& InNodeID, FVertexName InOutputName, int32& OutValue) const
+		{
+			return GetPrimativeValue<Frontend::FVertexAnalyzerForwardInt, int32>(InNodeID, InOutputName, IntConnectionValues, OutValue);
+		}
+
+		bool FGraphConnectionManager::GetValue(const FGuid& InNodeID, FVertexName InOutputName, FString& OutValue) const
+		{
+			return GetPrimativeValue<Frontend::FVertexAnalyzerForwardString, FString>(InNodeID, InOutputName, StringConnectionValues, OutValue);
+		}
+
+		const FFloatMovingWindow* FGraphConnectionManager::GetValueWindow(const FGuid& InNodeID, FVertexName InOutputName, FName InAnalyzerName) const
+		{
+			return WindowedValues.Find({ InNodeID, InOutputName, InAnalyzerName });
+		}
+
+		bool FGraphConnectionManager::IsTracked(const FGuid& InNodeID, FVertexName InOutputName, FName InAnalyzerName) const
+		{
+			return WindowedValues.Contains({ InNodeID, InOutputName, InAnalyzerName });
+		}
+
+		void FGraphConnectionManager::TrackValue(const FGuid& InNodeID, FVertexName InOutputName, FName InAnalyzerName, int32 InWindowSize)
+		{
+			FFloatMovingWindow& Window = WindowedValues.FindOrAdd({ InNodeID, InOutputName, InAnalyzerName });
+			Window.SetNumZeroed(InWindowSize);
+		}
+
+		void FGraphConnectionManager::Update(float InDeltaTime)
+		{
+			using namespace Frontend;
+
+			if (AudioComponent.IsValid())
+			{
+				float FadeDuration = 0.0001f;
+				auto UpdateSmoothedFloat = [this, FadeDur = &FadeDuration](const FAnalyzerAddress& AnalyzerAddress, const float& NewValue)
+				{
+					const FString AnalyzerKey = AnalyzerAddress.ToString();
+					Audio::FVolumeFader* Fader = ConnectionFaders.Find(AnalyzerKey);
+					if (!Fader)
+					{
+						Audio::FVolumeFader NewFader;
+						NewFader.SetVolume(0.0f);
+						Fader = &ConnectionFaders.Add(AnalyzerKey, MoveTemp(NewFader));
+					}
+
+					if (!FMath::IsNearlyEqual(NewValue, Fader->GetTargetVolume(), UE_KINDA_SMALL_NUMBER))
+					{
+						constexpr Audio::EFaderCurve Curve = Audio::EFaderCurve::Linear;
+
+						Fader->StartFade(NewValue, *FadeDur, Curve);
+					}
+				};
+
+				UpdateConnections<FVertexAnalyzerForwardFloat, float>(InDeltaTime, UpdateSmoothedFloat);
+				UpdateConnections<FVertexAnalyzerForwardInt, int32>(InDeltaTime, IntConnectionValues);
+				UpdateConnections<FVertexAnalyzerForwardBool, bool>(InDeltaTime, BoolConnectionValues);
+				UpdateConnections<FVertexAnalyzerForwardString, FString>(InDeltaTime, StringConnectionValues);
+
+				auto UpdateWindowedFloat = [this, &UpdateSmoothedFloat](const FAnalyzerAddress& AnalyzerAddress, const float& NewValue)
+				{
+					UpdateSmoothedFloat(AnalyzerAddress, NewValue);
+					WindowedValues.FindOrAdd({ AnalyzerAddress.NodeID, AnalyzerAddress.OutputName, AnalyzerAddress.AnalyzerName });
+				};
+				UpdateConnections<FVertexAnalyzerEnvelopeFollower, float>(InDeltaTime, UpdateWindowedFloat);
+
+				FadeDuration = 0.01f;
+				UpdateConnections<FVertexAnalyzerTriggerDensity, float>(InDeltaTime, UpdateWindowedFloat);
+
+				for (TPair<FString, Audio::FVolumeFader>& Pair : ConnectionFaders)
+				{
+					Pair.Value.Update(InDeltaTime);
+				}
+			}
+		}
+
+		float FGraphConnectionManager::UpdateValueWindow(const FGuid& InNodeID, FVertexName InOutputName, FName InAnalyzerName, FName InAnalyzerOutputName)
+		{
+			using namespace Frontend;
+
+			float CurrentValue = 0.0f;
+
+			TArray<FMetasoundAnalyzerView*> Views = GraphAnalyzerView.GetAnalyzerViews(InAnalyzerName);
+			for (const FMetasoundAnalyzerView* View : Views)
+			{
+				check(View);
+				const FAnalyzerAddress& Address = View->AnalyzerAddress;
+				const bool bOutputsMatch = Address.AnalyzerMemberName == InAnalyzerOutputName
+					&& Address.OutputName == InOutputName
+					&& Address.NodeID == InNodeID;
+				if (bOutputsMatch)
+				{
+					const FString AnalyzerKey = Address.ToString();
+					if (Audio::FVolumeFader* Fader = ConnectionFaders.Find(AnalyzerKey))
+					{
+						CurrentValue = Fader->GetVolume();
+						FFloatMovingWindow* Window = WindowedValues.Find({ Address.NodeID, Address.OutputName, Address.AnalyzerName });
+						if (Window && !Window->IsEmpty())
+						{
+							Window->SetValue(CurrentValue);
+							return CurrentValue;
+						}
+					}
+				}
+			}
+
+			return CurrentValue;
+		}
+	} // namespace Editor
+} // Metasound
diff --git a/Engine/Plugins/Runtime/Metasound/Source/MetasoundEditor/Private/MetasoundEditorGraphConnectionManager.h b/Engine/Plugins/Runtime/Metasound/Source/MetasoundEditor/Private/MetasoundEditorGraphConnectionManager.h
new file mode 100644
--- /dev/null
+++ b/Engine/Plugins/Runtime/Metasound/Source/MetasoundEditor/Private/MetasoundEditorGraphConnectionManager.h
@@ -0,0 +1,241 @@
+// Copyright Epic Games, Inc. All Rights Reserved.
+#pragma once
+
+#include "Analysis/MetasoundFrontendAnalyzerAddress.h"
+#include "Analysis/MetasoundFrontendGraphAnalyzerView.h"
+#include "Analysis/MetasoundFrontendVertexAnalyzerForwardValue.h"
+#include "Components/AudioComponent.h"
+#include "DSP/Dsp.h"
+#include "DSP/VolumeFader.h"
+#include "HAL/Platform.h"
+#include "MetasoundAssetBase.h"
+#include "MetasoundOperatorSettings.h"
+#include "Templates/Function.h"
+#include "Templates/TypeHash.h"
+#include "UObject/WeakObjectPtrTemplates.h"
+
+
+// Forward Declarations
+class UEdGraphPin;
+
+namespace Metasound
+{
+	namespace Editor
+	{
+		class FFloatMovingWindow
+		{
+			TArray<float> Data;
+			int32 WriteIndex = -1;
+
+		public:
+			bool IsEmpty() const
+			{
+				return Data.IsEmpty();
+			}
+
+			int32 Num() const
+			{
+				return Data.Num();
+			}
+
+			void SetNumZeroed(int32 InNum)
+			{
+				check(InNum >= 0);
+
+				constexpr bool bAllowShrinking = false;
+				Data.SetNumZeroed(InNum, bAllowShrinking);
+
+				// Initialize to back of array to avoid initial
+				// value being behind when copying to initial array.
+				if (WriteIndex < 0)
+				{
+					WriteIndex = Data.Num() - 1;
+				}
+				else
+				{
+					WriteIndex %= Data.Num();
+				}
+			}
+
+			void SetValue(float InValue)
+			{
+				check(!Data.IsEmpty());
+				Data[WriteIndex] = InValue;
+				WriteIndex--;
+				if (WriteIndex < 0)
+				{
+					WriteIndex += Data.Num();
+				}
+			}
+
+			TArray<float> ToArray() const
+			{
+				TArray<float> Values;
+
+				if (!Data.IsEmpty())
+				{
+					check(WriteIndex >= 0);
+					check(WriteIndex < Data.Num());
+
+					const int32 ReadIndex = (WriteIndex + 1) % Data.Num();
+					if (ReadIndex > WriteIndex)
+					{
+						Values.SetNumUninitialized(Data.Num());
+						const int32 Offset = Data.Num() - ReadIndex;
+						FMemory::Memcpy(Values.GetData(), Data.GetData() + ReadIndex, sizeof(float) * Offset);
+						FMemory::Memcpy(Values.GetData() + Offset, Data.GetData(), sizeof(float) * ReadIndex);
+					}
+					else
+					{
+						Values = Data;
+					}
+				}
+
+				return Values;
+			}
+		};
+
+		class FGraphConnectionManager
+		{
+		public:
+			FGraphConnectionManager() = default;
+			FGraphConnectionManager(const FMetasoundAssetBase& InAssetBase, const UAudioComponent& InAudioComponent, FSampleRate InSampleRate);
+
+			bool GetValue(const FGuid& InNodeID, FVertexName InOutputName, float& OutValue) const;
+			bool GetValue(const FGuid& InNodeID, FVertexName InOutputName, bool& OutValue) const;
+			bool GetValue(const FGuid& InNodeID, FVertexName InOutputName, int32& OutValue) const;
+			bool GetValue(const FGuid& InNodeID, FVertexName InOutputName, FString& OutValue) const;
+
+			const FFloatMovingWindow* GetValueWindow(const FGuid& InNodeID, FVertexName InOutputName, FName InAnalyzerName) const;
+
+			bool IsTracked(const FGuid& InNodeID, FVertexName InOutputName, FName InAnalyzerName) const;
+
+			// Marks to track (if not yet tracked) & sets the window size of cached values.
+			// Retains existing values provided from update up to given size.
+			void TrackValue(const FGuid& InNodeID, FVertexName InOutputName, FName InAnalyzerName, int32 InWindowSize);
+
+			void Update(float InDeltaTime);
+
+			template <typename TAnalyzerName>
+			float UpdateValueWindow(const FGuid& InNodeID, FVertexName InOutputName)
+			{
+				using namespace Frontend;
+				return UpdateValueWindow(InNodeID, InOutputName, TAnalyzerName::GetAnalyzerName(), TAnalyzerName::FOutputs::GetValue().Name);
+			}
+
+		private:
+			template <typename ForwardValueAnalyzerClass>
+			FString GetAnalyzerKey(const FGuid& InNodeID, FVertexName InOutputName) const
+			{
+				using namespace Frontend;
+
+				TArray<const FMetasoundAnalyzerView*> Views = GraphAnalyzerView.GetAnalyzerViewsForOutput(InNodeID, InOutputName, ForwardValueAnalyzerClass::GetAnalyzerName());
+				if (!Views.IsEmpty())
+				{
+					const FMetasoundAnalyzerView** View = Views.FindByPredicate([](const FMetasoundAnalyzerView* Candidate)
+						{
+							check(Candidate);
+							return Candidate->AnalyzerAddress.AnalyzerMemberName == ForwardValueAnalyzerClass::FOutputs::GetValue().Name;
+						});
+
+					if (View)
+					{
+						check(*View);
+						const FAnalyzerAddress& Address = (*View)->AnalyzerAddress;
+						const FString AnalyzerKey = Address.ToString();
+						return AnalyzerKey;
+					}
+				}
+
+				return { };
+			}
+
+			template<typename TVertexAnalyzerValue, typename TDataType>
+			bool GetPrimativeValue(const FGuid& InNodeID, FVertexName InOutputName, const TMap<FString, TDataType>& InDataMap, TDataType& OutValue) const
+			{
+				const FString AnalyzerKey = GetAnalyzerKey<TVertexAnalyzerValue>(InNodeID, InOutputName);
+				if (!AnalyzerKey.IsEmpty())
+				{
+					if (const TDataType* Value = InDataMap.Find(AnalyzerKey))
+					{
+						OutValue = *Value;
+						return true;
+					}
+				}
+
+				return false;
+			}
+
+			// Updates the value window, returning the most recent value
+			float UpdateValueWindow(const FGuid& InNodeID, FVertexName InOutputName, FName InAnalyzerName, FName InAnalyzerOutputName);
+
+			template <typename ForwardValueAnalyzerClass, typename AnalyzerOutputDataType>
+			void UpdateConnections(float InDeltaTime, TFunctionRef<void(const Frontend::FAnalyzerAddress& /* AnalyzerAddress */, const AnalyzerOutputDataType& /* NewValue */)> InUpdateFunction)
+			{
+				using namespace Frontend;
+
+				TArray<FMetasoundAnalyzerView*> Views = GraphAnalyzerView.GetAnalyzerViews(ForwardValueAnalyzerClass::GetAnalyzerName());
+				for (FMetasoundAnalyzerView* View : Views)
+				{
+					check(View);
+					const FAnalyzerAddress& Address = View->AnalyzerAddress;
+					const bool bOutputsMatch = Address.AnalyzerMemberName == ForwardValueAnalyzerClass::FOutputs::GetValue().Name;
+					if (bOutputsMatch)
+					{
+						AnalyzerOutputDataType Value;
+						if (View->TryGetOutputData<AnalyzerOutputDataType>(Address.AnalyzerMemberName, Value))
+						{
+							InUpdateFunction(Address, Value);
+						}
+					}
+				}
+			}
+
+			template <typename ForwardValueAnalyzerClass, typename AnalyzerOutputDataType>
+			void UpdateConnections(float InDeltaTime, TMap<FString, AnalyzerOutputDataType>& OutMap)
+			{
+				UpdateConnections<ForwardValueAnalyzerClass, AnalyzerOutputDataType>(InDeltaTime, [this, &OutMap](const Frontend::FAnalyzerAddress& AnalyzerAddress, const AnalyzerOutputDataType& NewValue)
+				{
+					FString AnalyzerKey = AnalyzerAddress.ToString();
+					OutMap.FindOrAdd(AnalyzerKey) = NewValue;
+				});
+			}
+
+			struct FWindowValueKey
+			{
+				const FGuid NodeID;
+				const FName OutputName;
+				const FName AnalyzerName;
+
+				friend FORCEINLINE uint32 GetTypeHash(const FWindowValueKey& Key)
+				{
+					uint32 Hash = GetTypeHash(Key.NodeID);
+					Hash = HashCombineFast(Hash, GetTypeHash(Key.OutputName));
+					Hash = HashCombineFast(Hash, GetTypeHash(Key.AnalyzerName));
+					return Hash;
+				}
+
+				FORCEINLINE bool operator==(const FWindowValueKey& Other) const
+				{
+					return (NodeID == Other.NodeID) && (OutputName == Other.OutputName) && (AnalyzerName == Other.AnalyzerName);
+				}
+
+				FORCEINLINE bool operator!=(const FWindowValueKey& Other) const
+				{
+					return !(*this == Other);
+				}
+			};
+
+			TMap<FWindowValueKey, FFloatMovingWindow> WindowedValues;
+
+			TWeakObjectPtr<const UAudioComponent> AudioComponent;
+			TMap<FString, Audio::FVolumeFader> ConnectionFaders;
+
+			TMap<FString, int32> IntConnectionValues;
+			TMap<FString, bool> BoolConnectionValues;
+			TMap<FString, FString> StringConnectionValues;
+
+			Frontend::FMetasoundGraphAnalyzerView GraphAnalyzerView;
+		};
+	} // namespace Editor
+} // Metasound
diff --git a/Engine/Plugins/Runtime/Metasound/Source/MetasoundEditor/Private/MetasoundEditorModule.cpp b/Engine/Plugins/Runtime/Metasound/Source/MetasoundEditor/Private/MetasoundEditorModule.cpp
--- a/Engine/Plugins/Runtime/Metasound/Source/MetasoundEditor/Private/MetasoundEditorModule.cpp
+++ b/Engine/Plugins/Runtime/Metasound/Source/MetasoundEditor/Private/MetasoundEditorModule.cpp
@@ -67,25 +67,6 @@ namespace Metasound
 	namespace Editor
 	{
 		static const FName AssetToolName { "AssetTools" };
-		
-		const FSlateBrush* FEditorDataType::GetIconBrush(const bool bIsConstructorType) const
-		{
-			if (const ISlateStyle* MetasoundStyle = FSlateStyleRegistry::FindSlateStyle("MetaSoundStyle"))
-			{
-				if (PinType.IsArray())
-				{
-					return bIsConstructorType ? MetasoundStyle->GetBrush(TEXT("MetasoundEditor.Graph.ConstructorPinArray")) : MetasoundStyle->GetBrush(TEXT("MetasoundEditor.Graph.ArrayPin"));
-				}
-				else
-				{
-					return bIsConstructorType ? MetasoundStyle->GetBrush(TEXT("MetasoundEditor.Graph.ConstructorPin")) : FAppStyle::GetBrush("Icons.BulletPoint");
-				}
-			}
-			else
-			{
-				return PinType.IsArray() ? FAppStyle::GetBrush("Graph.ArrayPin.Connected") : FAppStyle::GetBrush("Icons.BulletPoint");
-			}
-		}
 
 		template <typename T>
 		void AddAssetAction(IAssetTools& AssetTools, TArray<TSharedPtr<FAssetTypeActions_Base>>& AssetArray)
@@ -429,7 +410,7 @@ namespace Metasound
 				}
 			}
 
-			void RegisterCoreDataTypes()
+			void RegisterCorePinTypes()
 			{
 				using namespace Metasound::Frontend;
 
@@ -530,17 +511,30 @@ namespace Metasound
 							}
 						}
 
-						const bool bIsArray = RegistryInfo.IsArrayType();
-						const EPinContainerType ContainerType = bIsArray ? EPinContainerType::Array : EPinContainerType::None;
-						FEdGraphPinType PinType(PinCategory, PinSubCategory, nullptr, ContainerType, false, FEdGraphTerminalType());
-						UClass* ClassToUse = DataTypeRegistry.GetUClassForDataType(DataTypeName);
-						PinType.PinSubCategoryObject = Cast<UObject>(ClassToUse);
-
-						DataTypeInfo.Emplace(DataTypeName, FEditorDataType(MoveTemp(PinType), MoveTemp(RegistryInfo)));
+						RegisterPinType(DataTypeName, PinCategory, PinSubCategory);
 					}
 				}
 			}
 
+			void RegisterPinType(FName InDataTypeName, FName InPinCategory, FName InPinSubCategory)
+			{
+				using namespace Frontend;
+
+				FDataTypeRegistryInfo DataTypeInfo;
+				IDataTypeRegistry::Get().GetDataTypeInfo(InDataTypeName, DataTypeInfo);
+
+				// Default to object as most calls to this outside of the MetaSound Editor will be for custom UObject types
+				const FName PinCategory = InPinCategory.IsNone() ? FGraphBuilder::PinCategoryObject : InPinCategory;
+
+				const bool bIsArray = DataTypeInfo.IsArrayType();
+				const EPinContainerType ContainerType = bIsArray ? EPinContainerType::Array : EPinContainerType::None;
+				FEdGraphPinType PinType(PinCategory, InPinSubCategory, nullptr, ContainerType, false, FEdGraphTerminalType());
+				UClass* ClassToUse = IDataTypeRegistry::Get().GetUClassForDataType(InDataTypeName);
+				PinType.PinSubCategoryObject = Cast<UObject>(ClassToUse);
+
+				PinTypes.Emplace(InDataTypeName, MoveTemp(PinType));
+			}
+
 			void ShutdownAssetClassRegistry()
 			{
 				if (FAssetRegistryModule* AssetRegistryModule = static_cast<FAssetRegistryModule*>(FModuleManager::Get().GetModule("AssetRegistry")))
@@ -617,27 +611,36 @@ namespace Metasound
 				return InputDefaultLiteralClassRegistry.FindRef(InLiteralType);
 			}
 
-			virtual const FEditorDataType* FindDataType(FName InDataTypeName) const override
+			virtual const FSlateBrush* GetIconBrush(FName InDataType, const bool bIsConstructorType) const override
 			{
-				return DataTypeInfo.Find(InDataTypeName);
-			}
+				bool bIsArrayType = false;
 
-			virtual const FEditorDataType& FindDataTypeChecked(FName InDataTypeName) const override
-			{
-				return DataTypeInfo.FindChecked(InDataTypeName);
-			}
+				Frontend::FDataTypeRegistryInfo Info;
+				if (Frontend::IDataTypeRegistry::Get().GetDataTypeInfo(InDataType, Info))
+				{
+					bIsArrayType = Info.IsArrayType();
+				}
 
-			virtual bool IsRegisteredDataType(FName InDataTypeName) const override
-			{
-				return DataTypeInfo.Contains(InDataTypeName);
+				if (const ISlateStyle* MetasoundStyle = FSlateStyleRegistry::FindSlateStyle("MetaSoundStyle"))
+				{
+					if (bIsArrayType)
+					{
+						return bIsConstructorType ? MetasoundStyle->GetBrush(TEXT("MetasoundEditor.Graph.ConstructorPinArray")) : MetasoundStyle->GetBrush(TEXT("MetasoundEditor.Graph.ArrayPin"));
+					}
+					else
+					{
+						return bIsConstructorType ? MetasoundStyle->GetBrush(TEXT("MetasoundEditor.Graph.ConstructorPin")) : FAppStyle::GetBrush("Icons.BulletPoint");
+					}
+				}
+				else
+				{
+					return bIsArrayType ? FAppStyle::GetBrush("Graph.ArrayPin.Connected") : FAppStyle::GetBrush("Icons.BulletPoint");
+				}
 			}
 
-			virtual void IterateDataTypes(TUniqueFunction<void(const FEditorDataType&)> InDataTypeFunction) const override
+			virtual const FEdGraphPinType* FindPinType(FName InDataTypeName) const
 			{
-				for (const TPair<FName, FEditorDataType>& Pair : DataTypeInfo)
-				{
-					InDataTypeFunction(Pair.Value);
-				}
+				return PinTypes.Find(InDataTypeName);
 			}
 
 			virtual bool IsMetaSoundAssetClass(const FName InClassName) const override
@@ -698,7 +701,7 @@ namespace Metasound
 
 				StyleSet = MakeShared<FSlateStyle>();
 
-				RegisterCoreDataTypes();
+				RegisterCorePinTypes();
 				RegisterInputDefaultClasses();
 
 				GraphConnectionFactory = MakeShared<FGraphConnectionDrawingPolicyFactory>();
@@ -787,15 +790,15 @@ namespace Metasound
 				ShutdownAssetClassRegistry();
 
 				AssetActions.Reset();
-				DataTypeInfo.Reset();
+				PinTypes.Reset();
 				MetaSoundClassNames.Reset();
 			}
 
 			TArray<FName> MetaSoundClassNames;
 
 			TArray<TSharedPtr<FAssetTypeActions_Base>> AssetActions;
-			TMap<FName, FEditorDataType> DataTypeInfo;
 			TMap<EMetasoundFrontendLiteralType, const TSubclassOf<UMetasoundEditorGraphMemberDefaultLiteral>> InputDefaultLiteralClassRegistry;
+			TMap<FName, FEdGraphPinType> PinTypes;
 
 			TMap<UClass*, TUniquePtr<IMemberDefaultLiteralCustomizationFactory>> LiteralCustomizationFactories;
 
diff --git a/Engine/Plugins/Runtime/Metasound/Source/MetasoundEditor/Private/MetasoundEditorSettings.cpp b/Engine/Plugins/Runtime/Metasound/Source/MetasoundEditor/Private/MetasoundEditorSettings.cpp
--- a/Engine/Plugins/Runtime/Metasound/Source/MetasoundEditor/Private/MetasoundEditorSettings.cpp
+++ b/Engine/Plugins/Runtime/Metasound/Source/MetasoundEditor/Private/MetasoundEditorSettings.cpp
@@ -26,5 +26,10 @@ UMetasoundEditorSettings::UMetasoundEditorSettings(const FObjectInitializer& Obj
 	AssetReferenceNodeTitleColor = FLinearColor(0.047f, 0.686f, 0.988f);		// sky blue
 	InputNodeTitleColor = FLinearColor(0.168f, 1.0f, 0.7294f);					// sea foam
 	OutputNodeTitleColor = FLinearColor(1.0f, 0.878f, 0.1686f);					// yellow
-	VariableNodeTitleColor = FLinearColor(0.211f, 0.513f, 0.035f);				// copper 
+	VariableNodeTitleColor = FLinearColor(0.211f, 0.513f, 0.035f);				// copper
+
+	bAnimateConnections = false;
+	ActiveAnalyzerWireThickness = 8.0f;
+	ActiveAnalyzerWireSignalScalarMin = 1.0f;
+	ActiveAnalyzerWireSignalScalarMax = 1.0f;
 }
diff --git a/Engine/Plugins/Runtime/Metasound/Source/MetasoundEditor/Private/MetasoundEditorSettings.h b/Engine/Plugins/Runtime/Metasound/Source/MetasoundEditor/Private/MetasoundEditorSettings.h
--- a/Engine/Plugins/Runtime/Metasound/Source/MetasoundEditor/Private/MetasoundEditorSettings.h
+++ b/Engine/Plugins/Runtime/Metasound/Source/MetasoundEditor/Private/MetasoundEditorSettings.h
@@ -100,6 +100,22 @@ public:
 	UPROPERTY(EditAnywhere, config, Category = NodeTitleColors)
 	FLinearColor VariableNodeTitleColor;
 
+	/** Whether or not animated connections are enabled. */
+	UPROPERTY(EditAnywhere, config, Category = GraphAnimation, meta = (DisplayName = "Animate Connections (Experimental)"))
+	bool bAnimateConnections;
+
+	/** Thickness of default analyzer wire thickness when connection analyzer is active. */
+	UPROPERTY(EditAnywhere, config, Category = GraphAnimation, meta = (EditCondition = "bAnimateConnections", UIMin = 1, UIMax = 8, ClampMin = 1))
+	float ActiveAnalyzerWireThickness;
+
+	/** Minimum height scalar of wire signal analyzers (ex. audio, triggers). */
+	UPROPERTY(EditAnywhere, config, Category = GraphAnimation, meta = (EditCondition = "bAnimateConnections", UIMin = 1, UIMax = 3, ClampMin = 1))
+	float ActiveAnalyzerWireSignalScalarMin;
+
+	/** Maximum height scalar of wire signal analyzers (ex. audio, triggers). */
+	UPROPERTY(EditAnywhere, config, Category = GraphAnimation, meta = (EditCondition = "bAnimateConnections", UIMin = 1, UIMax = 3, ClampMin = 1))
+	float ActiveAnalyzerWireSignalScalarMax;
+
 	/** Determines which details view to show in Metasounds Editor */
 	UPROPERTY(Transient)
 	EMetasoundActiveDetailView DetailView = EMetasoundActiveDetailView::Metasound;
diff --git a/Engine/Plugins/Runtime/Metasound/Source/MetasoundEditor/Private/MetasoundNodeDetailCustomization.cpp b/Engine/Plugins/Runtime/Metasound/Source/MetasoundEditor/Private/MetasoundNodeDetailCustomization.cpp
--- a/Engine/Plugins/Runtime/Metasound/Source/MetasoundEditor/Private/MetasoundNodeDetailCustomization.cpp
+++ b/Engine/Plugins/Runtime/Metasound/Source/MetasoundEditor/Private/MetasoundNodeDetailCustomization.cpp
@@ -620,28 +620,27 @@ namespace Metasound
 
 			// Update data type based on "Is Array" checkbox and support for arrays.
 			// If an array type is not supported, default to the base data type.
-			IMetasoundEditorModule& EditorModule = FModuleManager::GetModuleChecked<IMetasoundEditorModule>("MetaSoundEditor");
 			if (DataTypeArrayCheckbox->GetCheckedState() == ECheckBoxState::Checked)
 			{
-				if (EditorModule.IsRegisteredDataType(ArrayDataTypeName))
+				if (Frontend::IDataTypeRegistry::Get().IsRegistered(ArrayDataTypeName))
 				{
 					NewDataTypeName = ArrayDataTypeName;
 				}
 				else
 				{
-					ensure(EditorModule.IsRegisteredDataType(InSelectedTypeName));
+					ensure(Frontend::IDataTypeRegistry::Get().IsRegistered(InSelectedTypeName));
 					NewDataTypeName = InSelectedTypeName;
 				}
 			}
 			else
 			{
-				if (EditorModule.IsRegisteredDataType(InSelectedTypeName))
+				if (Frontend::IDataTypeRegistry::Get().IsRegistered(InSelectedTypeName))
 				{
 					NewDataTypeName = InSelectedTypeName;
 				}
 				else
 				{
-					ensure(EditorModule.IsRegisteredDataType(ArrayDataTypeName));
+					ensure(Frontend::IDataTypeRegistry::Get().IsRegistered(ArrayDataTypeName));
 					NewDataTypeName = ArrayDataTypeName;
 				}
 			}
diff --git a/Engine/Plugins/Runtime/Metasound/Source/MetasoundEditor/Private/SMetasoundGraphPin.h b/Engine/Plugins/Runtime/Metasound/Source/MetasoundEditor/Private/SMetasoundGraphPin.h
--- a/Engine/Plugins/Runtime/Metasound/Source/MetasoundEditor/Private/SMetasoundGraphPin.h
+++ b/Engine/Plugins/Runtime/Metasound/Source/MetasoundEditor/Private/SMetasoundGraphPin.h
@@ -13,6 +13,7 @@
 #include "MetasoundEditorGraphNode.h"
 #include "MetasoundFrontendController.h"
 #include "SGraphPin.h"
+#include "SMetasoundPinValueInspector.h"
 #include "Styling/SlateStyleRegistry.h"
 #include "Widgets/Input/SButton.h"
 #include "Widgets/SBoxPanel.h"
@@ -29,6 +30,8 @@ namespace Metasound
 		template <typename ParentPinType>
 		class METASOUNDEDITOR_API TMetasoundGraphPin : public ParentPinType
 		{
+			TSharedPtr<SMetasoundPinValueInspector> PinInspector;
+
 		public:
 			SLATE_BEGIN_ARGS(TMetasoundGraphPin<ParentPinType>)
 			{
@@ -229,18 +232,78 @@ namespace Metasound
 					{
 						if (ParentPinType::IsArray())
 						{
-							return bIsConnected ? MetasoundStyle->GetBrush(TEXT("MetasoundEditor.Graph.ConstructorPinArray")) : 
+							return bIsConnected ? MetasoundStyle->GetBrush(TEXT("MetasoundEditor.Graph.ConstructorPinArray")) :
 								MetasoundStyle->GetBrush(TEXT("MetasoundEditor.Graph.ConstructorPinArrayDisconnected"));
 						}
 						else
 						{
-							return bIsConnected ? MetasoundStyle->GetBrush(TEXT("MetasoundEditor.Graph.ConstructorPin")) : 
+							return bIsConnected ? MetasoundStyle->GetBrush(TEXT("MetasoundEditor.Graph.ConstructorPin")) :
 								MetasoundStyle->GetBrush(TEXT("MetasoundEditor.Graph.ConstructorPinDisconnected"));
 						}
 					}
 				}
 				return SGraphPin::GetPinIcon();
 			}
+
+			virtual void Tick(const FGeometry& AllottedGeometry, const double InCurrentTime, const float InDeltaTime) override
+			{
+				ParentPinType::CachedNodeOffset = FVector2D(AllottedGeometry.AbsolutePosition) / AllottedGeometry.Scale - ParentPinType::OwnerNodePtr.Pin()->GetUnscaledPosition();
+				ParentPinType::CachedNodeOffset.Y += AllottedGeometry.Size.Y * 0.5f;
+
+				UEdGraphPin* GraphPin = ParentPinType::GetPinObj();
+
+				// Pause updates if menu is hosted (so user can capture state for ex. using the copy value action)
+				TSharedPtr<FPinValueInspectorTooltip> InspectorTooltip = ParentPinType::ValueInspectorTooltip.Pin();
+				if (InspectorTooltip.IsValid())
+				{
+					const bool bIsHoveringPin = ParentPinType::IsHovered();
+					const bool bIsInspectingPin = bIsHoveringPin && FGraphBuilder::CanInspectPin(GraphPin);
+					if (bIsInspectingPin)
+					{
+						if (PinInspector.IsValid())
+						{
+							PinInspector->UpdateMessage();
+						}
+					}
+					else
+					{
+						if (InspectorTooltip->TooltipCanClose())
+						{
+							PinInspector.Reset();
+							InspectorTooltip->TryDismissTooltip();
+						}
+					}
+				}
+				else
+				{
+					const bool bIsHoveringPin = ParentPinType::IsHovered();
+					const bool bCanInspectPin = FGraphBuilder::CanInspectPin(GraphPin);
+					if (bIsHoveringPin && bCanInspectPin)
+					{
+						const FEdGraphPinReference* CurrentRef = nullptr;
+						if (FPinValueInspectorTooltip::ValueInspectorWidget.IsValid())
+						{
+							CurrentRef = &FPinValueInspectorTooltip::ValueInspectorWidget->GetPinRef();
+						}
+
+						// Only update if reference is not already set.  This avoids ping-pong between pins which can happen
+						// when hovering connections as this state causes IsHovered to return true for all the associated pins
+						// for the given connection.
+						if (!CurrentRef || !CurrentRef->Get() || CurrentRef->Get() == GraphPin)
+						{
+							PinInspector = SNew(SMetasoundPinValueInspector);
+							ParentPinType::ValueInspectorTooltip = FPinValueInspectorTooltip::SummonTooltip(GraphPin, PinInspector);
+							InspectorTooltip = ParentPinType::ValueInspectorTooltip.Pin();
+							if (InspectorTooltip.IsValid())
+							{
+								FVector2D TooltipLocation;
+								ParentPinType::GetInteractiveTooltipLocation(TooltipLocation);
+								InspectorTooltip->MoveTooltip(TooltipLocation);
+							}
+						}
+					}
+				}
+			}
 		};
 
 		class SMetasoundGraphPin : public TMetasoundGraphPin<SGraphPin>
diff --git a/Engine/Plugins/Runtime/Metasound/Source/MetasoundEditor/Private/SMetasoundPinValueInspector.cpp b/Engine/Plugins/Runtime/Metasound/Source/MetasoundEditor/Private/SMetasoundPinValueInspector.cpp
new file mode 100644
--- /dev/null
+++ b/Engine/Plugins/Runtime/Metasound/Source/MetasoundEditor/Private/SMetasoundPinValueInspector.cpp
@@ -0,0 +1,702 @@
+// Copyright Epic Games, Inc. All Rights Reserved.
+#include "SMetasoundPinValueInspector.h"
+
+#include "EditorStyleSet.h"
+#include "Input/Events.h"
+#include "Layout/Geometry.h"
+#include "MetasoundEditor.h"
+#include "MetasoundEditorGraphBuilder.h"
+#include "MetasoundEditorGraphNode.h"
+#include "MetasoundEditorSettings.h"
+#include "MetasoundFrontendDocument.h"
+#include "SPinValueInspector.h"
+#include "Widgets/Colors/SColorBlock.h"
+#include "Widgets/Colors/SColorPicker.h"
+#include "Widgets/Input/SButton.h"
+#include "Widgets/Input/SEditableTextBox.h"
+#include "Widgets/SBoxPanel.h"
+#include "Widgets/SNullWidget.h"
+#include "Widgets/SWidget.h"
+#include "ScopedTransaction.h"
+
+#define LOCTEXT_NAMESPACE "MetasoundEditor"
+
+
+namespace Metasound
+{
+	namespace Editor
+	{
+		namespace PinValueInspectorPrivate
+		{
+			static const FText EnableValueColorizationText = LOCTEXT("PinValueInspector_EnableConnectionColorization", "Enable Colorized Connection");
+			static const FText DisableValueColorizationText = LOCTEXT("PinValueInspector_DisableConnectionColorization", "Disable Colorized Connection");
+
+			static const FText MinNameText = LOCTEXT("PinValueInspector_MinValue", "Min");
+			static const FText MaxNameText = LOCTEXT("PinValueInspector_MaxValue", "Max");
+
+			static const FText ColorRowNameFormat = LOCTEXT("PinValueInspector_ColorRowNameFormat", "{0}:");
+
+			bool IsMetasoundPrimitiveValueColorizationSupported(FName InDataType)
+			{
+				// Strings are the only non-numeric literal primitive which is not supported (for now).
+				return InDataType != GetMetasoundDataTypeName<FString>();
+			}
+
+			TArray<FMetasoundFrontendLiteral> InitializeLiterals(FName InDataType)
+			{
+				TArray<FMetasoundFrontendLiteral> Literals;
+
+				// For now, just hardcoded to two values.
+				// TODO: Support multiple literals, which could be
+				// really useful for strings, enums, and non-bipolar
+				// or non-binary numeric ranges of interest.
+				Literals.SetNum(2);
+
+				if (InDataType == GetMetasoundDataTypeName<float>())
+				{
+					Literals[0].Set(0.0f);
+					Literals[1].Set(1.0f);
+				}
+				else if (InDataType == GetMetasoundDataTypeName<int32>())
+				{
+					Literals[0].Set(0);
+					Literals[1].Set(1);
+				}
+				else if (InDataType == GetMetasoundDataTypeName<bool>())
+				{
+					Literals[0].Set(false);
+					Literals[1].Set(true);
+				}
+				// TODO: Potentially support strings
+// 				else if (InDataType == GetMetasoundDataTypeName<FString>())
+// 				{
+// 					Literals[0].Set(FString());
+// 					Literals[1].Set(FString());
+// 				}
+
+				return Literals;
+			}
+
+			void SetLiteralFromText(const FName InDataType, const FText& InText, FMetasoundFrontendLiteral& OutLiteral)
+			{
+				const FString TextString = InText.ToString();
+				if (InDataType == GetMetasoundDataTypeName<float>())
+				{
+					const float Value = FCString::Atof(*TextString);
+					OutLiteral.Set(Value);
+				}
+				else if (InDataType == GetMetasoundDataTypeName<int32>())
+				{
+					const int32 Value = FCString::Atoi(*TextString);
+					OutLiteral.Set(Value);
+				}
+				else if (InDataType == GetMetasoundDataTypeName<bool>())
+				{
+					const int32 Value = FCString::ToBool(*TextString);
+					OutLiteral.Set(Value);
+				}
+// 				else if (InDataType == GetMetasoundDataTypeName<FString>())
+// 				{
+// 					OutLiteral.Set(TextString);
+// 				}
+			}
+		} // namespace PinValueInspectorPrivate
+
+		FMetasoundNumericDebugLineItem::FMetasoundNumericDebugLineItem(UEdGraphPin* InGraphPinObj, FGetValueStringFunction&& InGetValueStringFunction)
+			: FDebugLineItem(DLT_Message)
+			, GraphPinObj(InGraphPinObj)
+			, GetValueStringFunction(MoveTemp(InGetValueStringFunction))
+		{
+			using namespace Frontend;
+
+			if (GraphPinObj && !GraphPinObj->LinkedTo.IsEmpty())
+			{
+				// Swap to show connected output if input (Only ever one)
+				if (GraphPinObj->Direction == EGPD_Input)
+				{
+					GraphPinObj = GraphPinObj->LinkedTo.Last();
+					check(GraphPinObj->Direction == EGPD_Output);
+				}
+
+				FConstOutputHandle OutputHandle = GetOutputHandle();
+				ColorLiterals = PinValueInspectorPrivate::InitializeLiterals(OutputHandle->GetDataType());
+
+				Update();
+			}
+
+			DisplayName = FGraphBuilder::GetDisplayName(*GetOutputHandle());
+		}
+
+		FColorPickerArgs FMetasoundNumericDebugLineItem::InitPickerArgs()
+		{
+			FColorPickerArgs PickerArgs;
+			PickerArgs.bIsModal = true;
+			PickerArgs.bUseAlpha = true;
+			PickerArgs.DisplayGamma = TAttribute<float>::Create(TAttribute<float>::FGetter::CreateUObject(GEngine, &UEngine::GetDisplayGamma));
+			PickerArgs.ParentWidget = ValueWidget;
+
+			return PickerArgs;
+		}
+
+		const UMetasoundEditorGraphNode* FMetasoundNumericDebugLineItem::GetNode() const
+		{
+			if (GraphPinObj)
+			{
+				return Cast<UMetasoundEditorGraphNode>(GraphPinObj->GetOwningNode());
+			}
+
+			return nullptr;
+		}
+
+		const UMetasoundEditorGraphNode& FMetasoundNumericDebugLineItem::GetNodeChecked() const
+		{
+			const UMetasoundEditorGraphNode* Node = GetNode();
+			check(Node);
+			return *Node;
+		}
+
+		Frontend::FConstOutputHandle FMetasoundNumericDebugLineItem::GetOutputHandle() const
+		{
+			return FGraphBuilder::GetOutputHandleFromPin(GraphPinObj);
+		}
+
+		Frontend::FOutputHandle FMetasoundNumericDebugLineItem::GetOutputHandle()
+		{
+			return FGraphBuilder::GetOutputHandleFromPin(GraphPinObj);
+		}
+
+		FLinearColor FMetasoundNumericDebugLineItem::GetEdgeStyleColorAtIndex(int32 InIndex) const
+		{
+			if (const FMetasoundFrontendEdgeStyle* EdgeStyle = GetEdgeStyle())
+			{
+				if (ensure(EdgeStyle->LiteralColorPairs.Num() > InIndex))
+				{
+					return EdgeStyle->LiteralColorPairs[InIndex].Color;
+				}
+			}
+			return FLinearColor::Transparent;
+		}
+
+		void FMetasoundNumericDebugLineItem::DisableValueColorization()
+		{
+			using namespace Frontend;
+
+			FScopedTransaction Transaction(PinValueInspectorPrivate::DisableValueColorizationText);
+
+			UObject* ParentObject = GetNodeChecked().GetOutermostObject();
+			check(ParentObject);
+			ParentObject->Modify();
+
+			FGraphHandle GraphHandle = GetGraphHandle();
+			FMetasoundFrontendGraphStyle Style = GraphHandle->GetGraphStyle();
+
+			const FGuid NodeID = GetNodeChecked().GetNodeID();
+			const FName OutputName = GetOutputHandle()->GetName();
+			Style.EdgeStyles.RemoveAllSwap([&NodeID, &OutputName](FMetasoundFrontendEdgeStyle& EdgeStyle)
+				{
+					return EdgeStyle.NodeID == NodeID && EdgeStyle.OutputName == OutputName;
+				});
+
+			GraphHandle->SetGraphStyle(Style);
+			bIsValueColorizationEnabled = false;
+		}
+
+		void FMetasoundNumericDebugLineItem::EnableValueColorization()
+		{
+			using namespace Frontend;
+
+			FScopedTransaction Transaction(PinValueInspectorPrivate::EnableValueColorizationText);
+
+			UObject* ParentObject = GetNodeChecked().GetOutermostObject();
+			check(ParentObject);
+			ParentObject->Modify();
+
+			TArray<FMetasoundFrontendEdgeStyleLiteralColorPair> DefaultPairs
+			{
+				{ ColorLiterals[0], FLinearColor::Transparent },
+				{ ColorLiterals[1], FLinearColor::White }
+			};
+
+			FGraphHandle GraphHandle = GetGraphHandle();
+			FMetasoundFrontendGraphStyle Style = GraphHandle->GetGraphStyle();
+			FConstOutputHandle OutputHandle = GetOutputHandle();
+
+			const FGuid NodeID = GetNodeChecked().GetNodeID();
+			const FName OutputName = OutputHandle->GetName();
+			FMetasoundFrontendEdgeStyle* EdgeStyle = Style.EdgeStyles.FindByPredicate([&NodeID, &OutputName](FMetasoundFrontendEdgeStyle& EdgeStyle)
+			{
+				return EdgeStyle.NodeID == NodeID && EdgeStyle.OutputName == OutputName;
+			});
+
+			if (EdgeStyle)
+			{
+				EdgeStyle->LiteralColorPairs = MoveTemp(DefaultPairs);
+			}
+			else
+			{
+				Style.EdgeStyles.Emplace(FMetasoundFrontendEdgeStyle{ NodeID, OutputName, MoveTemp(DefaultPairs) });
+			}
+
+			GraphHandle->SetGraphStyle(Style);
+			bIsValueColorizationEnabled = true;
+		}
+
+		void FMetasoundNumericDebugLineItem::OnColorCommitted(FLinearColor InColor, int32 InIndex)
+		{
+			using namespace Frontend;
+
+			static const FText SetValueFormat = LOCTEXT("PinValueInspector_SetMinColorFormat", "Set '{0}[{1}]' Color to '{2}'");
+			const FName Name = GetOutputHandle()->GetName();
+			FScopedTransaction Transaction(FText::Format(SetValueFormat, FText::FromName(Name), FText::AsNumber(InIndex), FText::FromString(InColor.ToString())));
+
+			UObject* ParentObject = GetNodeChecked().GetOutermostObject();
+			check(ParentObject);
+			ParentObject->Modify();
+
+			FGraphHandle GraphHandle = GetGraphHandle();
+			FMetasoundFrontendGraphStyle Style = GraphHandle->GetGraphStyle();
+
+			const FGuid NodeID = GetNodeChecked().GetNodeID();
+			const FName OutputName = GetOutputHandle()->GetName();
+			FMetasoundFrontendEdgeStyle* EdgeStyle = Style.EdgeStyles.FindByPredicate([&NodeID, &OutputName](FMetasoundFrontendEdgeStyle& EdgeStyle)
+			{
+				return EdgeStyle.NodeID == NodeID && EdgeStyle.OutputName == OutputName;
+			});
+
+			if (EdgeStyle)
+			{
+				if (ensure(InIndex < EdgeStyle->LiteralColorPairs.Num()))
+				{
+					EdgeStyle->LiteralColorPairs[InIndex].Color = InColor;
+				}
+			}
+
+			GraphHandle->SetGraphStyle(Style);
+		}
+
+		void FMetasoundNumericDebugLineItem::OnValueCommitted(const FText& InValueText, FMetasoundFrontendLiteral& OutNewLiteral, int32 InIndex)
+		{
+			using namespace Frontend;
+
+			if (const UMetasoundEditorGraphNode* Node = GetNode())
+			{
+				FOutputHandle OutputHandle = GetOutputHandle();
+
+				static const FText SetValueFormat = LOCTEXT("PinValueInspector_SetColorLiteralValueFormat", "Set '{0}[{1}]' EdgeStyle Color Value to '{2}'");
+				const FName Name = OutputHandle->GetName();
+				FScopedTransaction Transaction(FText::Format(SetValueFormat, FText::FromName(Name), FText::AsNumber(InIndex), InValueText));
+
+				UObject* ParentObject = Node->GetOutermostObject();
+				check(ParentObject);
+				ParentObject->Modify();
+
+				const FGuid NodeID = Node->GetNodeID();
+				const FName OutputName = OutputHandle->GetName();
+
+				FGraphHandle OwningGraph = GetGraphHandle();
+				FMetasoundFrontendGraphStyle Style = OwningGraph->GetGraphStyle();
+
+				bool bUpdated = false;
+				for (FMetasoundFrontendEdgeStyle& EdgeStyle : Style.EdgeStyles)
+				{
+					if (EdgeStyle.NodeID == NodeID && EdgeStyle.OutputName == OutputName)
+					{
+						if (ensure(EdgeStyle.LiteralColorPairs.Num() > InIndex))
+						{
+							const FName DataType = OutputHandle->GetDataType();
+							PinValueInspectorPrivate::SetLiteralFromText(DataType, InValueText, OutNewLiteral);
+							EdgeStyle.LiteralColorPairs[InIndex].Value = OutNewLiteral;
+						}
+					}
+				}
+
+				OwningGraph->SetGraphStyle(Style);
+			}
+		}
+
+		Frontend::FGraphHandle FMetasoundNumericDebugLineItem::GetGraphHandle()
+		{
+			Frontend::FOutputHandle OutputHandle = FGraphBuilder::GetOutputHandleFromPin(GraphPinObj);
+			return OutputHandle->GetOwningNode()->GetOwningGraph();
+		}
+
+		Frontend::FConstGraphHandle FMetasoundNumericDebugLineItem::GetGraphHandle() const
+		{
+			Frontend::FConstOutputHandle OutputHandle = FGraphBuilder::GetOutputHandleFromPin(GraphPinObj);
+			return OutputHandle->GetOwningNode()->GetOwningGraph();
+		}
+
+		FGraphConnectionManager* FMetasoundNumericDebugLineItem::GetConnectionManager()
+		{
+			using namespace Frontend;
+
+			TSharedPtr<FEditor> Editor = FGraphBuilder::GetEditorForNode(GetNodeChecked());
+			if (Editor.IsValid())
+			{
+				return &Editor->GetConnectionManager();
+			}
+
+			return nullptr;
+		}
+
+		const FMetasoundFrontendEdgeStyle* FMetasoundNumericDebugLineItem::GetEdgeStyle() const
+		{
+			using namespace Frontend;
+
+			FConstGraphHandle OwningGraph = GetGraphHandle();
+			if (!OwningGraph->IsValid())
+			{
+				return nullptr;
+			}
+
+			const FGuid NodeID = GetNodeChecked().GetNodeID();
+			const FName OutputName = GetOutputHandle()->GetName();
+			const FMetasoundFrontendGraphStyle& Style = OwningGraph->GetGraphStyle();
+
+			return Style.EdgeStyles.FindByPredicate([&NodeID, &OutputName](FMetasoundFrontendEdgeStyle& EdgeStyle)
+				{
+					return EdgeStyle.NodeID == NodeID && EdgeStyle.OutputName == OutputName;
+				});
+		}
+
+		void FMetasoundNumericDebugLineItem::Update()
+		{
+			using namespace Frontend;
+
+			bIsValueColorizationEnabled = false;
+
+			const FMetasoundFrontendEdgeStyle* EdgeStyle = GetEdgeStyle();
+			FGuid NodeID = GetNodeChecked().GetNodeID();
+			FName OutputName = GetOutputHandle()->GetName();
+			if (EdgeStyle)
+			{
+				if (ensure(EdgeStyle->LiteralColorPairs.Num() == 2))
+				{
+					const FMetasoundFrontendLiteral& CurrentMin = EdgeStyle->LiteralColorPairs[0].Value;
+					ColorLiterals[0] = CurrentMin;
+					const FMetasoundFrontendLiteral& CurrentMax = EdgeStyle->LiteralColorPairs[1].Value;
+					ColorLiterals[1] = CurrentMax;
+					bIsValueColorizationEnabled = true;
+				}
+			}
+
+			if (FGraphConnectionManager* ConnectionManager = GetConnectionManager())
+			{
+				Message = FString::Format(TEXT("Value: {0}"), { GetValueStringFunction(*ConnectionManager, NodeID, OutputName) });
+			}
+		}
+
+		FReply FMetasoundNumericDebugLineItem::OnColorBoxClicked(const FGeometry& MyGeometry, const FPointerEvent& MouseEvent, const FLinearColor& InInitColor, int32 InIndex)
+		{
+			using namespace Frontend;
+
+			if (MouseEvent.GetEffectingButton() != EKeys::LeftMouseButton)
+			{
+				return FReply::Unhandled();
+			}
+
+			FGraphHandle OwningGraph = GetGraphHandle();
+			if (!OwningGraph->IsValid())
+			{
+				return FReply::Unhandled();
+			}
+
+			FLinearColor Color = InInitColor;
+			TArray<FLinearColor*> LinearColorArray;
+			LinearColorArray.Add(&Color);
+			const FMetasoundFrontendEdgeStyle* EdgeStyle = GetEdgeStyle();
+			if (!EdgeStyle)
+			{
+				return FReply::Unhandled();
+			}
+
+			FColorPickerArgs PickerArgs = InitPickerArgs();
+
+			PickerArgs.LinearColorArray = &LinearColorArray;
+			PickerArgs.OnColorCommitted = FOnLinearColorValueChanged::CreateLambda([this, InIndex](FLinearColor NewColor) { OnColorCommitted(NewColor, InIndex); });
+
+			OpenColorPicker(PickerArgs);
+
+			return FReply::Handled();
+		}
+
+		FText FMetasoundNumericDebugLineItem::GetDescription() const
+		{
+			return FText::FromString(Message);
+		}
+
+		TSharedRef<SWidget> FMetasoundNumericDebugLineItem::GenerateNameWidget(TSharedPtr<FString> InSearchString)
+		{
+			if (!GraphPinObj)
+			{
+				return SNullWidget::NullWidget;
+			}
+
+			return SNew(STextBlock)
+				.Text(DisplayName)
+				.Font(FAppStyle::GetFontStyle(TEXT("PropertyWindow.NormalFont")));
+		}
+
+		TSharedRef<SWidget> FMetasoundNumericDebugLineItem::GenerateValueWidget(TSharedPtr<FString> InSearchString)
+		{
+			using namespace Frontend;
+			const FSlateFontInfo FontInfo = FAppStyle::GetFontStyle(TEXT("PropertyWindow.NormalFont"));
+
+			TSharedPtr<SVerticalBox> VerticalBox;
+			SAssignNew(VerticalBox, SVerticalBox)
+				+ SVerticalBox::Slot()
+				.AutoHeight()
+				.Padding(1.0f, 1.0f, 0.0f, 0.0f)
+				.VAlign(VAlign_Center)
+				[
+					SNew(STextBlock)
+					.Text_Lambda([this]() { return FText::FromString(Message); })
+					.Font(FontInfo)
+				];
+
+			const FName DataType = GetOutputHandle()->GetDataType();
+			const bool bConnectionAnimationsEnabled = GetDefault<UMetasoundEditorSettings>()->bAnimateConnections;
+			const bool bIsTypeSupported = PinValueInspectorPrivate::IsMetasoundPrimitiveValueColorizationSupported(DataType);
+			if (bConnectionAnimationsEnabled && bIsTypeSupported)
+			{
+				AddValueColorizationWidgets(VerticalBox);
+			}
+
+			ValueWidget = VerticalBox;
+			return ValueWidget->AsShared();
+		}
+
+		void FMetasoundNumericDebugLineItem::AddValueColorizationWidgets(TSharedPtr<SVerticalBox> VerticalBox)
+		{
+			using namespace Frontend;
+			const FSlateFontInfo FontInfo = FAppStyle::GetFontStyle(TEXT("PropertyWindow.NormalFont"));
+
+			VerticalBox->AddSlot()
+				.AutoHeight()
+				.Padding(1.0f, 1.0f, 0.0f, 1.0f)
+				.VAlign(VAlign_Center)
+				[
+					SNew(SButton)
+					.OnClicked_Lambda([this]()
+					{
+						bIsValueColorizationEnabled ? DisableValueColorization() : EnableValueColorization();
+						return FReply::Handled();
+					})
+				[
+					SNew(STextBlock)
+					.Text_Lambda([this]()
+					{
+						return bIsValueColorizationEnabled ? PinValueInspectorPrivate::DisableValueColorizationText : PinValueInspectorPrivate::EnableValueColorizationText;
+					})
+					.ToolTipText(LOCTEXT("EnableConnectionVisualizationToolTip", "Enable or disable if connection has fill colorization interpolating between provided min and max values."))
+					.Font(FontInfo)
+				]
+			];
+
+			auto CreateColorLabel = [this, &FontInfo](const FText& InLabel) -> TSharedRef<SWidget>
+			{
+				return SNew(STextBlock)
+					.Text(InLabel)
+					.Font(FontInfo)
+					.MinDesiredWidth(40)
+					.Visibility_Lambda([this]() { return bIsValueColorizationEnabled ? EVisibility::Visible : EVisibility::Collapsed; });
+			};
+
+			auto CreateColorValueWidget = [this, &FontInfo](int32 InIndex) -> TSharedRef<SWidget>
+			{
+				TSharedRef<SWidget> ColorValueWidget = SNew(SEditableTextBox)
+					.Text_Lambda([this, InIndex]()
+					{
+						return FText::FromString(ColorLiterals[InIndex].ToString());
+					})
+					.OnTextCommitted_Lambda([this, InIndex](const FText& InValueText, ETextCommit::Type InTextCommit)
+					{
+						OnValueCommitted(InValueText, ColorLiterals[InIndex], InIndex);
+					})
+					.SelectAllTextWhenFocused(true)
+					.MinDesiredWidth(30)
+					.Font(FontInfo)
+					.Visibility_Lambda([this]() { return bIsValueColorizationEnabled ? EVisibility::Visible : EVisibility::Collapsed; });
+
+				// Boolean values only have two states that default to true & false, so no reason to allow user to edit.
+				if (ColorLiterals[InIndex].GetType() == EMetasoundFrontendLiteralType::Boolean)
+				{
+					ColorValueWidget->SetEnabled(false);
+				}
+
+				return ColorValueWidget;
+			};
+
+			auto CreateColorPickerBlock = [this](int32 InIndex) -> TSharedRef<SWidget>
+			{
+				return SNew(SColorBlock)
+					.Color_Lambda([this, InIndex] () { return GetEdgeStyleColorAtIndex(InIndex); })
+					.ShowBackgroundForAlpha(true)
+					.OnMouseButtonDown_Lambda([this, InIndex](const FGeometry& MyGeometry, const FPointerEvent& MouseEvent)
+					{
+						const FLinearColor InitColor = GetEdgeStyleColorAtIndex(InIndex);
+						return OnColorBoxClicked(MyGeometry, MouseEvent, InitColor, InIndex);
+					})
+					.Visibility_Lambda([this]() { return bIsValueColorizationEnabled ? EVisibility::Visible : EVisibility::Collapsed; });
+			};
+
+			VerticalBox->AddSlot()
+			.AutoHeight()
+			.Padding(1.0f, 1.0f, 0.0f, 0.0f)
+			.VAlign(VAlign_Center)
+			[
+				SNew(SHorizontalBox)
+				+ SHorizontalBox::Slot()
+				.AutoWidth()
+				.Padding(1.0f, 0.0f, 0.0f, 0.0f)
+				.HAlign(HAlign_Center)
+				[
+					CreateColorLabel(FText::Format(PinValueInspectorPrivate::ColorRowNameFormat, PinValueInspectorPrivate::MinNameText))
+				]
+				+ SHorizontalBox::Slot()
+				.AutoWidth()
+				.Padding(1.0f, 0.0f, 0.0f, 0.0f)
+				.HAlign(HAlign_Center)
+				[
+					CreateColorPickerBlock(0)
+				]
+				+ SHorizontalBox::Slot()
+				.AutoWidth()
+				.Padding(3.0f, 0.0f, 0.0f, 1.0f)
+				.HAlign(HAlign_Center)
+				[
+					CreateColorValueWidget(0)
+				]
+			];
+
+			VerticalBox->AddSlot()
+			.AutoHeight()
+			.Padding(1.0f, 1.0f, 0.0f, 0.0f)
+			.VAlign(VAlign_Center)
+			[
+				SNew(SHorizontalBox)
+				+ SHorizontalBox::Slot()
+				.AutoWidth()
+				.Padding(1.0f, 0.0f, 0.0f, 0.0f)
+				.HAlign(HAlign_Center)
+				[
+					CreateColorLabel(FText::Format(PinValueInspectorPrivate::ColorRowNameFormat, PinValueInspectorPrivate::MaxNameText))
+				]
+				+ SHorizontalBox::Slot()
+				.AutoWidth()
+				.Padding(1.0f, 0.0f, 0.0f, 0.0f)
+				.HAlign(HAlign_Center)
+				[
+					CreateColorPickerBlock(1)
+				]
+				+ SHorizontalBox::Slot()
+				.AutoWidth()
+				.Padding(3.0f, 0.0f, 0.0f, 1.0f)
+				.HAlign(HAlign_Center)
+				[
+					CreateColorValueWidget(1)
+				]
+			];
+		}
+
+		FDebugLineItem* FMetasoundNumericDebugLineItem::Duplicate() const
+		{
+			FGetValueStringFunction DuplicateGetValueStringFunction = GetValueStringFunction;
+			return new FMetasoundNumericDebugLineItem(GraphPinObj, MoveTemp(DuplicateGetValueStringFunction));
+		}
+
+		bool FMetasoundNumericDebugLineItem::Compare(const FDebugLineItem* BaseOther) const
+		{
+			const FMetasoundNumericDebugLineItem* Other = static_cast<const FMetasoundNumericDebugLineItem*>(BaseOther);
+			return Message == Other->Message;
+		}
+
+		uint32 FMetasoundNumericDebugLineItem::GetHash()
+		{
+			return GetTypeHash(Message);
+		}
+
+		void SMetasoundPinValueInspector::Construct(const FArguments& InArgs)
+		{
+			SPinValueInspector::FArguments InspectorArgs;
+			SPinValueInspector::Construct(InspectorArgs);
+		}
+
+		void SMetasoundPinValueInspector::UpdateMessage()
+		{
+			if (LineItem.IsValid())
+			{
+				LineItem->Update();
+			}
+		}
+
+		void SMetasoundPinValueInspector::PopulateTreeView()
+		{
+			// Locate the class property associated with the source pin and set it to the root node.
+			GraphPinObj = GetPinRef().Get();
+
+			// Don't populate if no pin found or pin is not connected (analyzers are not enabled for disconnected node outputs)
+			if (!GraphPinObj || GraphPinObj->LinkedTo.IsEmpty())
+			{
+				LineItem.Reset();
+				return;
+			}
+
+			TSharedPtr<FMetasoundNumericDebugLineItem> NewLine;
+
+			static const FString UnsetValueString = TEXT("Unset");
+			if (GraphPinObj->PinType.PinCategory == FGraphBuilder::PinCategoryFloat)
+			{
+				NewLine = MakeShared<FMetasoundNumericDebugLineItem>(GraphPinObj, [&](const FGraphConnectionManager& ConnectionManager, const FGuid& NodeID, const FName OutputName)
+				{
+					float Value = 0.0f;
+					if (ConnectionManager.GetValue(NodeID, OutputName, Value))
+					{
+						return FString::Printf(TEXT("%.6f"), Value);
+					}
+					return UnsetValueString;
+				});
+			}
+			else if (GraphPinObj->PinType.PinCategory == FGraphBuilder::PinCategoryInt32)
+			{
+				NewLine = MakeShared<FMetasoundNumericDebugLineItem>(GraphPinObj, [&](const FGraphConnectionManager& ConnectionManager, const FGuid& NodeID, const FName OutputName)
+				{
+					int32 Value = 0;
+					if (ConnectionManager.GetValue(NodeID, OutputName, Value))
+					{
+						return FString::Printf(TEXT("%d"), Value);
+					}
+					return UnsetValueString;
+				});
+			}
+			else if (GraphPinObj->PinType.PinCategory == FGraphBuilder::PinCategoryString)
+			{
+				NewLine = MakeShared<FMetasoundNumericDebugLineItem>(GraphPinObj, [&](const FGraphConnectionManager& ConnectionManager, const FGuid& NodeID, const FName OutputName)
+				{
+					FString Value = UnsetValueString;
+					ConnectionManager.GetValue(NodeID, OutputName, Value);
+					return Value;
+				});
+			}
+			else if (GraphPinObj->PinType.PinCategory == FGraphBuilder::PinCategoryBoolean)
+			{
+				NewLine = MakeShared<FMetasoundNumericDebugLineItem>(GraphPinObj, [&](const FGraphConnectionManager& ConnectionManager, const FGuid& NodeID, const FName OutputName)
+				{
+					bool Value = false;
+					if (ConnectionManager.GetValue(NodeID, OutputName, Value))
+					{
+						return FString::Printf(TEXT("%s"), Value ? TEXT("true") : TEXT("false"));
+					}
+					return UnsetValueString;
+				});
+			}
+
+			if (NewLine.IsValid())
+			{
+				LineItem = NewLine;
+				AddTreeItemUnique(NewLine);
+			}
+		}
+	} // namespace Editor
+} // namespace Metasound
+#undef LOCTEXT_NAMESPACE
diff --git a/Engine/Plugins/Runtime/Metasound/Source/MetasoundEditor/Private/SMetasoundPinValueInspector.h b/Engine/Plugins/Runtime/Metasound/Source/MetasoundEditor/Private/SMetasoundPinValueInspector.h
new file mode 100644
--- /dev/null
+++ b/Engine/Plugins/Runtime/Metasound/Source/MetasoundEditor/Private/SMetasoundPinValueInspector.h
@@ -0,0 +1,121 @@
+// Copyright Epic Games, Inc. All Rights Reserved.
+#pragma once
+
+#include "Containers/UnrealString.h"
+#include "Debugging/SKismetDebugTreeView.h"
+#include "Engine/Engine.h"
+#include "InputCoreTypes.h"
+#include "Input/Reply.h"
+#include "MetasoundEditorGraphConnectionManager.h"
+#include "Misc/Attribute.h"
+#include "ScopedTransaction.h"
+#include "SPinValueInspector.h"
+#include "Templates/Function.h"
+#include "Templates/SharedPointer.h"
+#include "UObject/NameTypes.h"
+#include "UObject/NoExportTypes.h"
+#include "Widgets/SWidget.h"
+#include "Widgets/Colors/SColorPicker.h"
+
+
+// Forward Declarations
+class UMetasoundEditorGraphNode;
+
+
+namespace Metasound
+{
+	namespace Editor
+	{
+		struct FMetasoundNumericDebugLineItem : public FDebugLineItem
+		{
+			using FGetValueStringFunction = TFunction<FString(const FGraphConnectionManager&, const FGuid& /* InNodeID */, const FName /* InOutputName */)>;
+
+		private:
+			/** The GraphPin that this debug line is reporting. */
+			UEdGraphPin* GraphPinObj = nullptr;
+
+			/** Function used to get the string value from the connection manager. */
+			FGetValueStringFunction GetValueStringFunction;
+
+			FColorPickerArgs InitPickerArgs();
+
+		protected:
+			TArray<FMetasoundFrontendLiteral> ColorLiterals;
+
+			FString Message;
+			FText DisplayName;
+
+			TSharedPtr<SWidget> ValueWidget;
+
+			bool bIsValueColorizationEnabled = false;
+
+		public:
+			FMetasoundNumericDebugLineItem(UEdGraphPin* InGraphPinObj, FGetValueStringFunction&& InGetValueStringFunction);
+			virtual ~FMetasoundNumericDebugLineItem() = default;
+
+			const UMetasoundEditorGraphNode* GetNode() const;
+
+			const UMetasoundEditorGraphNode& GetNodeChecked() const;
+
+			Frontend::FConstOutputHandle GetOutputHandle() const;
+
+			Frontend::FOutputHandle GetOutputHandle();
+
+			Frontend::FGraphHandle GetGraphHandle();
+
+			Frontend::FConstGraphHandle GetGraphHandle() const;
+
+			FGraphConnectionManager* GetConnectionManager();
+
+			const FMetasoundFrontendEdgeStyle* GetEdgeStyle() const;
+
+			void Update();
+
+			virtual FText GetDescription() const override;
+
+			virtual TSharedRef<SWidget> GenerateNameWidget(TSharedPtr<FString> InSearchString) override;
+
+			virtual TSharedRef<SWidget> GenerateValueWidget(TSharedPtr<FString> InSearchString) override;
+
+		protected:
+			virtual void AddValueColorizationWidgets(TSharedPtr<SVerticalBox> VerticalBox);
+
+			virtual FDebugLineItem* Duplicate() const override;
+
+			virtual bool Compare(const FDebugLineItem* BaseOther) const override;
+
+			virtual uint32 GetHash() override;
+
+		private:
+			FLinearColor GetEdgeStyleColorAtIndex(int32 InIndex) const;
+
+			void DisableValueColorization();
+			void EnableValueColorization();
+
+			void OnColorCommitted(FLinearColor InColor, int32 InIndex);
+			FReply OnColorBoxClicked(const FGeometry& MyGeometry, const FPointerEvent& MouseEvent, const FLinearColor& InInitColor, int32 InIndex);
+			void OnValueCommitted(const FText& InValueText, FMetasoundFrontendLiteral& OutNewLiteral, int32 InIndex);
+		};
+
+		class METASOUNDEDITOR_API SMetasoundPinValueInspector : public SPinValueInspector
+		{
+			/** The GraphPin that this widget is inspecting. */
+			UEdGraphPin* GraphPinObj;
+
+			TSharedPtr<FMetasoundNumericDebugLineItem> LineItem;
+
+		public:
+			SLATE_BEGIN_ARGS(SMetasoundPinValueInspector)
+			{
+			}
+			SLATE_END_ARGS()
+
+			void Construct(const FArguments& InArgs);
+
+			void UpdateMessage();
+
+			virtual void PopulateTreeView() override;
+		};
+	} // namespace Editor
+} // namespace Metasound
+#undef LOCTEXT_NAMESPACE
diff --git a/Engine/Plugins/Runtime/Metasound/Source/MetasoundEditor/Public/IMetasoundEditor.h b/Engine/Plugins/Runtime/Metasound/Source/MetasoundEditor/Public/IMetasoundEditor.h
--- a/Engine/Plugins/Runtime/Metasound/Source/MetasoundEditor/Public/IMetasoundEditor.h
+++ b/Engine/Plugins/Runtime/Metasound/Source/MetasoundEditor/Public/IMetasoundEditor.h
@@ -10,9 +10,21 @@
 
 class IMetasoundEditor : public FAssetEditorToolkit
 {
+	// Returns the editor's assigned MetaSound UObject.
 	virtual UObject* GetMetasoundObject() const = 0;
+
+	// Sets the selected object(s).
 	virtual void SetSelection(const TArray<UObject*>& SelectedObjects) = 0;
-	virtual bool GetBoundsForSelectedNodes(FSlateRect& Rect, float Padding) = 0;
+
+	// Returns the bounds for the selected node(s).
+	virtual bool GetBoundsForSelectedNodes(FSlateRect& OutRect, float Padding) = 0;
+
+	// Plays the associated preview UAudioComponent & update editor state to reflect.
 	virtual void Play() = 0;
+
+	// Stop the associated preview UAudioComponent & update editor state to reflect.
 	virtual void Stop() = 0;
+
+	/** Whether or not the given editor is current associated with a UAudioComponent that is actively playing. */
+	virtual bool IsPlaying() const = 0;
 };
diff --git a/Engine/Plugins/Runtime/Metasound/Source/MetasoundEditor/Public/MetasoundEditorGraphSchema.h b/Engine/Plugins/Runtime/Metasound/Source/MetasoundEditor/Public/MetasoundEditorGraphSchema.h
--- a/Engine/Plugins/Runtime/Metasound/Source/MetasoundEditor/Public/MetasoundEditorGraphSchema.h
+++ b/Engine/Plugins/Runtime/Metasound/Source/MetasoundEditor/Public/MetasoundEditorGraphSchema.h
@@ -401,6 +401,7 @@ class UMetasoundEditorGraphSchema : public UEdGraphSchema
 {
 	GENERATED_UCLASS_BODY()
 
+public:
 	/** Check whether connecting these pins would cause a loop */
 	bool ConnectionCausesLoop(const UEdGraphPin* InputPin, const UEdGraphPin* OutputPin) const;
 
diff --git a/Engine/Plugins/Runtime/Metasound/Source/MetasoundEditor/Public/MetasoundEditorModule.h b/Engine/Plugins/Runtime/Metasound/Source/MetasoundEditor/Public/MetasoundEditorModule.h
--- a/Engine/Plugins/Runtime/Metasound/Source/MetasoundEditor/Public/MetasoundEditorModule.h
+++ b/Engine/Plugins/Runtime/Metasound/Source/MetasoundEditor/Public/MetasoundEditorModule.h
@@ -25,8 +25,6 @@ namespace Metasound
 {
 	namespace Editor
 	{
-		using FDataTypeRegistryInfo = Frontend::FDataTypeRegistryInfo;
-
 		// Status of initial asset scan when editor loads up.
 		enum class EAssetScanStatus : uint8
 		{
@@ -46,21 +44,6 @@ namespace Metasound
 			Complete = 3
 		};
 
-		struct FEditorDataType
-		{
-			FEdGraphPinType PinType;
-			FDataTypeRegistryInfo RegistryInfo;
-
-			FEditorDataType(FEdGraphPinType&& InPinType, FDataTypeRegistryInfo&& InRegistryInfo)
-				: PinType(MoveTemp(InPinType))
-				, RegistryInfo(InRegistryInfo)
-			{
-			}
-
-			// Get the corresponding icon brush for this type 
-			const FSlateBrush* GetIconBrush(const bool bIsConstructorType) const;
-		};
-
 		class METASOUNDEDITOR_API FMetasoundDefaultLiteralCustomizationBase
 		{
 		protected:
@@ -99,22 +82,22 @@ namespace Metasound
 			// By default, proxy classes support child classes & inheritance.
 			virtual void RegisterExplicitProxyClass(const UClass& InClass) = 0;
 
-			virtual const FEditorDataType* FindDataType(FName InDataTypeName) const = 0;
-			virtual const FEditorDataType& FindDataTypeChecked(FName InDataTypeName) const = 0;
 			virtual bool IsMetaSoundAssetClass(const FName InClassName) const = 0;
 
-			virtual bool IsRegisteredDataType(FName InDataTypeName) const = 0;
-
 			// Primes MetaSound assets, effectively loading the asset asynchronously (if not already
 			// loaded) & registers them if not already registered with the MetaSound Class Registry.
 			virtual void PrimeAssetRegistryAsync() = 0;
 			virtual EAssetPrimeStatus GetAssetRegistryPrimeStatus() const = 0;
 
-			virtual void IterateDataTypes(TUniqueFunction<void(const FEditorDataType&)> InDataTypeFunction) const = 0;
-
 			virtual TUniquePtr<FMetasoundDefaultLiteralCustomizationBase> CreateMemberDefaultLiteralCustomization(UClass& InClass, IDetailCategoryBuilder& DefaultCategoryBuilder) const = 0;
 
 			virtual const TSubclassOf<UMetasoundEditorGraphMemberDefaultLiteral> FindDefaultLiteralClass(EMetasoundFrontendLiteralType InLiteralType) const = 0;
+
+			virtual const FEdGraphPinType* FindPinType(FName InDataTypeName) const = 0;
+
+			virtual const FSlateBrush* GetIconBrush(FName InDataType, const bool bIsConstructorType) const = 0;
+
+			virtual void RegisterPinType(FName InDataTypeName, FName InPinCategory = { }, FName InPinSubCategory = { }) = 0;
 		};
 	} // namespace Editor
 } // namespace Metasound
