CIS build fix for MetasoundVertexData
#rnx

diff --git a/Engine/Plugins/Runtime/Metasound/Source/MetasoundGraphCore/Public/MetasoundVertexData.h b/Engine/Plugins/Runtime/Metasound/Source/MetasoundGraphCore/Public/MetasoundVertexData.h
--- a/Engine/Plugins/Runtime/Metasound/Source/MetasoundGraphCore/Public/MetasoundVertexData.h
+++ b/Engine/Plugins/Runtime/Metasound/Source/MetasoundGraphCore/Public/MetasoundVertexData.h
@@ -27,9 +27,9 @@ namespace Metasound
 		// Tests to see that the access type of the data reference is compatible 
 		// with the access type of the vertex.
 #if ENABLE_METASOUND_ACCESS_TYPE_COMPATIBILITY_TEST
-		void CheckAccessTypeCompatibility(const FDataVertex& InDataVertex, const FAnyDataReference& InDataReference);
+		void METASOUNDGRAPHCORE_API CheckAccessTypeCompatibility(const FDataVertex& InDataVertex, const FAnyDataReference& InDataReference);
 #else
-		void CheckAccessTypeCompatibility(const FDataVertex& InDataVertex, const FAnyDataReference& InDataReference) {}
+		FORCEINLINE void CheckAccessTypeCompatibility(const FDataVertex& InDataVertex, const FAnyDataReference& InDataReference) {}
 #endif // #if DO_CHECK
 
 		// Binds a vertex to a data reference.
