CIS fix for audio low level test failure.
#rnx

diff --git a/Engine/Plugins/Runtime/Metasound/Source/MetasoundFrontend/Private/MetasoundFrontendGraphController.cpp b/Engine/Plugins/Runtime/Metasound/Source/MetasoundFrontend/Private/MetasoundFrontendGraphController.cpp
--- a/Engine/Plugins/Runtime/Metasound/Source/MetasoundFrontend/Private/MetasoundFrontendGraphController.cpp
+++ b/Engine/Plugins/Runtime/Metasound/Source/MetasoundFrontend/Private/MetasoundFrontendGraphController.cpp
@@ -1524,7 +1524,8 @@ namespace Metasound
 					return TUniquePtr<IOperator>(nullptr);
 				}
 
-				FBuildGraphOperatorParams BuildParams { *Graph, InSettings, FInputVertexInterfaceData{}, InEnvironment };
+				FInputVertexInterfaceData InterfaceData;
+				FBuildGraphOperatorParams BuildParams { *Graph, InSettings, InterfaceData, InEnvironment };
 				return FOperatorBuilder(FOperatorBuilderSettings::GetDefaultSettings()).BuildGraphOperator(BuildParams, OutResults);
 			}
 			else
