Fix for implicit bool cast with MetaSound boolean connection animation throwing build warnings
#rnx

diff --git a/Engine/Plugins/Runtime/Metasound/Source/MetasoundEditor/Private/MetasoundEditorGraphConnectionDrawingPolicy.cpp b/Engine/Plugins/Runtime/Metasound/Source/MetasoundEditor/Private/MetasoundEditorGraphConnectionDrawingPolicy.cpp
--- a/Engine/Plugins/Runtime/Metasound/Source/MetasoundEditor/Private/MetasoundEditorGraphConnectionDrawingPolicy.cpp
+++ b/Engine/Plugins/Runtime/Metasound/Source/MetasoundEditor/Private/MetasoundEditorGraphConnectionDrawingPolicy.cpp
@@ -118,16 +118,23 @@ namespace Metasound
 				{
 					const FMetasoundFrontendEdgeStyleLiteralColorPair& PairA = InData.EdgeStyle->LiteralColorPairs[i];
 
+					float ValueFloat = static_cast<float>(Value);
 					TNumericType ValueA = InDefaultValue;
 					if (PairA.Value.TryGet(ValueA))
 					{
-						if (Value > ValueA)
+						const float ValueAFloat = static_cast<float>(ValueA);
+						if (ValueFloat > ValueAFloat)
 						{
 							const FMetasoundFrontendEdgeStyleLiteralColorPair& PairB = InData.EdgeStyle->LiteralColorPairs[i + 1];
 							TNumericType ValueB = InDefaultValue;
 							if (PairB.Value.TryGet(ValueB))
 							{
-								const float Alpha = (static_cast<float>(Value) - ValueA) / (static_cast<float>(ValueB) - ValueA);
+								float Denom = static_cast<float>(ValueB) - ValueAFloat;
+								if (FMath::IsNearlyZero(Denom))
+								{
+									Denom = SMALL_NUMBER;
+								}
+								const float Alpha = (ValueFloat - ValueAFloat) / Denom;
 								InnerColor = PairA.Color + FMath::Clamp(Alpha, 0.0f, 1.0f) * (PairB.Color - PairA.Color);
 							}
 							break;
