Properly update MetaSound node positions when nodes are moved with comment nodes or by alignment actions

diff --git a/Engine/Plugins/Runtime/Metasound/Source/MetasoundEditor/Private/MetasoundEditorGraphNode.cpp b/Engine/Plugins/Runtime/Metasound/Source/MetasoundEditor/Private/MetasoundEditorGraphNode.cpp
--- a/Engine/Plugins/Runtime/Metasound/Source/MetasoundEditor/Private/MetasoundEditorGraphNode.cpp
+++ b/Engine/Plugins/Runtime/Metasound/Source/MetasoundEditor/Private/MetasoundEditorGraphNode.cpp
@@ -74,19 +74,24 @@ UMetasoundEditorGraphNode::UMetasoundEditorGraphNode(const FObjectInitializer& O
 {
 }
 
-void UMetasoundEditorGraphNode::SetNodeLocation(const FVector2D& InLocation)
+void UMetasoundEditorGraphNode::UpdateFrontendNodeLocation(const FVector2D& InLocation)
 {
 	using namespace Metasound::Frontend;
-
-	NodePosX = FMath::TruncToInt(InLocation.X);
-	NodePosY = FMath::TruncToInt(InLocation.Y);
-
+	
 	FNodeHandle NodeHandle = GetNodeHandle();
 	FMetasoundFrontendNodeStyle Style = NodeHandle->GetNodeStyle();
 	Style.Display.Locations.FindOrAdd(NodeGuid) = InLocation;
 	NodeHandle->SetNodeStyle(Style);
 }
 
+void UMetasoundEditorGraphNode::SetNodeLocation(const FVector2D& InLocation)
+{
+	NodePosX = FMath::TruncToInt(InLocation.X);
+	NodePosY = FMath::TruncToInt(InLocation.Y);
+
+	UpdateFrontendNodeLocation(InLocation);
+}
+
 void UMetasoundEditorGraphNode::PostLoad()
 {
 	Super::PostLoad();
diff --git a/Engine/Plugins/Runtime/Metasound/Source/MetasoundEditor/Private/MetasoundEditorGraphNodeFactory.h b/Engine/Plugins/Runtime/Metasound/Source/MetasoundEditor/Private/MetasoundEditorGraphNodeFactory.h
--- a/Engine/Plugins/Runtime/Metasound/Source/MetasoundEditor/Private/MetasoundEditorGraphNodeFactory.h
+++ b/Engine/Plugins/Runtime/Metasound/Source/MetasoundEditor/Private/MetasoundEditorGraphNodeFactory.h
@@ -6,6 +6,7 @@
 #include "MetasoundEditorGraphNode.h"
 #include "SGraphNode.h"
 #include "SMetasoundGraphNode.h"
+#include "SMetasoundGraphNodeComment.h"
 #include "Templates/SharedPointer.h"
 #include "Widgets/DeclarativeSyntaxSupport.h"
 
@@ -20,6 +21,10 @@ class FMetasoundGraphNodeFactory : public FGraphPanelNodeFactory
 		{
 			return SNew(SMetaSoundGraphNode, InNode);
 		}
+		else if (UEdGraphNode_Comment* CommentNode = Cast<UEdGraphNode_Comment>(InNode))
+		{
+			return SNew(SMetasoundGraphNodeComment, CommentNode);
+		}
 
 		return nullptr;
 	}
diff --git a/Engine/Plugins/Runtime/Metasound/Source/MetasoundEditor/Private/MetasoundEditorGraphSchema.cpp b/Engine/Plugins/Runtime/Metasound/Source/MetasoundEditor/Private/MetasoundEditorGraphSchema.cpp
--- a/Engine/Plugins/Runtime/Metasound/Source/MetasoundEditor/Private/MetasoundEditorGraphSchema.cpp
+++ b/Engine/Plugins/Runtime/Metasound/Source/MetasoundEditor/Private/MetasoundEditorGraphSchema.cpp
@@ -1961,4 +1961,18 @@ TSharedPtr<FEdGraphSchemaAction> UMetasoundEditorGraphSchema::GetCreateCommentAc
 	TSharedPtr<FMetasoundGraphSchemaAction_NewComment> Comment = MakeShared<FMetasoundGraphSchemaAction_NewComment>();
 	return StaticCastSharedPtr<FEdGraphSchemaAction, FMetasoundGraphSchemaAction_NewComment>(Comment);
 }
+
+void UMetasoundEditorGraphSchema::SetNodePosition(UEdGraphNode* Node, const FVector2D& Position) const
+{
+	if (UMetasoundEditorGraphNode* MetasoundGraphNode = Cast<UMetasoundEditorGraphNode>(Node))
+	{
+		MetasoundGraphNode->GetMetasoundChecked().Modify();
+		UEdGraphSchema::SetNodePosition(Node, Position);
+		MetasoundGraphNode->UpdateFrontendNodeLocation(Position);
+	}
+	else
+	{
+		UEdGraphSchema::SetNodePosition(Node, Position);
+	}
+}
 #undef LOCTEXT_NAMESPACE
diff --git a/Engine/Plugins/Runtime/Metasound/Source/MetasoundEditor/Private/SMetasoundGraphNodeComment.cpp b/Engine/Plugins/Runtime/Metasound/Source/MetasoundEditor/Private/SMetasoundGraphNodeComment.cpp
new file mode 100644
--- /dev/null
+++ b/Engine/Plugins/Runtime/Metasound/Source/MetasoundEditor/Private/SMetasoundGraphNodeComment.cpp
@@ -0,0 +1,45 @@
+// Copyright Epic Games, Inc. All Rights Reserved.
+#include "SMetasoundGraphNodeComment.h"
+
+#include "EdGraphNode_Comment.h"
+#include "MetasoundEditorGraphNode.h"
+
+namespace Metasound
+{
+	namespace Editor
+	{
+		void SMetasoundGraphNodeComment::MoveTo(const FVector2D& NewPosition, FNodeSet& NodeFilter, bool bMarkDirty)
+		{
+			SGraphNodeComment::MoveTo(NewPosition, NodeFilter, bMarkDirty);
+
+			// Update Frontend node positions for unselected nodes that are dragged along with the comment box
+			// partially copied from SGraphNodeComment::MoveTo
+			// Don't drag note content if either of the shift keys are down.
+			FModifierKeysState KeysState = FSlateApplication::Get().GetModifierKeys();
+			if (!KeysState.IsShiftDown())
+			{
+				UEdGraphNode_Comment* CommentNode = Cast<UEdGraphNode_Comment>(GraphNode);
+				if (CommentNode && CommentNode->MoveMode == ECommentBoxMode::GroupMovement)
+				{
+					FVector2D PositionDelta = NewPosition - GetPosition();
+
+					// Now update any nodes which are touching the comment but *not* selected
+					// Selected nodes will be moved as part of the normal selection code
+					TSharedPtr< SGraphPanel > Panel = GetOwnerPanel();
+					for (FCommentNodeSet::TConstIterator NodeIt(CommentNode->GetNodesUnderComment()); NodeIt; ++NodeIt)
+					{
+						if (UMetasoundEditorGraphNode* MetasoundGraphNode = Cast<UMetasoundEditorGraphNode>(*NodeIt))
+						{
+							FVector2D MetasoundNodePosition = FVector2D(MetasoundGraphNode->NodePosX, MetasoundGraphNode->NodePosY);
+							MetasoundNodePosition += PositionDelta;
+
+							MetasoundGraphNode->GetMetasoundChecked().Modify();
+							MetasoundGraphNode->UpdateFrontendNodeLocation(MetasoundNodePosition);
+						}
+					}
+				}
+			}
+		}
+
+	} // namespace Editor
+} // namespace Metasound
diff --git a/Engine/Plugins/Runtime/Metasound/Source/MetasoundEditor/Private/SMetasoundGraphNodeComment.h b/Engine/Plugins/Runtime/Metasound/Source/MetasoundEditor/Private/SMetasoundGraphNodeComment.h
new file mode 100644
--- /dev/null
+++ b/Engine/Plugins/Runtime/Metasound/Source/MetasoundEditor/Private/SMetasoundGraphNodeComment.h
@@ -0,0 +1,19 @@
+// Copyright Epic Games, Inc. All Rights Reserved.
+#pragma once
+
+#include "Widgets/DeclarativeSyntaxSupport.h"
+#include "SGraphNodeComment.h"
+
+namespace Metasound
+{
+	namespace Editor
+	{
+		class SMetasoundGraphNodeComment : public SGraphNodeComment
+		{
+		public:
+			// SNodePanel::SNode interface
+			virtual void MoveTo(const FVector2D& NewPosition, FNodeSet& NodeFilter, bool bMarkDirty = true) override;
+			// End of SNodePanel::SNode interface
+		};
+	} // namespace Editor
+} // namespace Metasound
diff --git a/Engine/Plugins/Runtime/Metasound/Source/MetasoundEditor/Public/MetasoundEditorGraphNode.h b/Engine/Plugins/Runtime/Metasound/Source/MetasoundEditor/Public/MetasoundEditorGraphNode.h
--- a/Engine/Plugins/Runtime/Metasound/Source/MetasoundEditor/Public/MetasoundEditorGraphNode.h
+++ b/Engine/Plugins/Runtime/Metasound/Source/MetasoundEditor/Public/MetasoundEditorGraphNode.h
@@ -91,6 +91,8 @@ public:
 
 	// Sets the node's location, both on this graph member node and on the frontend handle
 	void SetNodeLocation(const FVector2D& InLocation);
+	// Helper function to update node location on frontend handle
+	void UpdateFrontendNodeLocation(const FVector2D& InLocation);
 
 	Metasound::Frontend::FGraphHandle GetRootGraphHandle() const;
 	Metasound::Frontend::FConstGraphHandle GetConstRootGraphHandle() const;
diff --git a/Engine/Plugins/Runtime/Metasound/Source/MetasoundEditor/Public/MetasoundEditorGraphSchema.h b/Engine/Plugins/Runtime/Metasound/Source/MetasoundEditor/Public/MetasoundEditorGraphSchema.h
--- a/Engine/Plugins/Runtime/Metasound/Source/MetasoundEditor/Public/MetasoundEditorGraphSchema.h
+++ b/Engine/Plugins/Runtime/Metasound/Source/MetasoundEditor/Public/MetasoundEditorGraphSchema.h
@@ -429,6 +429,7 @@ public:
 	virtual void DroppedAssetsOnPin(const TArray<FAssetData>& Assets, const FVector2D& GraphPosition, UEdGraphPin* Pin) const override;
 	virtual int32 GetNodeSelectionCount(const UEdGraph* Graph) const override;
 	virtual TSharedPtr<FEdGraphSchemaAction> GetCreateCommentAction() const override;
+	virtual void SetNodePosition(UEdGraphNode* Node, const FVector2D& Position) const;
 	//~ End EdGraphSchema Interface
 
 	void BreakNodeLinks(UEdGraphNode& TargetNode, bool bShouldActuallyTransact) const;
