Fix MetaSound editor meter channel count being incorrect resulting in incorrect meter values for mono sources

diff --git a/Engine/Plugins/Runtime/Metasound/Source/MetasoundEditor/Private/MetasoundEditor.cpp b/Engine/Plugins/Runtime/Metasound/Source/MetasoundEditor/Private/MetasoundEditor.cpp
--- a/Engine/Plugins/Runtime/Metasound/Source/MetasoundEditor/Private/MetasoundEditor.cpp
+++ b/Engine/Plugins/Runtime/Metasound/Source/MetasoundEditor/Private/MetasoundEditor.cpp
@@ -1150,7 +1150,7 @@
 				{
 					OutputMeter = MakeShared<FEditorMeter>();
 				}
-				OutputMeter->Init(EAudioBusChannels::Stereo, MetaSoundSource->NumChannels);
+				OutputMeter->Init(MetaSoundSource->NumChannels);
 			}
 			else
 			{
diff --git a/Engine/Plugins/Runtime/Metasound/Source/MetasoundEditor/Private/MetasoundEditorMeter.cpp b/Engine/Plugins/Runtime/Metasound/Source/MetasoundEditor/Private/MetasoundEditorMeter.cpp
--- a/Engine/Plugins/Runtime/Metasound/Source/MetasoundEditor/Private/MetasoundEditorMeter.cpp
+++ b/Engine/Plugins/Runtime/Metasound/Source/MetasoundEditor/Private/MetasoundEditorMeter.cpp
@@ -32,9 +32,9 @@
 			return Widget;
 		}
 
-		void FEditorMeter::Init(EAudioBusChannels InMaxBusChannels, int32 InNumChannels)
+		void FEditorMeter::Init(int32 InNumChannels)
 		{
-			check(InNumChannels <= static_cast<int32>(InMaxBusChannels) + 1);
+			check(InNumChannels > 0);
 
 			Settings = TStrongObjectPtr(NewObject<UMeterSettings>());
 			Settings->PeakHoldTime = 4000.0f;
@@ -43,7 +43,7 @@
 			Analyzer->Settings = Settings.Get();
 
 			AudioBus = TStrongObjectPtr(NewObject<UAudioBus>());
-			AudioBus->AudioBusChannels = InMaxBusChannels;
+			AudioBus->AudioBusChannels = EAudioBusChannels(InNumChannels - 1);
 
 			ResultsDelegateHandle = Analyzer->OnLatestPerChannelMeterResultsNative.AddRaw(this, &FEditorMeter::OnMeterOutput);
 
diff --git a/Engine/Plugins/Runtime/Metasound/Source/MetasoundEditor/Public/MetasoundEditorMeter.h b/Engine/Plugins/Runtime/Metasound/Source/MetasoundEditor/Public/MetasoundEditorMeter.h
--- a/Engine/Plugins/Runtime/Metasound/Source/MetasoundEditor/Public/MetasoundEditorMeter.h
+++ b/Engine/Plugins/Runtime/Metasound/Source/MetasoundEditor/Public/MetasoundEditorMeter.h
@@ -24,7 +24,7 @@
 
 			void Teardown();
 
-			void Init(EAudioBusChannels InMaxBusChannels, int32 InNumChannels);
+			void Init(int32 InNumChannels);
 
 		protected:
 			void OnMeterOutput(UMeterAnalyzer* InMeterAnalyzer, int32 ChannelIndex, const FMeterResults& InMeterResults);
