Metasound: Add constuctor vertex info to data type registration.

diff --git a/Engine/Plugins/Runtime/Metasound/Source/MetasoundFrontend/Private/MetasoundPrimitives.cpp b/Engine/Plugins/Runtime/Metasound/Source/MetasoundFrontend/Private/MetasoundPrimitives.cpp
--- a/Engine/Plugins/Runtime/Metasound/Source/MetasoundFrontend/Private/MetasoundPrimitives.cpp
+++ b/Engine/Plugins/Runtime/Metasound/Source/MetasoundFrontend/Private/MetasoundPrimitives.cpp
@@ -19,6 +19,13 @@
 		static constexpr bool Value = false;
 	};
 
+	// Disable constructor pins of audio buffers
+	template<>
+	struct TEnableConstructorVertex<FAudioBuffer>
+	{
+		static constexpr bool Value = false;
+	};
+
 	// Disable arrays of triggers
 	template<>
 	struct TEnableAutoArrayTypeRegistration<FTrigger>
@@ -26,6 +33,13 @@
 		static constexpr bool Value = false;
 	};
 
+	// Disable constructor pins of triggers 
+	template<>
+	struct TEnableConstructorVertex<FTrigger>
+	{
+		static constexpr bool Value = false;
+	};
+
 	// Disable auto-conversion using the FAudioBuffer(int32) constructor
 	template<typename FromDataType>
 	struct TEnableAutoConverterNodeRegistration<FromDataType, FAudioBuffer>
@@ -43,11 +57,8 @@
 
 REGISTER_METASOUND_DATATYPE(bool, "Bool", ::Metasound::ELiteralType::Boolean)
 REGISTER_METASOUND_DATATYPE(int32, "Int32", ::Metasound::ELiteralType::Integer)
-//REGISTER_METASOUND_DATATYPE(int64, "Int64", ::Metasound::ELiteralType::Integer)
 REGISTER_METASOUND_DATATYPE(float, "Float", ::Metasound::ELiteralType::Float)
-//REGISTER_METASOUND_DATATYPE(double, "Double", ::Metasound::ELiteralType::Float)
 REGISTER_METASOUND_DATATYPE(FString, "String", ::Metasound::ELiteralType::String)
-
 REGISTER_METASOUND_DATATYPE(Metasound::FTrigger, "Trigger", ::Metasound::ELiteralType::Boolean)
 REGISTER_METASOUND_DATATYPE(Metasound::FTime, "Time", ::Metasound::ELiteralType::Float)
 REGISTER_METASOUND_DATATYPE(Metasound::FAudioBuffer, "Audio")
diff --git a/Engine/Plugins/Runtime/Metasound/Source/MetasoundFrontend/Public/MetasoundDataTypeRegistrationMacro.h b/Engine/Plugins/Runtime/Metasound/Source/MetasoundFrontend/Public/MetasoundDataTypeRegistrationMacro.h
--- a/Engine/Plugins/Runtime/Metasound/Source/MetasoundFrontend/Public/MetasoundDataTypeRegistrationMacro.h
+++ b/Engine/Plugins/Runtime/Metasound/Source/MetasoundFrontend/Public/MetasoundDataTypeRegistrationMacro.h
@@ -57,23 +57,38 @@
 		static constexpr bool Value = true;
 	};
 
+	/** Enables or disables using a data type in constructor vertices. 
+	 * By default this is true and all data types that support copy assignment
+	 * and copy construction. */
+	template<typename ... Type>
+	struct TEnableConstructorVertex
+	{
+		static constexpr bool Value = true;
+	};
+
 	namespace MetasoundDataTypeRegistrationPrivate
 	{
-		// Helper utility to test if we can transmit a datatype between a send and a receive node.
+		// Helper utility to test if we a datatype is copyable
 		template <typename TDataType>
-		struct TIsTransmittable
+		struct TIsCopyable
 		{
 		private:
 			static constexpr bool bIsCopyConstructible = std::is_copy_constructible<TDataType>::value;
 			static constexpr bool bIsCopyAssignable = std::is_copy_assignable<TDataType>::value;
 
-			static constexpr bool bCanBeTransmitted = bIsCopyConstructible && bIsCopyAssignable;
-
 		public:
-
-			static constexpr bool Value = bCanBeTransmitted;
+			static constexpr bool Value = bIsCopyConstructible && bIsCopyAssignable;
 		};
 
+		// Helper utility to test if we can transmit a datatype between a send and a receive node.
+		template <typename TDataType>
+		struct TIsTransmittable
+		{
+			static constexpr bool Value = TIsCopyable<TDataType>::Value;
+		};
+
+
+
 		// Specialization of TIsTransmittable<> for TArray to handle lax TArray copy constructor 
 		// definition. This can be removed once updates to TArray are merged into the same codebase.
 		// TODO: delete me eventually.
@@ -236,6 +251,7 @@
 			RegistryInfo.bIsEnum = TEnumTraits<TDataType>::bIsEnum;
 			RegistryInfo.bIsVariable = TIsVariable<TDataType>::Value;
 			RegistryInfo.bIsTransmittable = TIsTransmittable<TDataType>::Value;
+			RegistryInfo.bIsConstructorType = TIsCopyable<TDataType>::Value && TEnableConstructorVertex<TDataType>::Value;
 			
 			if constexpr (std::is_base_of<UObject, UClassToUse>::value)
 			{
@@ -370,6 +386,7 @@
 						VariableDeferredAccessorClass = Metasound::Frontend::GenerateClass(VariableDeferredAccessorPrototype.GetMetadata(), EMetasoundFrontendClassType::VariableDeferredAccessor);
 					}
 				}
+
 				virtual ~FDataTypeRegistryEntry() {}
 
 				virtual const Frontend::FDataTypeRegistryInfo& GetDataTypeInfo() const override
diff --git a/Engine/Plugins/Runtime/Metasound/Source/MetasoundFrontend/Public/MetasoundFrontendDataTypeRegistry.h b/Engine/Plugins/Runtime/Metasound/Source/MetasoundFrontend/Public/MetasoundFrontendDataTypeRegistry.h
--- a/Engine/Plugins/Runtime/Metasound/Source/MetasoundFrontend/Public/MetasoundFrontendDataTypeRegistry.h
+++ b/Engine/Plugins/Runtime/Metasound/Source/MetasoundFrontend/Public/MetasoundFrontendDataTypeRegistry.h
@@ -55,6 +55,9 @@
 			// Is true if is a Variable type
 			bool bIsVariable = false;
 
+			// Is true if the data type can be used in on constructor vertices. 
+			bool bIsConstructorType = false;
+
 			// If provided in registration call, UClass this datatype was registered with
 			UClass* ProxyGeneratorClass = nullptr;
 
