Fix for frequency modulation going above Nyquist

diff --git a/Engine/Plugins/Runtime/Metasound/Source/MetasoundStandardNodes/Private/MetasoundOscillatorNodes.cpp b/Engine/Plugins/Runtime/Metasound/Source/MetasoundStandardNodes/Private/MetasoundOscillatorNodes.cpp
--- a/Engine/Plugins/Runtime/Metasound/Source/MetasoundStandardNodes/Private/MetasoundOscillatorNodes.cpp
+++ b/Engine/Plugins/Runtime/Metasound/Source/MetasoundStandardNodes/Private/MetasoundOscillatorNodes.cpp
@@ -224,6 +224,7 @@
 
 			void operator()(const FGeneratorArgs& InArgs)
 			{
+				float Nyquist = InArgs.SampleRate / 2.0f;
 				int32 RemainingSamplesInBlock = InArgs.AlignedBuffer.Num();
 				float* Out = InArgs.AlignedBuffer.GetData();
 				const float* FM = InArgs.FM.GetData();
@@ -237,7 +238,7 @@
 					{
 						while (RemainingSamplesInBlock > 0)
 						{
-							const float PerSampleFreq = InArgs.FrequencyHz + *FM++;
+							const float PerSampleFreq = FMath::Clamp(InArgs.FrequencyHz + *FM++, -Nyquist, Nyquist);
 							const float DeltaPhase = PerSampleFreq * OneOverSampleRate;
 
 							Wrap(Phase);
@@ -259,7 +260,7 @@
 					{
 						while (RemainingSamplesInBlock > 0)
 						{
-							const float ModulatedFreqSum = CurrentFreq + *FM++;
+							const float ModulatedFreqSum = FMath::Clamp(CurrentFreq + *FM++, -Nyquist, Nyquist);
 							const float DeltaPhase = ModulatedFreqSum * OneOverSampleRate;
 
 							Wrap(Phase);
