Doppler Pitch Shift Node

diff --git a/Engine/Plugins/Runtime/Metasound/Source/MetasoundStandardNodes/Private/MetasoundDopplerPitchShiftNode.cpp b/Engine/Plugins/Runtime/Metasound/Source/MetasoundStandardNodes/Private/MetasoundDopplerPitchShiftNode.cpp
new file mode 100644
--- /dev/null
+++ b/Engine/Plugins/Runtime/Metasound/Source/MetasoundStandardNodes/Private/MetasoundDopplerPitchShiftNode.cpp
@@ -0,0 +1,157 @@
+// Copyright Epic Games, Inc. All Rights Reserved.
+
+#include "Internationalization/Text.h"
+#include "MetasoundExecutableOperator.h"
+#include "MetasoundNodeRegistrationMacro.h"
+#include "MetasoundParamHelper.h"
+#include "MetasoundPrimitives.h"
+#include "MetasoundStandardNodesNames.h"
+#include "MetasoundAudioBuffer.h"
+#include "DSP/TapDelayPitchShifter.h"
+#include "MetasoundStandardNodesCategories.h"
+#include "MetasoundFacade.h"
+
+#define LOCTEXT_NAMESPACE "MetasoundStandardNodesDopplerPitchShift"
+
+namespace Metasound
+{
+	namespace DopplerPitchShift
+	{
+		METASOUND_PARAM(InParamAudioInput, "In", "Input audio buffer.")
+		METASOUND_PARAM(InParamPitchShift, "Pitch Shift", "The amount to pitch shift the audio signal, in semitones.")
+		METASOUND_PARAM(InParamDelayLength, "Delay Length", "The delay length of the internal delay buffer in milliseconds (10 ms to 100 ms). Changing this can reduce artifacts in certain pitch shift regions.")
+		METASOUND_PARAM(OutParamAudio, "Out", "Output audio buffer.")
+	}
+
+	class FDopplerPitchShiftOperator : public TExecutableOperator<FDopplerPitchShiftOperator>
+	{
+	public:
+		static const FNodeClassMetadata& GetNodeInfo();
+		static const FVertexInterface& GetVertexInterface();
+		static TUniquePtr<IOperator> CreateOperator(const FCreateOperatorParams& InParams, FBuildErrorArray& OutErrors);
+
+		FDopplerPitchShiftOperator(const FOperatorSettings& InSettings, const FAudioBufferReadRef& InAudioInput, const FFloatReadRef& InPitchShift, const FFloatReadRef& InDelayLength);
+
+		virtual FDataReferenceCollection GetInputs() const override;
+		virtual FDataReferenceCollection GetOutputs() const override;
+		void Execute();
+
+	private:
+		FAudioBufferReadRef AudioInput;
+		FFloatReadRef PitchShift;
+		FFloatReadRef DelayLength;
+		FAudioBufferWriteRef AudioOutput;
+		Audio::FDelay DelayBuffer;
+		Audio::FTapDelayPitchShifter DelayPitchShifter;
+	};
+
+	FDopplerPitchShiftOperator::FDopplerPitchShiftOperator(const FOperatorSettings& InSettings,
+		const FAudioBufferReadRef& InAudioInput,
+		const FFloatReadRef& InPitchShift,
+		const FFloatReadRef& InDelayLength)
+
+		: AudioInput(InAudioInput)
+		, PitchShift(InPitchShift)
+		, DelayLength(InDelayLength)
+		, AudioOutput(FAudioBufferWriteRef::CreateNew(InSettings))
+	{
+		DelayBuffer.Init(InSettings.GetSampleRate(), Audio::FTapDelayPitchShifter::MaxDelayLength);
+		DelayPitchShifter.Init(InSettings.GetSampleRate(), *PitchShift, *DelayLength);
+	}
+
+	FDataReferenceCollection FDopplerPitchShiftOperator::GetInputs() const
+	{
+		using namespace DopplerPitchShift;
+
+		FDataReferenceCollection InputDataReferences;
+		InputDataReferences.AddDataReadReference(METASOUND_GET_PARAM_NAME(InParamAudioInput), FAudioBufferReadRef(AudioInput));
+		InputDataReferences.AddDataReadReference(METASOUND_GET_PARAM_NAME(InParamPitchShift), FFloatReadRef(PitchShift));
+		InputDataReferences.AddDataReadReference(METASOUND_GET_PARAM_NAME(InParamDelayLength), FFloatReadRef(DelayLength));
+
+		return InputDataReferences;
+	}
+
+	FDataReferenceCollection FDopplerPitchShiftOperator::GetOutputs() const
+	{
+		using namespace DopplerPitchShift;
+
+		FDataReferenceCollection OutputDataReferences;
+		OutputDataReferences.AddDataReadReference(METASOUND_GET_PARAM_NAME(OutParamAudio), FAudioBufferReadRef(AudioOutput));
+		return OutputDataReferences;
+	}
+
+	void FDopplerPitchShiftOperator::Execute()
+	{
+		DelayPitchShifter.SetDelayLength(*DelayLength);
+		DelayPitchShifter.SetPitchShift(*PitchShift);
+		DelayPitchShifter.ProcessAudio(DelayBuffer, AudioInput->GetData(), AudioInput->Num(), AudioOutput->GetData());
+	}
+
+	const FVertexInterface& FDopplerPitchShiftOperator::GetVertexInterface()
+	{
+		using namespace DopplerPitchShift;
+
+		static const FVertexInterface Interface(
+			FInputVertexInterface(
+				TInputDataVertex<FAudioBuffer>(METASOUND_GET_PARAM_NAME_AND_METADATA(InParamAudioInput)),
+				TInputDataVertex<float>(METASOUND_GET_PARAM_NAME_AND_METADATA(InParamPitchShift), 0.0f),
+				TInputDataVertex<float>(METASOUND_GET_PARAM_NAME_AND_METADATA(InParamDelayLength), 30.0f)
+			),
+			FOutputVertexInterface(
+				TOutputDataVertex<FAudioBuffer>(METASOUND_GET_PARAM_NAME_AND_METADATA(OutParamAudio))
+			)
+		);
+
+		return Interface;
+	}
+
+	const FNodeClassMetadata& FDopplerPitchShiftOperator::GetNodeInfo()
+	{
+		auto InitNodeInfo = []() -> FNodeClassMetadata
+		{
+			FNodeClassMetadata Info;
+			Info.ClassName = { StandardNodes::Namespace, "Doppler Pitch Shift", StandardNodes::AudioVariant };
+			Info.MajorVersion = 1;
+			Info.MinorVersion = 0;
+			Info.DisplayName = METASOUND_LOCTEXT("DelayNode_DisplayName", "Doppler Pitch Shift");
+			Info.Description = METASOUND_LOCTEXT("DelayNode_Description", "Pitch shifts the audio buffer using a doppler-delay method.");
+			Info.Author = PluginAuthor;
+			Info.PromptIfMissing = PluginNodeMissingPrompt;
+			Info.DefaultInterface = GetVertexInterface();
+			Info.CategoryHierarchy.Emplace(NodeCategories::Delays);
+			return Info;
+		};
+
+		static const FNodeClassMetadata Info = InitNodeInfo();
+
+		return Info;
+	}
+
+	TUniquePtr<IOperator> FDopplerPitchShiftOperator::CreateOperator(const FCreateOperatorParams& InParams, FBuildErrorArray& OutErrors)
+	{
+		using namespace DopplerPitchShift;
+
+		const FDataReferenceCollection& InputCollection = InParams.InputDataReferences;
+		const FInputVertexInterface& InputInterface = GetVertexInterface().GetInputInterface();
+
+		FAudioBufferReadRef AudioIn = InputCollection.GetDataReadReferenceOrConstruct<FAudioBuffer>(METASOUND_GET_PARAM_NAME(InParamAudioInput), InParams.OperatorSettings);
+		FFloatReadRef PitchShift = InputCollection.GetDataReadReferenceOrConstructWithVertexDefault<float>(InputInterface, METASOUND_GET_PARAM_NAME(InParamPitchShift), InParams.OperatorSettings);
+		FFloatReadRef DelayLength = InputCollection.GetDataReadReferenceOrConstructWithVertexDefault<float>(InputInterface, METASOUND_GET_PARAM_NAME(InParamDelayLength), InParams.OperatorSettings);
+
+		return MakeUnique<FDopplerPitchShiftOperator>(InParams.OperatorSettings, AudioIn, PitchShift, DelayLength);
+	}
+
+	class FDopperPitchShiftNode : public FNodeFacade
+	{
+	public:
+		FDopperPitchShiftNode(const FNodeInitData& InitData)
+			: FNodeFacade(InitData.InstanceName, InitData.InstanceID, TFacadeOperatorClass<FDopplerPitchShiftOperator>())
+		{
+		}
+	};
+
+
+	METASOUND_REGISTER_NODE(FDopperPitchShiftNode)
+}
+
+#undef LOCTEXT_NAMESPACE
diff --git a/Engine/Source/Runtime/SignalProcessing/Private/TapDelayPitchShifter.cpp b/Engine/Source/Runtime/SignalProcessing/Private/TapDelayPitchShifter.cpp
new file mode 100644
--- /dev/null
+++ b/Engine/Source/Runtime/SignalProcessing/Private/TapDelayPitchShifter.cpp
@@ -0,0 +1,117 @@
+// Copyright Epic Games, Inc. All Rights Reserved.
+
+#include "DSP/TapDelayPitchShifter.h"
+
+namespace Audio
+{
+	namespace PitchShiftUtils
+	{
+		static float GetPitchShiftClamped(const float InPitchShift)
+		{
+			return FMath::Clamp(InPitchShift, -12.0f * FTapDelayPitchShifter::MaxAbsPitchShiftInOctaves, 12.0f * FTapDelayPitchShifter::MaxAbsPitchShiftInOctaves);
+		}
+
+		static float GetDelayLengthClamped(const float InBufferLength)
+		{
+			return FMath::Clamp(InBufferLength, FTapDelayPitchShifter::MinDelayLength, FTapDelayPitchShifter::MaxDelayLength);	
+		}
+	}
+
+	FTapDelayPitchShifter::FTapDelayPitchShifter()
+	{
+	}
+
+	FTapDelayPitchShifter::~FTapDelayPitchShifter()
+	{
+	}
+
+	void FTapDelayPitchShifter::Init(const float InSampleRate, const float InPitchShift, const float InDelayLength)
+	{
+		SampleRate = InSampleRate;
+		CurrentTargetDelayLength = PitchShiftUtils::GetDelayLengthClamped(InDelayLength);
+		CurrentDelayLength.Init(CurrentTargetDelayLength);
+		CurrentPitchShift = PitchShiftUtils::GetPitchShiftClamped(InPitchShift);
+		CurrentPitchShiftRatio = Audio::GetFrequencyMultiplier(CurrentPitchShift);
+
+		UpdatePhasorPhaseIncrement();
+	}
+	
+	void FTapDelayPitchShifter::UpdatePhasorPhaseIncrement()
+	{
+		const float CurrentLengthSecondsClamped = 0.001f * FMath::Max(CurrentDelayLength.PeekCurrentValue(), 1.0f);
+		const float PhasorFrequency =  (1.0f - CurrentPitchShiftRatio) / CurrentLengthSecondsClamped;
+		PhasorPhaseIncrement =  PhasorFrequency / SampleRate;
+	}
+	
+ 
+	void FTapDelayPitchShifter::SetDelayLength(const float InDelayLength)
+	{
+		const float NewDelayLength = PitchShiftUtils::GetDelayLengthClamped(InDelayLength);
+		if (!FMath::IsNearlyEqual(NewDelayLength, CurrentTargetDelayLength))
+		{
+			CurrentTargetDelayLength = InDelayLength;
+			CurrentDelayLength.SetValue(NewDelayLength);
+			UpdatePhasorPhaseIncrement();
+		}
+	}
+
+	void FTapDelayPitchShifter::SetPitchShift(const float InPitchShift)
+	{
+		const float NewPitchShift = PitchShiftUtils::GetPitchShiftClamped(InPitchShift);
+		if (!FMath::IsNearlyEqual(NewPitchShift, CurrentPitchShift))
+		{
+			CurrentPitchShift = NewPitchShift;
+			CurrentPitchShiftRatio = Audio::GetFrequencyMultiplier(CurrentPitchShift);
+			UpdatePhasorPhaseIncrement();
+		}
+	}
+
+	void FTapDelayPitchShifter::SetPitchShiftRatio(const float InPitchShiftRatio)
+	{
+		if (!FMath::IsNearlyEqual(InPitchShiftRatio, CurrentPitchShiftRatio))
+		{
+			CurrentPitchShiftRatio = InPitchShiftRatio;
+			UpdatePhasorPhaseIncrement();
+		}
+	}
+
+	float FTapDelayPitchShifter::ReadDopplerShiftedTapFromDelay(const Audio::FDelay& InDelayBuffer, const float InReadOffsetMilliseconds)
+	{
+		// Update the interpolated delay length value
+		if (!CurrentDelayLength.IsDone())
+		{
+			UpdatePhasorPhaseIncrement();
+			CurrentDelayLength.GetNextValue();
+		}
+		
+		// Compute the two tap delay read locations, one shifted 90 degrees out of phase
+		const float PhasorPhaseOffset = FMath::Fmod(PhasorPhase + 0.5f, 1.0f);
+		const float DelayTapRead1 = InReadOffsetMilliseconds + CurrentDelayLength.PeekCurrentValue() * PhasorPhase;
+		const float DelayTapRead2 = InReadOffsetMilliseconds + CurrentDelayLength.PeekCurrentValue() * PhasorPhaseOffset;
+
+		// This produces an overlapping cosine function that avoids pops in the output
+		const float DelayTapGain1 = FMath::Cos(PI * (PhasorPhase - 0.5f));
+		const float DelayTapGain2 = FMath::Cos(PI * (PhasorPhaseOffset - 0.5f));
+
+		// Update the phasor state
+		PhasorPhase += PhasorPhaseIncrement;
+		// Make sure we wrap to between 0.0 and 1.0 after incrementing the phase 
+		PhasorPhase = FMath::Wrap(PhasorPhase, 0.0f, 1.0f);
+		
+		// Read the delay lines at the given tap indices, apply the gains
+		const float Sample1 = DelayTapGain1 * InDelayBuffer.ReadDelayAt(DelayTapRead1);
+		const float Sample2 = DelayTapGain2 * InDelayBuffer.ReadDelayAt(DelayTapRead2);
+
+		return Sample1 + Sample2;
+	}
+
+	void FTapDelayPitchShifter::ProcessAudio(Audio::FDelay& InDelayBuffer, const float* InAudioBuffer, const int32 InNumFrames, float* OutAudioBuffer)
+	{
+		for (int32 FrameIndex = 0; FrameIndex < InNumFrames; ++FrameIndex)
+		{
+			OutAudioBuffer[FrameIndex] = ReadDopplerShiftedTapFromDelay(InDelayBuffer);
+			InDelayBuffer.WriteDelayAndInc(InAudioBuffer[FrameIndex]);
+		}
+	}
+
+}
diff --git a/Engine/Source/Runtime/SignalProcessing/Public/DSP/TapDelayPitchShifter.h b/Engine/Source/Runtime/SignalProcessing/Public/DSP/TapDelayPitchShifter.h
new file mode 100644
--- /dev/null
+++ b/Engine/Source/Runtime/SignalProcessing/Public/DSP/TapDelayPitchShifter.h
@@ -0,0 +1,55 @@
+// Copyright Epic Games, Inc. All Rights Reserved.
+
+#pragma once
+
+#include "CoreMinimal.h"
+#include "DSP/Dsp.h"
+#include "DSP/Delay.h"
+
+namespace Audio
+{
+	// Class which reads doppler-shifted tap from a delay buffer 
+	class SIGNALPROCESSING_API FTapDelayPitchShifter
+	{
+	public:
+		static constexpr float MinDelayLength = 10.0f;
+		static constexpr float MaxDelayLength = 100.0f;
+		static constexpr float MaxAbsPitchShiftInOctaves = 6.0f;
+		
+		// Constructor
+		FTapDelayPitchShifter();
+
+		// Virtual Destructor
+		virtual ~FTapDelayPitchShifter();
+
+		// Initialization of the internal delay with given sample rate, pitch shift in semitones, and delay length in milliseconds
+		void Init(const float InSampleRate, const float InPitchShift, const float InDelayLength);
+
+		// Sets the internal delay line length.
+		void SetDelayLength(const float InDelayLength);
+
+		// Sets the pitch shift in semitones
+		void SetPitchShift(const float InPitchScaleSemitones);
+
+		// Sets the pitch shift in terms of a sample rate ratio
+		void SetPitchShiftRatio(const float InPitchShiftRatio);
+		
+		// Reads the next sample tap from the input delay buffer
+		float ReadDopplerShiftedTapFromDelay(const Audio::FDelay& InDelayBuffer, const float ReadOffsetMilliseconds = 0.0f);
+
+		// Process whole block of audio reading to and writing from delay line from input buffer
+		void ProcessAudio(Audio::FDelay& InDelayBuffer, const float* InAudioBuffer, const int32 InNumFrames, float* OutAudioBuffer);
+		
+	private:
+		void UpdatePhasorPhaseIncrement();
+		
+		float SampleRate = 0.0f;
+		Audio::FExponentialEase CurrentDelayLength;
+		float CurrentTargetDelayLength = 0.0f;
+		float CurrentPitchShift = 0.0f;
+		float CurrentPitchShiftRatio = 1.0f;
+		float PhasorPhase = 0.0f;
+		float PhasorPhaseIncrement = 0.0f;
+	};
+
+}
