Fixed up TObjectPtr for Metasounds as I had warnings as errors on

diff --git a/Engine/Plugins/Runtime/Metasound/Source/MetasoundEditor/Private/MetasoundEditorGraphMemberDefaults.cpp b/Engine/Plugins/Runtime/Metasound/Source/MetasoundEditor/Private/MetasoundEditorGraphMemberDefaults.cpp
--- a/Engine/Plugins/Runtime/Metasound/Source/MetasoundEditor/Private/MetasoundEditorGraphMemberDefaults.cpp
+++ b/Engine/Plugins/Runtime/Metasound/Source/MetasoundEditor/Private/MetasoundEditorGraphMemberDefaults.cpp
@@ -341,7 +341,9 @@ EMetasoundFrontendLiteralType UMetasoundEditorGraphMemberDefaultObject::GetLiter
 
 void UMetasoundEditorGraphMemberDefaultObject::SetFromLiteral(const FMetasoundFrontendLiteral& InLiteral)
 {
-	ensure(InLiteral.TryGet(Default.Object));
+	UObject* Object = nullptr;
+	ensure(InLiteral.TryGet(Object));
+	Default.Object = Object;
 }
 
 void UMetasoundEditorGraphMemberDefaultObject::UpdatePreviewInstance(const Metasound::FVertexName& InParameterName, TScriptInterface<IAudioParameterControllerInterface>& InParameterInterface) const
diff --git a/Engine/Plugins/Runtime/Metasound/Source/MetasoundEditor/Public/MetasoundEditorGraph.h b/Engine/Plugins/Runtime/Metasound/Source/MetasoundEditor/Public/MetasoundEditorGraph.h
--- a/Engine/Plugins/Runtime/Metasound/Source/MetasoundEditor/Public/MetasoundEditorGraph.h
+++ b/Engine/Plugins/Runtime/Metasound/Source/MetasoundEditor/Public/MetasoundEditorGraph.h
@@ -172,7 +172,7 @@ public:
 protected:
 	/** Default literal value of member */
 	UPROPERTY()
-	UMetasoundEditorGraphMemberDefaultLiteral* Literal;
+	TObjectPtr<UMetasoundEditorGraphMemberDefaultLiteral> Literal;
 
 	/** Conforms literal object type to member's DataType */
 	void ConformLiteralDataType();
diff --git a/Engine/Plugins/Runtime/Metasound/Source/MetasoundEditor/Public/MetasoundEditorGraphMemberDefaults.h b/Engine/Plugins/Runtime/Metasound/Source/MetasoundEditor/Public/MetasoundEditorGraphMemberDefaults.h
--- a/Engine/Plugins/Runtime/Metasound/Source/MetasoundEditor/Public/MetasoundEditorGraphMemberDefaults.h
+++ b/Engine/Plugins/Runtime/Metasound/Source/MetasoundEditor/Public/MetasoundEditorGraphMemberDefaults.h
@@ -246,7 +246,7 @@ struct FMetasoundEditorGraphMemberDefaultObjectRef
 	GENERATED_BODY()
 
 	UPROPERTY(EditAnywhere, Category = DefaultValue)
-	UObject* Object = nullptr;
+	TObjectPtr<UObject> Object;
 };
 
 UCLASS(MinimalAPI)
diff --git a/Engine/Plugins/Runtime/Metasound/Source/MetasoundEditor/Public/MetasoundEditorGraphNode.h b/Engine/Plugins/Runtime/Metasound/Source/MetasoundEditor/Public/MetasoundEditorGraphNode.h
--- a/Engine/Plugins/Runtime/Metasound/Source/MetasoundEditor/Public/MetasoundEditorGraphNode.h
+++ b/Engine/Plugins/Runtime/Metasound/Source/MetasoundEditor/Public/MetasoundEditorGraphNode.h
@@ -166,7 +166,7 @@ class UMetasoundEditorGraphOutputNode : public UMetasoundEditorGraphMemberNode
 
 public:
 	UPROPERTY()
-	UMetasoundEditorGraphOutput* Output;
+	TObjectPtr<UMetasoundEditorGraphOutput> Output;
 
 	virtual FMetasoundFrontendClassName GetClassName() const override;
 	virtual FGuid GetNodeID() const override;
@@ -258,7 +258,7 @@ protected:
 public:
 	// Associated graph variable.
 	UPROPERTY()
-	UMetasoundEditorGraphVariable* Variable;
+	TObjectPtr<UMetasoundEditorGraphVariable> Variable;
 
 	// Variables do not have titles to distinguish more visually from vertex types
 	virtual void CacheTitle() override { }
diff --git a/Engine/Plugins/Runtime/Metasound/Source/MetasoundEngine/Public/Metasound.h b/Engine/Plugins/Runtime/Metasound/Source/MetasoundEngine/Public/Metasound.h
--- a/Engine/Plugins/Runtime/Metasound/Source/MetasoundEngine/Public/Metasound.h
+++ b/Engine/Plugins/Runtime/Metasound/Source/MetasoundEngine/Public/Metasound.h
@@ -154,7 +154,7 @@ protected:
 
 #if WITH_EDITORONLY_DATA
 	UPROPERTY()
-	UMetasoundEditorGraphBase* Graph;
+	TObjectPtr<UMetasoundEditorGraphBase> Graph;
 #endif // WITH_EDITORONLY_DATA
 
 public:
diff --git a/Engine/Plugins/Runtime/Metasound/Source/MetasoundEngine/Public/MetasoundSource.h b/Engine/Plugins/Runtime/Metasound/Source/MetasoundEngine/Public/MetasoundSource.h
--- a/Engine/Plugins/Runtime/Metasound/Source/MetasoundEngine/Public/MetasoundSource.h
+++ b/Engine/Plugins/Runtime/Metasound/Source/MetasoundEngine/Public/MetasoundSource.h
@@ -52,7 +52,7 @@ protected:
 
 #if WITH_EDITORONLY_DATA
 	UPROPERTY()
-	UMetasoundEditorGraphBase* Graph;
+	TObjectPtr<UMetasoundEditorGraphBase> Graph;
 #endif // WITH_EDITORONLY_DATA
 
 public:
diff --git a/Engine/Plugins/Runtime/Metasound/Source/MetasoundFrontend/Private/MetasoundFrontendLiteral.cpp b/Engine/Plugins/Runtime/Metasound/Source/MetasoundFrontend/Private/MetasoundFrontendLiteral.cpp
--- a/Engine/Plugins/Runtime/Metasound/Source/MetasoundFrontend/Private/MetasoundFrontendLiteral.cpp
+++ b/Engine/Plugins/Runtime/Metasound/Source/MetasoundFrontend/Private/MetasoundFrontendLiteral.cpp
@@ -29,9 +29,9 @@ namespace MetasoundFrontendLiteralPrivate
 
 	// String conversion specialization for UObject*
 	template<>
-	struct TLiteralValueToStringHelper<UObject*>
+	struct TLiteralValueToStringHelper<TObjectPtr<UObject>>
 	{
-		static FString Convert(const UObject* InObject)
+		static FString Convert(const TObjectPtr<UObject> InObject)
 		{
 			// Use empty string for null object as this is recognized
 			// as empty value in editor context (the string "nullptr"
diff --git a/Engine/Plugins/Runtime/Metasound/Source/MetasoundFrontend/Public/MetasoundFrontendLiteral.h b/Engine/Plugins/Runtime/Metasound/Source/MetasoundFrontend/Public/MetasoundFrontendLiteral.h
--- a/Engine/Plugins/Runtime/Metasound/Source/MetasoundFrontend/Public/MetasoundFrontendLiteral.h
+++ b/Engine/Plugins/Runtime/Metasound/Source/MetasoundFrontend/Public/MetasoundFrontendLiteral.h
@@ -88,7 +88,7 @@ private:
 	TArray<FString> AsString;
 
 	UPROPERTY()
-	TArray<UObject*> AsUObject;
+	TArray<TObjectPtr<UObject>> AsUObject;
 
 public:
 	// Returns true if the stored Type is an array type.
