Fix for Cook warning/content validation

diff --git a/Engine/Plugins/Runtime/Metasound/Source/MetasoundEditor/Public/MetasoundEditorGraphMemberDefaults.h b/Engine/Plugins/Runtime/Metasound/Source/MetasoundEditor/Public/MetasoundEditorGraphMemberDefaults.h
--- a/Engine/Plugins/Runtime/Metasound/Source/MetasoundEditor/Public/MetasoundEditorGraphMemberDefaults.h
+++ b/Engine/Plugins/Runtime/Metasound/Source/MetasoundEditor/Public/MetasoundEditorGraphMemberDefaults.h
@@ -246,7 +246,7 @@ struct FMetasoundEditorGraphMemberDefaultObjectRef
 	GENERATED_BODY()
 
 	UPROPERTY(EditAnywhere, Category = DefaultValue)
-	TObjectPtr<UObject> Object;
+	TObjectPtr<UObject> Object = nullptr;
 };
 
 UCLASS(MinimalAPI)
