MetaSound: Add vertex access type to frontend document in order to support constructor pins.

diff --git a/Engine/Plugins/Runtime/Metasound/Source/MetasoundFrontend/Private/MetasoundFrontendDocument.cpp b/Engine/Plugins/Runtime/Metasound/Source/MetasoundFrontend/Private/MetasoundFrontendDocument.cpp
--- a/Engine/Plugins/Runtime/Metasound/Source/MetasoundFrontend/Private/MetasoundFrontendDocument.cpp
+++ b/Engine/Plugins/Runtime/Metasound/Source/MetasoundFrontend/Private/MetasoundFrontendDocument.cpp
@@ -87,7 +87,20 @@ namespace Metasound
 
 			return false;
 		}
-	}
+
+		EMetasoundFrontendVertexAccessType CoreVertexAccessTypeToFrontendVertexAccessType(Metasound::EVertexAccessType InAccessType)
+		{
+			switch (InAccessType)
+			{
+				case EVertexAccessType::Value:
+					return EMetasoundFrontendVertexAccessType::Value;
+
+				case EVertexAccessType::Reference:
+				default:
+					return EMetasoundFrontendVertexAccessType::Reference;
+			}
+		}
+	} // namespace DocumentPrivate
 } // namespace Metasound
 
 FMetasoundFrontendNodeInterface::FMetasoundFrontendNodeInterface(const FMetasoundFrontendClassInterface& InClassInterface)
@@ -148,7 +161,7 @@ void FMetasoundFrontendClassVertex::SplitName(FName& OutNamespace, FName& OutPar
 
 bool FMetasoundFrontendClassVertex::IsFunctionalEquivalent(const FMetasoundFrontendClassVertex& InLHS, const FMetasoundFrontendClassVertex& InRHS)
 {
-	return FMetasoundFrontendVertex::IsFunctionalEquivalent(InLHS, InRHS);
+	return FMetasoundFrontendVertex::IsFunctionalEquivalent(InLHS, InRHS) && (InLHS.AccessType == InRHS.AccessType);
 }
 
 FMetasoundFrontendClassName::FMetasoundFrontendClassName(const FName& InNamespace, const FName& InName, const FName& InVariant)
@@ -210,6 +223,7 @@ FMetasoundFrontendClassInterface FMetasoundFrontendClassInterface::GenerateClass
 			ClassInput.Name = InputVertex.VertexName;
 			ClassInput.TypeName = InputVertex.DataTypeName;
 			ClassInput.VertexID = FGuid::NewGuid();
+			ClassInput.AccessType = Metasound::DocumentPrivate::CoreVertexAccessTypeToFrontendVertexAccessType(InputVertex.AccessType);
 
 
 #if WITH_EDITOR
@@ -264,6 +278,7 @@ FMetasoundFrontendClassInterface FMetasoundFrontendClassInterface::GenerateClass
 			ClassOutput.Name = OutputVertex.VertexName;
 			ClassOutput.TypeName = OutputVertex.DataTypeName;
 			ClassOutput.VertexID = FGuid::NewGuid();
+			ClassOutput.AccessType = Metasound::DocumentPrivate::CoreVertexAccessTypeToFrontendVertexAccessType(OutputVertex.AccessType);
 
 #if WITH_EDITOR
 			const FDataVertexMetadata& VertexMetadata = OutputVertex.Metadata;
diff --git a/Engine/Plugins/Runtime/Metasound/Source/MetasoundFrontend/Public/MetasoundFrontendDocument.h b/Engine/Plugins/Runtime/Metasound/Source/MetasoundFrontend/Public/MetasoundFrontendDocument.h
--- a/Engine/Plugins/Runtime/Metasound/Source/MetasoundFrontend/Public/MetasoundFrontendDocument.h
+++ b/Engine/Plugins/Runtime/Metasound/Source/MetasoundFrontend/Public/MetasoundFrontendDocument.h
@@ -47,6 +47,13 @@ namespace Metasound
 	} // namespace Frontend
 } // namespace Metasound
 
+// Describes how a vertex accesses the data connected to it. 
+UENUM()
+enum class EMetasoundFrontendVertexAccessType
+{
+	Reference,	//< The vertex accesses data by reference.
+	Value 		//< The vertex accesses data by value.
+};
 
 UENUM()
 enum class EMetasoundFrontendClassType : uint8
@@ -680,6 +687,9 @@ struct METASOUNDFRONTEND_API FMetasoundFrontendClassVertex : public FMetasoundFr
 	FMetasoundFrontendVertexMetadata Metadata;
 #endif // WITH_EDITORONLY_DATA
 
+	UPROPERTY()
+	EMetasoundFrontendVertexAccessType AccessType = EMetasoundFrontendVertexAccessType::Reference;
+
 	// Splits name into namespace & parameter name
 	void SplitName(FName& OutNamespace, FName& OutParameterName) const;
 
diff --git a/Engine/Plugins/Runtime/Metasound/Source/MetasoundGraphCore/Private/MetasoundVertex.cpp b/Engine/Plugins/Runtime/Metasound/Source/MetasoundGraphCore/Private/MetasoundVertex.cpp
--- a/Engine/Plugins/Runtime/Metasound/Source/MetasoundGraphCore/Private/MetasoundVertex.cpp
+++ b/Engine/Plugins/Runtime/Metasound/Source/MetasoundGraphCore/Private/MetasoundVertex.cpp
@@ -152,3 +152,18 @@ namespace Metasound
 		return !(InLHS == InRHS);
 	}
 }
+
+FString LexToString(Metasound::EVertexAccessType InAccessType)
+{
+	using namespace Metasound;
+
+	switch (InAccessType)
+	{
+		case EVertexAccessType::Value:
+			return TEXT("Value");
+
+		case EVertexAccessType::Reference:
+		default:
+			return TEXT("Reference");
+	}
+}
diff --git a/Engine/Plugins/Runtime/Metasound/Source/MetasoundGraphCore/Private/MetasoundVertexData.cpp b/Engine/Plugins/Runtime/Metasound/Source/MetasoundGraphCore/Private/MetasoundVertexData.cpp
--- a/Engine/Plugins/Runtime/Metasound/Source/MetasoundGraphCore/Private/MetasoundVertexData.cpp
+++ b/Engine/Plugins/Runtime/Metasound/Source/MetasoundGraphCore/Private/MetasoundVertexData.cpp
@@ -18,21 +18,12 @@ namespace Metasound
 
 			switch (InDataVertex.AccessType)
 			{
-				case EDataReferenceAccessType::None:
-					bIsCompatible = true;
+				case EVertexAccessType::Reference:
+					// Reference vertices can accept read, write and value types
+					bIsCompatible = (EDataReferenceAccessType::Read == ReferenceAccessType) || (EDataReferenceAccessType::Write == ReferenceAccessType) || (EDataReferenceAccessType::Value == ReferenceAccessType);
 					break;
 
-				case EDataReferenceAccessType::Read:
-					// Read vertices can be a read, write or value data reference type
-					bIsCompatible = EDataReferenceAccessType::None != ReferenceAccessType;
-					break;
-
-				case EDataReferenceAccessType::Write:
-					// Write vertices require a write data reference type
-					bIsCompatible = EDataReferenceAccessType::Write == ReferenceAccessType;
-					break;
-
-				case EDataReferenceAccessType::Value:
+				case EVertexAccessType::Value:
 					// value vertices require a value data reference type
 					bIsCompatible = EDataReferenceAccessType::Value == ReferenceAccessType;
 					break;
diff --git a/Engine/Plugins/Runtime/Metasound/Source/MetasoundGraphCore/Public/MetasoundVertex.h b/Engine/Plugins/Runtime/Metasound/Source/MetasoundGraphCore/Public/MetasoundVertex.h
--- a/Engine/Plugins/Runtime/Metasound/Source/MetasoundGraphCore/Public/MetasoundVertex.h
+++ b/Engine/Plugins/Runtime/Metasound/Source/MetasoundGraphCore/Public/MetasoundVertex.h
@@ -23,6 +23,13 @@ namespace Metasound
 		bool bIsAdvancedDisplay = false;
 	};
 
+	/** Describe how the vertex will access connected data. */
+	enum class EVertexAccessType
+	{
+		Reference, //< Vertex accesses the data references
+		Value      //< Vertex accesses the data value.
+	};
+
 
 	/** FDataVertex
 	 *
@@ -42,7 +49,7 @@ namespace Metasound
 		 * @InMetadata - Metadata pertaining to given vertex.
 		 * @InAccessType - The access type of the vertex.
 		 */
-		FDataVertex(const FVertexName& InVertexName, const FName& InDataTypeName, const FDataVertexMetadata& InMetadata, EDataReferenceAccessType InAccessType)
+		FDataVertex(const FVertexName& InVertexName, const FName& InDataTypeName, const FDataVertexMetadata& InMetadata, EVertexAccessType InAccessType)
 		: VertexName(InVertexName)
 		, DataTypeName(InDataTypeName)
 #if WITH_EDITORONLY_DATA
@@ -67,7 +74,7 @@ namespace Metasound
 #endif // WITH_EDITORONLY_DATA
 
 		/** Access type of the vertex. */
-		EDataReferenceAccessType AccessType;
+		EVertexAccessType AccessType;
 
 		UE_DEPRECATED(5.1, "GetVertexName() is deprecated. Use VertexName")
 		const FVertexName& GetVertexName() const
@@ -180,7 +187,7 @@ namespace Metasound
 		FInputDataVertex() = default;
 
 		/** Construct an FInputDataVertex. */
-		FInputDataVertex(const FVertexName& InVertexName, const FName& InDataTypeName, const FDataVertexMetadata& InMetadata, EDataReferenceAccessType InAccessType=EDataReferenceAccessType::Read)
+		FInputDataVertex(const FVertexName& InVertexName, const FName& InDataTypeName, const FDataVertexMetadata& InMetadata, EVertexAccessType InAccessType=EVertexAccessType::Reference)
 			: FDataVertex(InVertexName, InDataTypeName, InMetadata, InAccessType)
 			, LiteralFactory(FLiteral::FNone{})
 		{
@@ -188,16 +195,16 @@ namespace Metasound
 
 		/** Construct an FInputDataVertex with a default literal. */
 		template<typename LiteralValueType>
-		UE_DEPRECATED(5.1, "Use constructor with EDataReferenceAccessType argument.")
+		UE_DEPRECATED(5.1, "Use constructor with EVertexAccessType argument.")
 		FInputDataVertex(const FVertexName& InVertexName, const FName& InDataTypeName, const FDataVertexMetadata& InMetadata, const LiteralValueType& InLiteralValue)
-			: FDataVertex(InVertexName, InDataTypeName, InMetadata, EDataReferenceAccessType::Read)
+			: FDataVertex(InVertexName, InDataTypeName, InMetadata, EVertexAccessType::Reference)
 			, LiteralFactory(InLiteralValue)
 		{
 		}
 
 		/** Construct an FInputDataVertex with a default literal. */
 		template<typename LiteralValueType>
-		FInputDataVertex(const FVertexName& InVertexName, const FName& InDataTypeName, const FDataVertexMetadata& InMetadata, EDataReferenceAccessType InAccessType, const LiteralValueType& InLiteralValue)
+		FInputDataVertex(const FVertexName& InVertexName, const FName& InDataTypeName, const FDataVertexMetadata& InMetadata, EVertexAccessType InAccessType, const LiteralValueType& InLiteralValue)
 			: FDataVertex(InVertexName, InDataTypeName, InMetadata, InAccessType)
 			, LiteralFactory(InLiteralValue)
 		{
@@ -227,14 +234,14 @@ namespace Metasound
 
 		template<typename... ArgTypes>
 		TInputDataVertex(const FVertexName& InVertexName, const FDataVertexMetadata& InMetadata, ArgTypes&&... Args)
-		: FInputDataVertex(InVertexName, GetMetasoundDataTypeName<DataType>(), InMetadata, EDataReferenceAccessType::Read, Forward<ArgTypes>(Args)...)
+		: FInputDataVertex(InVertexName, GetMetasoundDataTypeName<DataType>(), InMetadata, EVertexAccessType::Reference, Forward<ArgTypes>(Args)...)
 		{
 		}
 
 		template<typename... ArgTypes>
  		UE_DEPRECATED(5.1, "TInputDataVertex constructor without explicit FDataVertexMatadata is deprecated. Use constructor which accepts TDataVertexMetadata")
 		TInputDataVertex(const FVertexName& InVertexName, const FText& InDescription, ArgTypes&&... Args)
-		: FInputDataVertex(InVertexName, GetMetasoundDataTypeName<DataType>(), { InDescription }, EDataReferenceAccessType::Read, Forward<ArgTypes>(Args)...)
+		: FInputDataVertex(InVertexName, GetMetasoundDataTypeName<DataType>(), { InDescription }, EVertexAccessType::Reference, Forward<ArgTypes>(Args)...)
 		{
 		}
 	};
@@ -252,7 +259,7 @@ namespace Metasound
 
 		template<typename... ArgTypes>
 		TInputConstructorVertex(const FVertexName& InVertexName, const FDataVertexMetadata& InMetadata, ArgTypes&&... Args)
-		: FInputDataVertex(InVertexName, GetMetasoundDataTypeName<DataType>(), InMetadata, EDataReferenceAccessType::Value, Forward<ArgTypes>(Args)...)
+		: FInputDataVertex(InVertexName, GetMetasoundDataTypeName<DataType>(), InMetadata, EVertexAccessType::Value, Forward<ArgTypes>(Args)...)
 		{
 		}
 	};
@@ -281,14 +288,14 @@ namespace Metasound
 
 		template<typename... ArgTypes>
 		TOutputDataVertex(const FVertexName& InVertexName, const FDataVertexMetadata& InMetadata, ArgTypes&&... Args)
-		: FOutputDataVertex(InVertexName, GetMetasoundDataTypeName<DataType>(), InMetadata, EDataReferenceAccessType::Read, Forward<ArgTypes>(Args)...)
+		: FOutputDataVertex(InVertexName, GetMetasoundDataTypeName<DataType>(), InMetadata, EVertexAccessType::Reference, Forward<ArgTypes>(Args)...)
 		{
 		}
 
 		template<typename... ArgTypes>
  		UE_DEPRECATED(5.1, "TOutputDataVertex constructor without explicit FDataVertexMatadata is deprecated. Use constructor which accepts TDataVertexMetadata")
 		TOutputDataVertex(const FVertexName& InVertexName, const FText& InDescription, ArgTypes&&... Args)
-		: FOutputDataVertex(InVertexName, GetMetasoundDataTypeName<DataType>(), { InDescription }, EDataReferenceAccessType::Read, Forward<ArgTypes>(Args)...)
+		: FOutputDataVertex(InVertexName, GetMetasoundDataTypeName<DataType>(), { InDescription }, EVertexAccessType::Reference, Forward<ArgTypes>(Args)...)
 		{
 		}
 	};
@@ -308,7 +315,7 @@ namespace Metasound
 
 		template<typename... ArgTypes>
 		TOutputConstructorVertex(const FVertexName& InVertexName, const FDataVertexMetadata& InMetadata, ArgTypes&&... Args)
-		: FOutputDataVertex(InVertexName, GetMetasoundDataTypeName<DataType>(), InMetadata, EDataReferenceAccessType::Value, Forward<ArgTypes>(Args)...)
+		: FOutputDataVertex(InVertexName, GetMetasoundDataTypeName<DataType>(), InMetadata, EVertexAccessType::Value, Forward<ArgTypes>(Args)...)
 		{
 		}
 	};
@@ -661,3 +668,7 @@ namespace Metasound
 		FGuid InstanceID;
 	};
 }
+
+/** Convert EVertexAccessType to string */
+METASOUNDGRAPHCORE_API FString LexToString(Metasound::EVertexAccessType InAccessType);
+
