Trigger route fix part 2

- Better fix, CL 20297675 introduced

diff --git a/Engine/Plugins/Runtime/Metasound/Source/MetasoundStandardNodes/Private/MetasoundTriggerRouteNode.cpp b/Engine/Plugins/Runtime/Metasound/Source/MetasoundStandardNodes/Private/MetasoundTriggerRouteNode.cpp
--- a/Engine/Plugins/Runtime/Metasound/Source/MetasoundStandardNodes/Private/MetasoundTriggerRouteNode.cpp
+++ b/Engine/Plugins/Runtime/Metasound/Source/MetasoundStandardNodes/Private/MetasoundTriggerRouteNode.cpp
@@ -171,6 +171,20 @@ namespace Metasound
 		{
 			check(InputValues.Num() > 0)
 			CurrentIndex = 0;
+
+			for (uint32 i = 0; i < NumInputs; ++i)
+			{
+				InputTriggers[i]->ExecuteBlock(
+					[&](int32 StartFrame, int32 EndFrame)
+					{
+					},
+					[this, i](int32 StartFrame, int32 EndFrame)
+					{
+						CurrentIndex = i;
+					}
+				);
+			}
+
 			*OutputValue = *InputValues[CurrentIndex];
 		}
 
@@ -215,11 +229,12 @@ namespace Metasound
 					[this, i](int32 StartFrame, int32 EndFrame)
 					{
 						CurrentIndex = i;
-						*OutputValue = *InputValues[CurrentIndex];
 						OutputTrigger->TriggerFrame(StartFrame);
 					}
 				);
-			}			
+			}
+
+			*OutputValue = *InputValues[CurrentIndex];
 		}
 
 	private:
