Trigger route fix

diff --git a/Engine/Plugins/Runtime/Metasound/Source/MetasoundStandardNodes/Private/MetasoundTriggerRouteNode.cpp b/Engine/Plugins/Runtime/Metasound/Source/MetasoundStandardNodes/Private/MetasoundTriggerRouteNode.cpp
--- a/Engine/Plugins/Runtime/Metasound/Source/MetasoundStandardNodes/Private/MetasoundTriggerRouteNode.cpp
+++ b/Engine/Plugins/Runtime/Metasound/Source/MetasoundStandardNodes/Private/MetasoundTriggerRouteNode.cpp
@@ -171,7 +171,8 @@ namespace Metasound
 		{
 			check(InputValues.Num() > 0)
 			CurrentIndex = 0;
-
+			
+			// Do the initial trigger input to determine the init index
 			for (uint32 i = 0; i < NumInputs; ++i)
 			{
 				InputTriggers[i]->ExecuteBlock(
@@ -184,7 +185,7 @@ namespace Metasound
 					}
 				);
 			}
-
+			
 			*OutputValue = *InputValues[CurrentIndex];
 		}
 
