Grain Delay Node

- Fix for compressor for limiter case (removed unncessarily clamp on detector output)

diff --git a/Engine/Plugins/Runtime/Metasound/Source/MetasoundStandardNodes/Private/MetasoundGrainDelayNode.cpp b/Engine/Plugins/Runtime/Metasound/Source/MetasoundStandardNodes/Private/MetasoundGrainDelayNode.cpp
new file mode 100644
--- /dev/null
+++ b/Engine/Plugins/Runtime/Metasound/Source/MetasoundStandardNodes/Private/MetasoundGrainDelayNode.cpp
@@ -0,0 +1,345 @@
+// Copyright Epic Games, Inc. All Rights Reserved.
+
+#include "MetasoundEnumRegistrationMacro.h"
+#include "MetasoundExecutableOperator.h"
+#include "MetasoundNodeRegistrationMacro.h"
+#include "MetasoundDataTypeRegistrationMacro.h"
+#include "MetasoundParamHelper.h"
+#include "MetasoundPrimitives.h"
+#include "MetasoundStandardNodesNames.h"
+#include "MetasoundTrigger.h"
+#include "MetasoundTime.h"
+#include "MetasoundAudioBuffer.h"
+#include "DSP/Delay.h"
+#include "MetasoundStandardNodesCategories.h"
+#include "MetasoundFacade.h"
+#include "DSP/GrainDelay.h"
+
+#define LOCTEXT_NAMESPACE "MetasoundStandardNodes_GrainDelayNode"
+
+namespace Metasound
+{
+	namespace GrainDelay
+	{
+		METASOUND_PARAM(InAudio, "In Audio", "Audio buffer input to be grain delayed.")
+		METASOUND_PARAM(InGrainSpawnTrigger, "Grain Spawn", "Spawns a new grain of audio.")
+		METASOUND_PARAM(InGrainDelay, "Grain Delay", "The delay of the next spawned grain (in milliseconds) between 0 and the maximum delay.")
+		METASOUND_PARAM(InGrainDelayDelta, "Grain Delay Range", "A random grain offset (in milliseconds) before and after the grain delay.")
+		METASOUND_PARAM(InGrainDuration, "Grain Duration", "The duration of the next spawned grain (in milliseconds).")
+		METASOUND_PARAM(InGrainDurationDelta, "Grain Duration Range", "A random grain duration range (in milliseconds) before and after the grain delay.")
+		METASOUND_PARAM(InGrainPitchShift, "Pitch Shift", "A pitch value (in semitones) to change the grain pitch of all rendering grains.")
+		METASOUND_PARAM(InGrainPitchShiftDelta, "Pitch Shift Range", "A random pitch shift delta (in semitones) randomly chosen when a grain is spawned.")
+		METASOUND_PARAM(InGrainDelayEnvelope, "Grain Envelope", "The type of envelope to use for the grains.")
+		METASOUND_PARAM(InGrainMaxCount, "Max Grain Count", "The maximum number of grains to render at a time (between 1 and 100). More grains will cost more CPU and potentially clip.")
+		METASOUND_PARAM(InGrainDelayFeedbackAmount, "Feedback Amount", "The amount of feedback of each grain. The grain delay will feed it's audio output back into itself.")
+
+		METASOUND_PARAM(OutAudio, "Out Audio", "Output audio buffer which has been grain-delayed.")
+	}
+	
+	DECLARE_METASOUND_ENUM(Audio::Grain::EEnvelope, Audio::Grain::EEnvelope::Gaussian, METASOUNDSTANDARDNODES_API, FEnumGrainDelayEnvelope, FEnumGrainDelayEnvelopeInfo, FGrainDelayEnvelopeReadRef, FEnumGrainDelayEnvelopeWriteRef);
+
+	DEFINE_METASOUND_ENUM_BEGIN(Audio::Grain::EEnvelope, FEnumGrainDelayEnvelope, "GrainDelayEnvelope")
+		DEFINE_METASOUND_ENUM_ENTRY(Audio::Grain::EEnvelope::Gaussian, "GrainDelayEnvelopeGaussianDescription", "Gaussian", "GrainDelayEnvelopeGaussianDescriptionTT", "A gaussian envelope."),
+		DEFINE_METASOUND_ENUM_ENTRY(Audio::Grain::EEnvelope::Triangle, "GrainDelayEnvelopeTriangleDescription", "Triangle", "GrainDelayEnvelopeTriangleDescriptionTT", "A triangular envelope."),
+		DEFINE_METASOUND_ENUM_ENTRY(Audio::Grain::EEnvelope::DownwardTriangle, "GrainDelayEnvelopeDownwardTriangleDescription", "Downward Triangle", "GrainDelayEnvelopeTriangleDescriptionTT", "A downward triangular envelope."),
+		DEFINE_METASOUND_ENUM_ENTRY(Audio::Grain::EEnvelope::UpwardTriangle, "GrainDelayEnvelopeUpwardTriangleDescription", "Upward Triangle", "GrainDelayEnvelopeUpwardTriangleDescriptionTT", "An upward triangular envelope."),
+		DEFINE_METASOUND_ENUM_ENTRY(Audio::Grain::EEnvelope::ExponentialDecay, "GrainDelayEnvelopeExponentialDecayDescription", "Exponential Decay", "GrainDelayEnvelopeExponentialDecayDescriptionTT", "An exponential decay envelope."),
+		DEFINE_METASOUND_ENUM_ENTRY(Audio::Grain::EEnvelope::ExponentialAttack, "GrainDelayEnvelopeExponentialAttackDescription", "Exponential Attack", "GrainDelayEnvelopeExponentialAttackDescriptionTT", "An exponential attack envelope."),
+	DEFINE_METASOUND_ENUM_END()
+
+	class FGrainDelayOperator final : public TExecutableOperator<FGrainDelayOperator>
+	{
+	public:
+
+		FGrainDelayOperator(const FOperatorSettings& InSettings,
+			const FAudioBufferReadRef& InAudioInput,
+			const FTriggerReadRef& InGrainSpawnTrigger,
+			const FFloatReadRef& InGrainDelay,
+			const FFloatReadRef& InGrainDelayDelta,
+			const FFloatReadRef& InGrainDuration,
+			const FFloatReadRef& InGrainDurationDelta,
+			const FFloatReadRef& InGrainPitchShift,
+			const FFloatReadRef& InGrainPitchShiftDelta,
+			const FGrainDelayEnvelopeReadRef& InGrainDelayEnvelope,
+			const FInt32ReadRef& InGrainMaxCount,	
+			const FFloatReadRef& InGrainDelayFeedbackAmount);
+
+		static const FNodeClassMetadata& GetNodeInfo();
+		static const FVertexInterface& GetVertexInterface();
+		static TUniquePtr<IOperator> CreateOperator(const FCreateOperatorParams& InParams, FBuildErrorArray& OutErrors);
+		virtual FDataReferenceCollection GetInputs() const override;
+		virtual FDataReferenceCollection GetOutputs() const override;
+		void Execute();
+
+	private:
+		const FOperatorSettings OperatorSettings;
+
+		// Input references
+		FAudioBufferReadRef AudioInput;
+		FTriggerReadRef GrainSpawnTrigger;
+		FFloatReadRef GrainDelay;
+		FFloatReadRef GrainDelayDelta;
+		FFloatReadRef GrainDuration;
+		FFloatReadRef GrainDurationDelta;
+		FFloatReadRef GrainPitchShift;
+		FFloatReadRef GrainPitchShiftDelta;
+		FGrainDelayEnvelopeReadRef GrainDelayEnvelope;
+		FInt32ReadRef GrainMaxCount;
+		FFloatReadRef GrainDelayFeedbackAmount;
+
+		// Output references
+		FAudioBufferWriteRef AudioOutput;
+
+		// Cached values from input refs
+		int32 PreviousGrainDelayMsec = 0;
+		float PreviousPitchShift = 0.0f;
+		Audio::Grain::EEnvelope PreviousGrainEnvelopeType = Audio::Grain::EEnvelope::Gaussian;
+
+		// Grain manager manages rendering grains
+		Audio::GrainDelay::FGrainDelay GrainDelayProcessor;
+	};
+	
+	FGrainDelayOperator::FGrainDelayOperator(const FOperatorSettings& InSettings,
+											  const FAudioBufferReadRef& InAudioInput,
+											  const FTriggerReadRef& InGrainSpawnTrigger,
+											  const FFloatReadRef& InGrainDelay,
+											  const FFloatReadRef& InGrainDelayDelta,
+											  const FFloatReadRef& InGrainDuration,
+											  const FFloatReadRef& InGrainDurationDelta,
+											  const FFloatReadRef& InGrainPitchShift,
+											  const FFloatReadRef& InGrainPitchShiftDelta,
+											  const FGrainDelayEnvelopeReadRef& InGrainDelayEnvelope,
+											  const FInt32ReadRef& InGrainMaxCount,
+											  const FFloatReadRef& InGrainDelayFeedbackAmount)
+		: OperatorSettings(InSettings),
+	      AudioInput(InAudioInput),
+		  GrainSpawnTrigger(InGrainSpawnTrigger),
+		  GrainDelay(InGrainDelay),
+		  GrainDelayDelta(InGrainDelayDelta),
+		  GrainDuration(InGrainDuration),
+		  GrainDurationDelta(InGrainDurationDelta),
+		  GrainPitchShift(InGrainPitchShift),
+		  GrainPitchShiftDelta(InGrainPitchShiftDelta),
+		  GrainDelayEnvelope(InGrainDelayEnvelope),
+		  GrainMaxCount(InGrainMaxCount),
+		  GrainDelayFeedbackAmount(InGrainDelayFeedbackAmount),
+	      AudioOutput(FAudioBufferWriteRef::CreateNew(OperatorSettings)),
+		  GrainDelayProcessor(OperatorSettings.GetSampleRate())
+	{
+		PreviousGrainDelayMsec = GrainDelayProcessor.GetGrainDelayClamped(*GrainDelay);
+		PreviousGrainEnvelopeType = *InGrainDelayEnvelope;
+		PreviousPitchShift = GrainDelayProcessor.GetGrainPitchShiftClamped(*GrainPitchShift);
+		
+		GrainDelayProcessor.SetGrainEnvelope(*GrainDelayEnvelope);
+		GrainDelayProcessor.SetGrainBasePitchShiftRatio(GrainDelayProcessor.GetGrainPitchShiftFrameRatio(PreviousPitchShift));
+	}
+
+	const FNodeClassMetadata& FGrainDelayOperator::GetNodeInfo()
+	{
+		auto InitNodeInfo = []() -> FNodeClassMetadata
+		{
+			FNodeClassMetadata Info;
+			Info.ClassName = { StandardNodes::Namespace,TEXT("GrainDelayNode"), StandardNodes::AudioVariant };
+			Info.MajorVersion = 1;
+			Info.MinorVersion = 0;
+			Info.DisplayName = METASOUND_LOCTEXT("GrainDelayNodeDisplayName", "Grain Delay");
+			Info.Description = METASOUND_LOCTEXT("Metasound_GrainDelayNodeDescription", "Performs delayed audio granulation on stereo input audio buffers.");
+			Info.Author = PluginAuthor;
+			Info.PromptIfMissing = PluginNodeMissingPrompt;
+			Info.DefaultInterface = GetVertexInterface();
+			Info.CategoryHierarchy.Emplace(NodeCategories::Delays);
+			return Info;
+		};
+
+		static const FNodeClassMetadata Info = InitNodeInfo();
+
+		return Info;
+	}
+
+	const FVertexInterface& FGrainDelayOperator::GetVertexInterface()
+	{
+		using namespace GrainDelay;
+
+		auto CreateDefaultInterface = []()-> FVertexInterface
+		{
+			FInputVertexInterface InputInterface;
+			InputInterface.Add(TInputDataVertex<FAudioBuffer>(METASOUND_GET_PARAM_NAME_AND_METADATA(InAudio)));
+			InputInterface.Add(TInputDataVertex<FTrigger>(METASOUND_GET_PARAM_NAME_AND_METADATA(InGrainSpawnTrigger)));
+			InputInterface.Add(TInputDataVertex<float>(METASOUND_GET_PARAM_NAME_AND_METADATA(InGrainDelay), 128.0f));
+			InputInterface.Add(TInputDataVertex<float>(METASOUND_GET_PARAM_NAME_AND_METADATA(InGrainDelayDelta), 0.0f));
+			InputInterface.Add(TInputDataVertex<float>(METASOUND_GET_PARAM_NAME_AND_METADATA(InGrainDuration), 128.0f));
+			InputInterface.Add(TInputDataVertex<float>(METASOUND_GET_PARAM_NAME_AND_METADATA(InGrainDurationDelta), 0.0f));
+			InputInterface.Add(TInputDataVertex<float>(METASOUND_GET_PARAM_NAME_AND_METADATA(InGrainPitchShift), 0.0f));
+			InputInterface.Add(TInputDataVertex<float>(METASOUND_GET_PARAM_NAME_AND_METADATA(InGrainPitchShiftDelta), 0.0f));
+			InputInterface.Add(TInputDataVertex<FEnumGrainDelayEnvelope>(METASOUND_GET_PARAM_NAME_AND_METADATA(InGrainDelayEnvelope), 0));
+			InputInterface.Add(TInputDataVertex<int32>(METASOUND_GET_PARAM_NAME_AND_METADATA(InGrainMaxCount), 16));
+			InputInterface.Add(TInputDataVertex<float>(METASOUND_GET_PARAM_NAME_AND_METADATA(InGrainDelayFeedbackAmount), 0.0f));
+
+			FOutputVertexInterface OutputInterface;
+			OutputInterface.Add(TOutputDataVertex<FAudioBuffer>(METASOUND_GET_PARAM_NAME_AND_METADATA(OutAudio)));
+
+			return FVertexInterface(InputInterface, OutputInterface);
+		};
+
+		static const FVertexInterface DefaultInterface = CreateDefaultInterface();
+		return DefaultInterface;
+	}
+
+	TUniquePtr<IOperator> FGrainDelayOperator::CreateOperator(const FCreateOperatorParams& InParams, FBuildErrorArray& OutErrors)
+	{
+		using namespace GrainDelay;
+
+		const FDataReferenceCollection& InputCollection = InParams.InputDataReferences;
+		const FInputVertexInterface& InputInterface = GetVertexInterface().GetInputInterface();
+		
+		FAudioBufferReadRef AudioInput = InputCollection.GetDataReadReferenceOrConstruct<FAudioBuffer>(METASOUND_GET_PARAM_NAME(InAudio), InParams.OperatorSettings);
+		FTriggerReadRef GrainSpawnTrigger = InputCollection.GetDataReadReferenceOrConstruct<FTrigger>(METASOUND_GET_PARAM_NAME(InGrainSpawnTrigger), InParams.OperatorSettings);
+		FFloatReadRef GrainDelay = InputCollection.GetDataReadReferenceOrConstructWithVertexDefault<float>(InputInterface, METASOUND_GET_PARAM_NAME(InGrainDelay), InParams.OperatorSettings);
+		FFloatReadRef GrainDelayDelta = InputCollection.GetDataReadReferenceOrConstructWithVertexDefault<float>(InputInterface, METASOUND_GET_PARAM_NAME(InGrainDelayDelta), InParams.OperatorSettings);
+		FFloatReadRef GrainDuration = InputCollection.GetDataReadReferenceOrConstructWithVertexDefault<float>(InputInterface, METASOUND_GET_PARAM_NAME(InGrainDuration), InParams.OperatorSettings);
+		FFloatReadRef GrainDurationDelta = InputCollection.GetDataReadReferenceOrConstructWithVertexDefault<float>(InputInterface, METASOUND_GET_PARAM_NAME(InGrainDurationDelta), InParams.OperatorSettings);
+		FFloatReadRef GrainPitchShift = InputCollection.GetDataReadReferenceOrConstructWithVertexDefault<float>(InputInterface, METASOUND_GET_PARAM_NAME(InGrainPitchShift), InParams.OperatorSettings);
+		FFloatReadRef GrainPitchShiftDelta = InputCollection.GetDataReadReferenceOrConstructWithVertexDefault<float>(InputInterface, METASOUND_GET_PARAM_NAME(InGrainPitchShiftDelta), InParams.OperatorSettings);
+		FGrainDelayEnvelopeReadRef GrainDelayEnvelope = InputCollection.GetDataReadReferenceOrConstruct<FEnumGrainDelayEnvelope>(METASOUND_GET_PARAM_NAME(InGrainDelayEnvelope));
+		FInt32ReadRef GrainMaxCount = InputCollection.GetDataReadReferenceOrConstructWithVertexDefault<int32>(InputInterface, METASOUND_GET_PARAM_NAME(InGrainMaxCount), InParams.OperatorSettings);
+		FFloatReadRef GrainDelayFeedbackAmount = InputCollection.GetDataReadReferenceOrConstructWithVertexDefault<float>(InputInterface, METASOUND_GET_PARAM_NAME(InGrainDelayFeedbackAmount), InParams.OperatorSettings);
+
+		return MakeUnique<FGrainDelayOperator>(InParams.OperatorSettings, 
+			AudioInput,
+			GrainSpawnTrigger,
+			GrainDelay,
+			GrainDelayDelta,
+			GrainDuration,
+			GrainDurationDelta,
+			GrainPitchShift,
+			GrainPitchShiftDelta,
+			GrainDelayEnvelope,
+			GrainMaxCount,
+			GrainDelayFeedbackAmount);
+	}
+	
+	FDataReferenceCollection FGrainDelayOperator::GetOutputs() const
+	{
+		using namespace GrainDelay;
+
+		FDataReferenceCollection OutputDataReferences;
+		OutputDataReferences.AddDataReadReference(METASOUND_GET_PARAM_NAME(OutAudio), FAudioBufferReadRef(AudioOutput));
+		return OutputDataReferences;
+	}
+	
+	FDataReferenceCollection FGrainDelayOperator::GetInputs() const
+	{
+		using namespace GrainDelay;
+
+		FDataReferenceCollection InputDataReferences;
+		InputDataReferences.AddDataReadReference(METASOUND_GET_PARAM_NAME(InAudio), FAudioBufferReadRef(AudioInput));
+		InputDataReferences.AddDataReadReference(METASOUND_GET_PARAM_NAME(InGrainSpawnTrigger), FTriggerReadRef(GrainSpawnTrigger));
+		InputDataReferences.AddDataReadReference(METASOUND_GET_PARAM_NAME(InGrainDelay), FFloatReadRef(GrainDelay));
+		InputDataReferences.AddDataReadReference(METASOUND_GET_PARAM_NAME(InGrainDelayDelta), FFloatReadRef(GrainDelayDelta));
+		InputDataReferences.AddDataReadReference(METASOUND_GET_PARAM_NAME(InGrainDuration), FFloatReadRef(GrainDuration));
+		InputDataReferences.AddDataReadReference(METASOUND_GET_PARAM_NAME(InGrainDurationDelta), FFloatReadRef(GrainDurationDelta));
+		InputDataReferences.AddDataReadReference(METASOUND_GET_PARAM_NAME(InGrainPitchShift), FFloatReadRef(GrainPitchShift));
+		InputDataReferences.AddDataReadReference(METASOUND_GET_PARAM_NAME(InGrainPitchShiftDelta), FFloatReadRef(GrainPitchShiftDelta));
+		InputDataReferences.AddDataReadReference(METASOUND_GET_PARAM_NAME(InGrainDelayEnvelope), FGrainDelayEnvelopeReadRef(GrainDelayEnvelope));
+		InputDataReferences.AddDataReadReference(METASOUND_GET_PARAM_NAME(InGrainMaxCount), FInt32ReadRef(GrainMaxCount));
+		InputDataReferences.AddDataReadReference(METASOUND_GET_PARAM_NAME(InGrainDelayFeedbackAmount), FFloatReadRef(GrainDelayFeedbackAmount));
+
+		return InputDataReferences;
+	}
+
+	void FGrainDelayOperator::Execute()
+	{
+		// Parse triggers and render audio correctly in the audio block
+		int32 SpawnTrigIndex = 0;
+		int32 NextSpawnFrame = 0;
+		const int32 NumSpawnTriggers = GrainSpawnTrigger->NumTriggeredInBlock();
+
+		int32 CurrAudioFrame = 0;
+		int32 NextAudioFrame = 0;
+		const int32 LastAudioFrame = OperatorSettings.GetNumFramesPerBlock() - 1;
+		const int32 NoTrigger = OperatorSettings.GetNumFramesPerBlock() << 1;
+
+		GrainDelayProcessor.SetMaxGrains(*GrainMaxCount);
+		
+		// Update the envelope before rendering the audio block
+		if (PreviousGrainEnvelopeType != *GrainDelayEnvelope)
+		{
+			PreviousGrainEnvelopeType = *GrainDelayEnvelope;
+			GrainDelayProcessor.SetGrainEnvelope(PreviousGrainEnvelopeType);
+		}
+
+		// Update the base pitch shift ratio for the grain manager if it has changed
+		const float CurrentPitchShiftClamped = GrainDelayProcessor.GetGrainPitchShiftClamped(*GrainPitchShift);
+		if (!FMath::IsNearlyEqual(PreviousPitchShift, CurrentPitchShiftClamped))
+		{
+			PreviousPitchShift = CurrentPitchShiftClamped;
+			GrainDelayProcessor.SetGrainBasePitchShiftRatio(GrainDelayProcessor.GetGrainPitchShiftFrameRatio(PreviousPitchShift));
+		}
+
+		const float CurrentDelayLengthMsec = GrainDelayProcessor.GetGrainDelayClamped(*GrainDelay);
+		if (!FMath::IsNearlyEqual(CurrentDelayLengthMsec, PreviousGrainDelayMsec))
+		{
+			PreviousGrainDelayMsec = CurrentDelayLengthMsec;
+		}
+
+		GrainDelayProcessor.SetFeedbackAmount(*GrainDelayFeedbackAmount);
+		
+		// Loop until we either hit the end of the frame or the next trigger
+		while (NextAudioFrame < LastAudioFrame)
+		{
+			if (SpawnTrigIndex < NumSpawnTriggers)
+			{
+				NextSpawnFrame = (*GrainSpawnTrigger)[SpawnTrigIndex];
+			}
+			else
+			{
+				NextSpawnFrame = NoTrigger;
+			}
+
+			NextAudioFrame = FMath::Min(NextSpawnFrame, LastAudioFrame);
+
+			// no more triggers, rendering to the end of the block
+			if (NextAudioFrame == NoTrigger)
+			{
+				NextAudioFrame = OperatorSettings.GetNumFramesPerBlock();
+			}
+
+			// render audio (while loop handles looping audio)
+			while (CurrAudioFrame != NextAudioFrame)
+			{
+				GrainDelayProcessor.SynthesizeAudio(CurrAudioFrame, NextAudioFrame, AudioInput->GetData(), AudioOutput->GetData());
+				
+				CurrAudioFrame = NextAudioFrame;
+			}
+
+			// execute the next spawn trigger
+			if (CurrAudioFrame == NextSpawnFrame)
+			{
+				const float NewDelay = GrainDelayProcessor.GetGrainDelayClamped(*GrainDelay + FMath::RandRange(-0.5f * (*GrainDelayDelta), 0.5f * (*GrainDelayDelta)));
+				const float NewDuration = GrainDelayProcessor.GetGrainDurationClamped(*GrainDuration + FMath::RandRange(-0.5f * (*GrainDurationDelta), 0.5f * (*GrainDurationDelta)));
+
+				// Clamp the pitch shift delta before determining the random range
+				const float ClampedPitchShiftDelta = GrainDelayProcessor.GetGrainPitchShiftClamped(*GrainPitchShiftDelta);
+				const float NewPitchShiftRatioOffset = GrainDelayProcessor.GetGrainPitchShiftFrameRatio(FMath::RandRange(-0.5f * ClampedPitchShiftDelta, 0.5f * ClampedPitchShiftDelta));
+				
+				// Spawn a new grain. Note the grain manager handles spawning logic and rate limiting.
+				GrainDelayProcessor.SpawnGrain(NewDelay, NewDuration, NewPitchShiftRatioOffset);
+				
+				++SpawnTrigIndex;
+			}
+		}
+	}
+	
+	class FGrainDelayNode final : public FNodeFacade
+	{
+	public:
+		explicit FGrainDelayNode(const FNodeInitData& InitData)
+			: FNodeFacade(InitData.InstanceName, InitData.InstanceID, TFacadeOperatorClass<FGrainDelayOperator>())
+		{
+		}
+	};
+
+	METASOUND_REGISTER_NODE(FGrainDelayNode)
+}
+
+#undef LOCTEXT_NAMESPACE //MetasoundStandardNodes_GrainDelayNode
diff --git a/Engine/Source/Runtime/SignalProcessing/Private/DynamicsProcessor.cpp b/Engine/Source/Runtime/SignalProcessing/Private/DynamicsProcessor.cpp
--- a/Engine/Source/Runtime/SignalProcessing/Private/DynamicsProcessor.cpp
+++ b/Engine/Source/Runtime/SignalProcessing/Private/DynamicsProcessor.cpp
@@ -385,7 +385,6 @@ namespace Audio
 		{
 			DetectorOuts[Channel] = EnvFollower[Channel].ProcessSample(DetectorGain * KeyIn[Channel]);
 		}
-		ArrayClampInPlace(DetectorOuts, 0.f, 1.f);
 
 		switch (LinkMode)
 		{
diff --git a/Engine/Source/Runtime/SignalProcessing/Private/GrainDelay.cpp b/Engine/Source/Runtime/SignalProcessing/Private/GrainDelay.cpp
new file mode 100644
--- /dev/null
+++ b/Engine/Source/Runtime/SignalProcessing/Private/GrainDelay.cpp
@@ -0,0 +1,188 @@
+// Copyright Epic Games, Inc. All Rights Reserved.
+
+#include "DSP/GrainDelay.h"
+
+namespace Audio
+{
+	namespace GrainDelay
+	{
+		// TODO: consider making this a construction pin or cvar
+		static constexpr float MaxDelaySeconds = 2.0f;
+		static constexpr int32 GrainEnvelopeSampleCount = 1024;
+		static constexpr float MaxAbsPitchShiftInOctaves = 6.0f;
+
+		FGrainDelay::FGrainDelay(const float InSampleRate)
+			: SampleRate(InSampleRate)
+		{
+			GrainDelayLine.Init(SampleRate, MaxDelaySeconds);
+			
+			// Generate the grain data
+			GenerateEnvelopeData(GrainEnvelope, GrainEnvelopeSampleCount, GrainEnvelopeType);
+			
+			// Setup the dynamics processor
+			DynamicsProcessor.Init(InSampleRate, 1);
+			DynamicsProcessor.SetLookaheadMsec(3.0f);
+			DynamicsProcessor.SetAttackTime(0.0f);
+			DynamicsProcessor.SetReleaseTime(10.0f);
+			DynamicsProcessor.SetThreshold(-6.0f);
+			DynamicsProcessor.SetKneeBandwidth(3.0f);
+			DynamicsProcessor.SetInputGain(0.0f);
+			DynamicsProcessor.SetOutputGain(0.0f);
+			DynamicsProcessor.SetAnalogMode(false);
+			DynamicsProcessor.SetPeakMode(EPeakMode::Peak);
+			DynamicsProcessor.SetProcessingMode(EDynamicsProcessingMode::Limiter);
+		}
+
+		FGrainDelay::~FGrainDelay()
+		{
+		}
+
+		float FGrainDelay::GetGrainDelayClamped(const float InDelay) const
+		{
+			return FMath::Clamp(InDelay, 0.0f, 1000.0f * MaxDelaySeconds);
+		}
+
+		float FGrainDelay::GetGrainDurationClamped(const float InDuration) const
+		{
+			return FMath::Clamp(InDuration, 0.0f, 1000.0f * MaxDelaySeconds);
+		}
+	
+		float FGrainDelay::GetGrainDelayRatioClamped(const float InGrainDelayRatio) const
+		{
+			return FMath::Clamp(InGrainDelayRatio, -1.0f, 1.0f);
+		}
+
+		float FGrainDelay::GetGrainPitchShiftClamped(const float InPitchShift) const
+		{
+			return FMath::Clamp(InPitchShift, -12.0f * MaxAbsPitchShiftInOctaves, 12.0f * MaxAbsPitchShiftInOctaves);
+		}
+
+		float FGrainDelay::GetGrainPitchShiftFrameRatio(const float InPitchShift) const
+		{
+			return FMath::Pow(2.0f, InPitchShift / 12.0f);
+		}
+
+		void FGrainDelay::SetMaxGrains(const int32 InMaxGrains)
+		{
+			MaxGrains = FMath::Clamp(InMaxGrains, 1, 100);
+		}
+		
+		void FGrainDelay::SetGrainEnvelope(const Audio::Grain::EEnvelope InGrainEnvelope)
+		{
+			if (GrainEnvelopeType != InGrainEnvelope)
+			{
+				GrainEnvelopeType = InGrainEnvelope;
+				GenerateEnvelopeData(GrainEnvelope, 1024, GrainEnvelopeType);
+			}
+		}
+
+		void FGrainDelay::SetFeedbackAmount(float InFeedbackAmount)
+		{
+			FeedbackAmount = FMath::Clamp(InFeedbackAmount, 0.0f, 0.95f);
+		}
+					
+		void FGrainDelay::SetGrainBasePitchShiftRatio(const float InPitchRatioBase)
+		{
+			if (!FMath::IsNearlyEqual(InPitchRatioBase, PitchShiftRatioBase))
+			{
+				PitchShiftRatioBase = InPitchRatioBase;
+				for (const int32 ActiveGrainId : ActiveGrains)
+				{
+					FGrain& Grain = GrainPool[ActiveGrainId];
+
+					// Set the pitch shift ratio of the grain.
+					// This function will take into account the initial random pitch offset
+					Grain.SetGrainPitchShiftRatio(InPitchRatioBase, SampleRate);
+				}
+			}
+		}
+		
+		void FGrainDelay::SpawnGrain(const float InDelay, const float InDuration, const float InPitchShiftRatioOffset)
+		{
+			if (ActiveGrains.Num() >= FMath::Max(MaxGrains, 1))
+			{
+				// Hit the max grain cap, no grain for you!
+				return;
+			}
+
+			// Find or create a new grain index
+			int32 NewGrainId = INDEX_NONE;
+			if (FreeGrains.Num() == 0)
+			{
+				NewGrainId = GrainPool.Num();
+				GrainPool.Add(FGrain());
+			}
+			else
+			{
+				NewGrainId = FreeGrains.Pop();
+			}
+
+			// Spawn a new grain and setup it's initial state and data
+			const float DurationSeconds = 0.001f * InDuration;
+			
+			FGrain& NewGrain = GrainPool[NewGrainId];
+			NewGrain.NumFramesRendered = 0.0f;
+			NewGrain.DelayTapPositionMilliseconds = InDelay;
+			NewGrain.DurationFrames = FMath::Max(DurationSeconds * SampleRate, 1.0f);
+			NewGrain.PitchShiftRatioOffset = InPitchShiftRatioOffset;
+			NewGrain.PitchShifter.Init(SampleRate, 0.0f, 1000.0f * DurationSeconds);
+			NewGrain.SetGrainPitchShiftRatio(PitchShiftRatioBase, SampleRate);
+
+			// Add to the active grain list
+			ActiveGrains.Add(NewGrainId);
+		}
+
+		float FGrainDelay::SynthesizeFrame(const Audio::FDelay& InDelayLine)
+		{
+			float OutFrame = 0.0f;
+			
+			// Loop through active grains from end so we can clean it up as we go when a grain is done
+			for (int32 ActiveGrainIndex = ActiveGrains.Num() - 1; ActiveGrainIndex >= 0; --ActiveGrainIndex)
+			{
+				// Get the grain data
+				const int32 GrainId = ActiveGrains[ActiveGrainIndex];
+				FGrain& Grain = GrainPool[GrainId];
+				
+				// Calculate the grain envelope
+				const float Fraction = FMath::Min(Grain.NumFramesRendered / Grain.DurationFrames, 1.0f);
+				const float GrainVolume = Grain::GetValue(GrainEnvelope, Fraction);
+
+				const float PitchShiftedSample = Grain.PitchShifter.ReadDopplerShiftedTapFromDelay(InDelayLine, Grain.DelayTapPositionMilliseconds);
+				OutFrame += GrainVolume * PitchShiftedSample;
+
+				// Note we are tracking real frames rendered not frames consumed from the delay for grain duration 
+				Grain.NumFramesRendered += 1.0f;
+
+				// Clean up the grain once it's done
+				if (Grain.NumFramesRendered >= Grain.DurationFrames)
+				{
+					// Pop back on this id so we can reuse it next time a grain is spawned
+					FreeGrains.Add(GrainId);
+
+					// Remove at swap allows us to clean up the active grain list while we loop through it
+					ActiveGrains.RemoveAtSwap(ActiveGrainIndex, 1, false);
+				}
+			}
+
+			// Feed audio through the dynamics processor to prevent clipping.
+			DynamicsProcessor.ProcessAudio(&OutFrame,1, &OutFrame);
+
+			return OutFrame;
+		}
+
+		void FGrainDelay::SynthesizeAudio(const int32 StartFrame, const int32 EndFrame, const float* InAudioBuffer, float* OutAudioBuffer)
+		{
+			int32 NumFramesProcessed = 0;
+			for (int32 FrameIndex = StartFrame; FrameIndex < EndFrame + 1; ++FrameIndex)
+			{
+				++NumFramesProcessed;
+				GrainDelayLine.WriteDelayAndInc(InAudioBuffer[FrameIndex] + FeedbackAmount * LastFrame);
+				LastFrame = SynthesizeFrame(GrainDelayLine);
+			
+				OutAudioBuffer[FrameIndex] = LastFrame;
+			}
+		}
+	};
+
+	
+}
diff --git a/Engine/Source/Runtime/SignalProcessing/Public/DSP/GrainDelay.h b/Engine/Source/Runtime/SignalProcessing/Public/DSP/GrainDelay.h
new file mode 100644
--- /dev/null
+++ b/Engine/Source/Runtime/SignalProcessing/Public/DSP/GrainDelay.h
@@ -0,0 +1,113 @@
+// Copyright Epic Games, Inc. All Rights Reserved.
+
+#pragma once
+
+#include "DSP/GrainEnvelope.h"
+#include "DSP/DynamicsProcessor.h"
+#include "DSP/TapDelayPitchShifter.h"
+
+namespace Audio
+{
+	namespace GrainDelay
+	{
+		class SIGNALPROCESSING_API FGrainDelay
+		{
+		public:
+			FGrainDelay(const float InSampleRate);
+			~FGrainDelay();
+
+			// Helpers to clamp audio
+			float GetGrainDelayClamped(const float InDelay) const;
+			float GetGrainDurationClamped(const float InDuration) const;
+			float GetGrainDelayRatioClamped(const float InGrainDelayRatio) const;
+			float GetGrainPitchShiftClamped(const float InPitchShift) const;
+			float GetGrainPitchShiftFrameRatio(const float InPitchShift) const;
+
+			// Dynamically sets the max grains 
+			void SetMaxGrains(const int32 InMaxGrains); 
+			
+			// Sets the grain envelope 
+			void SetGrainEnvelope(const Audio::Grain::EEnvelope InGrainEnvelope);
+
+			// Sets the feedback amount (how much output grain audio feeds into the delay line)
+			void SetFeedbackAmount(float InFeedbackAmount);
+
+			// Sets the base pitch shift ratio of all grains 
+			void SetGrainBasePitchShiftRatio(const float InPitchRatioBase);
+
+			// Spawns a new grain with the given parameters
+			void SpawnGrain(const float InDelay, const float InDuration, const float InPitchShiftRatioOffset);
+
+			// Synthesize audio in the given frame range from input audio. Writes to OutAudioBuffer
+			void SynthesizeAudio(const int32 StartFrame, const int32 EndFrame, const float* InAudioBuffer, float* OutAudioBuffer);
+			
+		private:
+			// Synthesize a single frame of audio
+			float SynthesizeFrame(const Audio::FDelay& InDelayLine);
+			
+			// Simple grain data structure for tracking grain state
+			struct FGrain
+			{
+				// Where the grain is currently tapped from the delay line
+				float DelayTapPositionMilliseconds = 0.0f;
+
+				// The current frame the grain is in it's duration
+				float NumFramesRendered = 0.0f;
+
+				// The total number of frames this grain will render for
+				float DurationFrames = 0.0f;
+
+				// The unique randomized pitch scale, so we can modulate the base pitch scale
+				float PitchShiftRatioOffset = 1.0f;
+				
+				// The tap delay pitch shifter which manages taps from the shared delay line for each grain
+				FTapDelayPitchShifter PitchShifter;
+				
+				// Helper function for computing the phaser frequency and phaser increment
+				void SetGrainPitchShiftRatio(const float InPitchShiftRatioBase, const float InSampleRate)
+				{
+					const float DurationMilliseconds = 1000.0f * DurationFrames / InSampleRate;
+					PitchShifter.SetDelayLength(DurationMilliseconds);
+					PitchShifter.SetPitchShiftRatio(InPitchShiftRatioBase * PitchShiftRatioOffset);
+				}
+			};
+
+			// Free grain indices in the grain pool
+			TArray<int32> FreeGrains;
+
+			// Active grain indices in the grain pool
+			TArray<int32> ActiveGrains;
+
+			// The actual grain pool. Reuses tracking memory for grain lifetimes.
+			TArray<FGrain> GrainPool;
+
+			// The shared envelope value. Recomputed when a new envelope is set. 
+			Grain::FEnvelope GrainEnvelope;
+			Grain::EEnvelope GrainEnvelopeType = Grain::EEnvelope::Gaussian;
+			
+			// A pitch shifting delay line
+			TArray<FTapDelayPitchShifter> DelayPitchShifters;
+
+			// The shared delay lines per channel
+			FDelay GrainDelayLine;
+
+			// The current feedback amount
+           	float FeedbackAmount = 0.0f;
+
+			// Cached value of last rendered frame. Used for feedback
+			float LastFrame = 0.0f;
+			
+			// Dynamics processor for the grain synthesis
+			FDynamicsProcessor DynamicsProcessor;
+
+			// The maximum grains to have rendering
+			int32 MaxGrains = 16;
+
+			// Sample rate track to easily deal with pitch scale calculations
+			float SampleRate = 0.0f;
+			
+			// The base pitch shift ratio (pitch scale) of all grains in the grain delay
+			float PitchShiftRatioBase = 1.0f;
+		};
+	}
+}
diff --git a/Engine/Source/Runtime/SignalProcessing/Public/DSP/GrainEnvelope.h b/Engine/Source/Runtime/SignalProcessing/Public/DSP/GrainEnvelope.h
new file mode 100644
--- /dev/null
+++ b/Engine/Source/Runtime/SignalProcessing/Public/DSP/GrainEnvelope.h
@@ -0,0 +1,108 @@
+// Copyright Epic Games, Inc. All Rights Reserved.
+
+#pragma once
+
+#include "CoreMinimal.h"
+
+namespace Audio
+{
+	// Utility for managing a grain data (enveloped buffer) 
+	namespace Grain
+	{
+		typedef TArray<float> FEnvelope;
+		
+		// Grain envelope types
+		enum class EEnvelope
+		{
+			Gaussian,
+			Triangle,
+			DownwardTriangle,
+			UpwardTriangle,
+			ExponentialDecay,
+			ExponentialAttack
+		};
+		
+		// Utility function generates an envelope with the input array for the number of desired frames
+		static void GenerateEnvelopeData(FEnvelope& InData, const int32 InNumFrames, const EEnvelope InEnvelopeType)
+		{
+			InData.Reset();
+			InData.AddUninitialized(InNumFrames);
+
+			const float N = static_cast<float>(InNumFrames);
+			const float N_1 = N - 1.0f;
+			float n = 0.0f;
+
+			switch (InEnvelopeType)
+			{
+			case EEnvelope::Gaussian:
+				{
+					const float Denominator = 0.3f * N_1 / 2.0f;
+					for (int32 i = 0; i < InNumFrames; ++i, n += 1.0f)
+					{
+						InData[i] = FMath::Exp(-0.5f * FMath::Pow((n - 0.5f * N_1) / Denominator, 2.0f));
+					}
+				}
+				break;
+
+			case EEnvelope::Triangle:
+				{
+					const float A = 0.5f * N_1;
+					for (int32 i = 0; i < InNumFrames; ++i, n += 1.0f)
+					{
+						InData[i] = 1.0f - FMath::Abs((n - A) / A);
+					}
+				}
+				break;
+
+			case EEnvelope::DownwardTriangle:
+				{
+					for (int32 i = 0; i < InNumFrames; ++i, n += 1.0f)
+					{
+						InData[i] = 1.0f - n / N_1;
+					}
+				}
+				break;
+
+			case EEnvelope::UpwardTriangle:
+				{
+					for (int32 i = 0; i < InNumFrames; ++i, n += 1.0f)
+					{
+						InData[i] = n / N_1;
+					}
+				}
+				break;
+
+			case EEnvelope::ExponentialDecay:
+				{
+					for (int32 i = 0; i < InNumFrames; ++i, n += 1.0f)
+					{
+						InData[i] = FMath::Pow((n - N + 1.0f) / N_1, 4.0f);
+					}
+				}
+				break;
+
+			case EEnvelope::ExponentialAttack:
+				{
+					for (int32 i = 0; i < InNumFrames; ++i, n += 1.0f)
+					{
+						InData[i] = FMath::Pow(n / N_1, 4.0f);
+					}
+				}
+				break;
+			}
+		}
+
+		// Returns the interpolated envelope value given the fraction through the grain
+		static float GetValue(const FEnvelope& InGrainData, const float InFraction)
+		{
+			const int32 NumFrames = InGrainData.Num();
+			const float Frame = static_cast<float>(InFraction) * (NumFrames - 1);
+			const int32 PrevIndex = static_cast<int32>(Frame);
+			const int32 NextIndex = FMath::Min(NumFrames, PrevIndex + 1);
+			const float AlphaIndex = Frame - static_cast<float>(PrevIndex);
+			const float* GrainEnvelopePtr = InGrainData.GetData();
+			return FMath::Lerp(GrainEnvelopePtr[PrevIndex], GrainEnvelopePtr[NextIndex], AlphaIndex);
+		}
+		
+	}
+}
