Changing display name of PitchShift node to "delay pitch shift"

diff --git a/Engine/Plugins/Runtime/Metasound/Source/MetasoundStandardNodes/Private/MetasoundDopplerPitchShiftNode.cpp b/Engine/Plugins/Runtime/Metasound/Source/MetasoundStandardNodes/Private/MetasoundDopplerPitchShiftNode.cpp
--- a/Engine/Plugins/Runtime/Metasound/Source/MetasoundStandardNodes/Private/MetasoundDopplerPitchShiftNode.cpp
+++ b/Engine/Plugins/Runtime/Metasound/Source/MetasoundStandardNodes/Private/MetasoundDopplerPitchShiftNode.cpp
@@ -110,11 +110,11 @@
 		auto InitNodeInfo = []() -> FNodeClassMetadata
 		{
 			FNodeClassMetadata Info;
-			Info.ClassName = { StandardNodes::Namespace, "Doppler Pitch Shift", StandardNodes::AudioVariant };
+			Info.ClassName = { StandardNodes::Namespace, "Delay Pitch Shift", StandardNodes::AudioVariant };
 			Info.MajorVersion = 1;
 			Info.MinorVersion = 0;
-			Info.DisplayName = METASOUND_LOCTEXT("DelayNode_DisplayName", "Doppler Pitch Shift");
-			Info.Description = METASOUND_LOCTEXT("DelayNode_Description", "Pitch shifts the audio buffer using a doppler-delay method.");
+			Info.DisplayName = METASOUND_LOCTEXT("DelayNode_DisplayName", "Delay Pitch Shift");
+			Info.Description = METASOUND_LOCTEXT("DelayNode_Description", "Pitch shifts the audio buffer using a delay-based doppler-shift method.");
 			Info.Author = PluginAuthor;
 			Info.PromptIfMissing = PluginNodeMissingPrompt;
 			Info.DefaultInterface = GetVertexInterface();
