Fixing tooltip loctext namespace

diff --git a/Engine/Plugins/Runtime/Metasound/Source/MetasoundStandardNodes/Private/MetasoundDopplerPitchShiftNode.cpp b/Engine/Plugins/Runtime/Metasound/Source/MetasoundStandardNodes/Private/MetasoundDopplerPitchShiftNode.cpp
--- a/Engine/Plugins/Runtime/Metasound/Source/MetasoundStandardNodes/Private/MetasoundDopplerPitchShiftNode.cpp
+++ b/Engine/Plugins/Runtime/Metasound/Source/MetasoundStandardNodes/Private/MetasoundDopplerPitchShiftNode.cpp
@@ -113,8 +113,8 @@ namespace Metasound
 			Info.ClassName = { StandardNodes::Namespace, "Delay Pitch Shift", StandardNodes::AudioVariant };
 			Info.MajorVersion = 1;
 			Info.MinorVersion = 0;
-			Info.DisplayName = METASOUND_LOCTEXT("DelayNode_DisplayName", "Delay Pitch Shift");
-			Info.Description = METASOUND_LOCTEXT("DelayNode_Description", "Pitch shifts the audio buffer using a delay-based doppler-shift method.");
+			Info.DisplayName = METASOUND_LOCTEXT("DopperDelayPitchShift_DisplayName", "Delay Pitch Shift");
+			Info.Description = METASOUND_LOCTEXT("DopperDelayPitchShift_Description", "Pitch shifts the audio buffer using a delay-based doppler-shift method.");
 			Info.Author = PluginAuthor;
 			Info.PromptIfMissing = PluginNodeMissingPrompt;
 			Info.DefaultInterface = GetVertexInterface();
