MetaSound animated connection optimizations/minor improvements
- Don't copy point array, just add wrapped accessor for simple FloatWindow
- Move to drawing lines instead of a bunch of bubble images
- Add max number of draw segments
- Add cvars for tweaking
- Split out numeric & envelope thickness options
#rnx

diff --git a/Engine/Plugins/Runtime/Metasound/Source/MetasoundEditor/Private/MetasoundEditorGraphConnectionDrawingPolicy.cpp b/Engine/Plugins/Runtime/Metasound/Source/MetasoundEditor/Private/MetasoundEditorGraphConnectionDrawingPolicy.cpp
--- a/Engine/Plugins/Runtime/Metasound/Source/MetasoundEditor/Private/MetasoundEditorGraphConnectionDrawingPolicy.cpp
+++ b/Engine/Plugins/Runtime/Metasound/Source/MetasoundEditor/Private/MetasoundEditorGraphConnectionDrawingPolicy.cpp
@@ -10,11 +10,16 @@
 #include "MetasoundEditorGraphBuilder.h"
 #include "MetasoundEditorGraphSchema.h"
 #include "MetasoundEditorSettings.h"
+#include "MetasoundTrace.h"
 #include "Misc/App.h"
 #include "Templates/Function.h"
 
+#define METASOUND_EDITOR_DEBUG_CONNECTIONS 0
+
+#if METASOUND_EDITOR_DEBUG_CONNECTIONS
+#include "HAL/IConsoleManager.h"
+#endif // METASOUND_EDITOR_DEBUG_CONNECTIONS
 
-#define METASOUND_DRAW_CONNECTION_DEBUG_DATA 0
 
 namespace Metasound
 {
@@ -22,7 +27,11 @@ namespace Metasound
 	{
 		namespace DrawingPolicyPrivate
 		{
-#if METASOUND_DRAW_CONNECTION_DEBUG_DATA
+			static float EnvelopeConnectionSpeedCVar = 1.0f;
+			static float EnvelopeConnectionSpacingCVar = 3.0f;
+			static int32 EnvelopeConnectionMaxPointsPerConnection = 1024;
+
+#if METASOUND_EDITOR_DEBUG_CONNECTIONS
 			// Draw the bounding box for debugging
 			void DrawDebugConnection(FSlateWindowElementList& InDrawElementsList, uint32 InLayerId, const FDrawConnectionData& InDrawConnectionData)
 			{
@@ -43,7 +52,28 @@ namespace Metasound
 				DrawSpaceLine(BR, BL, BoundsWireColor);
 				DrawSpaceLine(BL, TL, BoundsWireColor);
 			}
-#endif // METASOUND_DRAW_CONNECTION_DEBUG_DATA
+
+			FAutoConsoleVariableRef CVarMetaSoundEditorAudioConnectSpeed(
+				TEXT("au.MetaSounds.Editor.EnvelopeConnection.Speed"),
+				DrawingPolicyPrivate::EnvelopeConnectionSpeedCVar,
+				TEXT("Speed of lines drawn in audio connections in the MetaSoundEditor.\n")
+				TEXT("Default: 1.0f"),
+				ECVF_Default);
+
+			FAutoConsoleVariableRef CVarMetaSoundEditorAudioConnectSpacing(
+				TEXT("au.MetaSounds.Editor.EnvelopeConnection.Spacing"),
+				DrawingPolicyPrivate::EnvelopeConnectionSpacingCVar,
+				TEXT("Spacing of lines drawn in audio connections in the MetaSoundEditor.\n")
+				TEXT("Default: 3.0f"),
+				ECVF_Default);
+
+			FAutoConsoleVariableRef CVarEnvelopeConnectionMaxPointsPerConnection(
+				TEXT("au.MetaSounds.Editor.EnvelopeConnection.MaxPointsPerConnection"),
+				DrawingPolicyPrivate::EnvelopeConnectionMaxPointsPerConnection,
+				TEXT("Max number of draw points per connection for animated connections.\n")
+				TEXT("Default: 1024"),
+				ECVF_Default);
+#endif // METASOUND_EDITOR_DEBUG_CONNECTIONS
 
 			template<typename TAnalyzerType>
 			bool DetermineWiringStyle_Envelope(TSharedPtr<FEditor> MetasoundEditor, UEdGraphPin* OutputPin, UEdGraphPin* InputPin, FConnectionParams& OutParams)
@@ -155,7 +185,6 @@ namespace Metasound
 			}
 		} // namespace DrawingPolicyPrivate
 
-
 		FDrawConnectionData::FDrawConnectionData(const FVector2D& InStart, const FVector2D& InEnd, const FVector2D& InSplineTangent, const FConnectionParams& InParams, const UGraphEditorSettings& InSettings, const FVector2D& InMousePosition)
 			: P0(InStart)
 			, P0Tangent(InParams.StartDirection == EGPD_Output ? InSplineTangent : -InSplineTangent)
@@ -268,9 +297,10 @@ namespace Metasound
 			if (const UMetasoundEditorSettings* MetasoundSettings = GetDefault<UMetasoundEditorSettings>())
 			{
 				bAnimateConnections = MetasoundSettings->bAnimateConnections;
-				ActiveAnalyzerWireThickness = MetasoundSettings->ActiveAnalyzerWireThickness;
-				ActiveAnalyzerWireSignalScalarMin = MetasoundSettings->ActiveAnalyzerWireSignalScalarMin;
-				ActiveAnalyzerWireSignalScalarMax = MetasoundSettings->ActiveAnalyzerWireSignalScalarMax;
+				ActiveAnalyzerNumericWireThickness = MetasoundSettings->ActiveAnalyzerNumericWireThickness;
+				ActiveAnalyzerEnvelopeWireThickness = MetasoundSettings->ActiveAnalyzerEnvelopeWireThickness;
+				ActiveAnalyzerWireScalarMin = MetasoundSettings->ActiveAnalyzerWireScalarMin;
+				ActiveAnalyzerWireScalarMax = MetasoundSettings->ActiveAnalyzerWireScalarMax;
 			}
 
 			if (GraphObj)
@@ -283,7 +313,7 @@ namespace Metasound
 		{
 			if (OutParams.bDrawBubbles)
 			{
-				OutParams.WireThickness = FMath::Lerp(InactiveWireThickness * ActiveAnalyzerWireSignalScalarMin, ActiveAnalyzerWireThickness * ActiveAnalyzerWireSignalScalarMax, OutParams.WireThickness);
+				OutParams.WireThickness = FMath::Lerp(InactiveWireThickness * ActiveAnalyzerWireScalarMin, ActiveAnalyzerEnvelopeWireThickness * ActiveAnalyzerWireScalarMax, OutParams.WireThickness);
 			}
 			else
 			{
@@ -370,7 +400,7 @@ namespace Metasound
 
 					if (bEdgeStyleValid && EditorPtr->IsPlaying())
 					{
-						OutParams.WireThickness = ActiveAnalyzerWireThickness;
+						OutParams.WireThickness = ActiveAnalyzerNumericWireThickness;
 					}
 					else
 					{
@@ -386,6 +416,8 @@ namespace Metasound
 
 		void FGraphConnectionDrawingPolicy::DrawConnection(int32 LayerId, const FVector2D& Start, const FVector2D& End, const FConnectionParams& Params)
 		{
+			METASOUND_TRACE_CPUPROFILER_EVENT_SCOPE(FGraphConnectionDrawingPolicy::DrawConnection);
+
 			using namespace Frontend;
 
 			TSharedPtr<FEditor> EditorPtr = MetasoundEditor.Pin();
@@ -404,28 +436,32 @@ namespace Metasound
 			const FDrawConnectionData ConnectionData(Start, End, FGraphConnectionDrawingPolicy::ComputeSplineTangent(Start, End), Params, *Settings, LocalMousePosition);
 			if (Settings->bTreatSplinesLikePins)
 			{
-#if METASOUND_DRAW_CONNECTION_DEBUG_DATA
+#if METASOUND_EDITOR_DEBUG_CONNECTIONS
 				DrawingPolicyPrivate::DrawDebugConnection(DrawElementsList, ArrowLayerID, ConnectionData);
-#endif // METASOUND_DRAW_CONNECTION_DEBUG_DATA
+#endif // METASOUND_EDITOR_DEBUG_CONNECTIONS
 
 				ConnectionData.UpdateSplineOverlap(SplineOverlapResult);
 			}
 
 			if (ConnectionData.OutputHandle->IsValid() && MetasoundEditor.IsValid())
 			{
-				if (ConnectionData.OutputHandle->GetDataType() == GetMetasoundDataTypeName<float>())
+				const FName DataType = ConnectionData.OutputHandle->GetDataType();
+				if (DataType == GetMetasoundDataTypeName<float>())
 				{
 					DrawingPolicyPrivate::DrawConnectionSpline_Numeric<float>(*this, LayerId, ConnectionData, 0.0f);
 				}
-				else if (ConnectionData.OutputHandle->GetDataType() == GetMetasoundDataTypeName<int32>())
+				else if (DataType == GetMetasoundDataTypeName<int32>())
 				{
 					DrawingPolicyPrivate::DrawConnectionSpline_Numeric<int32>(*this, LayerId, ConnectionData, 0);
 				}
-				else if (ConnectionData.OutputHandle->GetDataType() == GetMetasoundDataTypeName<bool>())
+				else if (DataType == GetMetasoundDataTypeName<bool>())
 				{
 					DrawingPolicyPrivate::DrawConnectionSpline_Numeric<bool>(*this, LayerId, ConnectionData, 0);
 				}
-				else
+				// AudioBuffers & Triggers do not require drawing connection splines because they draw envelopes,
+				// which trace the same path
+				else if (DataType != GetMetasoundDataTypeName<FAudioBuffer>()
+					&& DataType != GetMetasoundDataTypeName<FTrigger>())
 				{
 					DrawConnectionSpline(LayerId, ConnectionData);
 				}
@@ -446,17 +482,18 @@ namespace Metasound
 				{
 					const float AppTime = FPlatformTime::Seconds() - GStartTime;
 					FDrawConnectionSignalData SignalParams = { AppTime, LayerId, SplineReparamTable, SplineLength, ConnectionData };
-					if (SignalParams.ConnectionData.OutputHandle->GetDataType() == GetMetasoundDataTypeName<FAudioBuffer>())
+					const FName DataType = SignalParams.ConnectionData.OutputHandle->GetDataType();
+					if (DataType == GetMetasoundDataTypeName<FAudioBuffer>())
 					{
-						SignalParams.SpacingFactor = 2.f;
-						SignalParams.SpeedFactor = 256.f;
+						SignalParams.SpacingFactor = DrawingPolicyPrivate::EnvelopeConnectionSpacingCVar;
+						SignalParams.SpeedFactor = DrawingPolicyPrivate::EnvelopeConnectionSpeedCVar;
 						const FName AnalyzerName = Frontend::FVertexAnalyzerEnvelopeFollower::GetAnalyzerName();
 						DrawConnectionSignal_Envelope(SignalParams, AnalyzerName);
 					}
-					else if (SignalParams.ConnectionData.OutputHandle->GetDataType() == GetMetasoundDataTypeName<FTrigger>())
+					else if (DataType == GetMetasoundDataTypeName<FTrigger>())
 					{
-						SignalParams.SpacingFactor = 2.f;
-						SignalParams.SpeedFactor = 64.f;
+						SignalParams.SpacingFactor = DrawingPolicyPrivate::EnvelopeConnectionSpacingCVar;
+						SignalParams.SpeedFactor = DrawingPolicyPrivate::EnvelopeConnectionSpeedCVar;
 						const FName AnalyzerName = Frontend::FVertexAnalyzerTriggerDensity::GetAnalyzerName();
 						DrawConnectionSignal_Envelope(SignalParams, AnalyzerName);
 					}
@@ -557,14 +594,14 @@ namespace Metasound
 				return;
 			}
 
-			const float BubbleSpacing = InParams.SpacingFactor * ZoomFactor;
-			float BubbleSpeed = InParams.SpeedFactor * ZoomFactor;
-			const FVector2D BubbleSize = BubbleImage->ImageSize * ZoomFactor * 0.2f * InParams.ConnectionData.Params.WireThickness;
-			const float BubbleOffset = FMath::Fmod(InParams.AppTime * BubbleSpeed, BubbleSpacing);
-			const int32 NumBubbles = FMath::CeilToInt(InParams.SplineLength / BubbleSpacing);
+			const float PointSpacing = FMath::Max(1.0f, InParams.SpacingFactor * ZoomFactor);
+			float PointSpeed = InParams.SpeedFactor * ZoomFactor;
+			const float PointSize = ZoomFactor * InParams.ConnectionData.Params.WireThickness;
+			const float PointOffset = FMath::Fmod(InParams.AppTime * PointSpeed, PointSpacing);
+			const int32 NumPoints = FMath::Min(FMath::CeilToInt(InParams.SplineLength / PointSpacing), DrawingPolicyPrivate::EnvelopeConnectionMaxPointsPerConnection);
 
-			TArray<FVector2D> ScaledBubbles;
-			ScaledBubbles.Init(BubbleSize, NumBubbles);
+			TArray<float> EnvMagnitudes;
+			EnvMagnitudes.Init(PointSize, NumPoints);
 			if (const UMetasoundEditorGraphNode* Node = Cast<UMetasoundEditorGraphNode>(InParams.ConnectionData.OutputPin->GetOwningNode()))
 			{
 				if (const UEdGraph* Graph = Node->GetGraph())
@@ -575,53 +612,91 @@ namespace Metasound
 						const FGuid NodeId = OutputHandle->GetOwningNodeID();
 						FName OutputName = OutputHandle->GetName();
 
-						TArray<float> Values;
+						// Only set to max as multiple edges can exist and each may have a different window size.
+						FGraphConnectionManager& ConnectionManager = Editor->GetConnectionManager();
+						if (const FFloatMovingWindow* Window = ConnectionManager.GetValueWindow(NodeId, OutputName, InAnalyzerName))
 						{
-							// Only set to max as multiple edges can exist and each may have a different window size.
-							FGraphConnectionManager& ConnectionManager = Editor->GetConnectionManager();
-							if (const FFloatMovingWindow* Window = ConnectionManager.GetValueWindow(NodeId, OutputName, InAnalyzerName))
+							if (NumPoints > Window->Num())
 							{
-								if (NumBubbles > Window->Num())
-								{
-									ConnectionManager.TrackValue(NodeId, OutputName, InAnalyzerName, NumBubbles);
-								}
-								Values = Window->ToArray();
-								Values.SetNum(NumBubbles);
-								for (int32 i = 0; i < Values.Num(); ++i)
-								{
-									const float ClampedScalar = FMath::Clamp(Values[i], 0.0f, 1.0f) * ActiveAnalyzerWireSignalScalarMax;
-									ScaledBubbles[i].X *= ClampedScalar;
-									ScaledBubbles[i].Y *= ClampedScalar;
-								}
+								ConnectionManager.TrackValue(NodeId, OutputName, InAnalyzerName, NumPoints);
 							}
-							else
+
+							for (int32 i = 0; i < EnvMagnitudes.Num(); ++i)
 							{
-								ConnectionManager.TrackValue(NodeId, OutputName, InAnalyzerName, NumBubbles);
+								const float ClampedScalar = FMath::Clamp(Window->GetValueWrapped(i), 0.0f, 1.0f) * ActiveAnalyzerWireScalarMax;
+								EnvMagnitudes[i] *= ClampedScalar;
 							}
 						}
+						else
+						{
+							ConnectionManager.TrackValue(NodeId, OutputName, InAnalyzerName, NumPoints);
+						}
 					}
 				}
 			}
 
-			for (int32 i = 0; i < ScaledBubbles.Num(); ++i)
+			TArray<FVector2f> PointsUp;
+			TArray<FVector2f> PointsDown;
+			FVector2D LastPointPos = FVector2D::Zero();
+			FVector2D LastTangent = FVector2D::Zero();
+
+			auto AddPoint = [&](const float EnvMagnitude, const float Distance)
 			{
-				const float Distance = (i * BubbleSpacing) + BubbleOffset;
-				if (Distance < InParams.SplineLength)
+				const float Alpha = InParams.SplineReparamTable.Eval(Distance, 0.f);
+				FVector2D PointPos = FMath::CubicInterp(InParams.ConnectionData.P0, InParams.ConnectionData.P0Tangent, InParams.ConnectionData.P1, InParams.ConnectionData.P1Tangent, Alpha);
+
+				LastTangent = FVector2D(PointPos.X - LastPointPos.X, PointPos.Y - LastPointPos.Y);
+				const FVector2D Normal = FVector2D(LastTangent.Y, -1.0f * LastTangent.X).GetSafeNormal();
+				LastPointPos = PointPos;
+
+				const float HalfEnvMagnitude = EnvMagnitude * 0.5f;
+				const FVector2f Top = { static_cast<float>(PointPos.X - (Normal.X * HalfEnvMagnitude)), static_cast<float>(PointPos.Y - (Normal.Y * HalfEnvMagnitude)) };
+				const FVector2f Bottom = { static_cast<float>(PointPos.X + (Normal.X * HalfEnvMagnitude)), static_cast<float>(PointPos.Y + (Normal.Y * HalfEnvMagnitude)) };
+
+				if (Top.Equals(Bottom))
 				{
-					const float Alpha = InParams.SplineReparamTable.Eval(Distance, 0.f);
-					FVector2D BubblePos = FMath::CubicInterp(InParams.ConnectionData.P0, InParams.ConnectionData.P0Tangent, InParams.ConnectionData.P1, InParams.ConnectionData.P1Tangent, Alpha);
-					BubblePos -= (ScaledBubbles[i] * 0.5f);
+					PointsUp.Add(Top);
+					PointsDown.Add(Bottom);
+				}
+				else
+				{
+					PointsUp.Add(Top);
+					PointsUp.Add(Bottom);
+					PointsDown.Add(Bottom);
+					PointsDown.Add(Top);
+				}
+			};
 
-					FSlateDrawElement::MakeBox(
-						DrawElementsList,
-						InParams.LayerId,
-						FPaintGeometry(BubblePos, ScaledBubbles[i], ZoomFactor),
-						BubbleImage,
-						ESlateDrawEffect::None,
-						InParams.ConnectionData.Params.WireColor
-					);
+			for (int32 i = 0; i < EnvMagnitudes.Num(); ++i)
+			{
+				const float Distance = (i * PointSpacing) + PointOffset;
+				if (Distance < InParams.SplineLength)
+				{
+					AddPoint(EnvMagnitudes[i], Distance);
 				}
 			}
+
+			bool bDrawRemainingSpline = false;
+			const float TermDistance = (EnvMagnitudes.Num() * PointSpacing) + PointOffset;
+			if (TermDistance < InParams.SplineLength)
+			{
+				AddPoint(0.0f, TermDistance);
+				bDrawRemainingSpline = true;
+			}
+			else
+			{
+				AddPoint(0.0f, InParams.SplineLength);
+			}
+
+			FSlateDrawElement::MakeLines(DrawElementsList, InParams.LayerId, FPaintGeometry(), PointsUp, ESlateDrawEffect::None, InParams.ConnectionData.Params.WireColor, true, InParams.ConnectionData.Params.WireThickness);
+			FSlateDrawElement::MakeLines(DrawElementsList, InParams.LayerId, FPaintGeometry(), PointsDown, ESlateDrawEffect::None, InParams.ConnectionData.Params.WireColor, true, InParams.ConnectionData.Params.WireThickness);
+			if (bDrawRemainingSpline)
+			{
+				FDrawConnectionData FinalData = InParams.ConnectionData;
+				FinalData.P0 = LastPointPos;
+				FinalData.P0Tangent = LastTangent;
+				DrawSpline(InParams.ConnectionData);
+			}
 		}
 
 		void FGraphConnectionDrawingPolicy::DrawSpline(const FDrawConnectionData& InData)
diff --git a/Engine/Plugins/Runtime/Metasound/Source/MetasoundEditor/Private/MetasoundEditorGraphConnectionDrawingPolicy.h b/Engine/Plugins/Runtime/Metasound/Source/MetasoundEditor/Private/MetasoundEditorGraphConnectionDrawingPolicy.h
--- a/Engine/Plugins/Runtime/Metasound/Source/MetasoundEditor/Private/MetasoundEditorGraphConnectionDrawingPolicy.h
+++ b/Engine/Plugins/Runtime/Metasound/Source/MetasoundEditor/Private/MetasoundEditorGraphConnectionDrawingPolicy.h
@@ -30,10 +30,10 @@ namespace Metasound
 
 		struct FDrawConnectionData
 		{
-			const FVector2D P0;
-			const FVector2D P0Tangent;
-			const FVector2D P1;
-			const FVector2D P1Tangent;
+			FVector2D P0;
+			FVector2D P0Tangent;
+			FVector2D P1;
+			FVector2D P1Tangent;
 
 			FConnectionParams Params;
 
@@ -104,9 +104,10 @@ namespace Metasound
 
 			bool bAnimateConnections = false;
 			float ActiveWireThickness = 1.0f;
-			float ActiveAnalyzerWireThickness = 1.0f;
-			float ActiveAnalyzerWireSignalScalarMin = 1.0f;
-			float ActiveAnalyzerWireSignalScalarMax = 1.0f;
+			float ActiveAnalyzerEnvelopeWireThickness = 1.0f;
+			float ActiveAnalyzerNumericWireThickness = 1.0f;
+			float ActiveAnalyzerWireScalarMin = 1.0f;
+			float ActiveAnalyzerWireScalarMax = 1.0f;
 
 			float InactiveWireThickness = 1.0f;
 
diff --git a/Engine/Plugins/Runtime/Metasound/Source/MetasoundEditor/Private/MetasoundEditorGraphConnectionManager.h b/Engine/Plugins/Runtime/Metasound/Source/MetasoundEditor/Private/MetasoundEditorGraphConnectionManager.h
--- a/Engine/Plugins/Runtime/Metasound/Source/MetasoundEditor/Private/MetasoundEditorGraphConnectionManager.h
+++ b/Engine/Plugins/Runtime/Metasound/Source/MetasoundEditor/Private/MetasoundEditorGraphConnectionManager.h
@@ -68,30 +68,12 @@ namespace Metasound
 				}
 			}
 
-			TArray<float> ToArray() const
+			float GetValueWrapped(int32 Index) const
 			{
-				TArray<float> Values;
-
-				if (!Data.IsEmpty())
-				{
-					check(WriteIndex >= 0);
-					check(WriteIndex < Data.Num());
-
-					const int32 ReadIndex = (WriteIndex + 1) % Data.Num();
-					if (ReadIndex > WriteIndex)
-					{
-						Values.SetNumUninitialized(Data.Num());
-						const int32 Offset = Data.Num() - ReadIndex;
-						FMemory::Memcpy(Values.GetData(), Data.GetData() + ReadIndex, sizeof(float) * Offset);
-						FMemory::Memcpy(Values.GetData() + Offset, Data.GetData(), sizeof(float) * ReadIndex);
-					}
-					else
-					{
-						Values = Data;
-					}
-				}
+				check(!Data.IsEmpty());
 
-				return Values;
+				const int32 WrappedIndex = (WriteIndex + Index) % Data.Num();
+				return Data[WrappedIndex];
 			}
 		};
 
diff --git a/Engine/Plugins/Runtime/Metasound/Source/MetasoundEditor/Private/MetasoundEditorSettings.cpp b/Engine/Plugins/Runtime/Metasound/Source/MetasoundEditor/Private/MetasoundEditorSettings.cpp
--- a/Engine/Plugins/Runtime/Metasound/Source/MetasoundEditor/Private/MetasoundEditorSettings.cpp
+++ b/Engine/Plugins/Runtime/Metasound/Source/MetasoundEditor/Private/MetasoundEditorSettings.cpp
@@ -29,7 +29,8 @@ UMetasoundEditorSettings::UMetasoundEditorSettings(const FObjectInitializer& Obj
 	VariableNodeTitleColor = FLinearColor(0.211f, 0.513f, 0.035f);				// copper
 
 	bAnimateConnections = false;
-	ActiveAnalyzerWireThickness = 8.0f;
-	ActiveAnalyzerWireSignalScalarMin = 1.0f;
-	ActiveAnalyzerWireSignalScalarMax = 1.0f;
+	ActiveAnalyzerNumericWireThickness = 5.0f;
+	ActiveAnalyzerEnvelopeWireThickness = 1.0f;
+	ActiveAnalyzerWireScalarMin = 1.0f;
+	ActiveAnalyzerWireScalarMax = 4.5f;
 }
diff --git a/Engine/Plugins/Runtime/Metasound/Source/MetasoundEditor/Private/MetasoundEditorSettings.h b/Engine/Plugins/Runtime/Metasound/Source/MetasoundEditor/Private/MetasoundEditorSettings.h
--- a/Engine/Plugins/Runtime/Metasound/Source/MetasoundEditor/Private/MetasoundEditorSettings.h
+++ b/Engine/Plugins/Runtime/Metasound/Source/MetasoundEditor/Private/MetasoundEditorSettings.h
@@ -104,17 +104,21 @@ public:
 	UPROPERTY(EditAnywhere, config, Category = GraphAnimation, meta = (DisplayName = "Animate Connections (Experimental)"))
 	bool bAnimateConnections;
 
-	/** Thickness of default analyzer wire thickness when connection analyzer is active. */
-	UPROPERTY(EditAnywhere, config, Category = GraphAnimation, meta = (EditCondition = "bAnimateConnections", UIMin = 1, UIMax = 8, ClampMin = 1))
-	float ActiveAnalyzerWireThickness;
+	/** Thickness of default envelope analyzer wire thickness when connection analyzer is active. */
+	UPROPERTY(EditAnywhere, config, Category = GraphAnimation, meta = (EditCondition = "bAnimateConnections", UIMin = 1, UIMax = 2, ClampMin = 1))
+	float ActiveAnalyzerEnvelopeWireThickness;
+
+	/** Thickness of default numeric analyzer wire thickness when connection analyzer is active. */
+	UPROPERTY(EditAnywhere, config, Category = GraphAnimation, meta = (EditCondition = "bAnimateConnections", UIMin = 1, UIMax = 10, ClampMin = 1))
+	float ActiveAnalyzerNumericWireThickness;
 
 	/** Minimum height scalar of wire signal analyzers (ex. audio, triggers). */
-	UPROPERTY(EditAnywhere, config, Category = GraphAnimation, meta = (EditCondition = "bAnimateConnections", UIMin = 1, UIMax = 3, ClampMin = 1))
-	float ActiveAnalyzerWireSignalScalarMin;
+	UPROPERTY(EditAnywhere, config, Category = GraphAnimation, meta = (EditCondition = "bAnimateConnections", UIMin = 1, UIMax = 5, ClampMin = 1))
+	float ActiveAnalyzerWireScalarMin;
 
 	/** Maximum height scalar of wire signal analyzers (ex. audio, triggers). */
-	UPROPERTY(EditAnywhere, config, Category = GraphAnimation, meta = (EditCondition = "bAnimateConnections", UIMin = 1, UIMax = 3, ClampMin = 1))
-	float ActiveAnalyzerWireSignalScalarMax;
+	UPROPERTY(EditAnywhere, config, Category = GraphAnimation, meta = (EditCondition = "bAnimateConnections", UIMin = 1, UIMax = 5, ClampMin = 1))
+	float ActiveAnalyzerWireScalarMax;
 
 	/** Determines which details view to show in Metasounds Editor */
 	UPROPERTY(Transient)
