Fix for missing enum ensure through second missed ctor
#rnx

diff --git a/Engine/Plugins/Runtime/Metasound/Source/MetasoundFrontend/Public/MetasoundEnum.h b/Engine/Plugins/Runtime/Metasound/Source/MetasoundFrontend/Public/MetasoundEnum.h
--- a/Engine/Plugins/Runtime/Metasound/Source/MetasoundFrontend/Public/MetasoundEnum.h
+++ b/Engine/Plugins/Runtime/Metasound/Source/MetasoundFrontend/Public/MetasoundEnum.h
@@ -2,8 +2,12 @@
 
 #pragma once
 
-#include "CoreMinimal.h"
+#include "Algo/Transform.h"
+#include "Internationalization/Text.h"
 #include "MetasoundLog.h"
+#include "Misc/Optional.h"
+#include "Templates/UnrealTypeTraits.h"
+#include "UObject/NameTypes.h"
 
 #include <type_traits>
 
@@ -117,10 +121,23 @@ namespace Metasound
 		{
 			// Try and convert to validate this is a valid value.
 			TOptional<FName> Converted = ToName(InValue);
-			if (ensure(Converted))
+			if (Converted)
 			{
 				EnumValue = InValue;
 			}
+			else if (!bHasWarnedNameToEnumConversionFailure)
+			{
+				TArray<FString> ValueStrings;
+				Algo::Transform(GetAllNames(), ValueStrings, [](const FName& Name) { return Name.ToString(); });
+
+				UE_LOG(LogMetaSound, Warning,
+					TEXT("Cannot create valid enum from value '%s'.\nPossible Values:\n"),
+					*FString::FromInt((int32)(InValue)),
+					*FString::Join(ValueStrings, TEXT("\n, "))
+				);
+				bHasWarnedNameToEnumConversionFailure = true;
+			}
+
 		}
 
 		// From Int32 (this is the common path from a Literal).
@@ -130,7 +147,7 @@ namespace Metasound
 		}
 
 		// From Name
-		explicit TEnum(FName InValueName)		
+		explicit TEnum(FName InValueName)
 		{
 			// Try and convert from Name to Value 
 			TOptional<EnumType> Converted = NameToEnum(InValueName);
@@ -142,8 +159,14 @@ namespace Metasound
 			{
 				if (!bHasWarnedNameToEnumConversionFailure)
 				{
-					UE_LOG(LogMetaSound, Warning, TEXT("Cannot create valid enum value from string '%s'"), *InValueName.ToString());
-					bHasWarnedNameToEnumConversionFailure = true;
+					TArray<FString> ValueStrings;
+					Algo::Transform(GetAllNames(), ValueStrings, [](const FName& Name) { return Name.ToString(); });
+
+					UE_LOG(LogMetaSound, Warning,
+						TEXT("Cannot create valid enum value from string '%s'.\nPossible Values:\n"),
+						*InValueName.ToString(),
+						*FString::Join(ValueStrings, TEXT("\n, "))
+					);
 				}
 			}
 		}
@@ -193,7 +216,7 @@ namespace Metasound
 		{
 			return TEnumStringHelper<EnumType>::GetAllNames();
 		}
-				
+
 	private:
 		// Keep the type in its fully typed form for debugging.
 		EnumType EnumValue = DefaultValue;
@@ -217,7 +240,6 @@ namespace Metasound
 	{
 		static constexpr bool bIsEnum = true;
 		using InnerType = T;
-		static constexpr T DefaultValue = D;		
-	};	
+		static constexpr T DefaultValue = D;
+	};
 }
-
