MetaSound: Add constructor inputs / outputs to frontend nodes
#p4v-preflight-copy 20381225

diff --git a/Engine/Plugins/Runtime/Metasound/Source/MetasoundEditor/Private/MetasoundEditorGraphBuilder.cpp b/Engine/Plugins/Runtime/Metasound/Source/MetasoundEditor/Private/MetasoundEditorGraphBuilder.cpp
--- a/Engine/Plugins/Runtime/Metasound/Source/MetasoundEditor/Private/MetasoundEditorGraphBuilder.cpp
+++ b/Engine/Plugins/Runtime/Metasound/Source/MetasoundEditor/Private/MetasoundEditorGraphBuilder.cpp
@@ -1049,8 +1049,15 @@ namespace Metasound
 			FMetasoundAssetBase* MetaSoundAsset = IMetasoundUObjectRegistry::Get().GetObjectAsAssetBase(&InMetaSound);
 			check(MetaSoundAsset);
 
-			const FName NewName = GenerateUniqueNameByClassType(InMetaSound, EMetasoundFrontendClassType::Input, InNameBase ? InNameBase->ToString() : TEXT("Input"));
-			return MetaSoundAsset->GetRootGraphHandle()->AddInputVertex(NewName, InTypeName, InDefaultValue);
+			FMetasoundFrontendClassInput ClassInput;
+			ClassInput.Name = GenerateUniqueNameByClassType(InMetaSound, EMetasoundFrontendClassType::Input, InNameBase ? InNameBase->ToString() : TEXT("Input"));
+			ClassInput.TypeName = InTypeName;
+			if (nullptr != InDefaultValue)
+			{
+				ClassInput.DefaultLiteral = *InDefaultValue;
+			}
+
+			return MetaSoundAsset->GetRootGraphHandle()->AddInputVertex(ClassInput);
 		}
 
 		Frontend::FNodeHandle FGraphBuilder::AddOutputNodeHandle(UObject& InMetaSound, const FName InTypeName, const FName* InNameBase)
diff --git a/Engine/Plugins/Runtime/Metasound/Source/MetasoundFrontend/Private/MetasoundFrontendDataTypeRegistry.cpp b/Engine/Plugins/Runtime/Metasound/Source/MetasoundFrontend/Private/MetasoundFrontendDataTypeRegistry.cpp
--- a/Engine/Plugins/Runtime/Metasound/Source/MetasoundFrontend/Private/MetasoundFrontendDataTypeRegistry.cpp
+++ b/Engine/Plugins/Runtime/Metasound/Source/MetasoundFrontend/Private/MetasoundFrontendDataTypeRegistry.cpp
@@ -191,6 +191,40 @@ namespace Metasound
 				TSharedRef<IDataTypeRegistryEntry, ESPMode::ThreadSafe> DataTypeEntry;
 			};
 
+			// Node registry entry for constructor input nodes created from a data type registry entry.
+			class FConstructorInputNodeRegistryEntry : public FDataTypeNodeRegistryEntry
+			{
+			public:
+				FConstructorInputNodeRegistryEntry() = delete;
+
+				FConstructorInputNodeRegistryEntry(const TSharedRef<IDataTypeRegistryEntry, ESPMode::ThreadSafe>& InDataTypeEntry)
+				: DataTypeEntry(InDataTypeEntry)
+				{
+					UpdateNodeClassInfo(DataTypeEntry->GetFrontendConstructorInputClass());
+				}
+
+				virtual ~FConstructorInputNodeRegistryEntry() = default;
+
+				virtual const FMetasoundFrontendClass& GetFrontendClass() const override
+				{
+					return DataTypeEntry->GetFrontendConstructorInputClass();
+				}
+
+				virtual TUniquePtr<INode> CreateNode(FDefaultNamedVertexWithLiteralNodeConstructorParams&& InParams) const override
+				{
+					return DataTypeEntry->CreateConstructorInputNode(MoveTemp(InParams));
+				}
+
+				virtual TUniquePtr<INodeRegistryEntry> Clone() const override
+				{
+					return MakeUnique<FConstructorInputNodeRegistryEntry>(DataTypeEntry);
+				}
+
+			private:
+				
+				TSharedRef<IDataTypeRegistryEntry, ESPMode::ThreadSafe> DataTypeEntry;
+			};
+
 			// Node registry entry for output nodes created from a data type registry entry.
 			class FOutputNodeRegistryEntry : public FDataTypeNodeRegistryEntry
 			{
@@ -225,6 +259,41 @@ namespace Metasound
 				TSharedRef<IDataTypeRegistryEntry, ESPMode::ThreadSafe> DataTypeEntry;
 			};
 
+			// Node registry entry for constructor output nodes created from a data type registry entry.
+			class FConstructorOutputNodeRegistryEntry : public FDataTypeNodeRegistryEntry
+			{
+			public:
+				FConstructorOutputNodeRegistryEntry() = delete;
+
+				FConstructorOutputNodeRegistryEntry(const TSharedRef<IDataTypeRegistryEntry, ESPMode::ThreadSafe>& InDataTypeEntry)
+				: DataTypeEntry(InDataTypeEntry)
+				{
+					UpdateNodeClassInfo(DataTypeEntry->GetFrontendConstructorOutputClass());
+				}
+
+				virtual ~FConstructorOutputNodeRegistryEntry() = default;
+
+				virtual const FMetasoundFrontendClass& GetFrontendClass() const override
+				{
+					return DataTypeEntry->GetFrontendConstructorOutputClass();
+				}
+
+				virtual TUniquePtr<INode> CreateNode(FDefaultNamedVertexNodeConstructorParams&& InParams) const override
+				{
+					return DataTypeEntry->CreateConstructorOutputNode(MoveTemp(InParams));
+				}
+
+				virtual TUniquePtr<INodeRegistryEntry> Clone() const override
+				{
+					return MakeUnique<FConstructorOutputNodeRegistryEntry>(DataTypeEntry);
+				}
+
+			private:
+				
+				TSharedRef<IDataTypeRegistryEntry, ESPMode::ThreadSafe> DataTypeEntry;
+			};
+
+
 			// Node registry entry for literal nodes created from a data type registry entry.
 			class FLiteralNodeRegistryEntry : public FDataTypeNodeRegistryEntry
 			{
@@ -439,8 +508,10 @@ namespace Metasound
 				virtual TSharedPtr<IDataChannel, ESPMode::ThreadSafe> CreateDataChannel(const FName& InDataType, const FOperatorSettings& InOperatorSettings) const override;
 
 				virtual bool GetFrontendInputClass(const FName& InDataType, FMetasoundFrontendClass& OutClass) const override;
+				virtual bool GetFrontendConstructorInputClass(const FName& InDataType, FMetasoundFrontendClass& OutClass) const override;
 				virtual bool GetFrontendLiteralClass(const FName& InDataType, FMetasoundFrontendClass& OutClass) const override;
 				virtual bool GetFrontendOutputClass(const FName& InDataType, FMetasoundFrontendClass& OutClass) const override;
+				virtual bool GetFrontendConstructorOutputClass(const FName& InDataType, FMetasoundFrontendClass& OutClass) const override;
 				virtual bool GetFrontendVariableClass(const FName& InDataType, FMetasoundFrontendClass& OutClass) const override;
 				virtual bool GetFrontendVariableMutatorClass(const FName& InDataType, FMetasoundFrontendClass& OutClass) const override;
 				virtual bool GetFrontendVariableAccessorClass(const FName& InDataType, FMetasoundFrontendClass& OutClass) const override;
@@ -448,8 +519,10 @@ namespace Metasound
 
 				// Create a new instance of a C++ implemented node from the registry.
 				virtual TUniquePtr<INode> CreateInputNode(const FName& InInputType, FInputNodeConstructorParams&& InParams) const override;
+				virtual TUniquePtr<INode> CreateConstructorInputNode(const FName& InInputType, FInputNodeConstructorParams&& InParams) const override;
 				virtual TUniquePtr<INode> CreateLiteralNode(const FName& InLiteralType, FLiteralNodeConstructorParams&& InParams) const override;
 				virtual TUniquePtr<INode> CreateOutputNode(const FName& InOutputType, FOutputNodeConstructorParams&& InParams) const override;
+				virtual TUniquePtr<INode> CreateConstructorOutputNode(const FName& InOutputType, FOutputNodeConstructorParams&& InParams) const override;
 				virtual TUniquePtr<INode> CreateReceiveNode(const FName& InOutputType, const FNodeInitData& InParams) const override;
 				virtual TUniquePtr<INode> CreateVariableNode(const FName& InDataType, FVariableNodeConstructorParams&&) const override;
 				virtual TUniquePtr<INode> CreateVariableMutatorNode(const FName& InDataType, const FNodeInitData&) const override;
@@ -491,12 +564,20 @@ namespace Metasound
 						if (Entry->GetDataTypeInfo().bIsParsable)
 						{
 							NodeRegistry->RegisterNode(MakeUnique<FInputNodeRegistryEntry>(Entry));
+							
 							NodeRegistry->RegisterNode(MakeUnique<FOutputNodeRegistryEntry>(Entry));
+							
 							NodeRegistry->RegisterNode(MakeUnique<FLiteralNodeRegistryEntry>(Entry));
 							NodeRegistry->RegisterNode(MakeUnique<FVariableNodeRegistryEntry>(Entry));
 							NodeRegistry->RegisterNode(MakeUnique<FVariableMutatorNodeRegistryEntry>(Entry));
 							NodeRegistry->RegisterNode(MakeUnique<FVariableAccessorNodeRegistryEntry>(Entry));
 							NodeRegistry->RegisterNode(MakeUnique<FVariableDeferredAccessorNodeRegistryEntry>(Entry));
+
+							if (Entry->GetDataTypeInfo().bIsConstructorType)
+							{
+								NodeRegistry->RegisterNode(MakeUnique<FConstructorInputNodeRegistryEntry>(Entry));
+								NodeRegistry->RegisterNode(MakeUnique<FConstructorOutputNodeRegistryEntry>(Entry));
+							}
 						}
 					}
 
@@ -772,6 +853,16 @@ namespace Metasound
 				return false;
 			}
 
+			bool FDataTypeRegistry::GetFrontendConstructorInputClass(const FName& InDataType, FMetasoundFrontendClass& OutClass) const
+			{
+				if (const IDataTypeRegistryEntry* Entry = FindDataTypeEntry(InDataType))
+				{
+					OutClass = Entry->GetFrontendConstructorInputClass();
+					return true;
+				}
+				return false;
+			}
+
 			bool FDataTypeRegistry::GetFrontendLiteralClass(const FName& InDataType, FMetasoundFrontendClass& OutClass) const
 			{
 				if (const IDataTypeRegistryEntry* Entry = FindDataTypeEntry(InDataType))
@@ -792,6 +883,16 @@ namespace Metasound
 				return false;
 			}
 
+			bool FDataTypeRegistry::GetFrontendConstructorOutputClass(const FName& InDataType, FMetasoundFrontendClass& OutClass) const
+			{
+				if (const IDataTypeRegistryEntry* Entry = FindDataTypeEntry(InDataType))
+				{
+					OutClass = Entry->GetFrontendConstructorOutputClass();
+					return true;
+				}
+				return false;
+			}
+
 			bool FDataTypeRegistry::GetFrontendVariableClass(const FName& InDataType, FMetasoundFrontendClass& OutClass) const
 			{
 				if (const IDataTypeRegistryEntry* Entry = FindDataTypeEntry(InDataType))
@@ -844,6 +945,15 @@ namespace Metasound
 				return nullptr;
 			}
 
+			TUniquePtr<INode> FDataTypeRegistry::CreateConstructorInputNode(const FName& InDataType, FInputNodeConstructorParams&& InParams) const
+			{
+				if (const IDataTypeRegistryEntry* Entry = FindDataTypeEntry(InDataType))
+				{
+					return Entry->CreateConstructorInputNode(MoveTemp(InParams));
+				}
+				return nullptr;
+			}
+
 			TUniquePtr<INode> FDataTypeRegistry::CreateLiteralNode(const FName& InDataType, FLiteralNodeConstructorParams&& InParams) const 
 			{
 				if (const IDataTypeRegistryEntry* Entry = FindDataTypeEntry(InDataType))
@@ -862,6 +972,15 @@ namespace Metasound
 				return nullptr;
 			}
 
+			TUniquePtr<INode> FDataTypeRegistry::CreateConstructorOutputNode(const FName& InDataType, FOutputNodeConstructorParams&& InParams) const
+			{
+				if (const IDataTypeRegistryEntry* Entry = FindDataTypeEntry(InDataType))
+				{
+					return Entry->CreateConstructorOutputNode(MoveTemp(InParams));
+				}
+				return nullptr;
+			}
+
 			TUniquePtr<INode> FDataTypeRegistry::CreateReceiveNode(const FName& InDataType, const FNodeInitData& InParams) const
 			{
 				if (const IDataTypeRegistryEntry* Entry = FindDataTypeEntry(InDataType))
@@ -925,7 +1044,6 @@ namespace Metasound
 			static MetasoundFrontendDataTypeRegistryPrivate::FDataTypeRegistry Registry;
 			return Registry;
 		}
-
 	}
 }
 
diff --git a/Engine/Plugins/Runtime/Metasound/Source/MetasoundFrontend/Private/MetasoundFrontendDocument.cpp b/Engine/Plugins/Runtime/Metasound/Source/MetasoundFrontend/Private/MetasoundFrontendDocument.cpp
--- a/Engine/Plugins/Runtime/Metasound/Source/MetasoundFrontend/Private/MetasoundFrontendDocument.cpp
+++ b/Engine/Plugins/Runtime/Metasound/Source/MetasoundFrontend/Private/MetasoundFrontendDocument.cpp
@@ -556,3 +556,16 @@ FMetasoundFrontendDocument::FMetasoundFrontendDocument()
 	RootGraph.Metadata.SetType(EMetasoundFrontendClassType::Graph);
 	ArchetypeVersion = FMetasoundFrontendVersion::GetInvalid();
 }
+
+const TCHAR* LexToString(EMetasoundFrontendVertexAccessType InVertexAccess)
+{
+	switch (InVertexAccess)
+	{
+		case EMetasoundFrontendVertexAccessType::Value:
+			return TEXT("Value");
+			
+		case EMetasoundFrontendVertexAccessType::Reference:
+		default:
+			return TEXT("Reference");
+	}
+}
diff --git a/Engine/Plugins/Runtime/Metasound/Source/MetasoundFrontend/Private/MetasoundFrontendInputController.cpp b/Engine/Plugins/Runtime/Metasound/Source/MetasoundFrontend/Private/MetasoundFrontendInputController.cpp
--- a/Engine/Plugins/Runtime/Metasound/Source/MetasoundFrontend/Private/MetasoundFrontendInputController.cpp
+++ b/Engine/Plugins/Runtime/Metasound/Source/MetasoundFrontend/Private/MetasoundFrontendInputController.cpp
@@ -57,6 +57,16 @@ namespace Metasound
 			return Invalid::GetInvalidName();
 		}
 
+		EMetasoundFrontendVertexAccessType FBaseInputController::GetVertexAccessType() const 
+		{
+			if (const FMetasoundFrontendClassVertex* ClassInput = ClassInputPtr.Get())
+			{
+				return ClassInput->AccessType;
+			}
+			
+			return EMetasoundFrontendVertexAccessType::Reference;
+		}
+
 #if WITH_EDITOR
 		FText FBaseInputController::GetDisplayName() const
 		{
@@ -193,6 +203,11 @@ namespace Metasound
 				OutConnectability.Connectable = FConnectability::EConnectable::No;
 				OutConnectability.Reason = FConnectability::EReason::IncompatibleDataTypes;
 			}
+			else if (!IsCompatibleAccessType(InController))
+			{
+				OutConnectability.Connectable = FConnectability::EConnectable::No;
+				OutConnectability.Reason = FConnectability::EReason::IncompatibleAccessTypes;
+			}
 			else if (OtherDataType == DataType)
 			{
 				// If data types are equal, connection can happen.
@@ -233,25 +248,37 @@ namespace Metasound
 			{
 				return false;
 			}
+			
 
 			if (FMetasoundFrontendGraph* Graph = GraphPtr.Get())
 			{
-				if (ensureAlwaysMsgf(OtherDataType == DataType, TEXT("Cannot connect incompatible types.")))
+				if (OtherDataType == DataType)
 				{
-					// Overwrite an existing connection if it exists.
-					FMetasoundFrontendEdge* Edge = FindEdge();
-
-					if (!Edge)
+					if (IsCompatibleAccessType(InController))
 					{
-						Edge = &Graph->Edges.AddDefaulted_GetRef();
-						Edge->ToNodeID = GetOwningNodeID();
-						Edge->ToVertexID = GetID();
-					}
+						// Overwrite an existing connection if it exists.
+						FMetasoundFrontendEdge* Edge = FindEdge();
+
+						if (!Edge)
+						{
+							Edge = &Graph->Edges.AddDefaulted_GetRef();
+							Edge->ToNodeID = GetOwningNodeID();
+							Edge->ToVertexID = GetID();
+						}
 
-					Edge->FromNodeID = InController.GetOwningNodeID();
-					Edge->FromVertexID = InController.GetID();
+						Edge->FromNodeID = InController.GetOwningNodeID();
+						Edge->FromVertexID = InController.GetID();
 
-					return true;
+						return true;
+					}
+					else
+					{
+						UE_LOG(LogMetaSound, Error, TEXT("Cannot connect incompatible vertex access types (Input)%s and (Output)%s."), *LexToString(GetVertexAccessType()), *LexToString(InController.GetVertexAccessType()));
+					}
+				}
+				else
+				{
+					UE_LOG(LogMetaSound, Error, TEXT("Cannot connect incompatible data types %s and %s."), *DataType.ToString(), *OtherDataType.ToString());
 				}
 			}
 
@@ -338,6 +365,19 @@ namespace Metasound
 			return false;
 		}
 
+		bool FBaseInputController::IsCompatibleAccessType(const IOutputController& InOutputController) const
+		{
+			if (EMetasoundFrontendVertexAccessType::Value == GetVertexAccessType())
+			{
+				// If the input vertex accesses by "Value" then the output vertex 
+				// must also access by "Value" to enforce unexpected consequences 
+				// of connecting data which varies over time to an input which only
+				// evaluates the data during operator initialization.
+				return (EMetasoundFrontendVertexAccessType::Value == InOutputController.GetVertexAccessType());
+			}
+			return true;
+		}
+
 		const FMetasoundFrontendEdge* FBaseInputController::FindEdge() const
 		{
 			if (const FMetasoundFrontendGraph* Graph = GraphPtr.Get())
diff --git a/Engine/Plugins/Runtime/Metasound/Source/MetasoundFrontend/Private/MetasoundFrontendInputController.h b/Engine/Plugins/Runtime/Metasound/Source/MetasoundFrontend/Private/MetasoundFrontendInputController.h
--- a/Engine/Plugins/Runtime/Metasound/Source/MetasoundFrontend/Private/MetasoundFrontendInputController.h
+++ b/Engine/Plugins/Runtime/Metasound/Source/MetasoundFrontend/Private/MetasoundFrontendInputController.h
@@ -39,6 +39,7 @@ namespace Metasound
 			virtual FGuid GetID() const override;
 			virtual const FName& GetDataType() const override;
 			virtual const FVertexName& GetName() const override;
+			virtual EMetasoundFrontendVertexAccessType GetVertexAccessType() const override;
 			virtual bool ClearLiteral() override;
 			virtual const FMetasoundFrontendLiteral* GetLiteral() const override;
 			virtual void SetLiteral(const FMetasoundFrontendLiteral& InLiteral) override;
@@ -79,6 +80,7 @@ namespace Metasound
 
 		protected:
 
+
 			virtual FDocumentAccess ShareAccess() override;
 			virtual FConstDocumentAccess ShareAccess() const override;
 
@@ -90,6 +92,11 @@ namespace Metasound
 			FConstClassInputAccessPtr ClassInputPtr;
 			FGraphAccessPtr GraphPtr;
 			FNodeHandle OwningNode;
+
+		private:
+
+			bool IsCompatibleAccessType(const IOutputController& InOutputController) const;
+
 		};
 
 		/** FOutputNodeInputController represents the input vertex of an output 
diff --git a/Engine/Plugins/Runtime/Metasound/Source/MetasoundFrontend/Private/MetasoundFrontendInvalidController.h b/Engine/Plugins/Runtime/Metasound/Source/MetasoundFrontend/Private/MetasoundFrontendInvalidController.h
--- a/Engine/Plugins/Runtime/Metasound/Source/MetasoundFrontend/Private/MetasoundFrontendInvalidController.h
+++ b/Engine/Plugins/Runtime/Metasound/Source/MetasoundFrontend/Private/MetasoundFrontendInvalidController.h
@@ -59,6 +59,7 @@ namespace Metasound
 			virtual FGuid GetID() const override { return Metasound::FrontendInvalidID; }
 			virtual const FName& GetDataType() const override { return Invalid::GetInvalidName(); }
 			virtual const FVertexName& GetName() const override { return Invalid::GetInvalidName(); }
+			virtual EMetasoundFrontendVertexAccessType GetVertexAccessType() const override { return EMetasoundFrontendVertexAccessType::Reference; }
 
 #if WITH_EDITOR
 			virtual FText GetDisplayName() const override { return FText::GetEmpty(); }
@@ -102,6 +103,7 @@ namespace Metasound
 			virtual bool IsConnected() const override { return false; }
 			virtual const FName& GetDataType() const override { return Invalid::GetInvalidName(); }
 			virtual const FVertexName& GetName() const override { return Invalid::GetInvalidName(); }
+			virtual EMetasoundFrontendVertexAccessType GetVertexAccessType() const override { return EMetasoundFrontendVertexAccessType::Reference; }
 
 #if WITH_EDITOR
 			virtual FText GetDisplayName() const override { return Invalid::GetInvalidText(); }
diff --git a/Engine/Plugins/Runtime/Metasound/Source/MetasoundFrontend/Private/MetasoundFrontendOutputController.cpp b/Engine/Plugins/Runtime/Metasound/Source/MetasoundFrontend/Private/MetasoundFrontendOutputController.cpp
--- a/Engine/Plugins/Runtime/Metasound/Source/MetasoundFrontend/Private/MetasoundFrontendOutputController.cpp
+++ b/Engine/Plugins/Runtime/Metasound/Source/MetasoundFrontend/Private/MetasoundFrontendOutputController.cpp
@@ -49,6 +49,15 @@ namespace Metasound
 			return Invalid::GetInvalidName();
 		}
 
+		EMetasoundFrontendVertexAccessType FBaseOutputController::GetVertexAccessType() const 
+		{
+			if (const FMetasoundFrontendClassVertex* ClassOutput = ClassOutputPtr.Get())
+			{
+				return ClassOutput->AccessType;
+			}
+			return EMetasoundFrontendVertexAccessType::Reference;
+		}
+
 		FGuid FBaseOutputController::GetOwningNodeID() const
 		{
 			return OwningNode->GetID();
diff --git a/Engine/Plugins/Runtime/Metasound/Source/MetasoundFrontend/Private/MetasoundFrontendOutputController.h b/Engine/Plugins/Runtime/Metasound/Source/MetasoundFrontend/Private/MetasoundFrontendOutputController.h
--- a/Engine/Plugins/Runtime/Metasound/Source/MetasoundFrontend/Private/MetasoundFrontendOutputController.h
+++ b/Engine/Plugins/Runtime/Metasound/Source/MetasoundFrontend/Private/MetasoundFrontendOutputController.h
@@ -43,6 +43,7 @@ namespace Metasound
 			virtual FGuid GetID() const override;
 			virtual const FName& GetDataType() const override;
 			virtual const FVertexName& GetName() const override;
+			virtual EMetasoundFrontendVertexAccessType GetVertexAccessType() const override;
 
 #if WITH_EDITOR
 			virtual FText GetDisplayName() const override;
diff --git a/Engine/Plugins/Runtime/Metasound/Source/MetasoundFrontend/Private/MetasoundInputNode.cpp b/Engine/Plugins/Runtime/Metasound/Source/MetasoundFrontend/Private/MetasoundInputNode.cpp
new file mode 100644
--- /dev/null
+++ b/Engine/Plugins/Runtime/Metasound/Source/MetasoundFrontend/Private/MetasoundInputNode.cpp
@@ -0,0 +1,45 @@
+// Copyright Epic Games, Inc. All Rights Reserved.
+
+#include "MetasoundInputNode.h"
+
+#include "MetasoundDataReference.h"
+#include "MetasoundVertex.h"
+#include "MetasoundVertexData.h"
+#include "UObject/NameTypes.h"
+
+namespace Metasound
+{
+	FDataReferenceCollection FInputValueOperator::GetInputs() const
+	{
+		checkNoEntry();
+		return {};
+	}
+
+	FDataReferenceCollection FInputValueOperator::GetOutputs() const
+	{
+		checkNoEntry();
+		return {};
+	}
+
+	void FInputValueOperator::Bind(FVertexInterfaceData& InVertexData) const
+	{
+		if (const FAnyDataReference* Ref = InVertexData.GetInputs().FindDataReference(VertexName))
+		{
+			checkf(EDataReferenceAccessType::Value == Ref->GetAccessType(), TEXT("Expected bound reference to have %s access. Actual access was %s"), *LexToString(EDataReferenceAccessType::Value), *LexToString(Ref->GetAccessType()));
+
+			// Pass through input to output
+			InVertexData.GetOutputs().BindVertex(VertexName, *Ref);
+		}
+		else
+		{
+			// Use stored default value
+			InVertexData.GetInputs().BindVertex(VertexName, Default);
+			InVertexData.GetOutputs().BindVertex(VertexName, Default);
+		}
+	}
+
+	IOperator::FExecuteFunction FInputValueOperator::GetExecuteFunction()
+	{
+		return nullptr;
+	}
+}
diff --git a/Engine/Plugins/Runtime/Metasound/Source/MetasoundFrontend/Private/MetasoundOutputNode.cpp b/Engine/Plugins/Runtime/Metasound/Source/MetasoundFrontend/Private/MetasoundOutputNode.cpp
new file mode 100644
--- /dev/null
+++ b/Engine/Plugins/Runtime/Metasound/Source/MetasoundFrontend/Private/MetasoundOutputNode.cpp
@@ -0,0 +1,39 @@
+// Copyright Epic Games, Inc. All Rights Reserved.
+
+#include "MetasoundOutputNode.h"
+
+namespace Metasound
+{
+	namespace OutputNodePrivate
+	{
+		FOutputOperator::FOutputOperator(const FVertexName& InVertexName, const FAnyDataReference& InDataReference)
+		: VertexName(InVertexName)
+		, DataReference(InDataReference)
+		{
+		}
+
+		FDataReferenceCollection FOutputOperator::GetInputs() const
+		{
+			// Slated for deprecation
+			return {};
+		}
+
+		FDataReferenceCollection FOutputOperator::GetOutputs() const
+		{
+			// Slated for deprecation
+			return {};
+		}
+
+		void FOutputOperator::Bind(FVertexInterfaceData& InVertexData) const
+		{
+			InVertexData.GetInputs().BindVertex(VertexName, DataReference);
+			InVertexData.GetOutputs().BindVertex(VertexName, DataReference);
+		}
+
+		IOperator::FExecuteFunction FOutputOperator::GetExecuteFunction()
+		{
+			return nullptr;
+		}
+	}
+}
+
diff --git a/Engine/Plugins/Runtime/Metasound/Source/MetasoundFrontend/Public/MetasoundDataTypeRegistrationMacro.h b/Engine/Plugins/Runtime/Metasound/Source/MetasoundFrontend/Public/MetasoundDataTypeRegistrationMacro.h
--- a/Engine/Plugins/Runtime/Metasound/Source/MetasoundFrontend/Public/MetasoundDataTypeRegistrationMacro.h
+++ b/Engine/Plugins/Runtime/Metasound/Source/MetasoundFrontend/Public/MetasoundDataTypeRegistrationMacro.h
@@ -13,6 +13,7 @@
 #include "MetasoundDataReferenceMacro.h"
 #include "MetasoundEnum.h"
 #include "MetasoundFrontendDataTypeRegistry.h"
+#include "MetasoundFrontendDataTypeTraits.h"
 #include "MetasoundFrontendRegistries.h"
 #include "MetasoundInputNode.h"
 #include "MetasoundLiteral.h"
@@ -29,81 +30,8 @@
 
 namespace Metasound
 {
-	/** Enables or disables automatic registration of array types given a 
-	 * MetaSound data type. By default this is true, and all data types will have
-	 * an associated TArray<DataType> registered if it is supported. */
-	template<typename ... Type>
-	struct TEnableAutoArrayTypeRegistration
-	{
-		static constexpr bool Value = true;
-	};
-
-	/** Enables or disables automatic registration of auto conversion nodes given a 
-	 * MetaSound data type. By default this is true, and all data types will have
-	 * associated conversion nodes registered based upon the data types supported
-	 * constructors and implicit conversions. */
-	template<typename ... Type>
-	struct TEnableAutoConverterNodeRegistration
-	{
-		static constexpr bool Value = true;
-	};
-
-	/** Enables or disables send and receive node registration for a given MetaSound
-	 * data type. By default this is true and all data types supported by the transmission
-	 * system will have associated send and receive nodes. */
-	template<typename ... Type>
-	struct TEnableTransmissionNodeRegistration
-	{
-		static constexpr bool Value = true;
-	};
-
-	/** Enables or disables using a data type in constructor vertices. 
-	 * By default this is true and all data types that support copy assignment
-	 * and copy construction. */
-	template<typename ... Type>
-	struct TEnableConstructorVertex
-	{
-		static constexpr bool Value = true;
-	};
-
 	namespace MetasoundDataTypeRegistrationPrivate
 	{
-		// Helper utility to test if we a datatype is copyable
-		template <typename TDataType>
-		struct TIsCopyable
-		{
-		private:
-			static constexpr bool bIsCopyConstructible = std::is_copy_constructible<TDataType>::value;
-			static constexpr bool bIsCopyAssignable = std::is_copy_assignable<TDataType>::value;
-
-		public:
-			static constexpr bool Value = bIsCopyConstructible && bIsCopyAssignable;
-		};
-
-		// Helper utility to test if we can transmit a datatype between a send and a receive node.
-		template <typename TDataType>
-		struct TIsTransmittable
-		{
-			static constexpr bool Value = TIsCopyable<TDataType>::Value;
-		};
-
-
-
-		// Specialization of TIsTransmittable<> for TArray to handle lax TArray copy constructor 
-		// definition. This can be removed once updates to TArray are merged into the same codebase.
-		// TODO: delete me eventually.
-		template<typename TElementType>
-		struct TIsTransmittable<TArray<TElementType>>
-		{
-			// Depend on copy constructor of elements to determine whether TArray is copy constructible
-			// and copy assignable. Definitions for TArray copy construction and copy assignment are
-			// defined by default whether or not the TArray element types support them. This logic
-			// skips checking the TArray object itselt and instead checks the TArray element type
-			// directly.
-
-			static constexpr bool Value = TIsTransmittable<TElementType>::Value;
-		};
-
 		// Returns the Array version of a literal type if it exists.
 		template<ELiteralType LiteralType>
 		struct TLiteralArrayEnum 
@@ -251,7 +179,7 @@
 			RegistryInfo.bIsEnum = TEnumTraits<TDataType>::bIsEnum;
 			RegistryInfo.bIsVariable = TIsVariable<TDataType>::Value;
 			RegistryInfo.bIsTransmittable = TIsTransmittable<TDataType>::Value;
-			RegistryInfo.bIsConstructorType = TIsCopyable<TDataType>::Value && TEnableConstructorVertex<TDataType>::Value;
+			RegistryInfo.bIsConstructorType = TIsConstructorVertexSupported<TDataType>::Value;
 			
 			if constexpr (std::is_base_of<UObject, UClassToUse>::value)
 			{
@@ -339,11 +267,11 @@
 		bool RegisterDataTypeWithFrontendInternal()
 		{
 			static constexpr bool bIsParsable = TLiteralTraits<TDataType>::bIsParsableFromAnyLiteralType;
+			static constexpr bool bIsConstructorType = TIsConstructorVertexSupported<TDataType>::Value;
+
 			static bool bAlreadyRegisteredThisDataType = false;
 			if (bAlreadyRegisteredThisDataType)
 			{
-				// if we reenter this code (because DECLARE_METASOUND_DATA_REFERENCE_TYPES was called twice with the same type),
-				// we catch it here.
 				UE_LOG(LogMetaSound, Display, TEXT("Tried to call REGISTER_METASOUND_DATATYPE twice with the same class %s. ignoring the second call. Likely because REGISTER_METASOUND_DATATYPE is in a header that's used in multiple modules. Consider moving it to a private header or cpp file."), TDataReferenceTypeInfo<TDataType>::TypeName)
 				return false;
 			}
@@ -364,10 +292,10 @@
 
 					if constexpr (bIsParsable)
 					{
-						TInputNode<TDataType> InputPrototype(TEXT(""), FGuid(), TEXT(""), FLiteral());
+						TInputNode<TDataType, EVertexAccessType::Reference> InputPrototype(TEXT(""), FGuid(), TEXT(""), FLiteral());
 						InputClass = Metasound::Frontend::GenerateClass(InputPrototype.GetMetadata(), EMetasoundFrontendClassType::Input);
 
-						TOutputNode<TDataType> OutputPrototype(TEXT(""), FGuid(), TEXT(""));
+						TOutputNode<TDataType, EVertexAccessType::Reference> OutputPrototype(TEXT(""), FGuid(), TEXT(""));
 						OutputClass = Metasound::Frontend::GenerateClass(OutputPrototype.GetMetadata(), EMetasoundFrontendClassType::Output);
 
 						TLiteralNode<TDataType> LiteralPrototype(TEXT(""), FGuid(), FLiteral());
@@ -384,6 +312,14 @@
 
 						TVariableDeferredAccessorNode<TDataType> VariableDeferredAccessorPrototype(TEXT(""), FGuid());
 						VariableDeferredAccessorClass = Metasound::Frontend::GenerateClass(VariableDeferredAccessorPrototype.GetMetadata(), EMetasoundFrontendClassType::VariableDeferredAccessor);
+
+						if constexpr (bIsConstructorType)
+						{
+							TInputNode<TDataType, EVertexAccessType::Value> ConstructorInputPrototype(TEXT(""), FGuid(), TEXT(""), FLiteral());
+							ConstructorInputClass = Metasound::Frontend::GenerateClass(ConstructorInputPrototype.GetMetadata(), EMetasoundFrontendClassType::Input);
+							TOutputNode<TDataType, EVertexAccessType::Value> ConstructorOutputPrototype(TEXT(""), FGuid(), TEXT(""));
+							ConstructorOutputClass = Metasound::Frontend::GenerateClass(ConstructorOutputPrototype.GetMetadata(), EMetasoundFrontendClassType::Output);
+						}
 					}
 				}
 
@@ -404,6 +340,11 @@
 					return InputClass;
 				}
 
+				virtual const FMetasoundFrontendClass& GetFrontendConstructorInputClass() const override
+				{
+					return ConstructorInputClass;
+				}
+
 				virtual const FMetasoundFrontendClass& GetFrontendLiteralClass() const override
 				{
 					return LiteralClass;
@@ -414,6 +355,11 @@
 					return OutputClass;
 				}
 
+				virtual const FMetasoundFrontendClass& GetFrontendConstructorOutputClass() const override
+				{
+					return ConstructorOutputClass;
+				}
+
 				virtual const FMetasoundFrontendClass& GetFrontendVariableClass() const override
 				{
 					return VariableClass;
@@ -438,7 +384,7 @@
 				{
 					if constexpr (bIsParsable)
 					{
-						return MakeUnique<TInputNode<TDataType>>(InParams.NodeName, InParams.InstanceID, InParams.VertexName, MoveTemp(InParams.InitParam));
+						return MakeUnique<TInputNode<TDataType, EVertexAccessType::Reference>>(InParams.NodeName, InParams.InstanceID, InParams.VertexName, MoveTemp(InParams.InitParam));
 					}
 					else
 					{
@@ -446,11 +392,23 @@
 					}
 				}
 
+				virtual TUniquePtr<INode> CreateConstructorInputNode(FInputNodeConstructorParams&& InParams) const override
+				{
+					if constexpr (bIsParsable && bIsConstructorType)
+					{
+						return MakeUnique<TInputNode<TDataType, EVertexAccessType::Value>>(InParams.NodeName, InParams.InstanceID, InParams.VertexName, MoveTemp(InParams.InitParam));
+					}
+					else
+					{
+						return TUniquePtr<INode>(nullptr);
+					}
+				}
+
 				virtual TUniquePtr<INode> CreateOutputNode(FOutputNodeConstructorParams&& InParams) const override
 				{
 					if constexpr (bIsParsable)
 					{
-						return MakeUnique<TOutputNode<TDataType>>(InParams.NodeName, InParams.InstanceID, InParams.VertexName);
+						return MakeUnique<TOutputNode<TDataType, EVertexAccessType::Reference>>(InParams.NodeName, InParams.InstanceID, InParams.VertexName);
 					}
 					else
 					{
@@ -458,6 +416,18 @@
 					}
 				}
 
+				virtual TUniquePtr<INode> CreateConstructorOutputNode(FOutputNodeConstructorParams&& InParams) const override
+				{
+					if constexpr (bIsParsable && bIsConstructorType)
+					{
+						return MakeUnique<TOutputNode<TDataType, EVertexAccessType::Value>>(InParams.NodeName, InParams.InstanceID, InParams.VertexName);
+					}
+					else
+					{
+						return TUniquePtr<INode>(nullptr);
+					}
+				}
+
 				virtual TUniquePtr<INode> CreateLiteralNode(FLiteralNodeConstructorParams&& InParams) const override
 				{
 					if constexpr (bIsParsable)
@@ -572,7 +542,9 @@
 			private:
 				Frontend::FDataTypeRegistryInfo Info;
 				FMetasoundFrontendClass InputClass;
+				FMetasoundFrontendClass ConstructorInputClass;
 				FMetasoundFrontendClass OutputClass;
+				FMetasoundFrontendClass ConstructorOutputClass;
 				FMetasoundFrontendClass LiteralClass;
 				FMetasoundFrontendClass VariableClass;
 				FMetasoundFrontendClass VariableMutatorClass;
@@ -662,7 +634,7 @@
 		static_assert(std::is_same<DataType, typename std::decay<DataType>::type>::value, "DataType and decayed DataType must be the same");
 		
 		// To register a data type, an input node must be able to instantiate it.
-		static constexpr bool bCanRegister = TInputNode<DataType>::bCanRegister;
+		static constexpr bool bCanRegister = TInputNode<DataType, EVertexAccessType::Reference>::bCanRegister;
 
 		// If a data type has been successfully registered, this will be true.
 		static const bool bSuccessfullyRegistered;
diff --git a/Engine/Plugins/Runtime/Metasound/Source/MetasoundFrontend/Public/MetasoundFrontend.h b/Engine/Plugins/Runtime/Metasound/Source/MetasoundFrontend/Public/MetasoundFrontend.h
--- a/Engine/Plugins/Runtime/Metasound/Source/MetasoundFrontend/Public/MetasoundFrontend.h
+++ b/Engine/Plugins/Runtime/Metasound/Source/MetasoundFrontend/Public/MetasoundFrontend.h
@@ -91,7 +91,8 @@ namespace Metasound
 			{
 				None,
 				IncompatibleDataTypes,
-				CausesLoop
+				CausesLoop,
+				IncompatibleAccessTypes
 			};
 
 			EConnectable Connectable = EConnectable::No;
diff --git a/Engine/Plugins/Runtime/Metasound/Source/MetasoundFrontend/Public/MetasoundFrontendController.h b/Engine/Plugins/Runtime/Metasound/Source/MetasoundFrontend/Public/MetasoundFrontendController.h
--- a/Engine/Plugins/Runtime/Metasound/Source/MetasoundFrontend/Public/MetasoundFrontendController.h
+++ b/Engine/Plugins/Runtime/Metasound/Source/MetasoundFrontend/Public/MetasoundFrontendController.h
@@ -94,7 +94,6 @@ namespace Metasound
 		  * when querying if a node's class has been updated */
 		struct FClassInterfaceUpdates
 		{
-
 			TArray<const FMetasoundFrontendClassInput*> AddedInputs;
 			TArray<const FMetasoundFrontendClassOutput*> AddedOutputs;
 			TArray<const FMetasoundFrontendClassInput*> RemovedInputs;
@@ -177,6 +176,9 @@ namespace Metasound
 			
 			/** Returns the name associated with this output. */
 			virtual const FVertexName& GetName() const = 0;
+
+			/** Returns the vertex access type. */
+			virtual EMetasoundFrontendVertexAccessType GetVertexAccessType() const = 0;
 			
 #if WITH_EDITOR
 			/** Returns the human readable name associated with this output. */
@@ -329,6 +331,9 @@ namespace Metasound
 			/** Returns the data type name associated with this input. */
 			virtual const FVertexName& GetName() const = 0;
 
+			/** Returns the vertex access type. */
+			virtual EMetasoundFrontendVertexAccessType GetVertexAccessType() const = 0;
+
 #if WITH_EDITOR
 			/** Returns the data type name associated with this input. */
 			virtual FText GetDisplayName() const = 0;
@@ -813,6 +818,8 @@ namespace Metasound
 			 * @return On success, a valid input node handle. On failure, an invalid node handle.
 			 */
 			virtual FNodeHandle AddInputVertex(const FMetasoundFrontendClassInput& InDescription) = 0;
+
+			UE_DEPRECATED(5.1, "Use AddInputVertex method which specifies EMetasoundFrontendVertexAccessType")
 			virtual FNodeHandle AddInputVertex(const FVertexName& InName, const FName InTypeName, const FMetasoundFrontendLiteral* InDefaultValue) = 0;
 
 			/** Remove the input with the given name. Returns true if successfully removed, false otherwise. */
diff --git a/Engine/Plugins/Runtime/Metasound/Source/MetasoundFrontend/Public/MetasoundFrontendDataTypeRegistry.h b/Engine/Plugins/Runtime/Metasound/Source/MetasoundFrontend/Public/MetasoundFrontendDataTypeRegistry.h
--- a/Engine/Plugins/Runtime/Metasound/Source/MetasoundFrontend/Public/MetasoundFrontendDataTypeRegistry.h
+++ b/Engine/Plugins/Runtime/Metasound/Source/MetasoundFrontend/Public/MetasoundFrontendDataTypeRegistry.h
@@ -137,12 +137,18 @@
 			/** Return an FMetasoundFrontendClass representing an input node of the data type. */
 			virtual const FMetasoundFrontendClass& GetFrontendInputClass() const = 0;
 
+			/** Return an FMetasoundFrontendClass representing an input node of the data type. */
+			virtual const FMetasoundFrontendClass& GetFrontendConstructorInputClass() const = 0;
+
 			/** Return an FMetasoundFrontendClass representing a variable node of the data type. */
 			virtual const FMetasoundFrontendClass& GetFrontendLiteralClass() const = 0;
 
 			/** Return an FMetasoundFrontendClass representing an output node of the data type. */
 			virtual const FMetasoundFrontendClass& GetFrontendOutputClass() const = 0;
 
+			/** Return an FMetasoundFrontendClass representing an input node of the data type. */
+			virtual const FMetasoundFrontendClass& GetFrontendConstructorOutputClass() const = 0;
+
 			/** Return an FMetasoundFrontendClass representing an init variable node of the data type. */
 			virtual const FMetasoundFrontendClass& GetFrontendVariableClass() const = 0;
 
@@ -158,8 +164,14 @@
 			/** Create an input node */
 			virtual TUniquePtr<INode> CreateInputNode(FInputNodeConstructorParams&&) const = 0;
 
+			/** Create an input node */
+			virtual TUniquePtr<INode> CreateConstructorInputNode(FInputNodeConstructorParams&&) const = 0;
+
 			/** Create an output node */
 			virtual TUniquePtr<INode> CreateOutputNode(FOutputNodeConstructorParams&&) const = 0;
+
+			/** Create an output node */
+			virtual TUniquePtr<INode> CreateConstructorOutputNode(FOutputNodeConstructorParams&&) const = 0;
 			
 			/** Create a variable node */
 			virtual TUniquePtr<INode> CreateLiteralNode(FLiteralNodeConstructorParams&&) const = 0;
@@ -245,12 +257,18 @@
 			/** Return an FMetasoundFrontendClass representing an input node of the data type. */
 			virtual bool GetFrontendInputClass(const FName& InDataType, FMetasoundFrontendClass& OutClass) const = 0;
 
+			/** Return an FMetasoundFrontendClass representing an input node of the data type. */
+			virtual bool GetFrontendConstructorInputClass(const FName& InDataType, FMetasoundFrontendClass& OutClass) const = 0;
+
 			/** Return an FMetasoundFrontendClass representing a variable node of the data type. */
 			virtual bool GetFrontendLiteralClass(const FName& InDataType, FMetasoundFrontendClass& OutClass) const = 0;
 
 			/** Return an FMetasoundFrontendClass representing an output node of the data type. */
 			virtual bool GetFrontendOutputClass(const FName& InDataType, FMetasoundFrontendClass& OutClass) const = 0;
 
+			/** Return an FMetasoundFrontendClass representing an output node of the data type. */
+			virtual bool GetFrontendConstructorOutputClass(const FName& InDataType, FMetasoundFrontendClass& OutClass) const = 0;
+
 			/** Return an FMetasoundFrontendClass representing an init variable node of the data type. */
 			virtual bool GetFrontendVariableClass(const FName& InDataType, FMetasoundFrontendClass& OutClass) const = 0;
 
@@ -265,8 +283,10 @@
 
 			// Create a new instance of a C++ implemented node from the registry.
 			virtual TUniquePtr<INode> CreateInputNode(const FName& InInputType, FInputNodeConstructorParams&& InParams) const = 0;
+			virtual TUniquePtr<INode> CreateConstructorInputNode(const FName& InInputType, FInputNodeConstructorParams&& InParams) const = 0;
 			virtual TUniquePtr<INode> CreateLiteralNode(const FName& InLiteralType, FLiteralNodeConstructorParams&& InParams) const = 0;
 			virtual TUniquePtr<INode> CreateOutputNode(const FName& InOutputType, FOutputNodeConstructorParams&& InParams) const = 0;
+			virtual TUniquePtr<INode> CreateConstructorOutputNode(const FName& InOutputType, FOutputNodeConstructorParams&& InParams) const = 0;
 			virtual TUniquePtr<INode> CreateReceiveNode(const FName& InDataType, const FNodeInitData&) const = 0;
 			virtual TUniquePtr<INode> CreateVariableNode(const FName& InDataType, FVariableNodeConstructorParams&& InParams) const = 0;
 			virtual TUniquePtr<INode> CreateVariableMutatorNode(const FName& InDataType, const FNodeInitData& InParams) const = 0;
diff --git a/Engine/Plugins/Runtime/Metasound/Source/MetasoundFrontend/Public/MetasoundFrontendDataTypeTraits.h b/Engine/Plugins/Runtime/Metasound/Source/MetasoundFrontend/Public/MetasoundFrontendDataTypeTraits.h
new file mode 100644
--- /dev/null
+++ b/Engine/Plugins/Runtime/Metasound/Source/MetasoundFrontend/Public/MetasoundFrontendDataTypeTraits.h
@@ -0,0 +1,78 @@
+// Copyright Epic Games, Inc. All Rights Reserved.
+
+#pragma once
+
+#include "MetasoundDataFactory.h"
+
+#include <type_traits>
+
+namespace Metasound
+{
+	/** Enables or disables automatic registration of array types given a 
+	 * MetaSound data type. By default this is true, and all data types will have
+	 * an associated TArray<DataType> registered if it is supported. */
+	template<typename ... Type>
+	struct TEnableAutoArrayTypeRegistration
+	{
+		static constexpr bool Value = true;
+	};
+
+	/** Enables or disables automatic registration of auto conversion nodes given a 
+	 * MetaSound data type. By default this is true, and all data types will have
+	 * associated conversion nodes registered based upon the data types supported
+	 * constructors and implicit conversions. */
+	template<typename ... Type>
+	struct TEnableAutoConverterNodeRegistration
+	{
+		static constexpr bool Value = true;
+	};
+
+	/** Enables or disables send and receive node registration for a given MetaSound
+	 * data type. By default this is true and all data types supported by the transmission
+	 * system will have associated send and receive nodes. */
+	template<typename ... Type>
+	struct TEnableTransmissionNodeRegistration
+	{
+		static constexpr bool Value = true;
+	};
+
+	/** Enables or disables using a data type in constructor vertices. 
+	 * By default this is true and all data types that support copy assignment
+	 * and copy construction. */
+	template<typename ... Type>
+	struct TEnableConstructorVertex
+	{
+		static constexpr bool Value = true;
+	};
+
+	// Helper utility to test if we a datatype is copyable
+	template <typename TDataType>
+	struct TIsCopyable
+	{
+	private:
+		static constexpr bool bIsCopyConstructible = std::is_copy_constructible<TDataType>::value;
+		static constexpr bool bIsCopyAssignable = std::is_copy_assignable<TDataType>::value;
+
+	public:
+		static constexpr bool Value = bIsCopyConstructible && bIsCopyAssignable;
+	};
+
+	// Helper to test if data type can be used on constructor pins.
+	template<typename TDataType>
+	struct TIsConstructorVertexSupported
+	{
+	private:
+		static constexpr bool bIsParsable = TLiteralTraits<TDataType>::bIsParsableFromAnyLiteralType;
+		static constexpr bool bIsCopyable = TIsCopyable<TDataType>::Value;
+		static constexpr bool bIsEnabled = TEnableConstructorVertex<TDataType>::Value;
+	public:
+		static constexpr bool Value = bIsParsable && bIsCopyable && bIsEnabled;
+	};
+
+	// Helper utility to test if we can transmit a datatype between a send and a receive node.
+	template <typename TDataType>
+	struct TIsTransmittable
+	{
+		static constexpr bool Value = TIsCopyable<TDataType>::Value;
+	};
+}
diff --git a/Engine/Plugins/Runtime/Metasound/Source/MetasoundFrontend/Public/MetasoundFrontendDocument.h b/Engine/Plugins/Runtime/Metasound/Source/MetasoundFrontend/Public/MetasoundFrontendDocument.h
--- a/Engine/Plugins/Runtime/Metasound/Source/MetasoundFrontend/Public/MetasoundFrontendDocument.h
+++ b/Engine/Plugins/Runtime/Metasound/Source/MetasoundFrontend/Public/MetasoundFrontendDocument.h
@@ -1415,3 +1415,5 @@ public:
 		return bDidEdit;
 	}
 };
+
+METASOUNDFRONTEND_API const TCHAR* LexToString(EMetasoundFrontendVertexAccessType InVertexAccess);
diff --git a/Engine/Plugins/Runtime/Metasound/Source/MetasoundFrontend/Public/MetasoundInputNode.h b/Engine/Plugins/Runtime/Metasound/Source/MetasoundFrontend/Public/MetasoundInputNode.h
--- a/Engine/Plugins/Runtime/Metasound/Source/MetasoundFrontend/Public/MetasoundInputNode.h
+++ b/Engine/Plugins/Runtime/Metasound/Source/MetasoundFrontend/Public/MetasoundInputNode.h
@@ -1,16 +1,21 @@
 // Copyright Epic Games, Inc. All Rights Reserved.
 #pragma once
 
-#include "CoreMinimal.h"
 #include "Internationalization/Text.h"
+#include "MetasoundBuilderInterface.h"
 #include "MetasoundDataReference.h"
+#include "MetasoundFrontendDataTypeTraits.h"
+#include "MetasoundLiteral.h"
+#include "MetasoundNode.h"
 #include "MetasoundNodeInterface.h"
+#include "MetasoundNodeRegistrationMacro.h"
 #include "MetasoundOperatorInterface.h"
 #include "MetasoundTrigger.h"
+#include "MetasoundVertexData.h"
+#include "UObject/NameTypes.h"
 
 #define LOCTEXT_NAMESPACE "MetasoundFrontend"
 
-
 namespace Metasound
 {
 	/** A writable input and a readable output. */
@@ -32,6 +37,20 @@ namespace Metasound
 			virtual ~TInputOperator() = default;
 
 			virtual FDataReferenceCollection GetInputs() const override
+			{
+				// This is slated to be deprecated and removed.
+				checkNoEntry();
+				return {};
+			}
+
+			virtual FDataReferenceCollection GetOutputs() const override
+			{
+				// This is slated to be deprecated and removed.
+				checkNoEntry();
+				return {};
+			}
+
+			virtual void Bind(FVertexInterfaceData& InOutVertexData) const override
 			{
 				// TODO: Expose a readable reference instead of a writable reference.
 				//
@@ -49,16 +68,9 @@ namespace Metasound
 				// complicated and hence triggers advancing should be managed by a 
 				// different object. Preferably the graph operator itself, or an
 				// explicit trigger manager tied to the environment.
-				FDataReferenceCollection Inputs;
-				Inputs.AddDataWriteReference<DataType>(DataReferenceName, InputValue);
-				return Inputs;
-			}
+				InOutVertexData.GetInputs().BindWriteVertex(DataReferenceName, InputValue);
 
-			virtual FDataReferenceCollection GetOutputs() const override
-			{
-				FDataReferenceCollection Outputs;
-				Outputs.AddDataReadReference<DataType>(DataReferenceName, OutputValue);
-				return Outputs;
+				InOutVertexData.GetOutputs().BindReadVertex(DataReferenceName, OutputValue);
 			}
 
 			void Execute()
@@ -93,7 +105,6 @@ namespace Metasound
 
 			FDataWriteReference InputValue;
 			FDataWriteReference OutputValue;
-
 	};
 
 	/** TPassThroughOperator supplies a readable input and a readable output. 
@@ -108,18 +119,17 @@ namespace Metasound
 		using Super = TInputOperator<DataType>;
 
 		TPassThroughOperator(const FVertexName& InDataReferenceName, FDataReadReference InDataReference)
-		: TInputOperator<DataType>(InDataReferenceName, WriteCast(InDataReference)) // Write cast is safe because `GetExecuteFunction() and GetInputs() are overridden, ensuring that data is not written.
+		: TInputOperator<DataType>(InDataReferenceName, WriteCast(InDataReference)) // Write cast is safe because `GetExecuteFunction() and Bind() are overridden, ensuring that data is not written.
 		, DataReferenceName(InDataReferenceName)
 		{
 		}
 
 		virtual ~TPassThroughOperator() = default;
 
-		virtual FDataReferenceCollection GetInputs() const override
+		virtual void Bind(FVertexInterfaceData& InOutVertexData) const override
 		{
-			FDataReferenceCollection Inputs;
-			ensure(Inputs.AddDataReadReferenceFrom(DataReferenceName, Super::GetInputs(), DataReferenceName, GetMetasoundDataTypeName<DataType>()));
-			return Inputs;
+			InOutVertexData.GetInputs().BindReadVertex(DataReferenceName, Super::InputValue);
+			InOutVertexData.GetOutputs().BindReadVertex(DataReferenceName, Super::OutputValue);
 		}
 
 		static void ExecuteFunction(IOperator* InOperator)
@@ -181,10 +191,35 @@ namespace Metasound
 		FVertexName DataReferenceName;
 	};
 
+	/** FInputValueOperator provides an input for value references. */
+	class METASOUNDFRONTEND_API FInputValueOperator : public IOperator
+	{
+	public:
+		/** Construct an FInputValueOperator with the name of the vertex and the 
+		 * value reference associated with input. 
+		 */
+		template<typename DataType>
+		explicit FInputValueOperator(const FName& InVertexName, const TDataValueReference<DataType>& InValueRef)
+		: VertexName(InVertexName)
+		, Default(InValueRef)
+		{
+		}
+
+		virtual ~FInputValueOperator() = default;
+
+		virtual FDataReferenceCollection GetInputs() const override;
+		virtual FDataReferenceCollection GetOutputs() const override;
+		virtual void Bind(FVertexInterfaceData& InVertexData) const override;
+		virtual FExecuteFunction GetExecuteFunction() override;
+
+	private:
+		FName VertexName;
+		FAnyDataReference Default;
+	};
 
 	/** Data type creation policy to create by copy construction. */
 	template<typename DataType>
-	struct FCreateDataReferenceWithCopy
+	struct UE_DEPRECATED(5.1, "Moved to private implementation.") FCreateDataReferenceWithCopy
 	{
 		template<typename... ArgTypes>
 		FCreateDataReferenceWithCopy(ArgTypes&&... Args)
@@ -201,9 +236,10 @@ namespace Metasound
 		DataType Data;
 	};
 
+
 	/** Data type creation policy to create by literal construction. */
 	template<typename DataType>
-	struct FCreateDataReferenceWithLiteral
+	struct UE_DEPRECATED(5.1, "Moved to private implementation.") FCreateDataReferenceWithLiteral
 	{
 		// If the data type is parsable from a literal type, then the data type 
 		// can be registered as an input type with the frontend.  To make a 
@@ -224,17 +260,15 @@ namespace Metasound
 		}
 
 	private:
-
 		FLiteral Literal;
 	};
 
-	
 
 	/** TInputOperatorFactory initializes the DataType at construction. It uses
 	 * the ReferenceCreatorType to create a data reference if one is not passed in.
 	 */
 	template<typename DataType, typename ReferenceCreatorType>
-	class TInputOperatorFactory : public IOperatorFactory
+	class UE_DEPRECATED(5.1, "Moved to private implementation") TInputOperatorFactory : public IOperatorFactory
 	{
 		public:
 
@@ -253,110 +287,240 @@ namespace Metasound
 	};
 
 	/** TInputNode represents an input to a metasound graph. */
-	template<typename DataType>
+	template<typename DataType, EVertexAccessType VertexAccess=EVertexAccessType::Reference>
 	class TInputNode : public FNode
 	{
+		static constexpr bool bIsConstructorInput = VertexAccess == EVertexAccessType::Value;
+		static constexpr bool bIsSupportedConstructorInput = TIsConstructorVertexSupported<DataType>::Value && bIsConstructorInput;
+		static constexpr bool bIsReferenceInput = VertexAccess == EVertexAccessType::Reference;
+		static constexpr bool bIsSupportedReferenceInput = TLiteralTraits<DataType>::bIsParsableFromAnyLiteralType && bIsReferenceInput;
 
-		public:
-			// If true, this node can be instantiated by the FrontEnd.
-			static constexpr bool bCanRegister = FCreateDataReferenceWithLiteral<DataType>::bCanCreateWithLiteral;
+		static constexpr bool bIsSupportedInput = bIsSupportedConstructorInput || bIsSupportedReferenceInput;
 
-			static FVertexInterface DeclareVertexInterface(const FVertexName& InVertexName)
+		// Use Variant names to differentiate between normal input nodes and constructor 
+		// input nodes.
+		static FName GetVariantName()
+		{
+			if constexpr (EVertexAccessType::Value == VertexAccess)
 			{
-				return FVertexInterface(
-					FInputVertexInterface(
-						TInputDataVertex<DataType>(InVertexName, FDataVertexMetadata{ FText::GetEmpty() })
-					),
-					FOutputVertexInterface(
-						TOutputDataVertex<DataType>(InVertexName, FDataVertexMetadata{ FText::GetEmpty() })
-					)
-				);
+				return FName("Constructor");
 			}
-
-			static FNodeClassMetadata GetNodeInfo(const FVertexName& InVertexName)
+			else
 			{
-				FNodeClassMetadata Info;
+				return FName();
+			}
+		}
 
-				Info.ClassName = { "Input", GetMetasoundDataTypeName<DataType>(), FName() };
-				Info.MajorVersion = 1;
-				Info.MinorVersion = 0;
-				Info.Description = METASOUND_LOCTEXT("Metasound_InputNodeDescription", "Input into the parent Metasound graph.");
-				Info.Author = PluginAuthor;
-				Info.PromptIfMissing = PluginNodeMissingPrompt;
-				Info.DefaultInterface = DeclareVertexInterface(InVertexName);
+		// Factory for creating input operators. 
+		class FInputNodeOperatorFactory : public IOperatorFactory
+		{
+			static constexpr bool bIsReferenceVertexAccess = VertexAccess == EVertexAccessType::Reference;
+			static constexpr bool bIsValueVertexAccess = VertexAccess == EVertexAccessType::Value;
 
-				return Info;
-			}
+			static_assert(bIsValueVertexAccess || bIsReferenceVertexAccess, "Unsupported EVertexAccessType");
+
+			// Choose which data reference type is created based on template parameters
+			using FDataReference = std::conditional_t<bIsReferenceVertexAccess, TDataWriteReference<DataType>, TDataValueReference<DataType>>;
+			using FPassThroughDataReference = std::conditional_t<bIsReferenceVertexAccess, TDataReadReference<DataType>, TDataValueReference<DataType>>;
 
-			template<typename... ArgTypes>
-			static FOperatorFactorySharedRef CreateOperatorFactoryWithArgs(ArgTypes&&... Args)
+			// Utility struct for creating data references for varying flavors of 
+			// runtime scenarios and vertex access types.
+			struct FDataReferenceCreatorBase
 			{
-				using FCreatorType = FCreateDataReferenceWithCopy<DataType>;
-				using FFactoryType = TInputOperatorFactory<DataType, FCreatorType>;
+				virtual ~FDataReferenceCreatorBase() = default;
 
-				return MakeOperatorFactoryRef<FFactoryType>(FCreatorType(Forward<ArgTypes>(Args)...));
-			}
+				// Create a a new data reference for constructing operators 
+				virtual FDataReference CreateDataReference(const FOperatorSettings& InOperatorSettings) const = 0;
 
-			static FOperatorFactorySharedRef CreateOperatorFactoryWithLiteral(FLiteral&& InLiteral)
+				// Create a data reference for constructing operators from a given data reference
+				virtual FPassThroughDataReference CreateDataReference(const FAnyDataReference& InRef) const
+				{
+					if constexpr (bIsReferenceVertexAccess)
+					{
+						return InRef.GetDataReadReference<DataType>();
+					}
+					else if constexpr (bIsValueVertexAccess)
+					{
+						return InRef.GetDataValueReference<DataType>();
+					}
+					else
+					{
+						static_assert("Unsupported EVertexAccessType");
+					}
+				}
+			};
+
+			// Create data references using a literal 
+			struct FCreateWithLiteral : FDataReferenceCreatorBase
 			{
-				using FCreatorType = FCreateDataReferenceWithLiteral<DataType>;
-				using FFactoryType = TInputOperatorFactory<DataType, FCreatorType>;
+				using FDataFactory = std::conditional_t<bIsReferenceVertexAccess, TDataWriteReferenceLiteralFactory<DataType>, TDataValueReferenceLiteralFactory<DataType>>;
 
-				return MakeOperatorFactoryRef<FFactoryType>(FCreatorType(MoveTemp(InLiteral)));
-			}
+				FLiteral Literal;
 
+				FCreateWithLiteral(FLiteral&& InLiteral)
+				: Literal(MoveTemp(InLiteral))
+				{
+				}
 
-			/* Construct a TInputNode using the TInputOperatorFactory<> and forwarding 
-			 * Args to the TInputOperatorFactory constructor.*/
-			template<typename... ArgTypes>
-			TInputNode(const FVertexName& InInstanceName, const FGuid& InInstanceID, const FVertexName& InVertexName, ArgTypes&&... Args)
-			:	FNode(InInstanceName, InInstanceID, GetNodeInfo(InVertexName))
-			,	VertexName(InVertexName)
-			,	Interface(DeclareVertexInterface(InVertexName))
-			,	Factory(CreateOperatorFactoryWithArgs(Forward<ArgTypes>(Args)...))
-			{
-			}
+				virtual FDataReference CreateDataReference(const FOperatorSettings& InOperatorSettings) const override
+				{
+					return FDataFactory::CreateExplicitArgs(InOperatorSettings, Literal);
+				}
+			};
 
-			/* Construct a TInputNode using the TInputOperatorLiteralFactory<> and moving
-			 * InParam to the TInputOperatorLiteralFactory constructor.*/
-			explicit TInputNode(const FVertexName& InNodeName, const FGuid& InInstanceID, const FVertexName& InVertexName, FLiteral&& InParam)
-			:	FNode(InNodeName, InInstanceID, GetNodeInfo(InVertexName))
-			,	VertexName(InVertexName)
-			,	Interface(DeclareVertexInterface(InVertexName))
-			, 	Factory(CreateOperatorFactoryWithLiteral(MoveTemp(InParam)))
+			// Create data references using a copy 
+			struct FCreateWithCopy : FDataReferenceCreatorBase
 			{
-			}
+				using FDataFactory = std::conditional_t<bIsReferenceVertexAccess, TDataWriteReferenceFactory<DataType>, TDataValueReferenceFactory<DataType>>;
 
-			const FVertexName& GetVertexName() const
-			{
-				return VertexName;
-			}
+				DataType Value;
 
-			virtual const FVertexInterface& GetVertexInterface() const override
-			{
-				return Interface;
-			}
+				virtual FDataReference CreateDataReference(const FOperatorSettings& InOperatorSettings) const override
+				{
+					return FDataFactory::CreateExplicitArgs(InOperatorSettings, Value);
+				}
+			};
 
-			virtual bool SetVertexInterface(const FVertexInterface& InInterface) override
+		public:
+
+			FInputNodeOperatorFactory(const DataType& InValue)
+			: ReferenceCreator(MakeUnique<FCreateWithCopy>(InValue))
 			{
-				return Interface == InInterface;
 			}
 
-			virtual bool IsVertexInterfaceSupported(const FVertexInterface& InInterface) const override
+			explicit FInputNodeOperatorFactory(FLiteral&& InLiteral)
+			: ReferenceCreator(MakeUnique<FCreateWithLiteral>(MoveTemp(InLiteral)))
 			{
-				return Interface == InInterface;
 			}
 
-			virtual TSharedRef<IOperatorFactory, ESPMode::ThreadSafe> GetDefaultOperatorFactory() const override
+			virtual TUniquePtr<IOperator> CreateOperator(const FBuildOperatorParams& InParams, FBuildResults& OutResults) override
 			{
-				return Factory;
+				using FInputNodeType = TInputNode<DataType, VertexAccess>;
+				using FOwnedInputOperatorType = std::conditional_t<bIsReferenceVertexAccess, TInputOperator<DataType>, FInputValueOperator>;
+				using FPassThroughInputOperatorType = std::conditional_t<bIsReferenceVertexAccess, TPassThroughOperator<DataType>, FInputValueOperator>;
+
+				const FInputNodeType& InputNode = static_cast<const FInputNodeType&>(InParams.Node);
+				const FVertexName& VertexKey = InputNode.GetVertexName();
+
+				if (const FAnyDataReference* Ref = InParams.InputData.FindDataReference(VertexKey))
+				{
+					// Pass through input value
+					return MakeUnique<FPassThroughInputOperatorType>(VertexKey, ReferenceCreator->CreateDataReference(*Ref));
+				}
+				else
+				{
+					// Owned input value
+					return MakeUnique<FOwnedInputOperatorType>(VertexKey, ReferenceCreator->CreateDataReference(InParams.OperatorSettings));
+				}
 			}
 
 		private:
-			FVertexName VertexName;
+			TUniquePtr<FDataReferenceCreatorBase> ReferenceCreator;
+		};
+
+
+	public:
+		// If true, this node can be instantiated by the Frontend.
+		static constexpr bool bCanRegister = bIsSupportedInput;
+
+		static FVertexInterface DeclareVertexInterface(const FVertexName& InVertexName)
+		{
+			return FVertexInterface(
+				FInputVertexInterface(
+					FInputDataVertex(InVertexName, GetMetasoundDataTypeName<DataType>(), FDataVertexMetadata{ FText::GetEmpty() }, VertexAccess)
+				),
+				FOutputVertexInterface(
+					FOutputDataVertex(InVertexName, GetMetasoundDataTypeName<DataType>(), FDataVertexMetadata{ FText::GetEmpty() }, VertexAccess)
+				)
+			);
+		}
+
+		static FNodeClassMetadata GetNodeInfo(const FVertexName& InVertexName)
+		{
+			FNodeClassMetadata Info;
+
+			Info.ClassName = { "Input", GetMetasoundDataTypeName<DataType>(), GetVariantName() };
+			Info.MajorVersion = 1;
+			Info.MinorVersion = 0;
+			Info.Description = METASOUND_LOCTEXT("Metasound_InputNodeDescription", "Input into the parent Metasound graph.");
+			Info.Author = PluginAuthor;
+			Info.PromptIfMissing = PluginNodeMissingPrompt;
+			Info.DefaultInterface = DeclareVertexInterface(InVertexName);
+
+			return Info;
+		}
+		
+		template<typename... ArgTypes>
+		UE_DEPRECATED(5.1, "Moved to internal implementation.")
+		static FOperatorFactorySharedRef CreateOperatorFactoryWithArgs(ArgTypes&&... Args)
+		{
+			using FCreatorType = FCreateDataReferenceWithCopy<DataType>;
+			using FFactoryType = TInputOperatorFactory<DataType, FCreatorType>;
+
+			return MakeOperatorFactoryRef<FFactoryType>(FCreatorType(Forward<ArgTypes>(Args)...));
+		}
+
+		UE_DEPRECATED(5.1, "Moved to internal implementation.")
+		static FOperatorFactorySharedRef CreateOperatorFactoryWithLiteral(FLiteral&& InLiteral)
+		{
+			using FCreatorType = FCreateDataReferenceWithLiteral<DataType>;
+			using FFactoryType = TInputOperatorFactory<DataType, FCreatorType>;
+
+			return MakeOperatorFactoryRef<FFactoryType>(FCreatorType(MoveTemp(InLiteral)));
+		}
+
+		/* Construct a TInputNode using the TInputOperatorFactory<> and forwarding 
+		 * Args to the TInputOperatorFactory constructor.*/
+		template<typename... ArgTypes>
+		UE_DEPRECATED(5.1, "Constructing an TInputNode with args will no longer be supported.")
+		TInputNode(const FVertexName& InInstanceName, const FGuid& InInstanceID, const FVertexName& InVertexName, ArgTypes&&... Args)
+		:	FNode(InInstanceName, InInstanceID, GetNodeInfo(InVertexName))
+		,	VertexName(InVertexName)
+		,	Interface(DeclareVertexInterface(InVertexName))
+		,	Factory(MakeShared<FInputNodeOperatorFactory>(Forward<ArgTypes>(Args)...))
+		{
+		}
+
+		/* Construct a TInputNode using the TInputOperatorLiteralFactory<> and moving
+		 * InParam to the TInputOperatorLiteralFactory constructor.*/
+		explicit TInputNode(const FVertexName& InNodeName, const FGuid& InInstanceID, const FVertexName& InVertexName, FLiteral&& InParam)
+		:	FNode(InNodeName, InInstanceID, GetNodeInfo(InVertexName))
+		,	VertexName(InVertexName)
+		,	Interface(DeclareVertexInterface(InVertexName))
+		, 	Factory(MakeShared<FInputNodeOperatorFactory>(MoveTemp(InParam)))
+		{
+		}
+
+		const FVertexName& GetVertexName() const
+		{
+			return VertexName;
+		}
+
+		virtual const FVertexInterface& GetVertexInterface() const override
+		{
+			return Interface;
+		}
+
+		virtual bool SetVertexInterface(const FVertexInterface& InInterface) override
+		{
+			return Interface == InInterface;
+		}
+
+		virtual bool IsVertexInterfaceSupported(const FVertexInterface& InInterface) const override
+		{
+			return Interface == InInterface;
+		}
+
+		virtual TSharedRef<IOperatorFactory, ESPMode::ThreadSafe> GetDefaultOperatorFactory() const override
+		{
+			return Factory;
+		}
+
+	private:
+		FVertexName VertexName;
 
-			FVertexInterface Interface;
-			FOperatorFactorySharedRef Factory;
+		FVertexInterface Interface;
+		FOperatorFactorySharedRef Factory;
 	};
 
 
@@ -371,13 +535,13 @@ namespace Metasound
 		if (InParams.InputData.IsVertexBound(VertexKey))
 		{
 			// Data is externally owned. Use pass through operator
-			FDataReadReference DataRef = InParams.InputData.GetDataReadReference<DataType>(VertexKey);
+			TDataReadReference<DataType> DataRef = InParams.InputData.GetDataReadReference<DataType>(VertexKey);
 			return MakeUnique<TPassThroughOperator<DataType>>(InputNode.GetVertexName(), DataRef);
 		}
 		else
 		{
 			// Create write reference by calling compatible constructor with literal.
-			FDataWriteReference DataRef = ReferenceCreator.CreateDataReference(InParams.OperatorSettings);
+			TDataWriteReference<DataType> DataRef = ReferenceCreator.CreateDataReference(InParams.OperatorSettings);
 			return MakeUnique<TInputOperator<DataType>>(InputNode.GetVertexName(), DataRef);
 		}
 	}
diff --git a/Engine/Plugins/Runtime/Metasound/Source/MetasoundFrontend/Public/MetasoundOutputNode.h b/Engine/Plugins/Runtime/Metasound/Source/MetasoundFrontend/Public/MetasoundOutputNode.h
--- a/Engine/Plugins/Runtime/Metasound/Source/MetasoundFrontend/Public/MetasoundOutputNode.h
+++ b/Engine/Plugins/Runtime/Metasound/Source/MetasoundFrontend/Public/MetasoundOutputNode.h
@@ -1,13 +1,17 @@
 // Copyright Epic Games, Inc. All Rights Reserved.
 #pragma once
 
-#include "CoreMinimal.h"
+#include "Internationalization/Text.h"
 #include "MetasoundBuildError.h"
 #include "MetasoundBuilderInterface.h"
+#include "MetasoundFrontendDataTypeTraits.h"
 #include "MetasoundNodeInterface.h"
+#include "MetasoundNodeRegistrationMacro.h"
 #include "MetasoundOperatorInterface.h"
 #include "MetasoundNode.h"
 #include "MetasoundVertex.h"
+#include "Templates/SharedPointer.h"
+#include "UObject/NameTypes.h"
 
 #define LOCTEXT_NAMESPACE "MetasoundGraphCore"
 
@@ -30,73 +34,83 @@ namespace Metasound
 		}
 
 		virtual ~FMissingOutputNodeInputReferenceError() = default;
-
 	};
 
-	template<typename DataType>
-	class TOutputNode : public FNode
+	namespace OutputNodePrivate
 	{
-		class FOutputOperator : public IOperator
+		class METASOUNDFRONTEND_API FOutputOperator : public IOperator
 		{
 			public:
-				using FDataReadReference = TDataReadReference<DataType>;
-
-				FOutputOperator(const FVertexName& InDataReferenceName, FDataReadReference InDataReference)
-				{
-					Outputs.AddDataReadReference<DataType>(InDataReferenceName, InDataReference);
-				}
+				FOutputOperator(const FVertexName& InVertexName, const FAnyDataReference& InDataReference);
 
-				virtual ~FOutputOperator() {}
+				virtual ~FOutputOperator() = default;
 
-				virtual FDataReferenceCollection GetInputs() const override
-				{
-					return {};
-				}
-
-				virtual FDataReferenceCollection GetOutputs() const override
-				{
-					return Outputs;
-				}
-
-				virtual FExecuteFunction GetExecuteFunction() override
-				{
-					return nullptr;
-				}
+				virtual FDataReferenceCollection GetInputs() const override;
+				virtual FDataReferenceCollection GetOutputs() const override;
+				virtual void Bind(FVertexInterfaceData& InVertexData) const;
+				virtual FExecuteFunction GetExecuteFunction() override;
 
 			private:
-				FDataReferenceCollection Outputs;
+				FVertexName VertexName;
+				FAnyDataReference DataReference;
 		};
+	}
+
+	template<typename DataType, EVertexAccessType VertexAccess=EVertexAccessType::Reference>
+	class TOutputNode : public FNode
+	{
+	private:
+		static FName GetVariantName()
+		{
+			if constexpr (EVertexAccessType::Value == VertexAccess)
+			{
+				return FName("Constructor");
+			}
+			else
+			{
+				return FName();
+			}
+		}
 
 		class FOutputOperatorFactory : public IOperatorFactory
 		{
-			public:
-				FOutputOperatorFactory(const FVertexName& InDataReferenceName)
-				:	DataReferenceName(InDataReferenceName)
+		public:
+			FOutputOperatorFactory(const FVertexName& InDataReferenceName)
+			:	DataReferenceName(InDataReferenceName)
+			{
+			}
+
+			virtual TUniquePtr<IOperator> CreateOperator(const FBuildOperatorParams& InParams, FBuildResults& OutResults) override
+			{
+				using namespace OutputNodePrivate;
+
+				if (const FAnyDataReference* Ref = InParams.InputData.FindDataReference(DataReferenceName))
 				{
+					return MakeUnique<FOutputOperator>(DataReferenceName, *Ref);
 				}
 
-				virtual TUniquePtr<IOperator> CreateOperator(const FBuildOperatorParams& InParams, FBuildResults& OutResults) override
+				// Only construct default if default construction is supported
+				if constexpr (TIsParsable<DataType>::Value)
 				{
-					const bool bContainsRef = InParams.InputData.IsVertexBound(DataReferenceName);
-					if (bContainsRef)
+					if constexpr (TIsConstructorVertexSupported<DataType>::Value)
 					{
-						TDataReadReference<DataType> InputReadRef = InParams.InputData.GetDataReadReference<DataType>(DataReferenceName);
-						return MakeUnique<FOutputOperator>(DataReferenceName, InputReadRef);
+						// Prefer to make a constant value
+						TDataValueReference<DataType> DefaultValueRef = TDataValueReferenceFactory<DataType>::CreateAny(InParams.OperatorSettings);
+						return MakeUnique<FOutputOperator>(DataReferenceName, DefaultValueRef);
 					}
-
-					// Only construct default if default construction is supported
-					if constexpr (TIsParsable<DataType>::Value)
+					else
 					{
 						TDataReadReference<DataType> DefaultReadRef = TDataReadReferenceFactory<DataType>::CreateAny(InParams.OperatorSettings);
 						return MakeUnique<FOutputOperator>(DataReferenceName, DefaultReadRef);
 					}
-
-					OutResults.Errors.Emplace(MakeUnique<FMissingOutputNodeInputReferenceError>(InParams.Node, GetMetasoundDataTypeDisplayText<DataType>()));
-					return TUniquePtr<IOperator>(nullptr);
 				}
 
-			private:
-				FVertexName DataReferenceName;
+				OutResults.Errors.Emplace(MakeUnique<FMissingOutputNodeInputReferenceError>(InParams.Node, GetMetasoundDataTypeDisplayText<DataType>()));
+				return TUniquePtr<IOperator>(nullptr);
+			}
+
+		private:
+			FVertexName DataReferenceName;
 		};
 
 		static FVertexInterface GetVertexInterface(const FVertexName& InVertexName)
@@ -105,10 +119,10 @@ namespace Metasound
 
 			return FVertexInterface(
 				FInputVertexInterface(
-					TInputDataVertex<DataType>(InVertexName, FDataVertexMetadata{VertexDescription})
+					FInputDataVertex(InVertexName, GetMetasoundDataTypeName<DataType>(), FDataVertexMetadata{VertexDescription}, VertexAccess)
 				),
 				FOutputVertexInterface(
-					TOutputDataVertex<DataType>(InVertexName, FDataVertexMetadata{VertexDescription})
+					FOutputDataVertex(InVertexName, GetMetasoundDataTypeName<DataType>(), FDataVertexMetadata{VertexDescription}, VertexAccess)
 				)
 			);
 		}
@@ -117,7 +131,7 @@ namespace Metasound
 		{
 			FNodeClassMetadata Info;
 
-			Info.ClassName = { "Output", GetMetasoundDataTypeName<DataType>(), FName() };
+			Info.ClassName = { "Output", GetMetasoundDataTypeName<DataType>(), GetVariantName() };
 			Info.MajorVersion = 1;
 			Info.MinorVersion = 0;
 			Info.Description = METASOUND_LOCTEXT("Metasound_OutputNodeDescription", "Output from the parent Metasound graph.");
@@ -128,8 +142,6 @@ namespace Metasound
 			return Info;
 		};
 
-
-
 		public:
 			TOutputNode(const FVertexName& InInstanceName, const FGuid& InInstanceID, const FVertexName& InVertexName)
 			:	FNode(InInstanceName, InInstanceID, GetNodeInfo(InVertexName))
diff --git a/Engine/Plugins/Runtime/Metasound/Source/MetasoundGraphCore/Public/MetasoundDataFactory.h b/Engine/Plugins/Runtime/Metasound/Source/MetasoundGraphCore/Public/MetasoundDataFactory.h
--- a/Engine/Plugins/Runtime/Metasound/Source/MetasoundGraphCore/Public/MetasoundDataFactory.h
+++ b/Engine/Plugins/Runtime/Metasound/Source/MetasoundGraphCore/Public/MetasoundDataFactory.h
@@ -934,6 +934,19 @@ namespace Metasound
 	{
 	};
 
+	/** TDataValueReferenceFactory creates TDataValueReferences for the given DataType.
+	 *
+	 * TDataValueReferenceFactory provides several factory methods for forwarding 
+	 * arguments from the factory method to the DataType constructor.  See TDataFactory
+	 * for more information on the provided factory methods.
+	 *
+	 * @tparam DataType - The Metasound DataType of the TDataValueReference.
+	 */
+	template<typename DataType>
+	struct TDataValueReferenceFactory : TDataFactory<DataType, TDataValueReference<DataType>>
+	{
+	};
+
 	/** TDataReadReferenceFactory creates TDataReadReferences for the given DataType.
 	 *
 	 * TDataReadReferenceFactory provides several factory methods for forwarding 
@@ -1079,6 +1092,19 @@ namespace Metasound
 	{
 	};
 
+	/** TDataValueReferenceLiteralFactory creates TDataValueReferences for the given DataType.
+	 *
+	 * TDataValueReferenceLiteralFactory provides several factory methods for forwarding 
+	 * arguments from the factory method to the DataType constructor.  See TDataFactory
+	 * for more information on the provided factory methods.
+	 *
+	 * @tparam DataType - The Metasound DataType of the TDataValueReference.
+	 */
+	template<typename DataType>
+	struct TDataValueReferenceLiteralFactory : TDataLiteralFactory<DataType, TDataValueReference<DataType>>
+	{
+	};
+
 	/** TDataReadReferenceLiteralFactory creates TDataReadReferences for the given DataType.
 	 *
 	 * TDataReadReferenceLiteralFactory provides several factory methods for forwarding 
diff --git a/Engine/Plugins/Runtime/Metasound/Source/MetasoundGraphCore/Public/MetasoundDataReference.h b/Engine/Plugins/Runtime/Metasound/Source/MetasoundGraphCore/Public/MetasoundDataReference.h
--- a/Engine/Plugins/Runtime/Metasound/Source/MetasoundGraphCore/Public/MetasoundDataReference.h
+++ b/Engine/Plugins/Runtime/Metasound/Source/MetasoundGraphCore/Public/MetasoundDataReference.h
@@ -219,6 +219,10 @@ namespace Metasound
 	template <typename DataType>
 	class TDataReadReference;
 
+	// Forward declare
+	template <typename DataType>
+	class TDataWriteReference;
+
 	/** TDataValueReference represents a constant value and provides read only access. 
 	 * A TDataValueReference can never change value. */
 	template<typename DataType>
@@ -227,12 +231,23 @@ namespace Metasound
 		// Construct operator with no arguments if the DataType has a default constructor.
 		template <typename... ArgTypes>
 		TDataValueReference(EDataRefShouldConstruct InToken, ArgTypes&&... Args)
-			: FDataReference(InToken, Forward<ArgTypes>(Args)...)
+			: TDataReference<DataType>(InToken, Forward<ArgTypes>(Args)...)
 		{
 		}
 
+		// Constructor taking data reference. Used for casting
+		TDataValueReference(const TDataReference<DataType>& InRef)
+		: TDataReference<DataType>(InRef)
+		{
+		}
+
+		template<typename T>
+		friend TDataValueReference<T> ValueCast(const TDataReadReference<T>& InRef);
+
+		template<typename T>
+		friend TDataValueReference<T> ValueCast(const TDataWriteReference<T>& InRef);
+
 	public:
-		typedef TDataReference<DataType> FDataReference;
 
 		/** This should be used to construct a new DataType object and return this TDataValueReference as a wrapper around it. */
 		template <typename... ArgTypes>
@@ -286,6 +301,32 @@ namespace Metasound
 		}
 	};
 
+	/** Cast a TDataReadReference to a TDataValueReference. 
+	 *
+	 * In general TDataReadReferences should not be converted into TDataValueReferences unless the caller
+	 * can be certain that no other TDataWriteReference exists for the underlying parameter. Having a 
+	 * TDataWriteReferences to an existing TDataValueReference can cause confusing behavior as values
+	 * references are not expected to change value.
+	 */
+	template<typename T>
+	TDataValueReference<T> ValueCast(const TDataReadReference<T>& InRef)
+	{
+		return TDataValueReference<T>(InRef);
+	}
+
+	/** Cast a TDataWriteReference to a TDataValueReference. 
+	 *
+	 * In general TDataWriteReferences should never beconverted into TDataValueReferences unless the caller
+	 * can be certain that no other TDataWriteReference exists for the underlying parameter. Having a 
+	 * TDataWriteReferences to an existing TDataValueReference can cause confusing behavior as values
+	 * references are not expected to change value.
+	 */
+	template<typename T>
+	TDataValueReference<T> ValueCast(const TDataWriteReference<T>& InRef)
+	{
+		return TDataValueReference<T>(InRef);
+	}
+
 	/** TDataWriteReference provides write access to a shared parameter reference. */
 	template <typename DataType>
 	class TDataWriteReference : public TDataReference<DataType>
@@ -293,85 +334,80 @@ namespace Metasound
 		// Construct operator with no arguments if the DataType has a default constructor.
 		template <typename... ArgTypes>
 		TDataWriteReference(EDataRefShouldConstruct InToken, ArgTypes&&... Args)
-			: FDataReference(InToken, Forward<ArgTypes>(Args)...)
+		: TDataReference<DataType>(InToken, Forward<ArgTypes>(Args)...)
 		{
 		}
 
-		public:
-			typedef TDataReference<DataType> FDataReference;
-
-			/** This should be used to construct a new DataType object and return this TDataWriteReference as a wrapper around it. */
-			template <typename... ArgTypes>
-			static TDataWriteReference<DataType> CreateNew(ArgTypes&&... Args)
-			{
-				static_assert(std::is_constructible<DataType, ArgTypes...>::value, "TDataWriteReference::CreateNew underlying type is not constructible with provided arguments.");
-				return TDataWriteReference<DataType>(EDataRefShouldConstruct::NewObject, Forward<ArgTypes>(Args)...);
-			}
+		/** Create a writable ref from a blank parameter ref. Should be done with care and understanding
+		 * of side-effects of converting a ref to a writable ref. 
+		 */
+		TDataWriteReference<DataType>(const TDataReference<DataType>& InDataReference)
+		: TDataReference<DataType>(InDataReference)
+		{
+		}
 
-			/** Enable copy constructor */
-			TDataWriteReference(const TDataWriteReference<DataType>& Other) = default;
+		// Friend because it calls protected constructor
+		template <typename T>
+		friend TDataWriteReference<T> WriteCast(const TDataReadReference<T>& InReadableRef);
 
-			/** Enable move constructor */
-			TDataWriteReference(TDataWriteReference<DataType>&& Other) = default;
+	public:
+		/** This should be used to construct a new DataType object and return this TDataWriteReference as a wrapper around it. */
+		template <typename... ArgTypes>
+		static TDataWriteReference<DataType> CreateNew(ArgTypes&&... Args)
+		{
+			static_assert(std::is_constructible<DataType, ArgTypes...>::value, "TDataWriteReference::CreateNew underlying type is not constructible with provided arguments.");
+			return TDataWriteReference<DataType>(EDataRefShouldConstruct::NewObject, Forward<ArgTypes>(Args)...);
+		}
 
-			/** Enable assignment operator. */
-			TDataWriteReference<DataType>& operator=(const TDataWriteReference<DataType>& Other) = default;
+		/** Enable copy constructor */
+		TDataWriteReference(const TDataWriteReference<DataType>& Other) = default;
 
-			/** Enable move operator. */
-			TDataWriteReference<DataType>& operator=(TDataWriteReference<DataType>&& Other) = default;
+		/** Enable move constructor */
+		TDataWriteReference(TDataWriteReference<DataType>&& Other) = default;
 
-			/** Implicit conversion to a readable parameter. */
-			operator TDataReadReference<DataType>() const
-			{
-				return TDataReadReference<DataType>(*this);
-			}
+		/** Enable assignment operator. */
+		TDataWriteReference<DataType>& operator=(const TDataWriteReference<DataType>& Other) = default;
 
-			/** Non-const access to the underlying parameter object. */
-			FORCEINLINE DataType& operator*() const
-			{
-				return *TDataReference<DataType>::ObjectReference;
-			}
+		/** Enable move operator. */
+		TDataWriteReference<DataType>& operator=(TDataWriteReference<DataType>&& Other) = default;
 
-			/** Non-const access to the underlying parameter object. */
-			FORCEINLINE DataType* operator->() const
-			{
-				return TDataReference<DataType>::ObjectReference.operator->();
-			}
+		/** Implicit conversion to a readable parameter. */
+		operator TDataReadReference<DataType>() const
+		{
+			return TDataReadReference<DataType>(*this);
+		}
 
-			/** Non-const access to the underlying parameter object. */
-			FORCEINLINE DataType* Get() const
-			{
-				return TDataReference<DataType>::ObjectReference.operator->();
-			}
+		/** Non-const access to the underlying parameter object. */
+		FORCEINLINE DataType& operator*() const
+		{
+			return *TDataReference<DataType>::ObjectReference;
+		}
 
-			/** Create a clone of this parameter reference. */
-			virtual TUniquePtr<IDataReference> Clone() const override
-			{
-				typedef TDataWriteReference<DataType> FDataWriteReference;
+		/** Non-const access to the underlying parameter object. */
+		FORCEINLINE DataType* operator->() const
+		{
+			return TDataReference<DataType>::ObjectReference.operator->();
+		}
 
-				return MakeUnique< FDataWriteReference >(*this);
-			}
+		/** Non-const access to the underlying parameter object. */
+		FORCEINLINE DataType* Get() const
+		{
+			return TDataReference<DataType>::ObjectReference.operator->();
+		}
 
-			// Provide access to ObjectReference when converting from Write to Read.
-			friend class TDataReadReference<DataType>;
+		/** Create a clone of this parameter reference. */
+		virtual TUniquePtr<IDataReference> Clone() const override
+		{
+			typedef TDataWriteReference<DataType> FDataWriteReference;
 
-			// Friend because it calls protected constructor
-			template <typename T>
-			friend TDataWriteReference<T> WriteCast(const TDataReadReference<T>& InReadableRef);
+			return MakeUnique< FDataWriteReference >(*this);
+		}
 
-		protected:
+		// Provide access to ObjectReference when converting from Write to Read.
+		friend class TDataReadReference<DataType>;
 
-			/** Create a writable ref from a blank parameter ref. Should be done with care and understanding
-			 * of side-effects of converting a readable ref to a writable ref. 
-			 */
-			TDataWriteReference<DataType>(const FDataReference& InDataReference)
-			:	FDataReference(InDataReference)
-			{
-			}
 	};
 
-
-
 	/** Cast a TDataReadReference to a TDataWriteReference. 
 	 *
 	 * In general TDataReadReferences should not be converted into TDataWriteReferences unless the caller
@@ -395,12 +431,11 @@ namespace Metasound
 		// Construct operator with no arguments if the DataType has a default constructor.
 		template <typename... ArgTypes>
 		TDataReadReference(EDataRefShouldConstruct InToken, ArgTypes&&... Args)
-			: FDataReference(InToken, Forward<ArgTypes>(Args)...)
+		: TDataReference<DataType>(InToken, Forward<ArgTypes>(Args)...)
 		{
 		}
 
 	public:
-		typedef TDataReference<DataType> FDataReference;
 
 		// This should be used to construct a new DataType object and return this TDataReadReference as a wrapper around it.
 		template <typename... ArgTypes>
@@ -414,13 +449,13 @@ namespace Metasound
 
 		/** Construct a readable parameter ref from a writable parameter ref. */
 		explicit TDataReadReference(const TDataWriteReference<DataType>& WritableRef)
-		:	FDataReference(WritableRef)
+		: TDataReference<DataType>(WritableRef)
 		{
 		}
 
 		/** Construct a readable reference from a value reference. */
 		explicit TDataReadReference(const TDataValueReference<DataType>& ValueRef)
-		:	FDataReference(ValueRef)
+		: TDataReference<DataType>(ValueRef)
 		{
 		}
 
@@ -598,6 +633,46 @@ namespace Metasound
 			return nullptr;
 		}
 
+		/** Get access to a TDataValueReference. 
+		 *
+		 * This method will return a valid TDataValueReference of the templated data
+		 * type. The returned object is only valid if:
+		 *     1. The template parameter DataType matches that of the underlying data reference.
+		 *     2. The underlying data reference is has Value access.
+		 *
+		 * If this method's behavior is undefined and will assert if it is called 
+		 * with a mismatched data type or unsupported access type.
+		 */
+		template<typename DataType>
+		TDataValueReference<DataType> GetDataValueReference() const
+		{
+			check(DataRefPtr.IsValid());
+			check(IsDataReferenceOfType<DataType>(*DataRefPtr));
+
+			switch (AccessType)
+			{
+				case EDataReferenceAccessType::Value:
+				{
+					return *static_cast<const TDataValueReference<DataType>*>(DataRefPtr.Get());
+				}
+				case EDataReferenceAccessType::Write:
+				{
+					checkf(false, TEXT("Retrieving a value reference from a write reference is not supported."));
+					return ValueCast<DataType>(*static_cast<const TDataWriteReference<DataType>*>(DataRefPtr.Get()));
+				}
+				case EDataReferenceAccessType::Read:
+				{
+					checkf(false, TEXT("Retrieving a value reference from a read reference is not supported."));
+					return ValueCast<DataType>(*static_cast<const TDataReadReference<DataType>*>(DataRefPtr.Get()));
+				}
+				default:
+				{
+					checkf(false, TEXT("Unknown data reference access type."));
+					return *static_cast<const TDataValueReference<DataType>*>(DataRefPtr.Get());
+				}
+			}
+		}
+
 		/** Get access to a TDataReadReference. 
 		 *
 		 * This method will return a valid TDataReadReference of the templated data
diff --git a/Engine/Plugins/Runtime/Metasound/Source/MetasoundGraphCore/Public/MetasoundLiteral.h b/Engine/Plugins/Runtime/Metasound/Source/MetasoundGraphCore/Public/MetasoundLiteral.h
--- a/Engine/Plugins/Runtime/Metasound/Source/MetasoundGraphCore/Public/MetasoundLiteral.h
+++ b/Engine/Plugins/Runtime/Metasound/Source/MetasoundGraphCore/Public/MetasoundLiteral.h
@@ -67,6 +67,10 @@ namespace Metasound
 	 */
 	struct METASOUNDGRAPHCORE_API FLiteral
 	{
+	private:
+		// Forward declare
+		template<typename U>
+		struct TIsSupportedLiteralType;
 	public:
 		struct FInvalid {};
 
@@ -82,25 +86,22 @@ namespace Metasound
 			FInvalid
 		>;
 
+
 		/** Construct a literal param with a single argument. */
 		template<
-			typename... ArgTypes,
-			typename std::enable_if<1 == sizeof...(ArgTypes), int>::type = 0
+			typename ArgType,
+			typename std::enable_if<TIsSupportedLiteralType<ArgType>::Value, int>::type = 0
 			>
-		FLiteral(ArgTypes&&... Args)
+		FLiteral(ArgType&& Arg)
 		{
-			Value.Set<typename std::decay<ArgTypes...>::type>(Forward<ArgTypes>(Args)...);
+			Value.Set<typename std::decay<ArgType>::type>(Forward<ArgType>(Arg));
 #if METASOUND_DEBUG_LITERALS
 			InitDebugString();
 #endif
 		}
 
 		/** Construct a literal param with no arguments. */
-		template<
-			typename... ArgTypes,
-			typename std::enable_if<0 == sizeof...(ArgTypes), int>::type = 0
-			>
-		FLiteral(ArgTypes&&... Args)
+		FLiteral()
 		{
 			Value.Set<FNone>(FNone());
 #if METASOUND_DEBUG_LITERALS
@@ -130,14 +131,35 @@ namespace Metasound
 #endif
 		}
 
-#if METASOUND_DEBUG_LITERALS
 	private:
+#if METASOUND_DEBUG_LITERALS
 
 		FString DebugString;
 
 		void InitDebugString() const;
 
 #endif // METAOUND_DEBUG_LITERALS
+		
+		// Helper function to determine if type "U" can be stored in this FLiteral.
+		template<typename U, typename ... Ts>
+		struct TIsSupportedLiteralTypeHelper
+		{
+			static constexpr bool Value = UE::Core::Private::TParameterPackTypeIndex<U, Ts...>::Value != (SIZE_T)-1;
+		};
+
+		// Helper function to determine if type "U" can be stored in this FLiteral.
+		template<typename U, typename T, typename ... Ts>
+		struct TIsSupportedLiteralTypeHelper<U, TVariant<T, Ts...>>
+		{
+			static constexpr bool Value = TIsSupportedLiteralTypeHelper<U, T, Ts...>::Value;
+		};
+
+		// Determine if type "U" can be stored in this FLiteral.
+		template<typename U>
+		struct TIsSupportedLiteralType
+		{
+			static constexpr bool Value = TIsSupportedLiteralTypeHelper<typename std::decay<U>::type, FVariantType>::Value;
+		};
 	};
 
 	namespace MetasoundLiteralIntrinsics
