Update all MetaSound Frontend class input properties when adding a input node

diff --git a/Engine/Plugins/Runtime/Metasound/Source/MetasoundEditor/Private/MetasoundEditorGraphBuilder.cpp b/Engine/Plugins/Runtime/Metasound/Source/MetasoundEditor/Private/MetasoundEditorGraphBuilder.cpp
--- a/Engine/Plugins/Runtime/Metasound/Source/MetasoundEditor/Private/MetasoundEditorGraphBuilder.cpp
+++ b/Engine/Plugins/Runtime/Metasound/Source/MetasoundEditor/Private/MetasoundEditorGraphBuilder.cpp
@@ -1048,14 +1048,22 @@ namespace Metasound
 		{
 			FMetasoundAssetBase* MetaSoundAsset = IMetasoundUObjectRegistry::Get().GetObjectAsAssetBase(&InMetaSound);
 			check(MetaSoundAsset);
+			const FGuid VertexID = FGuid::NewGuid();
 
 			FMetasoundFrontendClassInput ClassInput;
 			ClassInput.Name = GenerateUniqueNameByClassType(InMetaSound, EMetasoundFrontendClassType::Input, InNameBase ? InNameBase->ToString() : TEXT("Input"));
 			ClassInput.TypeName = InTypeName;
+			ClassInput.VertexID = VertexID;
+
 			if (nullptr != InDefaultValue)
 			{
 				ClassInput.DefaultLiteral = *InDefaultValue;
 			}
+			else
+			{
+				Metasound::FLiteral Literal = Frontend::IDataTypeRegistry::Get().CreateDefaultLiteral(InTypeName);
+				ClassInput.DefaultLiteral.SetFromLiteral(Literal);
+			}
 
 			return MetaSoundAsset->GetRootGraphHandle()->AddInputVertex(ClassInput);
 		}
