Fix string format for warning

diff --git a/Engine/Plugins/Runtime/Metasound/Source/MetasoundFrontend/Public/MetasoundEnum.h b/Engine/Plugins/Runtime/Metasound/Source/MetasoundFrontend/Public/MetasoundEnum.h
--- a/Engine/Plugins/Runtime/Metasound/Source/MetasoundFrontend/Public/MetasoundEnum.h
+++ b/Engine/Plugins/Runtime/Metasound/Source/MetasoundFrontend/Public/MetasoundEnum.h
@@ -131,7 +131,7 @@ namespace Metasound
 				Algo::Transform(GetAllNames(), ValueStrings, [](const FName& Name) { return Name.ToString(); });
 
 				UE_LOG(LogMetaSound, Warning,
-					TEXT("Cannot create valid enum from value '%s'.\nPossible Values:\n"),
+					TEXT("Cannot create valid enum from value '%s'.\nPossible Values:\n%s"),
 					*FString::FromInt((int32)(InValue)),
 					*FString::Join(ValueStrings, TEXT("\n, "))
 				);
@@ -163,7 +163,7 @@ namespace Metasound
 					Algo::Transform(GetAllNames(), ValueStrings, [](const FName& Name) { return Name.ToString(); });
 
 					UE_LOG(LogMetaSound, Warning,
-						TEXT("Cannot create valid enum value from string '%s'.\nPossible Values:\n"),
+						TEXT("Cannot create valid enum value from string '%s'.\nPossible Values:\n%s"),
 						*InValueName.ToString(),
 						*FString::Join(ValueStrings, TEXT("\n, "))
 					);
