Replace short asset name from DisabledClasses set with Top Level Path Name

diff --git a/Engine/Plugins/Runtime/Metasound/Source/MetasoundEditor/Private/MetasoundEditorGraphNode.cpp b/Engine/Plugins/Runtime/Metasound/Source/MetasoundEditor/Private/MetasoundEditorGraphNode.cpp
--- a/Engine/Plugins/Runtime/Metasound/Source/MetasoundEditor/Private/MetasoundEditorGraphNode.cpp
+++ b/Engine/Plugins/Runtime/Metasound/Source/MetasoundEditor/Private/MetasoundEditorGraphNode.cpp
@@ -435,7 +435,7 @@ TSet<FString> UMetasoundEditorGraphNode::GetDisallowedPinClassNames(const UEdGra
 
 		if (EditorModule.IsExplicitProxyClass(*ProxyGenClass) && Class->IsChildOf(ProxyGenClass))
 		{
-			DisallowedClasses.Add(ClassIt->GetName());
+			DisallowedClasses.Add(ClassIt->GetClassPathName().ToString());
 		}
 	}
 
