Fixing crash when selecting Wave Asset combo box/dropdown menu of Wave Player node (ANY_PACKAGE removal fallout)

diff --git a/Engine/Plugins/Runtime/Metasound/Source/MetasoundEditor/Private/MetasoundEditorGraphNode.cpp b/Engine/Plugins/Runtime/Metasound/Source/MetasoundEditor/Private/MetasoundEditorGraphNode.cpp
--- a/Engine/Plugins/Runtime/Metasound/Source/MetasoundEditor/Private/MetasoundEditorGraphNode.cpp
+++ b/Engine/Plugins/Runtime/Metasound/Source/MetasoundEditor/Private/MetasoundEditorGraphNode.cpp
@@ -414,7 +414,7 @@ TSet<FString> UMetasoundEditorGraphNode::GetDisallowedPinClassNames(const UEdGra
 	}
 
 	TSet<FString> DisallowedClasses;
-	const FName ClassName = ProxyGenClass->GetFName();
+	const FTopLevelAssetPath ClassName = ProxyGenClass->GetClassPathName();
 	for (TObjectIterator<UClass> ClassIt; ClassIt; ++ClassIt)
 	{
 		UClass* Class = *ClassIt;
@@ -428,7 +428,7 @@ TSet<FString> UMetasoundEditorGraphNode::GetDisallowedPinClassNames(const UEdGra
 			continue;
 		}
 
-		if (ClassIt->GetFName() == ClassName)
+		if (ClassIt->GetClassPathName() == ClassName)
 		{
 			continue;
 		}
diff --git a/Engine/Plugins/Runtime/Metasound/Source/MetasoundEditor/Private/MetasoundEditorGraphSchema.cpp b/Engine/Plugins/Runtime/Metasound/Source/MetasoundEditor/Private/MetasoundEditorGraphSchema.cpp
--- a/Engine/Plugins/Runtime/Metasound/Source/MetasoundEditor/Private/MetasoundEditorGraphSchema.cpp
+++ b/Engine/Plugins/Runtime/Metasound/Source/MetasoundEditor/Private/MetasoundEditorGraphSchema.cpp
@@ -1284,7 +1284,7 @@ void UMetasoundEditorGraphSchema::TrySetDefaultValue(UEdGraphPin& Pin, const FSt
 			{
 				if (UClass* Class = Object->GetClass())
 				{
-					if (DisallowedClassNames.Contains(Class->GetName()))
+					if (DisallowedClassNames.Contains(Class->GetClassPathName().ToString()))
 					{
 						return;
 					}
