Fix MetaSound connection splines not visible for new connections while playing when using animated connections

diff --git a/Engine/Plugins/Runtime/Metasound/Source/MetasoundEditor/Private/MetasoundEditorGraphConnectionDrawingPolicy.cpp b/Engine/Plugins/Runtime/Metasound/Source/MetasoundEditor/Private/MetasoundEditorGraphConnectionDrawingPolicy.cpp
--- a/Engine/Plugins/Runtime/Metasound/Source/MetasoundEditor/Private/MetasoundEditorGraphConnectionDrawingPolicy.cpp
+++ b/Engine/Plugins/Runtime/Metasound/Source/MetasoundEditor/Private/MetasoundEditorGraphConnectionDrawingPolicy.cpp
@@ -443,34 +443,7 @@ namespace Metasound
 				ConnectionData.UpdateSplineOverlap(SplineOverlapResult);
 			}
 
-			if (ConnectionData.OutputHandle->IsValid() && MetasoundEditor.IsValid())
-			{
-				const FName DataType = ConnectionData.OutputHandle->GetDataType();
-				if (DataType == GetMetasoundDataTypeName<float>())
-				{
-					DrawingPolicyPrivate::DrawConnectionSpline_Numeric<float>(*this, LayerId, ConnectionData, 0.0f);
-				}
-				else if (DataType == GetMetasoundDataTypeName<int32>())
-				{
-					DrawingPolicyPrivate::DrawConnectionSpline_Numeric<int32>(*this, LayerId, ConnectionData, 0);
-				}
-				else if (DataType == GetMetasoundDataTypeName<bool>())
-				{
-					DrawingPolicyPrivate::DrawConnectionSpline_Numeric<bool>(*this, LayerId, ConnectionData, 0);
-				}
-				// AudioBuffers & Triggers do not require drawing connection splines because they draw envelopes,
-				// which trace the same path
-				else if (DataType != GetMetasoundDataTypeName<FAudioBuffer>()
-					&& DataType != GetMetasoundDataTypeName<FTrigger>())
-				{
-					DrawConnectionSpline(LayerId, ConnectionData);
-				}
-			}
-			else
-			{
-				DrawConnectionSpline(LayerId, ConnectionData);
-			}
-
+			// AudioBuffers & Triggers only draw envelopes, not connection splines because they trace the same path
 			if (Params.bDrawBubbles || MidpointImage)
 			{
 				// Maps distance along curve to alpha
@@ -507,6 +480,33 @@ namespace Metasound
 				{
 					DrawConnection_MidpointImage(LayerId, SplineReparamTable, SplineLength, ConnectionData);
 				}
+				return;
+			}
+
+			// Draw other types
+			if (ConnectionData.OutputHandle->IsValid() && MetasoundEditor.IsValid())
+			{
+				const FName DataType = ConnectionData.OutputHandle->GetDataType();
+				if (DataType == GetMetasoundDataTypeName<float>())
+				{
+					DrawingPolicyPrivate::DrawConnectionSpline_Numeric<float>(*this, LayerId, ConnectionData, 0.0f);
+				}
+				else if (DataType == GetMetasoundDataTypeName<int32>())
+				{
+					DrawingPolicyPrivate::DrawConnectionSpline_Numeric<int32>(*this, LayerId, ConnectionData, 0);
+				}
+				else if (DataType == GetMetasoundDataTypeName<bool>())
+				{
+					DrawingPolicyPrivate::DrawConnectionSpline_Numeric<bool>(*this, LayerId, ConnectionData, 0);
+				}
+				else 
+				{
+					DrawConnectionSpline(LayerId, ConnectionData);
+				}
+			}
+			else
+			{
+				DrawConnectionSpline(LayerId, ConnectionData);
 			}
 		}
 
@@ -695,7 +695,7 @@ namespace Metasound
 				FDrawConnectionData FinalData = InParams.ConnectionData;
 				FinalData.P0 = LastPointPos;
 				FinalData.P0Tangent = LastTangent;
-				DrawSpline(InParams.ConnectionData);
+				DrawSpline(FinalData);
 			}
 		}
 
