Update MetaSound presets when an interface is deleted from the original MetaSound

diff --git a/Engine/Plugins/Runtime/Metasound/Source/MetasoundEditor/Private/MetasoundDetailCustomization.cpp b/Engine/Plugins/Runtime/Metasound/Source/MetasoundEditor/Private/MetasoundDetailCustomization.cpp
--- a/Engine/Plugins/Runtime/Metasound/Source/MetasoundEditor/Private/MetasoundDetailCustomization.cpp
+++ b/Engine/Plugins/Runtime/Metasound/Source/MetasoundEditor/Private/MetasoundDetailCustomization.cpp
@@ -477,6 +477,8 @@ namespace Metasound
 					InterfaceComboBox->RefreshOptions();
 					MetaSoundAsset->SetUpdateDetailsOnSynchronization();
 					MetaSoundAsset->SetSynchronizationRequired();
+					FGraphBuilder::RegisterGraphWithFrontend(*MetaSound.Get());
+
 				}), LOCTEXT("RemoveInterfaceTooltip2", "Removes the associated interface from the MetaSound."));
 
 				TSharedRef<SWidget> EntryWidget = SNew(SHorizontalBox)
