Fix null MetaSound output literal on graph synchronization

diff --git a/Engine/Plugins/Runtime/Metasound/Source/MetasoundEditor/Private/MetasoundEditorGraph.cpp b/Engine/Plugins/Runtime/Metasound/Source/MetasoundEditor/Private/MetasoundEditorGraph.cpp
--- a/Engine/Plugins/Runtime/Metasound/Source/MetasoundEditor/Private/MetasoundEditorGraph.cpp
+++ b/Engine/Plugins/Runtime/Metasound/Source/MetasoundEditor/Private/MetasoundEditorGraph.cpp
@@ -89,7 +89,7 @@ const UMetasoundEditorGraph* UMetasoundEditorGraphMember::GetOwningGraph() const
 	return Cast<const UMetasoundEditorGraph>(GetOuter());
 }
 
-void UMetasoundEditorGraphMember::ConformLiteralDataType()
+void UMetasoundEditorGraphMember::InitializeLiteral()
 {
 	using namespace Metasound::Editor;
 	using namespace Metasound::Frontend;
@@ -155,7 +155,7 @@ void UMetasoundEditorGraphVertex::InitMember(FName InDataType, const FMetasoundF
 	NodeID = InNodeID;
 	ClassName = MoveTemp(InClassName);
 
-	ConformLiteralDataType();
+	InitializeLiteral();
 
 	if (ensure(Literal))
 	{
@@ -347,7 +347,7 @@ void UMetasoundEditorGraphVertex::SetDataType(FName InNewType, bool bPostTransac
 	NodeID = NewNodeHandle->GetID();
 
 	TypeName = InNewType;
-	ConformLiteralDataType();
+	InitializeLiteral();
 
 	// 6. Create new node references in the same locations as the old locations
 	for (FVector2D Location : NodeLocations)
@@ -757,7 +757,7 @@ void UMetasoundEditorGraphVariable::InitMember(FName InDataType, const FMetasoun
 	TypeName = InDataType;
 	VariableID = InVariableID;
 
-	ConformLiteralDataType();
+	InitializeLiteral();
 
 	if (ensure(Literal))
 	{
@@ -1002,7 +1002,7 @@ void UMetasoundEditorGraphVariable::SetDataType(FName InNewType, bool bPostTrans
 		SetDescription(CachedData.Description, bPostSubTransaction);
 
 		TypeName = InNewType;
-		ConformLiteralDataType();
+		InitializeLiteral();
 
 		// Add the nodes with the same identifier data but new datatype.
 		UObject& Metasound = Graph->GetMetasoundChecked();
diff --git a/Engine/Plugins/Runtime/Metasound/Source/MetasoundEditor/Private/MetasoundEditorGraphBuilder.cpp b/Engine/Plugins/Runtime/Metasound/Source/MetasoundEditor/Private/MetasoundEditorGraphBuilder.cpp
--- a/Engine/Plugins/Runtime/Metasound/Source/MetasoundEditor/Private/MetasoundEditorGraphBuilder.cpp
+++ b/Engine/Plugins/Runtime/Metasound/Source/MetasoundEditor/Private/MetasoundEditorGraphBuilder.cpp
@@ -2463,6 +2463,12 @@ namespace Metasound
 				if (UMetasoundEditorGraphOutput* Output = Graph->FindOutput(NodeHandle->GetID()))
 				{
 					SynchronizeMemberDataType(*Output);
+
+					// Fix up corrupted assets with no literal set 
+					if (!Output->GetLiteral())
+					{
+						Output->InitializeLiteral();
+					}
 				}
 			}, EMetasoundFrontendClassType::Output);
 
diff --git a/Engine/Plugins/Runtime/Metasound/Source/MetasoundEditor/Public/MetasoundEditorGraph.h b/Engine/Plugins/Runtime/Metasound/Source/MetasoundEditor/Public/MetasoundEditorGraph.h
--- a/Engine/Plugins/Runtime/Metasound/Source/MetasoundEditor/Public/MetasoundEditorGraph.h
+++ b/Engine/Plugins/Runtime/Metasound/Source/MetasoundEditor/Public/MetasoundEditorGraph.h
@@ -169,14 +169,15 @@ public:
 
 	/** Returns literal associated with the given member */
 	UMetasoundEditorGraphMemberDefaultLiteral* GetLiteral() const { return Literal; }
+
+	/** Creates new literal if there is none and/or conforms literal object type to member's DataType */
+	void InitializeLiteral();
+
 protected:
 	/** Default literal value of member */
 	UPROPERTY()
 	TObjectPtr<UMetasoundEditorGraphMemberDefaultLiteral> Literal;
 
-	/** Conforms literal object type to member's DataType */
-	void ConformLiteralDataType();
-
 	/** Metasound Data Type. */
 	UPROPERTY()
 	FName TypeName;
