Fix non-unity issues in light of CL# 20648535
#rnx

diff --git a/Engine/Plugins/Runtime/AudioWidgets/Source/AudioWidgets/Private/SAudioSlider.cpp b/Engine/Plugins/Runtime/AudioWidgets/Source/AudioWidgets/Private/SAudioSlider.cpp
--- a/Engine/Plugins/Runtime/AudioWidgets/Source/AudioWidgets/Private/SAudioSlider.cpp
+++ b/Engine/Plugins/Runtime/AudioWidgets/Source/AudioWidgets/Private/SAudioSlider.cpp
@@ -1,6 +1,7 @@
 // Copyright Epic Games, Inc. All Rights Reserved.
 
 #include "SAudioSlider.h"
+#include "AudioDefines.h"
 #include "Brushes/SlateImageBrush.h"
 #include "DSP/Dsp.h"
 #include "Fonts/FontMeasure.h"
diff --git a/Engine/Plugins/Runtime/Metasound/Source/MetasoundStandardNodes/Private/MetasoundCompressorNode.cpp b/Engine/Plugins/Runtime/Metasound/Source/MetasoundStandardNodes/Private/MetasoundCompressorNode.cpp
--- a/Engine/Plugins/Runtime/Metasound/Source/MetasoundStandardNodes/Private/MetasoundCompressorNode.cpp
+++ b/Engine/Plugins/Runtime/Metasound/Source/MetasoundStandardNodes/Private/MetasoundCompressorNode.cpp
@@ -2,6 +2,7 @@
 
 #include "DSP/Dsp.h"
 #include "DSP/DynamicsProcessor.h"
+#include "DSP/FloatArrayMath.h"
 #include "Internationalization/Text.h"
 #include "MetasoundAudioBuffer.h"
 #include "MetasoundEnumRegistrationMacro.h"
diff --git a/Engine/Plugins/Runtime/Metasound/Source/MetasoundStandardNodes/Private/MetasoundCrossfadeNode.cpp b/Engine/Plugins/Runtime/Metasound/Source/MetasoundStandardNodes/Private/MetasoundCrossfadeNode.cpp
--- a/Engine/Plugins/Runtime/Metasound/Source/MetasoundStandardNodes/Private/MetasoundCrossfadeNode.cpp
+++ b/Engine/Plugins/Runtime/Metasound/Source/MetasoundStandardNodes/Private/MetasoundCrossfadeNode.cpp
@@ -4,6 +4,7 @@
 #include "MetasoundAudioBuffer.h"
 #include "CoreMinimal.h"
 #include "DSP/BufferVectorOperations.h"
+#include "DSP/FloatArrayMath.h"
 #include "Internationalization/Text.h"
 #include "MetasoundExecutableOperator.h"
 #include "MetasoundFacade.h"
diff --git a/Engine/Plugins/Runtime/Metasound/Source/MetasoundStandardNodes/Private/MetasoundITDPannerNode.cpp b/Engine/Plugins/Runtime/Metasound/Source/MetasoundStandardNodes/Private/MetasoundITDPannerNode.cpp
--- a/Engine/Plugins/Runtime/Metasound/Source/MetasoundStandardNodes/Private/MetasoundITDPannerNode.cpp
+++ b/Engine/Plugins/Runtime/Metasound/Source/MetasoundStandardNodes/Private/MetasoundITDPannerNode.cpp
@@ -13,6 +13,7 @@
 #include "DSP/BufferVectorOperations.h"
 #include "DSP/Delay.h"
 #include "DSP/Dsp.h"
+#include "DSP/FloatArrayMath.h"
 #include "MetasoundFacade.h"
 #include "MetasoundStandardNodesCategories.h"
 #include "MetasoundVertex.h"
diff --git a/Engine/Plugins/Runtime/Metasound/Source/MetasoundStandardNodes/Private/MetasoundMathNodes.cpp b/Engine/Plugins/Runtime/Metasound/Source/MetasoundStandardNodes/Private/MetasoundMathNodes.cpp
--- a/Engine/Plugins/Runtime/Metasound/Source/MetasoundStandardNodes/Private/MetasoundMathNodes.cpp
+++ b/Engine/Plugins/Runtime/Metasound/Source/MetasoundStandardNodes/Private/MetasoundMathNodes.cpp
@@ -2,6 +2,7 @@
 
 #include "CoreMinimal.h"
 #include "DSP/BufferVectorOperations.h"
+#include "DSP/FloatArrayMath.h"
 #include "MetasoundAudioBuffer.h"
 #include "MetasoundBuilderInterface.h"
 #include "MetasoundDataReference.h"
diff --git a/Engine/Plugins/Runtime/Metasound/Source/MetasoundStandardNodes/Private/MetasoundMixerNode.cpp b/Engine/Plugins/Runtime/Metasound/Source/MetasoundStandardNodes/Private/MetasoundMixerNode.cpp
--- a/Engine/Plugins/Runtime/Metasound/Source/MetasoundStandardNodes/Private/MetasoundMixerNode.cpp
+++ b/Engine/Plugins/Runtime/Metasound/Source/MetasoundStandardNodes/Private/MetasoundMixerNode.cpp
@@ -1,6 +1,7 @@
 // Copyright Epic Games, Inc. All Rights Reserved.
 #include "DSP/Dsp.h"
 #include "DSP/BufferVectorOperations.h"
+#include "DSP/FloatArrayMath.h"
 
 #include "MetasoundAudioBuffer.h"
 #include "MetasoundExecutableOperator.h"
diff --git a/Engine/Plugins/Runtime/Metasound/Source/MetasoundStandardNodes/Private/MetasoundRingModNode.cpp b/Engine/Plugins/Runtime/Metasound/Source/MetasoundStandardNodes/Private/MetasoundRingModNode.cpp
--- a/Engine/Plugins/Runtime/Metasound/Source/MetasoundStandardNodes/Private/MetasoundRingModNode.cpp
+++ b/Engine/Plugins/Runtime/Metasound/Source/MetasoundStandardNodes/Private/MetasoundRingModNode.cpp
@@ -1,5 +1,6 @@
 // Copyright Epic Games, Inc. All Rights Reserved.
 
+#include "DSP/FloatArrayMath.h"
 #include "Internationalization/Text.h"
 #include "MetasoundExecutableOperator.h"
 #include "MetasoundEnumRegistrationMacro.h"
diff --git a/Engine/Plugins/Runtime/Metasound/Source/MetasoundStandardNodes/Private/MetasoundStereoPannerNode.cpp b/Engine/Plugins/Runtime/Metasound/Source/MetasoundStandardNodes/Private/MetasoundStereoPannerNode.cpp
--- a/Engine/Plugins/Runtime/Metasound/Source/MetasoundStandardNodes/Private/MetasoundStereoPannerNode.cpp
+++ b/Engine/Plugins/Runtime/Metasound/Source/MetasoundStandardNodes/Private/MetasoundStereoPannerNode.cpp
@@ -12,6 +12,7 @@
 #include "MetasoundTime.h"
 #include "MetasoundAudioBuffer.h"
 #include "DSP/BufferVectorOperations.h"
+#include "DSP/FloatArrayMath.h"
 #include "MetasoundStandardNodesCategories.h"
 #include "MetasoundFacade.h"
 #include "MetasoundVertex.h"
diff --git a/Engine/Source/Runtime/AudioMixer/Private/Effects/AudioMixerSubmixEffectReverb.cpp b/Engine/Source/Runtime/AudioMixer/Private/Effects/AudioMixerSubmixEffectReverb.cpp
--- a/Engine/Source/Runtime/AudioMixer/Private/Effects/AudioMixerSubmixEffectReverb.cpp
+++ b/Engine/Source/Runtime/AudioMixer/Private/Effects/AudioMixerSubmixEffectReverb.cpp
@@ -6,6 +6,7 @@
 #include "Audio.h"
 #include "AudioMixer.h"
 #include "DSP/BufferVectorOperations.h"
+#include "DSP/FloatArrayMath.h"
 #include "DSP/ReverbFast.h"
 #include "DSP/Amp.h"
 #include "ProfilingDebugging/CsvProfiler.h"
diff --git a/Engine/Source/Runtime/SignalProcessing/Private/ConvertDeinterleave.cpp b/Engine/Source/Runtime/SignalProcessing/Private/ConvertDeinterleave.cpp
--- a/Engine/Source/Runtime/SignalProcessing/Private/ConvertDeinterleave.cpp
+++ b/Engine/Source/Runtime/SignalProcessing/Private/ConvertDeinterleave.cpp
@@ -6,6 +6,7 @@
 #include "Containers/ArrayView.h"
 #include "DSP/ChannelMap.h"
 #include "DSP/MultichannelBuffer.h"
+#include "DSP/FloatArrayMath.h"
 
 namespace Audio
 {
diff --git a/Engine/Source/Runtime/SignalProcessing/Private/Flanger.cpp b/Engine/Source/Runtime/SignalProcessing/Private/Flanger.cpp
--- a/Engine/Source/Runtime/SignalProcessing/Private/Flanger.cpp
+++ b/Engine/Source/Runtime/SignalProcessing/Private/Flanger.cpp
@@ -1,6 +1,7 @@
 // Copyright Epic Games, Inc. All Rights Reserved.
 
 #include "DSP/Flanger.h"
+#include "DSP/FloatArrayMath.h"
 
 namespace Audio
 {
