Fix Compressor Node's Wet/Dry Functionality

diff --git a/Engine/Plugins/Runtime/Metasound/Source/MetasoundStandardNodes/Private/MetasoundCompressorNode.cpp b/Engine/Plugins/Runtime/Metasound/Source/MetasoundStandardNodes/Private/MetasoundCompressorNode.cpp
--- a/Engine/Plugins/Runtime/Metasound/Source/MetasoundStandardNodes/Private/MetasoundCompressorNode.cpp
+++ b/Engine/Plugins/Runtime/Metasound/Source/MetasoundStandardNodes/Private/MetasoundCompressorNode.cpp
@@ -3,6 +3,7 @@
 #include "DSP/Dsp.h"
 #include "DSP/DynamicsProcessor.h"
 #include "DSP/FloatArrayMath.h"
+#include "DSP/IntegerDelay.h"
 #include "Internationalization/Text.h"
 #include "MetasoundAudioBuffer.h"
 #include "MetasoundEnumRegistrationMacro.h"
@@ -27,7 +28,7 @@ namespace Metasound
 		METASOUND_PARAM(InputThreshold, "Threshold dB", "Amplitude threshold (dB) above which gain will be reduced.");
 		METASOUND_PARAM(InputAttackTime, "Attack Time", "How long it takes for audio above the threshold to reach its compressed volume level.");
 		METASOUND_PARAM(InputReleaseTime, "Release Time", "How long it takes for audio below the threshold to return to its original volume level.");
-		METASOUND_PARAM(InputLookaheadTime, "Lookahead Time", "How long to delay the compressed signal behind the analyzed input signal.");
+		METASOUND_PARAM(InputLookaheadTime, "Lookahead Time", "How much time the compressor has to lookahead and catch peaks. This delays the signal.");
 		METASOUND_PARAM(InputKnee, "Knee", "How hard or soft the gain reduction blends from no gain reduction to gain reduction. 0 dB = no blending.");
 		METASOUND_PARAM(InputSidechain, "Sidechain", "(Optional) External audio signal to control the compressor with. If empty, uses the input audio signal.");
 		METASOUND_PARAM(InputEnvelopeMode, "Envelope Mode", "The envelope-following method the compressor will use for gain detection.");
@@ -72,8 +73,10 @@ namespace Metasound
 			, WetDryMixInput(InWetDryMix)
 			, AudioOutput(FAudioBufferWriteRef::CreateNew(InSettings))
 			, EnvelopeOutput(FAudioBufferWriteRef::CreateNew(InSettings))
-			, Compressor()
+			, InputDelay(FMath::CeilToInt(InSettings.GetSampleRate() * Compressor.GetMaxLookaheadMsec() / 1000.f) + 1, InSettings.GetSampleRate() * 10.0f / 1000.0f)
+			, DelayedInputSignal(InSettings.GetNumFramesPerBlock())
 			, bUseSidechain(bInUseSidechain)
+			, MsToSamples(InSettings.GetSampleRate() / 1000.0f)
 			, PrevAttackTime(FMath::Max(FTime::ToMilliseconds(*InAttackTime), 0.0))
 			, PrevReleaseTime(FMath::Max(FTime::ToMilliseconds(*InReleaseTime), 0.0))
 			, PrevLookaheadTime(FMath::Max(FTime::ToMilliseconds(*InLookaheadTime), 0.0))
@@ -251,13 +254,13 @@ namespace Metasound
 				PrevReleaseTime = CurrRelease;
 			}
 
-			// Lookahead time cannot be negative
-			float CurrLookahead = FMath::Max(FTime::ToMilliseconds(*LookaheadTimeInput), 0.0f);
+			float CurrLookahead = FMath::Clamp(FTime::ToMilliseconds(*LookaheadTimeInput), 0.0f, Compressor.GetMaxLookaheadMsec());
 			if (FMath::IsNearlyEqual(CurrLookahead, PrevLookaheadTime) == false)
 			{
 				Compressor.SetLookaheadMsec(CurrLookahead);
 				PrevLookaheadTime = CurrLookahead;
 			}
+			InputDelay.SetDelayLengthSamples(CurrLookahead * MsToSamples);
 
 			Compressor.SetKneeBandwidth(*KneeInput);
 
@@ -286,6 +289,9 @@ namespace Metasound
 				break;
 			}
 
+			// Apply lookahead delay to dry signal
+			InputDelay.ProcessAudio(*AudioInput, DelayedInputSignal);
+
 			if (bUseSidechain)
 			{
 				Compressor.ProcessAudio(AudioInput->GetData(), AudioInput->Num(), AudioOutput->GetData(), SidechainInput->GetData(), EnvelopeOutput->GetData());
@@ -300,7 +306,7 @@ namespace Metasound
 			// Wet signal
 			Audio::ArrayMultiplyByConstantInPlace(*AudioOutput, NewWetDryMix);
 			// Add Dry signal
-			Audio::ArrayMultiplyAddInPlace(*AudioInput, 1.0f - NewWetDryMix, *AudioOutput);
+			Audio::ArrayMultiplyAddInPlace(DelayedInputSignal, 1.0f - NewWetDryMix, *AudioOutput);
 			
 		}
 
@@ -325,10 +331,17 @@ namespace Metasound
 
 		// Internal DSP Compressor
 		Audio::FDynamicsProcessor Compressor;
+		// Compressor does not have a wet/dry signal built in, so
+		// we need to account for lookahead delay in the input to prevent phase issues.
+		Audio::FIntegerDelay InputDelay;
+		FAudioBuffer DelayedInputSignal;
 
 		// Whether or not to use sidechain input (is false if no input pin is connected to sidechain input)
 		bool bUseSidechain;
 
+		// Conversion from milliseconds to samples
+		float MsToSamples;
+
 		// Cached variables
 		float PrevAttackTime;
 		float PrevReleaseTime;
diff --git a/Engine/Source/Runtime/SignalProcessing/Private/DynamicsProcessor.cpp b/Engine/Source/Runtime/SignalProcessing/Private/DynamicsProcessor.cpp
--- a/Engine/Source/Runtime/SignalProcessing/Private/DynamicsProcessor.cpp
+++ b/Engine/Source/Runtime/SignalProcessing/Private/DynamicsProcessor.cpp
@@ -43,7 +43,7 @@ namespace Audio
 
 		for (int32 Channel = 0; Channel < InNumChannels; ++Channel)
 		{
-			LookaheadDelay[Channel].Init(SampleRate, 0.1f);
+			LookaheadDelay[Channel].Init(SampleRate, MaxLookaheadMsec / 1000.0f);
 			LookaheadDelay[Channel].SetDelayMsec(LookaheedDelayMsec);
 
 			EnvFollower[Channel].Init(FInlineEnvelopeFollowerInitParams{SampleRate, AttackTimeMsec, ReleaseTimeMsec, EnvelopeFollowerPeakMode, bIsAnalogMode});
@@ -68,6 +68,11 @@ namespace Audio
 	{
 		return EnvFollower.Num();
 	}
+
+	float Audio::FDynamicsProcessor::GetMaxLookaheadMsec() const
+	{
+		return MaxLookaheadMsec;
+	}
 	
 	void FDynamicsProcessor::SetLookaheadMsec(const float InLookAheadMsec)
 	{
diff --git a/Engine/Source/Runtime/SignalProcessing/Public/DSP/DynamicsProcessor.h b/Engine/Source/Runtime/SignalProcessing/Public/DSP/DynamicsProcessor.h
--- a/Engine/Source/Runtime/SignalProcessing/Public/DSP/DynamicsProcessor.h
+++ b/Engine/Source/Runtime/SignalProcessing/Public/DSP/DynamicsProcessor.h
@@ -43,6 +43,7 @@ namespace Audio
 
 		int32 GetNumChannels() const;
 		int32 GetKeyNumChannels() const;
+		float GetMaxLookaheadMsec() const;
 
 		void SetLookaheadMsec(const float InLookAheadMsec);
 		void SetAttackTime(const float InAttackTimeMsec);
@@ -155,5 +156,7 @@ namespace Audio
 		bool bKeyLowshelfEnabled;
 
 		static constexpr float UpwardsCompressionMaxGain = 36.0f;
+
+		static constexpr float MaxLookaheadMsec = 100.0f;
 	};
 }
