Various MetaSound Editor UX improvements for member defaults and selection

- Add Metasound editor setting for default input widget type
- Change member default type from bool to float
- Make widgets only available for inputs (no longer for variables)
- Automatically select/focus newly added/promoted nodes
- Selecting a single node selects the corresponding member in the graph member list if applicable

diff --git a/Engine/Plugins/Runtime/Metasound/Source/MetasoundEditor/Private/MetasoundEditor.cpp b/Engine/Plugins/Runtime/Metasound/Source/MetasoundEditor/Private/MetasoundEditor.cpp
--- a/Engine/Plugins/Runtime/Metasound/Source/MetasoundEditor/Private/MetasoundEditor.cpp
+++ b/Engine/Plugins/Runtime/Metasound/Source/MetasoundEditor/Private/MetasoundEditor.cpp
@@ -193,6 +193,11 @@
 					{
 						NewGraphNode->Modify();
 						FGraphBuilder::RegisterGraphWithFrontend(ParentMetasound);
+						TSharedPtr<FEditor> MetasoundEditor = FGraphBuilder::GetEditorForGraph(*MetasoundGraph);
+						if (MetasoundEditor.IsValid())
+						{
+							MetasoundEditor->ClearSelectionAndSelectNode(NewGraphNode);
+						}
 						return FReply::Handled();
 					}
 				}
@@ -212,6 +217,11 @@
 						{
 							NewGraphNode->Modify();
 							FGraphBuilder::RegisterGraphWithFrontend(ParentMetasound);
+							TSharedPtr<FEditor> MetasoundEditor = FGraphBuilder::GetEditorForGraph(*MetasoundGraph);
+							if (MetasoundEditor.IsValid())
+							{
+								MetasoundEditor->ClearSelectionAndSelectNode(NewGraphNode);
+							}
 							return FReply::Handled();
 						}
 					}
@@ -301,6 +311,11 @@
 					{
 						NewGraphNode->Modify();
 						FGraphBuilder::RegisterGraphWithFrontend(ParentMetasound);
+						TSharedPtr<FEditor> MetasoundEditor = FGraphBuilder::GetEditorForGraph(*MetasoundGraph);
+						if (MetasoundEditor.IsValid())
+						{
+							MetasoundEditor->ClearSelectionAndSelectNode(NewGraphNode);
+						}
 						return FReply::Handled();
 					}
 				}
@@ -898,6 +913,19 @@
 
 		void FEditor::SetSelection(const TArray<UObject*>& SelectedObjects)
 		{
+			if (GraphMembersMenu.IsValid())
+			{
+				// Only support menu selection of a single object until multiselect functionality is added 
+				if (SelectedObjects.Num() == 1)
+				{
+					if (UMetasoundEditorGraphMember* Member = Cast<UMetasoundEditorGraphMember>(SelectedObjects[0]))
+					{
+						const FName ActionName = Member->GetMemberName();
+						GraphMembersMenu->SelectItemByName(ActionName, ESelectInfo::Direct, static_cast<int32>(Member->GetSectionID()));
+					}
+				}
+			}
+
 			if (MetasoundDetails.IsValid())
 			{
 				MetasoundDetails->SetObjects(SelectedObjects);
@@ -2624,7 +2652,7 @@
 		{
 			if (GraphMembersMenu.IsValid())
 			{
-				GraphMembersMenu->RefreshAllActions(true /* bPreserveExpansion */);
+				GraphMembersMenu->RefreshAllActions(true /* bPreserveExpansion */, false /*bHandleOnSelectionEvent*/);
 			}
 		}
 
@@ -3216,6 +3244,15 @@
 			return MetasoundAsset->GetRootGraphHandle()->GetGraphStyle().bIsGraphEditable;
 		}
 
+		void FEditor::ClearSelectionAndSelectNode(UEdGraphNode* Node)
+		{
+			if (MetasoundGraphEditor.IsValid())
+			{
+				MetasoundGraphEditor->ClearSelectionSet();
+				MetasoundGraphEditor->SetNodeSelection(Node, /*bSelect=*/true);
+			}
+		}
+
 		TSharedRef<SWidget> FEditor::OnGetMenuSectionWidget(TSharedRef<SWidget> RowWidget, int32 InSectionID)
 		{
 			TWeakPtr<SWidget> WeakRowWidget = RowWidget;
@@ -3283,7 +3320,7 @@
 					const FScopedTransaction Transaction(LOCTEXT("AddInputNode", "Add MetaSound Input"));
 					Metasound->Modify();
 
-					const FName DataTypeName = GetMetasoundDataTypeName<bool>();
+					const FName DataTypeName = GetMetasoundDataTypeName<float>();
 					Frontend::FNodeHandle NodeHandle = FGraphBuilder::AddInputNodeHandle(*Metasound, DataTypeName);
 					if (ensure(NodeHandle->IsValid()))
 					{
@@ -3309,7 +3346,7 @@
 					const FScopedTransaction Transaction(TEXT(""), LOCTEXT("AddOutputNode", "Add MetaSound Output"), Metasound);
 					Metasound->Modify();
 
-					const FName DataTypeName = GetMetasoundDataTypeName<bool>();
+					const FName DataTypeName = GetMetasoundDataTypeName<float>();
 					Frontend::FNodeHandle NodeHandle = FGraphBuilder::AddOutputNodeHandle(*Metasound, DataTypeName);
 					if (ensure(NodeHandle->IsValid()))
 					{
@@ -3335,7 +3372,7 @@
 					const FScopedTransaction Transaction(TEXT(""), LOCTEXT("AddVariableNode", "Add MetaSound Variable"), Metasound);
 					Metasound->Modify();
 
-					const FName DataTypeName = GetMetasoundDataTypeName<bool>();
+					const FName DataTypeName = GetMetasoundDataTypeName<float>();
 					
 					Frontend::FVariableHandle FrontendVariable = FGraphBuilder::AddVariableHandle(*Metasound, DataTypeName);
 					if (ensure(FrontendVariable->IsValid()))
diff --git a/Engine/Plugins/Runtime/Metasound/Source/MetasoundEditor/Private/MetasoundEditor.h b/Engine/Plugins/Runtime/Metasound/Source/MetasoundEditor/Private/MetasoundEditor.h
--- a/Engine/Plugins/Runtime/Metasound/Source/MetasoundEditor/Private/MetasoundEditor.h
+++ b/Engine/Plugins/Runtime/Metasound/Source/MetasoundEditor/Private/MetasoundEditor.h
@@ -175,6 +175,8 @@ namespace Metasound
 			/* Whether the displayed graph is marked as editable */
 			bool IsGraphEditable() const;
 
+			void ClearSelectionAndSelectNode(UEdGraphNode* Node);
+
 			int32 GetNumNodesSelected() const
 			{
 				return MetasoundGraphEditor->GetSelectedNodes().Num();
diff --git a/Engine/Plugins/Runtime/Metasound/Source/MetasoundEditor/Private/MetasoundEditorGraphMemberDefaults.cpp b/Engine/Plugins/Runtime/Metasound/Source/MetasoundEditor/Private/MetasoundEditorGraphMemberDefaults.cpp
--- a/Engine/Plugins/Runtime/Metasound/Source/MetasoundEditor/Private/MetasoundEditorGraphMemberDefaults.cpp
+++ b/Engine/Plugins/Runtime/Metasound/Source/MetasoundEditor/Private/MetasoundEditorGraphMemberDefaults.cpp
@@ -10,7 +10,6 @@
 #include "MetasoundDataReference.h"
 #include "MetasoundEditorGraph.h"
 #include "MetasoundEditorGraphNode.h"
-#include "MetasoundEditorSettings.h"
 #include "MetasoundFrontend.h"
 #include "MetasoundFrontendController.h"
 #include "MetasoundFrontendDocument.h"
@@ -145,6 +144,14 @@ void UMetasoundEditorGraphMemberDefaultIntArray::UpdatePreviewInstance(const Met
 	InParameterInterface->SetIntArrayParameter(InParameterName, IntArray);
 }
 
+UMetasoundEditorGraphMemberDefaultFloat::UMetasoundEditorGraphMemberDefaultFloat()
+{
+	if (const UMetasoundEditorSettings* EditorSettings = ::GetDefault<UMetasoundEditorSettings>())
+	{
+		WidgetType = EditorSettings->DefaultInputWidgetType;
+	}
+}
+
 FMetasoundFrontendLiteral UMetasoundEditorGraphMemberDefaultFloat::GetDefault() const
 {
 	FMetasoundFrontendLiteral Literal;
diff --git a/Engine/Plugins/Runtime/Metasound/Source/MetasoundEditor/Private/MetasoundEditorGraphSchema.cpp b/Engine/Plugins/Runtime/Metasound/Source/MetasoundEditor/Private/MetasoundEditorGraphSchema.cpp
--- a/Engine/Plugins/Runtime/Metasound/Source/MetasoundEditor/Private/MetasoundEditorGraphSchema.cpp
+++ b/Engine/Plugins/Runtime/Metasound/Source/MetasoundEditor/Private/MetasoundEditorGraphSchema.cpp
@@ -283,6 +283,13 @@ namespace Metasound
 							{
 								FGraphBuilder::RegisterGraphWithFrontend(ParentMetaSound);
 								MetaSoundGraph->SetSynchronizationRequired();
+
+								TSharedPtr<FEditor> ParentEditor = FGraphBuilder::GetEditorForMetasound(ParentMetaSound);
+								if (ParentEditor.IsValid())
+								{
+									ParentEditor->RefreshGraphMemberMenu();
+								}
+
 								return EdGraphNode;
 							}
 						}
@@ -352,6 +359,8 @@ namespace Metasound
 
 UEdGraphNode* FMetasoundGraphSchemaAction_NodeWithMultipleOutputs::PerformAction(class UEdGraph* ParentGraph, TArray<UEdGraphPin*>& FromPins, const FVector2D Location, bool bSelectNewNode/* = true*/)
 {
+	using namespace Metasound::Editor;
+	
 	UEdGraphNode* ResultNode = NULL;
 
 	if (FromPins.Num() > 0)
@@ -372,6 +381,12 @@ UEdGraphNode* FMetasoundGraphSchemaAction_NodeWithMultipleOutputs::PerformAction
 		ResultNode = PerformAction(ParentGraph, NULL, Location, bSelectNewNode);
 	}
 
+	TSharedPtr<FEditor> MetasoundEditor = FGraphBuilder::GetEditorForGraph(*ParentGraph);
+	if (MetasoundEditor.IsValid() && bSelectNewNode)
+	{
+		MetasoundEditor->ClearSelectionAndSelectNode(ResultNode);
+	}
+
 	return ResultNode;
 }
 
@@ -429,6 +444,13 @@ UEdGraphNode* FMetasoundGraphSchemaAction_NewNode::PerformAction(UEdGraph* Paren
 		NewGraphNode->Modify();
 		SchemaPrivate::TryConnectNewNodeToMatchingDataTypePin(*NewGraphNode, FromPin);
 		MetaSoundGraph->SetSynchronizationRequired();
+		
+		TSharedPtr<FEditor> ParentEditor = FGraphBuilder::GetEditorForMetasound(ParentMetasound);
+		if (ParentEditor.IsValid() && bSelectNewNode)
+		{
+			ParentEditor->ClearSelectionAndSelectNode(NewGraphNode);
+		}
+		
 		return NewGraphNode;
 	}
 
@@ -644,7 +666,13 @@ UEdGraphNode* FMetasoundGraphSchemaAction_PromoteToVariable_MutatorNode::Perform
 	FMetasoundFrontendClass VariableClass;
 	if (ensure(IDataTypeRegistry::Get().GetFrontendVariableMutatorClass(DataType, VariableClass)))
 	{
-		return SchemaPrivate::PromoteToVariable(NodeName, *FromPin, DataType, VariableClass, InLocation);
+		UEdGraphNode* NewGraphNode = SchemaPrivate::PromoteToVariable(NodeName, *FromPin, DataType, VariableClass, InLocation);
+		TSharedPtr<FEditor> MetasoundEditor = FGraphBuilder::GetEditorForGraph(*ParentGraph);
+		if (MetasoundEditor.IsValid() && bSelectNewNode)
+		{
+			MetasoundEditor->ClearSelectionAndSelectNode(NewGraphNode);
+		}
+		return NewGraphNode;
 	}
 
 	return nullptr;
@@ -757,6 +785,7 @@ UEdGraphNode* FMetasoundGraphSchemaAction_PromoteToOutput::PerformAction(UEdGrap
 					}
 
 					FGraphBuilder::RegisterGraphWithFrontend(ParentMetasound);
+					MetasoundEditor->ClearSelectionAndSelectNode(EdGraphNode);
 					return EdGraphNode;
 				}
 			}
diff --git a/Engine/Plugins/Runtime/Metasound/Source/MetasoundEditor/Private/MetasoundEditorSettings.h b/Engine/Plugins/Runtime/Metasound/Source/MetasoundEditor/Private/MetasoundEditorSettings.h
--- a/Engine/Plugins/Runtime/Metasound/Source/MetasoundEditor/Private/MetasoundEditorSettings.h
+++ b/Engine/Plugins/Runtime/Metasound/Source/MetasoundEditor/Private/MetasoundEditorSettings.h
@@ -8,6 +8,13 @@
 
 #include "MetasoundEditorSettings.generated.h"
 
+UENUM()
+enum class EMetasoundMemberDefaultWidget : uint8
+{
+	None,
+	Slider,
+	RadialSlider UMETA(DisplayName = "Knob"),
+};
 
 UENUM()
 enum class EMetasoundActiveDetailView : uint8
@@ -120,6 +127,10 @@ public:
 	UPROPERTY(EditAnywhere, config, Category = GraphAnimation, meta = (EditCondition = "bAnimateConnections", UIMin = 1, UIMax = 5, ClampMin = 1))
 	float ActiveAnalyzerWireScalarMax;
 
+	/** Widget type to show on input nodes by default */
+	UPROPERTY(EditAnywhere, config, Category = General)
+	EMetasoundMemberDefaultWidget DefaultInputWidgetType = EMetasoundMemberDefaultWidget::None;
+
 	/** Determines which details view to show in Metasounds Editor */
 	UPROPERTY(Transient)
 	EMetasoundActiveDetailView DetailView = EMetasoundActiveDetailView::Metasound;
diff --git a/Engine/Plugins/Runtime/Metasound/Source/MetasoundEditor/Private/MetasoundNodeDetailCustomization.cpp b/Engine/Plugins/Runtime/Metasound/Source/MetasoundEditor/Private/MetasoundNodeDetailCustomization.cpp
--- a/Engine/Plugins/Runtime/Metasound/Source/MetasoundEditor/Private/MetasoundNodeDetailCustomization.cpp
+++ b/Engine/Plugins/Runtime/Metasound/Source/MetasoundEditor/Private/MetasoundNodeDetailCustomization.cpp
@@ -172,17 +172,18 @@ namespace Metasound
 				}
 			}
 
-			bool bIsGraphEditable = false;
-			if (const UMetasoundEditorGraphMember* ParentMember = InLiteral.GetParentMember())
+			// Enable widget options for editable inputs only 
+			bool bShowWidgetOptions = false;
+			if (const UMetasoundEditorGraphInput* ParentMember = Cast <UMetasoundEditorGraphInput>(InLiteral.GetParentMember()))
 			{
 				if (const UMetasoundEditorGraph* OwningGraph = ParentMember->GetOwningGraph())
 				{
-					bIsGraphEditable = OwningGraph->IsEditable();
+					bShowWidgetOptions = OwningGraph->IsEditable();
 				}
 			}
 
 			// add input widget properties
-			if (bIsGraphEditable)
+			if (bShowWidgetOptions)
 			{
 				IDetailCategoryBuilder& WidgetCategoryBuilder = InDetailLayout.EditCategory("EditorOptions");
 				DefaultRows.Add(WidgetCategoryBuilder.AddExternalObjectProperty(TArray<UObject*>({ DefaultFloat }), GET_MEMBER_NAME_CHECKED(UMetasoundEditorGraphMemberDefaultFloat, WidgetType)));
diff --git a/Engine/Plugins/Runtime/Metasound/Source/MetasoundEditor/Private/SMetasoundGraphNode.cpp b/Engine/Plugins/Runtime/Metasound/Source/MetasoundEditor/Private/SMetasoundGraphNode.cpp
--- a/Engine/Plugins/Runtime/Metasound/Source/MetasoundEditor/Private/SMetasoundGraphNode.cpp
+++ b/Engine/Plugins/Runtime/Metasound/Source/MetasoundEditor/Private/SMetasoundGraphNode.cpp
@@ -632,7 +632,7 @@ namespace Metasound
 
 			// If float input node, check if custom widget required
 			bool IsFloatMemberNode = false;
-			if (UMetasoundEditorGraphMember* GraphMember = GetMetaSoundMember())
+			if (UMetasoundEditorGraphInput* GraphMember = Cast<UMetasoundEditorGraphInput>(GetMetaSoundMember()))
 			{
 				if (UMetasoundEditorGraphMemberDefaultFloat* DefaultFloat = Cast<UMetasoundEditorGraphMemberDefaultFloat>(GraphMember->GetLiteral()))
 				{
diff --git a/Engine/Plugins/Runtime/Metasound/Source/MetasoundEditor/Public/MetasoundEditorGraphMemberDefaults.h b/Engine/Plugins/Runtime/Metasound/Source/MetasoundEditor/Public/MetasoundEditorGraphMemberDefaults.h
--- a/Engine/Plugins/Runtime/Metasound/Source/MetasoundEditor/Public/MetasoundEditorGraphMemberDefaults.h
+++ b/Engine/Plugins/Runtime/Metasound/Source/MetasoundEditor/Public/MetasoundEditorGraphMemberDefaults.h
@@ -11,6 +11,7 @@
 #include "MetasoundDataReference.h"
 #include "MetasoundEditorGraph.h"
 #include "MetasoundEditorGraphNode.h"
+#include "MetasoundEditorSettings.h"
 #include "MetasoundFrontend.h"
 #include "MetasoundFrontendController.h"
 #include "MetasoundFrontendDocument.h"
@@ -119,13 +120,6 @@ DECLARE_MULTICAST_DELEGATE_OneParam(FOnMetasoundInputValueChangedEvent, float);
 DECLARE_MULTICAST_DELEGATE_OneParam(FOnMetasoundRangeChangedEvent, FVector2D);
 DECLARE_MULTICAST_DELEGATE_OneParam(FOnMetasoundInputClampDefaultChangedEvent, bool);
 
-UENUM()
-enum class EMetasoundMemberDefaultWidget : uint8
-{
-	None,
-	Slider,
-	RadialSlider UMETA(DisplayName = "Knob"),
-};
 
 UENUM()
 enum class EMetasoundMemberDefaultWidgetValueType : uint8
@@ -145,6 +139,7 @@ private:
 	float Default = 0.f;
 
 public:
+	UMetasoundEditorGraphMemberDefaultFloat();
 	virtual ~UMetasoundEditorGraphMemberDefaultFloat() = default;
 
 	UPROPERTY(EditAnywhere, Category = DefaultValue, meta=(EditCondition = "WidgetType == EMetasoundMemberDefaultWidget::None", EditConditionHides))
