Fix crash in flanger MetaSound node

diff --git a/Engine/Plugins/Runtime/Metasound/Source/MetasoundStandardNodes/Private/MetasoundFlangerNode.cpp b/Engine/Plugins/Runtime/Metasound/Source/MetasoundStandardNodes/Private/MetasoundFlangerNode.cpp
--- a/Engine/Plugins/Runtime/Metasound/Source/MetasoundStandardNodes/Private/MetasoundFlangerNode.cpp
+++ b/Engine/Plugins/Runtime/Metasound/Source/MetasoundStandardNodes/Private/MetasoundFlangerNode.cpp
@@ -148,6 +148,9 @@ namespace Metasound
 			AlignedAudioIn.Reset();
 			AlignedAudioIn.Append(InputAudio, NumFrames);
 
+			AlignedAudioOut.Reset();
+			AlignedAudioOut.AddUninitialized(NumFrames);
+
 			Flanger.ProcessAudio(AlignedAudioIn, NumFrames, AlignedAudioOut);
 			FMemory::Memcpy(OutputAudio, AlignedAudioOut.GetData(), sizeof(float) * NumFrames);
 		}
