Minor MetaSound UX fixes

- Only apply input widget editor setting to new inputs
- Change default input widget to knob
- Don't show input widgets on read only graphs
- Tooltip clarification for setting audio component default parameters on MetaSounds

diff --git a/Engine/Plugins/Runtime/Metasound/Source/MetasoundEditor/Private/MetasoundEditorGraph.cpp b/Engine/Plugins/Runtime/Metasound/Source/MetasoundEditor/Private/MetasoundEditorGraph.cpp
--- a/Engine/Plugins/Runtime/Metasound/Source/MetasoundEditor/Private/MetasoundEditorGraph.cpp
+++ b/Engine/Plugins/Runtime/Metasound/Source/MetasoundEditor/Private/MetasoundEditorGraph.cpp
@@ -108,6 +108,15 @@ void UMetasoundEditorGraphMember::InitializeLiteral()
 	if (!Literal || Literal->GetClass() != LiteralClass)
 	{
 		UMetasoundEditorGraphMemberDefaultLiteral* NewLiteral = NewObject<UMetasoundEditorGraphMemberDefaultLiteral>(this, LiteralClass, FName(), RF_Transactional);
+
+		// Set default widget type from editor settings for float inputs 
+		if (Cast<UMetasoundEditorGraphInput>(this) && LiteralClass == UMetasoundEditorGraphMemberDefaultFloat::StaticClass())
+		{
+			if (const UMetasoundEditorSettings* EditorSettings = ::GetDefault<UMetasoundEditorSettings>())
+			{
+				Cast<UMetasoundEditorGraphMemberDefaultFloat>(NewLiteral)->WidgetType = EditorSettings->DefaultInputWidgetType;
+			}
+		}
 		Literal = NewLiteral;
 	}
 }
diff --git a/Engine/Plugins/Runtime/Metasound/Source/MetasoundEditor/Private/MetasoundEditorGraphMemberDefaults.cpp b/Engine/Plugins/Runtime/Metasound/Source/MetasoundEditor/Private/MetasoundEditorGraphMemberDefaults.cpp
--- a/Engine/Plugins/Runtime/Metasound/Source/MetasoundEditor/Private/MetasoundEditorGraphMemberDefaults.cpp
+++ b/Engine/Plugins/Runtime/Metasound/Source/MetasoundEditor/Private/MetasoundEditorGraphMemberDefaults.cpp
@@ -144,14 +144,6 @@ void UMetasoundEditorGraphMemberDefaultIntArray::UpdatePreviewInstance(const Met
 	InParameterInterface->SetIntArrayParameter(InParameterName, IntArray);
 }
 
-UMetasoundEditorGraphMemberDefaultFloat::UMetasoundEditorGraphMemberDefaultFloat()
-{
-	if (const UMetasoundEditorSettings* EditorSettings = ::GetDefault<UMetasoundEditorSettings>())
-	{
-		WidgetType = EditorSettings->DefaultInputWidgetType;
-	}
-}
-
 FMetasoundFrontendLiteral UMetasoundEditorGraphMemberDefaultFloat::GetDefault() const
 {
 	FMetasoundFrontendLiteral Literal;
diff --git a/Engine/Plugins/Runtime/Metasound/Source/MetasoundEditor/Private/MetasoundEditorSettings.h b/Engine/Plugins/Runtime/Metasound/Source/MetasoundEditor/Private/MetasoundEditorSettings.h
--- a/Engine/Plugins/Runtime/Metasound/Source/MetasoundEditor/Private/MetasoundEditorSettings.h
+++ b/Engine/Plugins/Runtime/Metasound/Source/MetasoundEditor/Private/MetasoundEditorSettings.h
@@ -129,7 +129,7 @@ public:
 
 	/** Widget type to show on input nodes by default */
 	UPROPERTY(EditAnywhere, config, Category = General)
-	EMetasoundMemberDefaultWidget DefaultInputWidgetType = EMetasoundMemberDefaultWidget::None;
+	EMetasoundMemberDefaultWidget DefaultInputWidgetType = EMetasoundMemberDefaultWidget::RadialSlider;
 
 	/** Determines which details view to show in Metasounds Editor */
 	UPROPERTY(Transient)
diff --git a/Engine/Plugins/Runtime/Metasound/Source/MetasoundEditor/Private/SMetasoundGraphNode.cpp b/Engine/Plugins/Runtime/Metasound/Source/MetasoundEditor/Private/SMetasoundGraphNode.cpp
--- a/Engine/Plugins/Runtime/Metasound/Source/MetasoundEditor/Private/SMetasoundGraphNode.cpp
+++ b/Engine/Plugins/Runtime/Metasound/Source/MetasoundEditor/Private/SMetasoundGraphNode.cpp
@@ -630,19 +630,21 @@ namespace Metasound
 			TSharedPtr<SHorizontalBox> ContentBox = SNew(SHorizontalBox);
 			TSharedPtr<SWidget> OuterContentBox; // currently only used for input float nodes to accommodate the input widget
 
-			// If float input node, check if custom widget required
-			bool IsFloatMemberNode = false;
+			// If editable float input node, check if custom widget required
+			bool bShowInputWidget = false;
 			if (UMetasoundEditorGraphInput* GraphMember = Cast<UMetasoundEditorGraphInput>(GetMetaSoundMember()))
 			{
-				if (UMetasoundEditorGraphMemberDefaultFloat* DefaultFloat = Cast<UMetasoundEditorGraphMemberDefaultFloat>(GraphMember->GetLiteral()))
+				const UMetasoundEditorGraph* OwningGraph = GraphMember->GetOwningGraph();
+				if (OwningGraph && OwningGraph->IsEditable())
 				{
-					if (DefaultFloat->WidgetType != EMetasoundMemberDefaultWidget::None)
+					UMetasoundEditorGraphMemberDefaultFloat* DefaultFloat = Cast<UMetasoundEditorGraphMemberDefaultFloat>(GraphMember->GetLiteral());
+					if (DefaultFloat && DefaultFloat->WidgetType != EMetasoundMemberDefaultWidget::None)
 					{
 						constexpr float WidgetPadding = 3.0f;
 						static const FVector2D SliderDesiredSizeVertical = FVector2D(30.0f, 250.0f);
 						static const FVector2D RadialSliderDesiredSize = FVector2D(56.0f, 87.0f);
 
-						IsFloatMemberNode = true;
+						bShowInputWidget = true;
 
 						auto OnValueChangedLambda = [DefaultFloat, GraphMember, this](float Value)
 						{
@@ -887,7 +889,7 @@ namespace Metasound
 				.VAlign(VAlign_Fill)
 				.Padding(FMargin(0,3))
 				[
-					(IsFloatMemberNode ? OuterContentBox : ContentBox).ToSharedRef()
+					(bShowInputWidget ? OuterContentBox : ContentBox).ToSharedRef()
 				];
 		}
 	} // namespace Editor
diff --git a/Engine/Plugins/Runtime/Metasound/Source/MetasoundEditor/Public/MetasoundEditorGraphMemberDefaults.h b/Engine/Plugins/Runtime/Metasound/Source/MetasoundEditor/Public/MetasoundEditorGraphMemberDefaults.h
--- a/Engine/Plugins/Runtime/Metasound/Source/MetasoundEditor/Public/MetasoundEditorGraphMemberDefaults.h
+++ b/Engine/Plugins/Runtime/Metasound/Source/MetasoundEditor/Public/MetasoundEditorGraphMemberDefaults.h
@@ -139,7 +139,6 @@ private:
 	float Default = 0.f;
 
 public:
-	UMetasoundEditorGraphMemberDefaultFloat();
 	virtual ~UMetasoundEditorGraphMemberDefaultFloat() = default;
 
 	UPROPERTY(EditAnywhere, Category = DefaultValue, meta=(EditCondition = "WidgetType == EMetasoundMemberDefaultWidget::None", EditConditionHides))
diff --git a/Engine/Plugins/Runtime/Metasound/Source/MetasoundEngine/MetasoundEngine.Build.cs b/Engine/Plugins/Runtime/Metasound/Source/MetasoundEngine/MetasoundEngine.Build.cs
--- a/Engine/Plugins/Runtime/Metasound/Source/MetasoundEngine/MetasoundEngine.Build.cs
+++ b/Engine/Plugins/Runtime/Metasound/Source/MetasoundEngine/MetasoundEngine.Build.cs
@@ -9,8 +9,6 @@ namespace UnrealBuildTool.Rules
 	{
 		public MetasoundEngine(ReadOnlyTargetRules Target) : base(Target)
 		{
-			OptimizeCode = CodeOptimization.Never;
-
 			PrivateDependencyModuleNames.AddRange
 			(
 				new string[]
diff --git a/Engine/Plugins/Runtime/Metasound/Source/MetasoundStandardNodes/MetasoundStandardNodes.Build.cs b/Engine/Plugins/Runtime/Metasound/Source/MetasoundStandardNodes/MetasoundStandardNodes.Build.cs
--- a/Engine/Plugins/Runtime/Metasound/Source/MetasoundStandardNodes/MetasoundStandardNodes.Build.cs
+++ b/Engine/Plugins/Runtime/Metasound/Source/MetasoundStandardNodes/MetasoundStandardNodes.Build.cs
@@ -9,8 +9,6 @@ namespace UnrealBuildTool.Rules
 	{
 		public MetasoundStandardNodes(ReadOnlyTargetRules Target) : base(Target)
 		{
-			OptimizeCode = CodeOptimization.Never;
-
 			PublicDependencyModuleNames.AddRange
 			(
 				new string[]
diff --git a/Engine/Source/Runtime/AudioExtensions/Public/AudioParameter.h b/Engine/Source/Runtime/AudioExtensions/Public/AudioParameter.h
--- a/Engine/Source/Runtime/AudioExtensions/Public/AudioParameter.h
+++ b/Engine/Source/Runtime/AudioExtensions/Public/AudioParameter.h
@@ -35,8 +35,8 @@ enum class EAudioParameterType : uint8
 	// 'Default' results in behavior that is resolved
 	// based on the system interpreting it.  To support
 	// legacy implementation, SoundCues cache all typed values
-	// associated with a given parameter name, whereas MetaSounds
-	// only use the last typed value set.
+	// associated with a given parameter name. 
+	// For MetaSounds, use a specific Type instead of this one. 
 	None UMETA(DisplayName = "Default"),
 
 	// Boolean value
