Adding MetaSound environment data for audio mixer output frames.

diff --git a/Engine/Plugins/Runtime/Metasound/Source/MetasoundEngine/Private/MetasoundSource.cpp b/Engine/Plugins/Runtime/Metasound/Source/MetasoundEngine/Private/MetasoundSource.cpp
--- a/Engine/Plugins/Runtime/Metasound/Source/MetasoundEngine/Private/MetasoundSource.cpp
+++ b/Engine/Plugins/Runtime/Metasound/Source/MetasoundEngine/Private/MetasoundSource.cpp
@@ -843,6 +843,7 @@ Metasound::FMetasoundEnvironment UMetaSoundSource::CreateEnvironment(const FSoun
 	Environment.SetValue<bool>(SourceInterface::Environment::IsPreview, InParams.bIsPreviewSound);
 	Environment.SetValue<uint64>(SourceInterface::Environment::TransmitterID, InParams.InstanceID);
 	Environment.SetValue<Audio::FDeviceId>(SourceInterface::Environment::DeviceID, InParams.AudioDeviceID);
+	Environment.SetValue<int32>(SourceInterface::Environment::AudioMixerNumOutputFrames, InParams.AudioMixerNumOutputFrames);
 
 #if WITH_METASOUND_DEBUG_ENVIRONMENT
 	Environment.SetValue<FString>(SourceInterface::Environment::GraphName, GetFullName());
diff --git a/Engine/Plugins/Runtime/Metasound/Source/MetasoundFrontend/Private/MetasoundSourceInterface.cpp b/Engine/Plugins/Runtime/Metasound/Source/MetasoundFrontend/Private/MetasoundSourceInterface.cpp
--- a/Engine/Plugins/Runtime/Metasound/Source/MetasoundFrontend/Private/MetasoundSourceInterface.cpp
+++ b/Engine/Plugins/Runtime/Metasound/Source/MetasoundFrontend/Private/MetasoundSourceInterface.cpp
@@ -174,6 +174,7 @@ namespace Metasound
 				const FName IsPreview = AUDIO_PARAMETER_INTERFACE_MEMBER_DEFINE("IsPreviewSound");
 				const FName SoundUniqueID = AUDIO_PARAMETER_INTERFACE_MEMBER_DEFINE("SoundUniqueID");
 				const FName TransmitterID = AUDIO_PARAMETER_INTERFACE_MEMBER_DEFINE("TransmitterID");
+				const FName AudioMixerNumOutputFrames = AUDIO_PARAMETER_INTERFACE_MEMBER_DEFINE("AudioMixerNumOutputFrames");
 			}
 
 			Audio::FParameterInterfacePtr CreateInterface(const UClass& InUClass)
@@ -224,6 +225,12 @@ namespace Metasound
 								LOCTEXT("SoundUniqueDescription", "ID of unique source instance."),
 								FName(), // TODO: Align environment data types with environment (ex. this is actually set/get as a uint32)
 								Environment::SoundUniqueID
+							},
+							{
+								LOCTEXT("AudioMixerOutputFramesDisplayName", "Audio Mixer Output Frames"),
+								LOCTEXT("AudioMixerOutputFramesDescription", "The number of output frames rendered per buffer in the audio mixer."),
+								FName(), 
+								Environment::AudioMixerOutputFrames
 							}
 						};
 					}
diff --git a/Engine/Plugins/Runtime/Metasound/Source/MetasoundFrontend/Public/MetasoundSourceInterface.h b/Engine/Plugins/Runtime/Metasound/Source/MetasoundFrontend/Public/MetasoundSourceInterface.h
--- a/Engine/Plugins/Runtime/Metasound/Source/MetasoundFrontend/Public/MetasoundSourceInterface.h
+++ b/Engine/Plugins/Runtime/Metasound/Source/MetasoundFrontend/Public/MetasoundSourceInterface.h
@@ -64,6 +64,7 @@ namespace Metasound
 				METASOUNDFRONTEND_API const extern FName IsPreview;
 				METASOUNDFRONTEND_API const extern FName SoundUniqueID;
 				METASOUNDFRONTEND_API const extern FName TransmitterID;
+				METASOUNDFRONTEND_API const extern FName AudioMixerNumOutputFrames;
 			}
 
 			METASOUNDFRONTEND_API const FMetasoundFrontendVersion& GetVersion();
diff --git a/Engine/Source/Runtime/AudioMixer/Private/AudioMixerSource.cpp b/Engine/Source/Runtime/AudioMixer/Private/AudioMixerSource.cpp
--- a/Engine/Source/Runtime/AudioMixer/Private/AudioMixerSource.cpp
+++ b/Engine/Source/Runtime/AudioMixer/Private/AudioMixerSource.cpp
@@ -868,6 +868,7 @@ namespace Audio
 		BufferInitArgs.AudioDeviceID = AudioDevice->DeviceID;
 		BufferInitArgs.InstanceID = InstanceID;
 		BufferInitArgs.SampleRate = AudioDevice->GetSampleRate();
+		BufferInitArgs.AudioMixerNumOutputFrames = MixerDevice->GetNumOutputFrames();
 		BufferInitArgs.Buffer = MixerBuffer;
 		BufferInitArgs.SoundWave = &SoundWave;
 		BufferInitArgs.LoopingMode = InWaveInstance->LoopingMode;
diff --git a/Engine/Source/Runtime/AudioMixer/Private/AudioMixerSourceBuffer.cpp b/Engine/Source/Runtime/AudioMixer/Private/AudioMixerSourceBuffer.cpp
--- a/Engine/Source/Runtime/AudioMixer/Private/AudioMixerSourceBuffer.cpp
+++ b/Engine/Source/Runtime/AudioMixer/Private/AudioMixerSourceBuffer.cpp
@@ -95,6 +95,7 @@ namespace Audio
 			FSoundGeneratorInitParams InitParams;
 			InitParams.AudioDeviceID = InArgs.AudioDeviceID;
 			InitParams.SampleRate = InArgs.SampleRate;
+			InitParams.AudioMixerNumOutputFrames = InArgs.AudioMixerNumOutputFrames;
 			InitParams.NumChannels = NumChannels;
 			InitParams.NumFramesPerCallback = MONO_PCM_BUFFER_SAMPLES;
 			InitParams.InstanceID = InArgs.InstanceID;
diff --git a/Engine/Source/Runtime/AudioMixer/Private/AudioMixerSourceBuffer.h b/Engine/Source/Runtime/AudioMixer/Private/AudioMixerSourceBuffer.h
--- a/Engine/Source/Runtime/AudioMixer/Private/AudioMixerSourceBuffer.h
+++ b/Engine/Source/Runtime/AudioMixer/Private/AudioMixerSourceBuffer.h
@@ -51,6 +51,7 @@ namespace Audio
 		FDeviceId AudioDeviceID = 0;
 		uint32 InstanceID = 0;
 		int32 SampleRate = 0;
+		int32 AudioMixerNumOutputFrames = 0;
 		FMixerBuffer* Buffer = nullptr;
 		USoundWave* SoundWave = nullptr;
 		ELoopingMode LoopingMode = ELoopingMode::LOOP_Never;
diff --git a/Engine/Source/Runtime/Engine/Classes/Sound/SoundGenerator.h b/Engine/Source/Runtime/Engine/Classes/Sound/SoundGenerator.h
--- a/Engine/Source/Runtime/Engine/Classes/Sound/SoundGenerator.h
+++ b/Engine/Source/Runtime/Engine/Classes/Sound/SoundGenerator.h
@@ -17,6 +17,7 @@ struct FSoundGeneratorInitParams
 {
 	Audio::FDeviceId AudioDeviceID;
 	float SampleRate = 0.0f;
+	int32 AudioMixerNumOutputFrames = 0;
 	int32 NumChannels = 0;
 	int32 NumFramesPerCallback = 0;
 	uint64 InstanceID = 0;
