CIS Build fix for MetaSound source interface

diff --git a/Engine/Plugins/Runtime/Metasound/Source/MetasoundFrontend/Private/MetasoundSourceInterface.cpp b/Engine/Plugins/Runtime/Metasound/Source/MetasoundFrontend/Private/MetasoundSourceInterface.cpp
--- a/Engine/Plugins/Runtime/Metasound/Source/MetasoundFrontend/Private/MetasoundSourceInterface.cpp
+++ b/Engine/Plugins/Runtime/Metasound/Source/MetasoundFrontend/Private/MetasoundSourceInterface.cpp
@@ -230,7 +230,7 @@ namespace Metasound
 								LOCTEXT("AudioMixerOutputFramesDisplayName", "Audio Mixer Output Frames"),
 								LOCTEXT("AudioMixerOutputFramesDescription", "The number of output frames rendered per buffer in the audio mixer."),
 								FName(), 
-								Environment::AudioMixerOutputFrames
+								Environment::AudioMixerNumOutputFrames
 							}
 						};
 					}
