CIS build fix. Removing uneeded module dependency which was breaking AudioUnitTest builds

diff --git a/Engine/Plugins/Runtime/Metasound/Source/MetasoundStandardNodes/MetasoundStandardNodes.Build.cs b/Engine/Plugins/Runtime/Metasound/Source/MetasoundStandardNodes/MetasoundStandardNodes.Build.cs
--- a/Engine/Plugins/Runtime/Metasound/Source/MetasoundStandardNodes/MetasoundStandardNodes.Build.cs
+++ b/Engine/Plugins/Runtime/Metasound/Source/MetasoundStandardNodes/MetasoundStandardNodes.Build.cs
@@ -26,8 +26,7 @@ namespace UnrealBuildTool.Rules
 				new string[]
 				{
 					"CoreUObject",
-					"MetasoundGraphCore",
-					"WaveTable"
+					"MetasoundGraphCore"
 				}
 			);
 		}
