MetaSound audio bus reader node

diff --git a/Engine/Plugins/Runtime/Metasound/Source/MetasoundEngine/Private/MetasoundAudioBusReaderNode.cpp b/Engine/Plugins/Runtime/Metasound/Source/MetasoundEngine/Private/MetasoundAudioBusReaderNode.cpp
new file mode 100644
--- /dev/null
+++ b/Engine/Plugins/Runtime/Metasound/Source/MetasoundEngine/Private/MetasoundAudioBusReaderNode.cpp
@@ -0,0 +1,232 @@
+// Copyright Epic Games, Inc. All Rights Reserved.
+
+#include "AudioDevice.h"
+#include "Internationalization/Text.h"
+#include "MetasoundExecutableOperator.h"
+#include "MetasoundEngineNodesNames.h"
+#include "MetasoundNodeRegistrationMacro.h"
+#include "MetasoundParamHelper.h"
+#include "DSP/ConvertDeinterleave.h"
+#include "MetasoundAudioBuffer.h"
+#include "MetasoundAudioBus.h"
+#include "MetasoundStandardNodesCategories.h"
+#include "MetasoundFacade.h"
+#include "AVEncoder/Public/MediaPacket.h"
+
+#define LOCTEXT_NAMESPACE "MetasoundAudioBusNode"
+
+namespace Metasound
+{
+	namespace AudioBusReaderNode
+	{
+		METASOUND_PARAM(InParamAudioBusInput, "Audio Bus", "Audio Bus Asset.")
+
+		METASOUND_PARAM(OutParamAudio, "Out {0}", "Audio bus output for channel {0}.");
+	}
+
+	template <uint32 NumChannels>
+	class TAudioBusReaderOperator : public TExecutableOperator<TAudioBusReaderOperator<NumChannels>>
+	{
+	public:
+
+		static const FNodeClassMetadata& GetNodeInfo()
+		{
+			auto InitNodeInfo = []() -> FNodeClassMetadata
+			{
+				FName OperatorName = *FString::Printf(TEXT("Audio Bus Reader (%d)"), NumChannels);
+				FText NodeDisplayName = METASOUND_LOCTEXT_FORMAT("AudioBusReaderDisplayNamePattern", "Audio Bus Reader ({0})", NumChannels);
+				
+				FNodeClassMetadata Info;
+				Info.ClassName = { EngineNodes::Namespace, OperatorName, TEXT("") };
+				Info.MajorVersion = 1;
+				Info.MinorVersion = 0;
+				Info.DisplayName = NodeDisplayName;
+				Info.Description = METASOUND_LOCTEXT("AudioBusReader_Description", "Outputs audio data from the audio bus asset.");
+				Info.Author = PluginAuthor;
+				Info.PromptIfMissing = PluginNodeMissingPrompt;
+				Info.DefaultInterface = GetVertexInterface();
+				Info.CategoryHierarchy.Emplace(NodeCategories::Io);
+				return Info;
+			};
+
+			static const FNodeClassMetadata Info = InitNodeInfo();
+
+			return Info;
+		}
+		
+		static const FVertexInterface& GetVertexInterface()
+		{
+			using namespace AudioBusReaderNode;
+
+			auto CreateVertexInterface = []() -> FVertexInterface
+			{
+				FInputVertexInterface InputInterface;
+				InputInterface.Add(TInputDataVertex<FAudioBusAsset>(METASOUND_GET_PARAM_NAME_AND_METADATA(InParamAudioBusInput)));
+
+				FOutputVertexInterface OutputInterface;
+				for (uint32 i = 0; i < NumChannels; ++i)
+				{
+					OutputInterface.Add(TOutputDataVertex<FAudioBuffer>(METASOUND_GET_PARAM_NAME_WITH_INDEX_AND_METADATA(OutParamAudio, i)));
+				}
+
+				return FVertexInterface(InputInterface, OutputInterface);
+			};
+			
+			static const FVertexInterface Interface = CreateVertexInterface();
+			return Interface;
+		}
+
+		static TUniquePtr<IOperator> CreateOperator(const FCreateOperatorParams& InParams, FBuildErrorArray& OutErrors)
+		{
+			using namespace Frontend;
+			
+			using namespace AudioBusReaderNode; 
+			const FDataReferenceCollection& InputCollection = InParams.InputDataReferences;
+
+			Audio::FDeviceId AudioDeviceId = InParams.Environment.GetValue<Audio::FDeviceId>(SourceInterface::Environment::DeviceID);
+			int32 AudioMixerOutputFrames = InParams.Environment.GetValue<int32>(SourceInterface::Environment::AudioMixerNumOutputFrames);
+
+			FAudioBusAssetReadRef AudioBusIn = InputCollection.GetDataReadReferenceOrConstruct<FAudioBusAsset>(METASOUND_GET_PARAM_NAME(InParamAudioBusInput));
+			
+			return MakeUnique<TAudioBusReaderOperator<NumChannels>>(InParams.OperatorSettings, AudioMixerOutputFrames, AudioDeviceId, AudioBusIn);
+		}
+
+		TAudioBusReaderOperator(const FOperatorSettings& InSettings, int32 InAudioMixerOutputFrames, Audio::FDeviceId InAudioDeviceId, const FAudioBusAssetReadRef& InAudioBusAsset)
+			: AudioBusAsset(InAudioBusAsset)
+			, AudioMixerOutputFrames(InAudioMixerOutputFrames)
+			, AudioDeviceId(InAudioDeviceId)
+			, SampleRate(InSettings.GetSampleRate())
+		{
+			for (int32 ChannelIndex = 0; ChannelIndex < NumChannels; ++ChannelIndex)
+			{
+				AudioOutputs.Add(FAudioBufferWriteRef::CreateNew(InSettings));
+			}
+			
+			const FAudioBusProxyPtr& AudioBusProxy = AudioBusAsset->GetAudioBusProxy();
+
+			if (FAudioDeviceManager* ADM = FAudioDeviceManager::Get())
+			{
+				if (FAudioDevice* AudioDevice = ADM->GetAudioDeviceRaw(AudioDeviceId))
+				{
+					// Start the audio bus in case it's not already started
+					AudioBusChannels = AudioBusProxy->NumChannels;
+					AudioDevice->StartAudioBus(AudioBusProxy->AudioBusId, AudioBusChannels, false);
+					
+					AudioBusPatchOutput = AudioDevice->AddPatchForAudioBus(AudioBusProxy->AudioBusId);
+				}
+			}
+		}
+
+		virtual ~TAudioBusReaderOperator()
+		{
+		}
+
+		virtual FDataReferenceCollection GetInputs() const override
+		{
+			using namespace AudioBusReaderNode;
+
+			FDataReferenceCollection InputDataReferences;
+			InputDataReferences.AddDataReadReference(METASOUND_GET_PARAM_NAME(InParamAudioBusInput), FAudioBusAssetReadRef(AudioBusAsset));
+			return InputDataReferences;
+		}
+
+		virtual FDataReferenceCollection GetOutputs() const override
+		{
+			using namespace AudioBusReaderNode;
+
+			FDataReferenceCollection OutputDataReferences;
+
+			for (int32 ChannelIndex = 0; ChannelIndex < NumChannels; ++ChannelIndex)
+			{
+				OutputDataReferences.AddDataReadReference(METASOUND_GET_PARAM_NAME_WITH_INDEX(OutParamAudio, ChannelIndex), FAudioBufferWriteRef(AudioOutputs[ChannelIndex]));
+			}
+			
+			return OutputDataReferences;			
+		}
+
+		void Execute()
+		{
+			const int32 BlockSizeFrames = AudioOutputs[0]->Num();
+			const int32 NumSamplesToPop = BlockSizeFrames * AudioBusChannels;
+
+			// We want to write out zeroed audio until there is enough audio queued up in the patch output to avoid any underruns due to thread timing.
+			// From testing, waiting for the first buffer for 3 times the size of a single metasound block was sufficient to avoid pops. 
+			// Once we have started popping audio off the patch output, we want to keep it popping. Unless there are CPU underruns, there should be enough
+			// runway to keep it from underruning. We employ a little bit of a wait with a time out just to decrease the chance that there will be any missed buffers.
+			// In practice, this wait should rarely, if ever, actually cause a wait in this metasound execute.
+
+			bool bPerformPop = false;
+			if (bFirstBlock)
+			{
+				// Tuned amount to wait for first rendered block to account for any timing issues between the audio render thread and the metasound task
+				constexpr int32 FirstBlockBufferCount = 3;
+				if (AudioBusPatchOutput->GetNumSamplesAvailable() > FirstBlockBufferCount * NumSamplesToPop)
+				{
+					bFirstBlock = false;
+					bPerformPop = true;
+
+					InterleavedBuffer.Reset();
+					InterleavedBuffer.AddUninitialized(NumSamplesToPop);
+				}
+			}
+			else if (AudioBusPatchOutput->WaitUntilNumSamplesAvailable(NumSamplesToPop, static_cast<uint32>(0.5f * BlockSizeFrames * SampleRate)))
+			{
+				bPerformPop = true;
+			}
+
+			if (bPerformPop)
+			{
+				// Pop off the interleaved data from the audio bus
+				AudioBusPatchOutput->PopAudio(InterleavedBuffer.GetData(), NumSamplesToPop, false);
+
+				const uint32 MinChannels = FMath::Min(NumChannels, AudioBusChannels);
+				for (uint32 ChannelIndex = 0; ChannelIndex < MinChannels; ++ChannelIndex)
+				{
+					float* AudioOutputBufferPtr = AudioOutputs[ChannelIndex]->GetData();
+					for (int32 FrameIndex = 0; FrameIndex < BlockSizeFrames; ++FrameIndex)
+					{
+						AudioOutputBufferPtr[FrameIndex] = InterleavedBuffer[FrameIndex * AudioBusChannels + ChannelIndex];
+					}
+				}
+			}
+			else
+			{
+				UE_CLOG(!bFirstBlock, LogMetaSound, Warning, TEXT("Underrun detected in audio bus reader node."));
+			}
+		}
+
+	private:
+		FAudioBusAssetReadRef AudioBusAsset;
+		TArray<FAudioBufferWriteRef> AudioOutputs;
+
+		TArray<float> InterleavedBuffer;
+		int32 AudioMixerOutputFrames = INDEX_NONE;
+		Audio::FDeviceId AudioDeviceId = INDEX_NONE;
+		float SampleRate = 0.0f;
+		Audio::FPatchOutputStrongPtr AudioBusPatchOutput;
+		TUniquePtr<Audio::IConvertDeinterleave> ConvertDeinterleave;
+		Audio::FMultichannelBuffer DeinterleavedBuffer;
+		uint32 AudioBusChannels = INDEX_NONE;
+		bool bFirstBlock = true;
+	};
+
+	template<uint32 NumChannels>
+	class TAudioBusReaderNode : public FNodeFacade
+	{
+	public:
+		TAudioBusReaderNode(const FNodeInitData& InitData)
+			: FNodeFacade(InitData.InstanceName, InitData.InstanceID, TFacadeOperatorClass<TAudioBusReaderOperator<NumChannels>>())
+		{
+		}
+	};
+
+#define REGISTER_AUDIO_BUS_READER_NODE(ChannelCount) \
+	using FTriggerSequenceNode_##ChannelCount = TAudioBusReaderNode<ChannelCount>; \
+	METASOUND_REGISTER_NODE(FTriggerSequenceNode_##ChannelCount) \
+
+	
+	REGISTER_AUDIO_BUS_READER_NODE(1)
+	REGISTER_AUDIO_BUS_READER_NODE(2)
+}
+
+#undef LOCTEXT_NAMESPACE
