Do not call "pure virtual" function for error case
#rnx

diff --git a/Engine/Plugins/Runtime/Metasound/Source/MetasoundEditor/Private/MetasoundEditorGraph.cpp b/Engine/Plugins/Runtime/Metasound/Source/MetasoundEditor/Private/MetasoundEditorGraph.cpp
--- a/Engine/Plugins/Runtime/Metasound/Source/MetasoundEditor/Private/MetasoundEditorGraph.cpp
+++ b/Engine/Plugins/Runtime/Metasound/Source/MetasoundEditor/Private/MetasoundEditorGraph.cpp
@@ -503,7 +503,7 @@ Metasound::Frontend::FNodeHandle UMetasoundEditorGraphInput::AddNodeHandle(const
 	UMetasoundEditorGraph* Graph = Cast<UMetasoundEditorGraph>(GetOuter());
 	if (!ensure(Graph))
 	{
-		return UMetasoundEditorGraphVertex::AddNodeHandle(InName, InDataType);
+		return Metasound::Frontend::INodeController::GetInvalidHandle();
 	}
 
 	UObject& Metasound = Graph->GetMetasoundChecked();
@@ -648,7 +648,7 @@ Metasound::Frontend::FNodeHandle UMetasoundEditorGraphOutput::AddNodeHandle(cons
 	UMetasoundEditorGraph* Graph = Cast<UMetasoundEditorGraph>(GetOuter());
 	if (!ensure(Graph))
 	{
-		return UMetasoundEditorGraphVertex::AddNodeHandle(InName, InDataType);
+		return Metasound::Frontend::INodeController::GetInvalidHandle();
 	}
 
 	UObject& Metasound = Graph->GetMetasoundChecked();
