WaveTable envelope, getter, & oscillator nodes
#rnx

diff --git a/Engine/Plugins/Runtime/Metasound/Source/MetasoundEngine/Private/MetasoundWaveTableEnvelopeNode.cpp b/Engine/Plugins/Runtime/Metasound/Source/MetasoundEngine/Private/MetasoundWaveTableEnvelopeNode.cpp
new file mode 100644
--- /dev/null
+++ b/Engine/Plugins/Runtime/Metasound/Source/MetasoundEngine/Private/MetasoundWaveTableEnvelopeNode.cpp
@@ -0,0 +1,358 @@
+// Copyright Epic Games, Inc. All Rights Reserved.
+
+#include "MetasoundDataFactory.h"
+#include "MetasoundEngineNodesNames.h"
+#include "MetasoundEnumRegistrationMacro.h"
+#include "MetasoundExecutableOperator.h"
+#include "MetasoundFacade.h"
+#include "MetasoundNodeRegistrationMacro.h"
+#include "MetasoundPrimitives.h"
+#include "MetasoundSourceInterface.h"
+#include "MetasoundStandardNodesNames.h"
+#include "MetasoundTime.h"
+#include "MetasoundTrace.h"
+#include "MetasoundTrigger.h"
+#include "MetasoundWaveTable.h"
+#include "WaveTableSampler.h"
+
+#define LOCTEXT_NAMESPACE "MetasoundStandardNodes"
+
+
+namespace Metasound
+{
+	DECLARE_METASOUND_ENUM(WaveTable::FWaveTableSampler::ESingleSampleMode, WaveTable::FWaveTableSampler::ESingleSampleMode::Zero, METASOUNDENGINE_API,
+	FEnumWaveTableEnvelopeMode, FEnumWaveTableEnvelopeModeTypeInfo, FEnumWaveTableEnvelopeModeReadRef, FEnumWaveTableEnvelopeModeWriteRef);
+
+	DECLARE_METASOUND_ENUM(WaveTable::FWaveTableSampler::EInterpolationMode, WaveTable::FWaveTableSampler::EInterpolationMode::Linear, METASOUNDENGINE_API,
+	FEnumWaveTableInterpolationMode, FEnumWaveTableInterpModeTypeInfo, FEnumWaveTableInterpModeReadRef, FEnumWaveTableInterpModeWriteRef);
+
+	DEFINE_METASOUND_ENUM_BEGIN(WaveTable::FWaveTableSampler::ESingleSampleMode, FEnumWaveTableEnvelopeMode, "WaveTableEnvelopeMode")
+		DEFINE_METASOUND_ENUM_ENTRY(WaveTable::FWaveTableSampler::ESingleSampleMode::Loop, "LoopDisplayName", "Loop", "EnvMode_LoopDescription", "Interpolates last value and first value in WaveTable, starting over interpolation of envelope on completion."),
+		DEFINE_METASOUND_ENUM_ENTRY(WaveTable::FWaveTableSampler::ESingleSampleMode::Hold, "HoldDisplayName", "Hold", "EnvMode_HoldDescription", "Holds last value in table if elapsed beyond WaveTable length"),
+		DEFINE_METASOUND_ENUM_ENTRY(WaveTable::FWaveTableSampler::ESingleSampleMode::Unit, "UnitDisplayName", "Unit", "EnvMode_UnitDescription", "Interpolates last value in table with unit (1.0f) if elapsed beyond WaveTable length."),
+		DEFINE_METASOUND_ENUM_ENTRY(WaveTable::FWaveTableSampler::ESingleSampleMode::Zero, "ZeroDisplayName", "Zero", "EnvMode_ZeroDescription", "Interpolates last value in table with zero (0.0f) if elapsed beyond WaveTable length"),
+	DEFINE_METASOUND_ENUM_END()
+
+	DEFINE_METASOUND_ENUM_BEGIN(WaveTable::FWaveTableSampler::EInterpolationMode, FEnumWaveTableInterpolationMode, "WaveTableInterpolation")
+		DEFINE_METASOUND_ENUM_ENTRY(WaveTable::FWaveTableSampler::EInterpolationMode::None, "InterpModeDisplayName", "None (Step)", "EnvMode_InterpDescription", "No interpolation between values (uses lowest)."),
+		DEFINE_METASOUND_ENUM_ENTRY(WaveTable::FWaveTableSampler::EInterpolationMode::Linear, "InterpModeName", "Linear", "EnvMode_InterpDescription", "Linearly interpolates between values."),
+		DEFINE_METASOUND_ENUM_ENTRY(WaveTable::FWaveTableSampler::EInterpolationMode::Cubic, "InterpModeName", "Cubic", "EnvMode_InterpDescription", "Cubically interpolates between values.")
+	DEFINE_METASOUND_ENUM_END()
+
+	class FMetasoundWaveTableEnvelopeNodeOperator : public TExecutableOperator<FMetasoundWaveTableEnvelopeNodeOperator>
+	{
+	public:
+		static const FVertexInterface& GetDefaultInterface()
+		{
+			using namespace WaveTable;
+
+			static const FVertexInterface DefaultInterface(
+				FInputVertexInterface(
+					TInputDataVertex<FWaveTable>("WaveTable", FDataVertexMetadata { LOCTEXT("MetasoundWaveTableEnvelopeNode_InputWaveTable", "WaveTable") }),
+					TInputDataVertex<FTrigger>("Play", FDataVertexMetadata { LOCTEXT("MetasoundWaveTableEnvelopeNode_InputOnPlay", "Play") }),
+					TInputDataVertex<FTrigger>("Stop", FDataVertexMetadata { LOCTEXT("MetasoundWaveTableEnvelopeNode_InputOnStop", "Stop") }),
+					TInputDataVertex<FTrigger>("Pause", FDataVertexMetadata { LOCTEXT("MetasoundWaveTableEnvelopeNode_InputOnPause", "Pause") }),
+					TInputDataVertex<FTime>("Duration", FDataVertexMetadata { LOCTEXT("MetasoundWaveTableEnvelopeNode_Duration", "Duration") }, 1.0f),
+					TInputDataVertex<FEnumWaveTableEnvelopeMode>("Mode", FDataVertexMetadata
+					{
+						LOCTEXT("MetasoundWaveTableEnvelopeNode_ModeDescription", "What value the envelope completes on (or whether it loops)."),
+						LOCTEXT("MetasoundWaveTableEnvelopeNode_Mode", "Mode"),
+						true /* bIsAdvancedDisplay */
+					}, static_cast<int32>(WaveTable::FWaveTableSampler::ESingleSampleMode::Zero)),
+					TInputDataVertex<FEnumWaveTableInterpolationMode>("Interpolation", FDataVertexMetadata
+					{
+						LOCTEXT("MetasoundWaveTableEnvelopeNode_InterpDescription", "How the envelope interpolates between WaveTable values."),
+						LOCTEXT("MetasoundWaveTableEnvelopeNode_Interp", "Interpolation"),
+						true /* bIsAdvancedDisplay */
+					}, static_cast<int32>(FWaveTableSampler::EInterpolationMode::Linear))
+				),
+				FOutputVertexInterface(
+					TOutputDataVertex<FTrigger>("OnFinished", FDataVertexMetadata{ LOCTEXT("MetasoundWaveTableEnvelopeNode_OnFinished", "OnFinished") }),
+					TOutputDataVertex<float>("Out", FDataVertexMetadata { LOCTEXT("MetasoundWaveTableEnvelopeNode_Output", "Out") })
+				)
+			);
+
+			return DefaultInterface;
+		}
+
+		static const FNodeClassMetadata& GetNodeInfo()
+		{
+			auto CreateNodeClassMetadata = []() -> FNodeClassMetadata
+			{
+				FNodeClassMetadata Metadata
+				{
+					{ EngineNodes::Namespace, "WaveTableEnvelope", "" },
+					1, // Major Version
+					0, // Minor Version
+					LOCTEXT("MetasoundWaveTableEnvelopeNode_Name", "WaveTable Envelope"),
+					LOCTEXT("MetasoundWaveTableEnvelopeNode_Description", "Reads through the given WaveTable at the provided frequency."),
+					PluginAuthor,
+					PluginNodeMissingPrompt,
+					GetDefaultInterface(),
+					{ },
+					{ },
+					{ }
+				};
+
+				return Metadata;
+			};
+
+			static const FNodeClassMetadata Metadata = CreateNodeClassMetadata();
+			return Metadata;
+		}
+
+		static TUniquePtr<IOperator> CreateOperator(const FCreateOperatorParams& InParams, TArray<TUniquePtr<IOperatorBuildError>>& OutErrors)
+		{
+			using namespace WaveTable;
+
+			const FInputVertexInterface& InputInterface = InParams.Node.GetVertexInterface().GetInputInterface();
+			const FDataReferenceCollection& InputCollection = InParams.InputDataReferences;
+
+			FWaveTableReadRef InWaveTableReadRef = InputCollection.GetDataReadReferenceOrConstruct<FWaveTable>("WaveTable");
+			FTriggerReadRef InPlayReadRef = InputCollection.GetDataReadReferenceOrConstructWithVertexDefault<FTrigger>(InputInterface, "Play", InParams.OperatorSettings);
+			FTriggerReadRef InStopReadRef = InputCollection.GetDataReadReferenceOrConstructWithVertexDefault<FTrigger>(InputInterface, "Stop", InParams.OperatorSettings);
+			FTriggerReadRef InPauseReadRef = InputCollection.GetDataReadReferenceOrConstructWithVertexDefault<FTrigger>(InputInterface, "Pause", InParams.OperatorSettings);
+			FTimeReadRef InDurationReadRef = InputCollection.GetDataReadReferenceOrConstructWithVertexDefault<FTime>(InputInterface, "Duration", InParams.OperatorSettings);
+			FEnumWaveTableEnvelopeModeReadRef InModeReadRef = InputCollection.GetDataReadReferenceOrConstructWithVertexDefault<FEnumWaveTableEnvelopeMode>( InputInterface, "Mode", InParams.OperatorSettings);
+			FEnumWaveTableInterpModeReadRef InInterpReadRef = InputCollection.GetDataReadReferenceOrConstructWithVertexDefault<FEnumWaveTableInterpolationMode>(InputInterface, "Interpolation", InParams.OperatorSettings);
+
+			return MakeUnique<FMetasoundWaveTableEnvelopeNodeOperator>(InParams, InWaveTableReadRef, InPlayReadRef, InStopReadRef, InPauseReadRef, InDurationReadRef, InModeReadRef, InInterpReadRef);
+		}
+
+		FMetasoundWaveTableEnvelopeNodeOperator(
+			const FCreateOperatorParams& InParams,
+			const FWaveTableReadRef& InWaveTableReadRef,
+			const FTriggerReadRef& InPlayReadRef,
+			const FTriggerReadRef& InStopReadRef,
+			const FTriggerReadRef& InPauseReadRef,
+			const FTimeReadRef& InDurationReadRef,
+			const FEnumWaveTableEnvelopeModeReadRef& InModeReadRef,
+			const FEnumWaveTableInterpModeReadRef& InInterpModeReadRef)
+			: WaveTableReadRef(InWaveTableReadRef)
+			, PlayReadRef(InPlayReadRef)
+			, StopReadRef(InStopReadRef)
+			, PauseReadRef(InPauseReadRef)
+			, DurationReadRef(InDurationReadRef)
+			, ModeReadRef(InModeReadRef)
+			, InterpModeReadRef(InInterpModeReadRef)
+			, OnFinishedWriteRef(TDataWriteReference<FTrigger>::CreateNew(InParams.OperatorSettings))
+			, OutWriteRef(TDataWriteReferenceFactory<float>::CreateAny(InParams.OperatorSettings))
+			, SampleRate(InParams.OperatorSettings.GetSampleRate())
+		{
+			{
+				using namespace WaveTable;
+				FWaveTableSampler::FSettings Settings;
+				Settings.Freq = 0.0f; // Sampler phase is manually progressed via this node
+				Sampler = FWaveTableSampler(MoveTemp(Settings));
+			}
+
+			{
+				const float BlockRate = InParams.OperatorSettings.GetActualBlockRate();
+				check(BlockRate > 0.0f && !FMath::IsNearlyZero(BlockRate));
+				SecondsPerBlock = 1.0f / BlockRate;
+			}
+
+			{
+				BlockSize = InParams.OperatorSettings.GetNumFramesPerBlock();
+				check(BlockSize > 0);
+			}
+		}
+
+		virtual ~FMetasoundWaveTableEnvelopeNodeOperator() = default;
+
+		virtual FDataReferenceCollection GetInputs() const override
+		{
+			FDataReferenceCollection Inputs;
+			Inputs.AddDataReadReference("WaveTable", WaveTableReadRef);
+			Inputs.AddDataReadReference("Play", PlayReadRef);
+			Inputs.AddDataReadReference("Stop", StopReadRef);
+			Inputs.AddDataReadReference("Pause", PauseReadRef);
+			Inputs.AddDataReadReference("Duration", DurationReadRef);
+			Inputs.AddDataReadReference("Mode", ModeReadRef);
+
+
+			return Inputs;
+		}
+
+		virtual FDataReferenceCollection GetOutputs() const override
+		{
+			FDataReferenceCollection Outputs;
+			Outputs.AddDataReadReference("OnFinished", TDataReadReference<FTrigger>(OnFinishedWriteRef));
+			Outputs.AddDataReadReference("Out", TDataReadReference<float>(OutWriteRef));
+
+			return Outputs;
+		}
+
+		float GetStopValue(WaveTable::FWaveTableSampler::ESingleSampleMode InSampleMode) const
+		{
+			using namespace WaveTable;
+
+			float LastValue = *OutWriteRef;
+
+			if (bPaused)
+			{
+				return LastValue;
+			}
+
+			switch (InSampleMode)
+			{
+				case FWaveTableSampler::ESingleSampleMode::Unit:
+				return 1.0f;
+
+				case FWaveTableSampler::ESingleSampleMode::Hold:
+				return LastValue;
+
+				case FWaveTableSampler::ESingleSampleMode::Zero:
+				case FWaveTableSampler::ESingleSampleMode::Loop:
+				default:
+				return 0.0f;
+			}
+		}
+
+		void Execute()
+		{
+			METASOUND_TRACE_CPUPROFILER_EVENT_SCOPE(FMetasoundWaveTableEnvelopeNodeOperator::Execute);
+
+			// To avoid unexpected behavior of only firing a single "OnFinished" per block if set to looping,
+			// limit the duration to a minimum threshold of SecondsPerBlock.
+			const float Dur = FMath::Max(DurationReadRef->GetSeconds(), SecondsPerBlock);
+
+			WaveTable::FWaveTableSampler::ESingleSampleMode Mode = *ModeReadRef;
+			WaveTable::FWaveTableSampler::EInterpolationMode InterpMode = *InterpModeReadRef;
+
+			OnFinishedWriteRef->AdvanceBlock();
+
+			float NextValue = GetStopValue(Mode);
+
+			auto GetLastIndex = [](const FTriggerReadRef& Trigger)
+			{
+				int32 LastIndex = -1;
+				Trigger->ExecuteBlock([](int32, int32) {}, [&LastIndex](int32 StartFrame, int32 EndFrame)
+				{
+					LastIndex = FMath::Max(LastIndex, StartFrame);
+				});
+				return LastIndex;
+			};
+
+			const int32 LastPlayIndex = GetLastIndex(PlayReadRef);
+			const int32 LastStopIndex = GetLastIndex(StopReadRef);
+			const int32 LastPauseIndex = GetLastIndex(PauseReadRef);
+
+			if (LastPlayIndex >= 0 || LastStopIndex >= 0)
+			{
+				if (LastPlayIndex > LastStopIndex)
+				{
+					Elapsed = 0.0f;
+					bPaused = LastPauseIndex > LastPlayIndex;
+				}
+				else // Stop wins if same frame
+				{
+					if (Elapsed >= 0.0f)
+					{
+						Elapsed = -1.0f;
+						OnFinishedWriteRef->TriggerFrame(LastStopIndex);
+					}
+					bPaused = false;
+				}
+			}
+			else if (LastPauseIndex >= 0)
+			{
+				if (Elapsed >= 0.0f)
+				{
+					bPaused = !bPaused;
+					if (bPaused)
+					{
+						NextValue = *OutWriteRef;
+					}
+				}
+			}
+
+			TArrayView<const float> WaveTableView = WaveTableReadRef->GetView();
+
+			if (WaveTableView.IsEmpty())
+			{
+				if (Elapsed >= 0.0f)
+				{
+					Elapsed = -1.0f;
+					OnFinishedWriteRef->TriggerFrame(0);
+					bPaused = false;
+				}
+			}
+			else
+			{
+				if (Elapsed >= 0.0f && !bPaused)
+				{
+					if (Dur > 0.0f)
+					{
+						Sampler.Reset();
+						Sampler.SetInterpolationMode(*InterpModeReadRef);
+						Sampler.SetPhase(FMath::Clamp(Elapsed / Dur, 0.0f, 1.0f));
+						Sampler.Process(WaveTableView, NextValue, Mode);
+					}
+
+					Elapsed += SecondsPerBlock;
+				}
+			}
+
+			if (Elapsed > Dur)
+			{
+				int32 FinishSampleIndex = 0;
+				if (Mode == WaveTable::FWaveTableSampler::ESingleSampleMode::Loop)
+				{
+					while (Elapsed >= Dur)
+					{
+						Elapsed -= Dur;
+					}
+					FinishSampleIndex = Elapsed * SampleRate;
+				}
+				else
+				{
+					FinishSampleIndex = (Elapsed - Dur) * SampleRate;
+					Elapsed = -1.0f;
+				}
+
+				FinishSampleIndex %= BlockSize;
+				OnFinishedWriteRef->TriggerFrame(FinishSampleIndex);
+			}
+
+			*OutWriteRef = NextValue;
+		}
+
+	private:
+		FWaveTableReadRef WaveTableReadRef;
+		FTriggerReadRef PlayReadRef;
+		FTriggerReadRef StopReadRef;
+		FTriggerReadRef PauseReadRef;
+		FTimeReadRef DurationReadRef;
+		FEnumWaveTableEnvelopeModeReadRef ModeReadRef;
+		FEnumWaveTableInterpModeReadRef InterpModeReadRef;
+
+		WaveTable::FWaveTableSampler Sampler;
+
+		FTriggerWriteRef OnFinishedWriteRef;
+		FFloatWriteRef OutWriteRef;
+
+		int32 BlockSize = 0;
+
+		float SecondsPerBlock = 0.0f;
+		float SampleRate = 0.0f;
+		float Elapsed = -1.0f;
+
+		bool bPaused = false;
+	};
+
+	class FMetasoundWaveTableEnvelopeNode : public FNodeFacade
+	{
+	public:
+		FMetasoundWaveTableEnvelopeNode(const FNodeInitData& InInitData)
+			: FNodeFacade(InInitData.InstanceName, InInitData.InstanceID, TFacadeOperatorClass<FMetasoundWaveTableEnvelopeNodeOperator>())
+		{
+		}
+
+		virtual ~FMetasoundWaveTableEnvelopeNode() = default;
+	};
+
+	METASOUND_REGISTER_NODE(FMetasoundWaveTableEnvelopeNode)
+} // namespace Metasound
+
+#undef LOCTEXT_NAMESPACE // MetasoundStandardNodes
diff --git a/Engine/Plugins/Runtime/Metasound/Source/MetasoundEngine/Private/MetasoundWaveTableGetNode.cpp b/Engine/Plugins/Runtime/Metasound/Source/MetasoundEngine/Private/MetasoundWaveTableGetNode.cpp
new file mode 100644
--- /dev/null
+++ b/Engine/Plugins/Runtime/Metasound/Source/MetasoundEngine/Private/MetasoundWaveTableGetNode.cpp
@@ -0,0 +1,243 @@
+// Copyright Epic Games, Inc. All Rights Reserved.
+
+#include "DSP/FloatArrayMath.h"
+#include "MetasoundDataFactory.h"
+#include "MetasoundEngineNodesNames.h"
+#include "MetasoundExecutableOperator.h"
+#include "MetasoundFacade.h"
+#include "MetasoundNodeRegistrationMacro.h"
+#include "MetasoundPrimitives.h"
+#include "MetasoundSourceInterface.h"
+#include "MetasoundTrace.h"
+#include "MetasoundWaveTable.h"
+#include "WaveTable.h"
+
+#define LOCTEXT_NAMESPACE "MetasoundEngine"
+
+
+namespace Metasound
+{
+	class FMetasoundWaveTableGetNodeOperator : public TExecutableOperator<FMetasoundWaveTableGetNodeOperator>
+	{
+	public:
+		static const FVertexInterface& GetDefaultInterface()
+		{
+			using namespace WaveTable;
+
+			static const FVertexInterface DefaultInterface(
+				FInputVertexInterface(
+					TInputDataVertex<FWaveTableBankAsset>("WaveTableBank", FDataVertexMetadata { LOCTEXT("MetasoundWaveTableGetNode_InputWaveTableBank", "WaveTableBank") }),
+					TInputDataVertex<float>("TableIndex", FDataVertexMetadata { LOCTEXT("MetasoundWaveTableGetNode_InputTableIndex", "Index") }, 0.0f)
+				),
+				FOutputVertexInterface(
+					TOutputDataVertex<FWaveTable>("Out", FDataVertexMetadata { LOCTEXT("MetasoundWaveTableGetNode_OutputWaveTable", "Out") })
+				)
+			);
+
+			return DefaultInterface;
+		}
+
+		static const FNodeClassMetadata& GetNodeInfo()
+		{
+			auto CreateNodeClassMetadata = []() -> FNodeClassMetadata
+			{
+				FNodeClassMetadata Metadata
+				{
+					FNodeClassName { EngineNodes::Namespace, "MetasoundWaveTableGet", "" },
+					1, // Major Version
+					0, // Minor Version
+					LOCTEXT("MetasoundWaveTableGetNode_Name", "Get WaveTable From Bank"),
+					LOCTEXT("MetasoundWaveTableGetNode_Description", "Gets or generates interpolated WaveTable from provided WaveTableBank asset."),
+					PluginAuthor,
+					PluginNodeMissingPrompt,
+					GetDefaultInterface(),
+					{ },
+					{ },
+					{ }
+				};
+
+				return Metadata;
+			};
+
+			static const FNodeClassMetadata Metadata = CreateNodeClassMetadata();
+			return Metadata;
+		}
+
+		static TUniquePtr<IOperator> CreateOperator(const FCreateOperatorParams& InParams, TArray<TUniquePtr<IOperatorBuildError>>& OutErrors)
+		{
+			using namespace Metasound;
+
+			const FInputVertexInterface& InputInterface = InParams.Node.GetVertexInterface().GetInputInterface();
+			const FDataReferenceCollection& InputCollection = InParams.InputDataReferences;
+
+			FWaveTableBankAssetReadRef WaveTableBankReadRef = InputCollection.GetDataReadReferenceOrConstruct<FWaveTableBankAsset>("WaveTableBank");
+			FFloatReadRef TableIndexReadRef = InputCollection.GetDataReadReferenceOrConstructWithVertexDefault<float>(InputInterface, "TableIndex", InParams.OperatorSettings);
+
+			return MakeUnique<FMetasoundWaveTableGetNodeOperator>(InParams, WaveTableBankReadRef, TableIndexReadRef);
+		}
+
+		FMetasoundWaveTableGetNodeOperator(
+			const FCreateOperatorParams& InParams,
+			const FWaveTableBankAssetReadRef& InWaveTableBankReadRef,
+			const FFloatReadRef& InTableIndexReadRef)
+			: WaveTableBankReadRef(InWaveTableBankReadRef)
+			, TableIndexReadRef(InTableIndexReadRef)
+			, OutTable(TDataWriteReferenceFactory<WaveTable::FWaveTable>::CreateAny(InParams.OperatorSettings))
+		{
+		}
+
+		virtual ~FMetasoundWaveTableGetNodeOperator() = default;
+
+		virtual FDataReferenceCollection GetInputs() const override
+		{
+			using namespace Metasound;
+
+			FDataReferenceCollection Inputs;
+			Inputs.AddDataReadReference("WaveTableBank", WaveTableBankReadRef);
+			Inputs.AddDataReadReference("TableIndex", TableIndexReadRef);
+
+			return Inputs;
+		}
+
+		virtual FDataReferenceCollection GetOutputs() const override
+		{
+			using namespace Metasound;
+
+			FDataReferenceCollection Outputs;
+			Outputs.AddDataReadReference("Out", FWaveTableWriteRef(OutTable));
+
+			return Outputs;
+		}
+
+		void Execute()
+		{
+			using namespace WaveTable;
+
+			METASOUND_TRACE_CPUPROFILER_EVENT_SCOPE(FMetasoundWaveTableGetNodeOperator::Execute);
+
+			const FWaveTableBankAsset& WaveTableBankAsset = *WaveTableBankReadRef;
+
+			FWaveTableBankAssetProxyPtr Proxy = WaveTableBankAsset.GetProxy();
+			if (Proxy.IsValid())
+			{
+				const TArray<FWaveTable>& WaveTables = Proxy->GetWaveTables();
+				if (!WaveTables.IsEmpty())
+				{
+					int32 NumSamples = WaveTables.Last().Num();
+					FWaveTable& OutputWaveTable = *OutTable;
+					OutputWaveTable.SetNum(NumSamples);
+					OutputWaveTable.Zero();
+
+					if (NumSamples > 0)
+					{
+						auto WrapIndex = [Max = WaveTables.Num()](float& InOutIndex)
+						{
+							InOutIndex = FMath::Abs(InOutIndex); // Avoids remainder offset flip at 0 crossing
+							const int32 WrapIndex = FMath::TruncToInt32(InOutIndex) % Max;
+							const float Remainder = FMath::Frac(InOutIndex);
+							InOutIndex = WrapIndex + Remainder;
+						};
+
+						float NextTableIndex = *TableIndexReadRef;
+						WrapIndex(NextTableIndex);
+
+						if (WaveTables.Num() == 1)
+						{
+							OutputWaveTable = WaveTables.Last();
+						}
+						else
+						{
+							if (LastTableIndex < 0.0f)
+							{
+								LastTableIndex = NextTableIndex;
+							}
+							else
+							{
+								WrapIndex(LastTableIndex);
+							}
+
+							const int32 StartHighIndex = FMath::CeilToInt32(LastTableIndex) % WaveTables.Num();
+							const int32 StartLowIndex = FMath::FloorToInt32(LastTableIndex);
+
+							const float StartHighGain = LastTableIndex - StartLowIndex;
+							const float StartLowGain = 1.0f - StartHighGain;
+
+							const int32 EndHighIndex = FMath::CeilToInt32(NextTableIndex) % WaveTables.Num();
+							const int32 EndLowIndex = FMath::FloorToInt32(NextTableIndex);
+
+							const float EndHighGain = NextTableIndex - EndLowIndex;
+							const float EndLowGain = 1.0f - EndHighGain;
+
+							for (int32 TableIndex = 0; TableIndex < WaveTables.Num(); ++TableIndex)
+							{
+								float Gain = 0.0f;
+								float GainDelta = 0.0f;
+
+								if (TableIndex == StartLowIndex)
+								{
+									Gain = StartLowGain;
+								}
+								else if (TableIndex == StartHighIndex)
+								{
+									Gain = StartHighGain;
+								}
+
+								if (TableIndex == EndLowIndex)
+								{
+									GainDelta = (EndLowGain - Gain) / NumSamples;
+								}
+								else if (TableIndex == EndHighIndex)
+								{
+									GainDelta = (EndHighGain - Gain) / NumSamples;
+								}
+
+								if (Gain > 0.0f || GainDelta > 0.0f)
+								{
+									const FWaveTable& InputWaveTable = WaveTables[TableIndex];
+									Audio::ArrayMixIn(InputWaveTable.GetView(), OutputWaveTable.GetView(), Gain);
+								}
+							}
+						}
+
+						LastTableIndex = NextTableIndex;
+					}
+					else
+					{
+						LastTableIndex = -1.0f;
+					}
+				}
+				else
+				{
+					LastTableIndex = -1.0f;
+				}
+			}
+			else
+			{
+				LastTableIndex = -1.0f;
+			}
+		}
+
+	private:
+		FWaveTableBankAssetReadRef WaveTableBankReadRef;
+		FFloatReadRef TableIndexReadRef;
+
+		float LastTableIndex = -1.0f;
+
+		FWaveTableWriteRef OutTable;
+	};
+
+	class FMetasoundWaveTableGetNode : public FNodeFacade
+	{
+	public:
+		FMetasoundWaveTableGetNode(const FNodeInitData& InInitData)
+			: FNodeFacade(InInitData.InstanceName, InInitData.InstanceID, TFacadeOperatorClass<FMetasoundWaveTableGetNodeOperator>())
+		{
+		}
+
+		virtual ~FMetasoundWaveTableGetNode() = default;
+	};
+
+	METASOUND_REGISTER_NODE(FMetasoundWaveTableGetNode)
+} // namespace Metasound
+
+#undef LOCTEXT_NAMESPACE
diff --git a/Engine/Plugins/Runtime/Metasound/Source/MetasoundEngine/Private/MetasoundWaveTableOscillatorNode.cpp b/Engine/Plugins/Runtime/Metasound/Source/MetasoundEngine/Private/MetasoundWaveTableOscillatorNode.cpp
new file mode 100644
--- /dev/null
+++ b/Engine/Plugins/Runtime/Metasound/Source/MetasoundEngine/Private/MetasoundWaveTableOscillatorNode.cpp
@@ -0,0 +1,240 @@
+// Copyright Epic Games, Inc. All Rights Reserved.
+
+#include "MetasoundDataFactory.h"
+#include "MetasoundEngineNodesNames.h"
+#include "MetasoundExecutableOperator.h"
+#include "MetasoundFacade.h"
+#include "MetasoundNodeRegistrationMacro.h"
+#include "MetasoundPrimitives.h"
+#include "WaveTableSampler.h"
+#include "MetasoundSourceInterface.h"
+#include "MetasoundStandardNodesCategories.h"
+#include "MetasoundStandardNodesNames.h"
+#include "MetasoundTrace.h"
+#include "MetasoundTrigger.h"
+#include "MetasoundWaveTable.h"
+
+#define LOCTEXT_NAMESPACE "MetasoundStandardNodes"
+
+
+namespace Metasound
+{
+	class FMetasoundWaveTableOscillatorNodeOperator : public TExecutableOperator<FMetasoundWaveTableOscillatorNodeOperator>
+	{
+	public:
+		static const FVertexInterface& GetDefaultInterface()
+		{
+			using namespace WaveTable;
+
+			static const FVertexInterface DefaultInterface(
+				FInputVertexInterface(
+					TInputDataVertex<FWaveTable>("WaveTable", FDataVertexMetadata { LOCTEXT("MetasoundWaveTableOscillatorNode_InputWaveTable", "WaveTable") }),
+					TInputDataVertex<FTrigger>("Play", FDataVertexMetadata { LOCTEXT("MetasoundWaveTableOscillatorNode_InputPlay", "Play") }),
+					TInputDataVertex<FTrigger>("Stop", FDataVertexMetadata { LOCTEXT("MetasoundWaveTableOscillatorNode_InputStop", "Stop") }),
+					TInputDataVertex<FTrigger>("Sync", FDataVertexMetadata { LOCTEXT("MetasoundWaveTableOscillatorNode_InputSync", "Sync") }),
+					TInputDataVertex<float>("Freq", FDataVertexMetadata { LOCTEXT("MetasoundWaveTableOscillatorNode_Freq", "Freq") }, 440.0f),
+					TInputDataVertex<FAudioBuffer>("PhaseMod", FDataVertexMetadata
+					{
+						LOCTEXT("MetasoundWaveTableOscillatorNode_PhaseModDescription", "Modulation audio source for modulating oscillation phase of provided table."),
+						LOCTEXT("MetasoundWaveTableOscillatorNode_PhaseMod", "Phase Modulator"),
+						true /* bIsAdvancedDisplay */
+					})
+				),
+				FOutputVertexInterface(
+					TOutputDataVertex<FAudioBuffer>("Out", FDataVertexMetadata { LOCTEXT("MetasoundWaveTableOscillatorNode_Output", "Out") })
+				)
+			);
+
+			return DefaultInterface;
+		}
+
+		static const FNodeClassMetadata& GetNodeInfo()
+		{
+			auto CreateNodeClassMetadata = []() -> FNodeClassMetadata
+			{
+				FNodeClassMetadata Metadata
+				{
+					{ EngineNodes::Namespace, "WaveTableOscillator", "" },
+					1, // Major Version
+					0, // Minor Version
+					LOCTEXT("MetasoundWaveTableOscillatorNode_Name", "WaveTable Oscillator"),
+					LOCTEXT("MetasoundWaveTableOscillatorNode_Description", "Reads through the given WaveTable at the provided frequency."),
+					PluginAuthor,
+					PluginNodeMissingPrompt,
+					GetDefaultInterface(),
+					{ NodeCategories::Envelopes },
+					{ },
+					{ }
+				};
+
+				return Metadata;
+			};
+
+			static const FNodeClassMetadata Metadata = CreateNodeClassMetadata();
+			return Metadata;
+		}
+
+		static TUniquePtr<IOperator> CreateOperator(const FCreateOperatorParams& InParams, TArray<TUniquePtr<IOperatorBuildError>>& OutErrors)
+		{
+			using namespace WaveTable;
+
+			const FInputVertexInterface& InputInterface = InParams.Node.GetVertexInterface().GetInputInterface();
+			const FDataReferenceCollection& InputCollection = InParams.InputDataReferences;
+
+			FWaveTableReadRef InWaveTableReadRef = InputCollection.GetDataReadReferenceOrConstruct<FWaveTable>("WaveTable");
+			FTriggerReadRef InPlayReadRef = InputCollection.GetDataReadReferenceOrConstructWithVertexDefault<FTrigger>(InputInterface, "Play", InParams.OperatorSettings);
+			FTriggerReadRef InStopReadRef = InputCollection.GetDataReadReferenceOrConstructWithVertexDefault<FTrigger>(InputInterface, "Stop", InParams.OperatorSettings);
+			FTriggerReadRef InSyncReadRef = InputCollection.GetDataReadReferenceOrConstructWithVertexDefault<FTrigger>(InputInterface, "Sync", InParams.OperatorSettings);
+			FFloatReadRef InFreqReadRef = InputCollection.GetDataReadReferenceOrConstructWithVertexDefault<float>(InputInterface, "Freq", InParams.OperatorSettings);
+
+			TOptional<FAudioBufferReadRef> InPhaseModReadRef;
+			if (InputCollection.ContainsDataReadReference<FAudioBuffer>("PhaseMod"))
+			{
+				InPhaseModReadRef = InputCollection.GetDataReadReference<FAudioBuffer>("PhaseMod");
+			}
+
+			return MakeUnique<FMetasoundWaveTableOscillatorNodeOperator>(InParams, InWaveTableReadRef, InPlayReadRef, InStopReadRef, InSyncReadRef, InFreqReadRef, MoveTemp(InPhaseModReadRef));
+		}
+
+		FMetasoundWaveTableOscillatorNodeOperator(
+			const FCreateOperatorParams& InParams,
+			const FWaveTableReadRef& InWaveTableReadRef,
+			const FTriggerReadRef& InPlayReadRef,
+			const FTriggerReadRef& InStopReadRef,
+			const FTriggerReadRef& InSyncReadRef,
+			const FFloatReadRef& InFreqReadRef,
+			TOptional<FAudioBufferReadRef>&& InPhaseModReadRef
+		)
+			: WaveTableReadRef(InWaveTableReadRef)
+			, PlayReadRef(InPlayReadRef)
+			, StopReadRef(InStopReadRef)
+			, SyncReadRef(InSyncReadRef)
+			, FreqReadRef(InFreqReadRef)
+			, PhaseModReadRef(MoveTemp(InPhaseModReadRef))
+			, OutBufferWriteRef(TDataWriteReferenceFactory<FAudioBuffer>::CreateAny(InParams.OperatorSettings))
+		{
+			const float BlockRate = InParams.OperatorSettings.GetActualBlockRate();
+			if (BlockRate > 0.0f)
+			{
+				BlockPeriod = 1.0f / BlockRate;
+			}
+		}
+
+		virtual ~FMetasoundWaveTableOscillatorNodeOperator() = default;
+
+		virtual FDataReferenceCollection GetInputs() const override
+		{
+			FDataReferenceCollection Inputs;
+
+			Inputs.AddDataReadReference("WaveTable", WaveTableReadRef);
+			Inputs.AddDataReadReference("Sync", SyncReadRef);
+			Inputs.AddDataReadReference("Freq", FreqReadRef);
+
+			if (PhaseModReadRef.IsSet())
+			{
+				Inputs.AddDataReadReference("PhaseMod", *PhaseModReadRef);
+			}
+
+			return Inputs;
+		}
+
+		virtual FDataReferenceCollection GetOutputs() const override
+		{
+			FDataReferenceCollection Outputs;
+			Outputs.AddDataReadReference("Out", TDataReadReference<FAudioBuffer>(OutBufferWriteRef));
+
+			return Outputs;
+		}
+
+		void Execute()
+		{
+			using namespace WaveTable;
+
+			METASOUND_TRACE_CPUPROFILER_EVENT_SCOPE(FMetasoundWaveTableOscillatorNodeOperator::Execute);
+
+			FAudioBuffer& OutBuffer = *OutBufferWriteRef;
+			OutBuffer.Zero();
+
+			auto GetLastIndex = [](const FTriggerReadRef& Trigger)
+			{
+				int32 LastIndex = -1;
+				Trigger->ExecuteBlock([](int32, int32) {}, [&LastIndex](int32 StartFrame, int32 EndFrame)
+				{
+					LastIndex = FMath::Max(LastIndex, StartFrame);
+				});
+				return LastIndex;
+			};
+
+			const int32 LastPlayIndex = GetLastIndex(PlayReadRef);
+			const int32 LastStopIndex = GetLastIndex(StopReadRef);
+			if (LastPlayIndex >= 0 || LastStopIndex >= 0)
+			{
+				bPlaying = LastPlayIndex > LastStopIndex;
+			}
+
+			if (bPlaying)
+			{
+				const FTrigger& SyncTrigger = *SyncReadRef;
+				TArrayView<float> SyncBufferView;
+				if (SyncTrigger.IsTriggered())
+				{
+					SyncBuffer.SetNum(OutBuffer.Num());
+					FMemory::Memset(SyncBuffer.GetData(), 0, sizeof(float) * SyncBuffer.Num());
+					SyncTrigger.ExecuteBlock(
+						[](int32 StartFrame, int32 EndFrame) {},
+						[this](int32 StartFrame, int32 EndFrame)
+						{
+							SyncBuffer[StartFrame] = 1.0f;
+						}
+					);
+					SyncBufferView = SyncBuffer;
+				}
+
+				TArrayView<const float> PhaseMod;
+				if (PhaseModReadRef.IsSet())
+				{
+					const FAudioBuffer& Buffer = *(*PhaseModReadRef);
+					PhaseMod = { Buffer.GetData(), Buffer.Num() };
+				}
+
+				// Limit wrap operations running off toward infinity while allowing sampler to play in reverse
+				Sampler.SetFreq(FMath::Clamp(*FreqReadRef, -20000.f, 20000.f) * BlockPeriod);
+
+				const FWaveTable& InputTable = *WaveTableReadRef;
+				Sampler.Process(InputTable.GetView(), { }, PhaseMod, SyncBufferView, OutBuffer);
+			}
+		}
+
+	private:
+		float BlockPeriod = 0.0f;
+		bool bPlaying = false;
+
+		FWaveTableReadRef WaveTableReadRef;
+		FTriggerReadRef PlayReadRef;
+		FTriggerReadRef StopReadRef;
+		FTriggerReadRef SyncReadRef;
+		FFloatReadRef FreqReadRef;
+
+		Audio::FAlignedFloatBuffer SyncBuffer;
+		TOptional<FAudioBufferReadRef> PhaseModReadRef;
+
+		WaveTable::FWaveTableSampler Sampler;
+
+		TDataWriteReference<FAudioBuffer> OutBufferWriteRef;
+	};
+
+	class FMetasoundWaveTableOscillatorNode : public FNodeFacade
+	{
+	public:
+		FMetasoundWaveTableOscillatorNode(const FNodeInitData& InInitData)
+			: FNodeFacade(InInitData.InstanceName, InInitData.InstanceID, TFacadeOperatorClass<FMetasoundWaveTableOscillatorNodeOperator>())
+		{
+		}
+
+		virtual ~FMetasoundWaveTableOscillatorNode() = default;
+	};
+
+	METASOUND_REGISTER_NODE(FMetasoundWaveTableOscillatorNode)
+} // namespace Metasound
+
+#undef LOCTEXT_NAMESPACE // MetasoundStandardNodes
