[MetaSound] enforce new vertices to have valid IDs.

diff --git a/Engine/Plugins/Runtime/Metasound/Source/MetasoundFrontend/Private/MetasoundFrontendGraphController.cpp b/Engine/Plugins/Runtime/Metasound/Source/MetasoundFrontend/Private/MetasoundFrontendGraphController.cpp
--- a/Engine/Plugins/Runtime/Metasound/Source/MetasoundFrontend/Private/MetasoundFrontendGraphController.cpp
+++ b/Engine/Plugins/Runtime/Metasound/Source/MetasoundFrontend/Private/MetasoundFrontendGraphController.cpp
@@ -954,10 +954,10 @@ namespace Metasound
 						FConstClassAccessPtr InputClassPtr = OwningDocument->FindOrAddClass(Key);
 						if (const FMetasoundFrontendClass* InputClass = InputClassPtr.Get())
 						{
-							FMetasoundFrontendNode& Node = GraphClass->Graph.Nodes.Emplace_GetRef(*InputClass);
-
 							const FVertexName NewName = InClassInput.Name;
 
+							// Setup input node
+							FMetasoundFrontendNode& Node = GraphClass->Graph.Nodes.Emplace_GetRef(*InputClass);
 							Node.Name = NewName;
 
 							if (InClassInput.NodeID.IsValid())
@@ -969,6 +969,7 @@ namespace Metasound
 								Node.UpdateID(FGuid::NewGuid());
 							}
 
+							// Set name on related vertices of input node
 							auto IsVertexWithTypeName = [&](FMetasoundFrontendVertex& Vertex) { return Vertex.TypeName == InClassInput.TypeName; };
 							if (FMetasoundFrontendVertex* InputVertex = Node.Interface.Inputs.FindByPredicate(IsVertexWithTypeName))
 							{
@@ -979,17 +980,25 @@ namespace Metasound
 								UE_LOG(LogMetaSound, Error, TEXT("Input node [TypeName:%s] does not contain input vertex with type [TypeName:%s]"), *InClassInput.TypeName.ToString(), *InClassInput.TypeName.ToString());
 							}
 
-							if (Node.Interface.Outputs.Num() == 1)
+							if (FMetasoundFrontendVertex* OutputVertex = Node.Interface.Outputs.FindByPredicate(IsVertexWithTypeName))
 							{
-								Node.Interface.Outputs[0].Name = NewName;
+								OutputVertex->Name = NewName;
 							}
-							else if (FMetasoundFrontendVertex* OutputVertex = Node.Interface.Outputs.FindByPredicate(IsVertexWithTypeName))
+							else
 							{
-								OutputVertex->Name = NewName;
+								UE_LOG(LogMetaSound, Error, TEXT("Input node [TypeName:%s] does not contain output vertex with type [TypeName:%s]"), *InClassInput.TypeName.ToString(), *InClassInput.TypeName.ToString());
 							}
 
+							// Add input to this graph class interface
 							FMetasoundFrontendClassInput& NewInput = GraphClass->Interface.Inputs.Add_GetRef(InClassInput);
+
 							NewInput.NodeID = Node.GetID();
+							if (!NewInput.VertexID.IsValid())
+							{
+								// Create a new guid if there wasn't a valid guid attached
+								// to input.
+								NewInput.VertexID = FGuid::NewGuid();
+							}
 
 							GraphClass->Interface.UpdateChangeID();
 
@@ -1080,9 +1089,9 @@ namespace Metasound
 						FConstClassAccessPtr OutputClassPtr = OwningDocument->FindOrAddClass(Key);
 						if (const FMetasoundFrontendClass* OutputClass = OutputClassPtr.Get())
 						{
-							FMetasoundFrontendNode& Node = GraphClass->Graph.Nodes.Add_GetRef(*OutputClass);
-
 							const FVertexName NewName = InClassOutput.Name;
+
+							FMetasoundFrontendNode& Node = GraphClass->Graph.Nodes.Add_GetRef(*OutputClass);
 							Node.Name = NewName;
 
 							if (InClassOutput.NodeID.IsValid())
@@ -1094,15 +1103,15 @@ namespace Metasound
 								Node.UpdateID(FGuid::NewGuid());
 							}
 
-							// TODO: have something that checks if input node has valid interface.
+							// Set vertex name on output node
 							auto IsVertexWithTypeName = [&](FMetasoundFrontendVertex& Vertex) { return Vertex.TypeName == InClassOutput.TypeName; };
-							if (Node.Interface.Inputs.Num() == 1)
+							if (FMetasoundFrontendVertex* InputVertex = Node.Interface.Inputs.FindByPredicate(IsVertexWithTypeName))
 							{
-								Node.Interface.Inputs[0].Name = NewName;
+								InputVertex->Name = NewName;
 							}
-							else if (FMetasoundFrontendVertex* InputVertex = Node.Interface.Inputs.FindByPredicate(IsVertexWithTypeName))
+							else
 							{
-								InputVertex->Name = NewName;
+								UE_LOG(LogMetaSound, Error, TEXT("Output node [TypeName:%s] does not contain input vertex with type [TypeName:%s]"), *InClassOutput.TypeName.ToString(), *InClassOutput.TypeName.ToString());
 							}
 
 							if (FMetasoundFrontendVertex* OutputVertex = Node.Interface.Outputs.FindByPredicate(IsVertexWithTypeName))
@@ -1114,9 +1123,19 @@ namespace Metasound
 								UE_LOG(LogMetaSound, Error, TEXT("Output node [TypeName:%s] does not contain output vertex with type [TypeName:%s]"), *InClassOutput.TypeName.ToString(), *InClassOutput.TypeName.ToString());
 							}
 
+							// Add output to graph interface
 							FMetasoundFrontendClassOutput& NewOutput = GraphClass->Interface.Outputs.Add_GetRef(InClassOutput);
+							
+							// Setup new output
 							NewOutput.NodeID = Node.GetID();
+							if (!NewOutput.VertexID.IsValid())
+							{
+								// Create a new guid if there wasn't a valid guid attached
+								// to output.
+								NewOutput.VertexID = FGuid::NewGuid();
+							}
 
+							// Mark interface as changed.
 							GraphClass->Interface.UpdateChangeID();
 
 							FNodeAccessPtr NodePtr = GraphClassPtr.GetNodeWithNodeID(Node.GetID());
