Update MetaSound transmission addresses
#rnx

diff --git a/Engine/Plugins/Runtime/Metasound/Source/MetasoundFrontend/Private/MetasoundRouter.cpp b/Engine/Plugins/Runtime/Metasound/Source/MetasoundFrontend/Private/MetasoundRouter.cpp
--- a/Engine/Plugins/Runtime/Metasound/Source/MetasoundFrontend/Private/MetasoundRouter.cpp
+++ b/Engine/Plugins/Runtime/Metasound/Source/MetasoundFrontend/Private/MetasoundRouter.cpp
@@ -8,100 +8,20 @@
 #include "MetasoundTrace.h"
 #include "HAL/IConsoleManager.h"
 
-// Convenience exec commands to push values to global params.
-static FAutoConsoleCommand GPushFloatCommand(
-	TEXT("au.MetaSound.SetFloat"),
-	TEXT("Use this with au.MetaSound.SetFloat [type] [address] [value]. Pushes a parameter value directly to a global address, which can then be received by Metasounds using a Receive node."),
-	FConsoleCommandWithArgsDelegate::CreateStatic(
-		[](const TArray< FString >& Args)
-		{
-			if (Args.Num() < 3)
-			{
-				UE_LOG(LogMetaSound, Warning, TEXT("au.MetaSound.Set* should be called with three args- the data type of the channel, the address to send to and the value to send."));
-				return;
-			}
-
-			FName DataType = FName(*Args[0]);
-			FName ChannelName = FName(*Args[1]);
-			float ValueToPush = TCString<TCHAR>::Atof(*Args[2]);
-
-			Metasound::FLiteral LiteralParam(ValueToPush);
-
-			Metasound::FDataTransmissionCenter::Get().PushLiteral(DataType, ChannelName, LiteralParam);
-		})
-);
-
-static FAutoConsoleCommand GPushBoolCommand(
-	TEXT("au.MetaSound.SetBool"),
-	TEXT("Use this with au.MetaSound.SetBool [address] [value]. Pushes a parameter value directly to a global address, which can then be received by Metasounds using a Receive node."),
-	FConsoleCommandWithArgsDelegate::CreateStatic(
-		[](const TArray< FString >& Args)
-		{
-			if (Args.Num() < 2)
-			{
-				UE_LOG(LogMetaSound, Warning, TEXT("au.MetaSound.SetBool should be called with two args- the address to send to and the value to send."));
-				UE_LOG(LogMetaSound, Warning, TEXT("au.MetaSound.Set* should be called with three args- the data type of the channel, the address to send to and the value to send."));
-				return;
-			}
-
-			FName DataType = FName(*Args[0]);
-			FName ChannelName = FName(*Args[1]);
-			int32 ValueAsInt = TCString<TCHAR>::Atoi(*Args[2]);
-
-			bool ValueToPush = ValueAsInt != 0;
-
-			Metasound::FLiteral LiteralParam(ValueToPush);
-
-			Metasound::FDataTransmissionCenter::Get().PushLiteral(DataType, ChannelName, LiteralParam);
-		})
-);
-
-static FAutoConsoleCommand GPushIntCommand(
-	TEXT("au.MetaSound.SetInt"),
-	TEXT("Use this with au.MetaSound.SetInt [address] [value]. Pushes a parameter value directly to a global address, which can then be received by Metasounds using a Receive node."),
-	FConsoleCommandWithArgsDelegate::CreateStatic(
-		[](const TArray< FString >& Args)
-		{
-			if (Args.Num() < 2)
-			{
-				UE_LOG(LogMetaSound, Warning, TEXT("au.MetaSound.Set* should be called with three args- the data type of the channel, the address to send to and the value to send."));
-				return;
-			}
-
-			FName DataType = FName(*Args[0]);
-			FName ChannelName = FName(*Args[1]);
-			int32 ValueToPush = TCString<TCHAR>::Atoi(*Args[2]);
-
-			Metasound::FLiteral LiteralParam(ValueToPush);
-
-			Metasound::FDataTransmissionCenter::Get().PushLiteral(DataType, ChannelName, LiteralParam);
-		})
-);
-
-static FAutoConsoleCommand GPushStringCommand(
-	TEXT("au.MetaSound.SetString"),
-	TEXT("Use this with au.MetaSound.SetString [address] [value]. Pushes a parameter value directly to a global address, which can then be received by Metasounds using a Receive node."),
-	FConsoleCommandWithArgsDelegate::CreateStatic(
-		[](const TArray< FString >& Args)
-		{
-			if (Args.Num() < 3)
-			{
-				UE_LOG(LogMetaSound, Warning, TEXT("au.MetaSound.SetBool should be called with three args- the data type of the channel, the address to send to and the value to send."));
-				UE_LOG(LogMetaSound, Warning, TEXT("au.MetaSound.Set* should be called with three args- the data type of the channel, the address to send to and the value to send."));
-				return;
-			}
-
-			FName DataType = FName(*Args[0]);
-			FName ChannelName = FName(*Args[1]);
+namespace Metasound
+{
+	bool FTransmissionAddress::operator==(const FTransmissionAddress& InOther) const
+	{
+		return IsEqual(InOther);
+	}
 
-			Metasound::FLiteral LiteralParam(Args[2]);
+	bool FTransmissionAddress::operator!=(const FTransmissionAddress& InOther) const
+	{
+		return !IsEqual(InOther);
+	}
 
-			Metasound::FDataTransmissionCenter::Get().PushLiteral(DataType, ChannelName, LiteralParam);
-		})
-);
+	const FLazyName FSendAddress::AddressType("MetaSoundSendAddress");
 
-namespace Metasound
-{
 	FSendAddress::FSendAddress(const FString& InChannelName)
 	: ChannelName(*InChannelName)
 	, DataType()
@@ -116,6 +36,52 @@ namespace Metasound
 	{
 	}
 
+	FName FSendAddress::GetAddressType() const
+	{
+		return FSendAddress::AddressType;
+	}
+
+	TUniquePtr<FTransmissionAddress> FSendAddress::Clone() const 
+	{
+		return MakeUnique<FSendAddress>(*this);
+	}
+
+	FName FSendAddress::GetDataType() const
+	{
+		return DataType;
+	}
+
+	const FName& FSendAddress::GetChannelName() const
+	{
+		return ChannelName;
+	}
+
+	uint64 FSendAddress::GetInstanceID() const
+	{
+		return InstanceID;
+	}
+
+	FString FSendAddress::ToString() const
+	{
+		return FString::Format(TEXT("SendAddress {0}:{1}[Type={2}]"), {ChannelName.ToString(), InstanceID, DataType.ToString()});
+	}
+
+	uint32 FSendAddress::GetHash() const
+	{
+		uint32 HashedChannel = HashCombineFast(::GetTypeHash(DataType), ::GetTypeHash(ChannelName));
+		HashedChannel = HashCombineFast(HashedChannel, ::GetTypeHash(InstanceID));
+		return HashedChannel;
+	}
+
+	bool FSendAddress::IsEqual(const FTransmissionAddress& InOther) const
+	{
+		if (const FSendAddress* OtherSendAddress = CastAddressType<FSendAddress>(InOther))
+		{
+			return (InstanceID == OtherSendAddress->InstanceID) && (ChannelName == OtherSendAddress->ChannelName) && (DataType == OtherSendAddress->DataType); 
+		}
+		return false;
+	}
+
 	IReceiver::~IReceiver()
 	{
 		DataChannel->OnReceiverDestroyed();
@@ -132,33 +98,39 @@ namespace Metasound
 		return Singleton;
 	}
 
-	TUniquePtr<ISender> FDataTransmissionCenter::RegisterNewSender(const FSendAddress& InAddress, const FSenderInitParams& InitParams)
+	TUniquePtr<ISender> FDataTransmissionCenter::RegisterNewSender(const FTransmissionAddress& InAddress, const FSenderInitParams& InitParams)
 	{
 		return GlobalRouter.RegisterNewSender(InAddress, InitParams);
 	}
 
-	TUniquePtr<IReceiver> FDataTransmissionCenter::RegisterNewReceiver(const FSendAddress& InAddress, const FReceiverInitParams& InitParams)
+	TUniquePtr<IReceiver> FDataTransmissionCenter::RegisterNewReceiver(const FTransmissionAddress& InAddress, const FReceiverInitParams& InitParams)
 	{
 		return GlobalRouter.RegisterNewReceiver(InAddress, InitParams);
 	}
 
-	bool FDataTransmissionCenter::UnregisterDataChannel(const FSendAddress& InAddress)
+	bool FDataTransmissionCenter::UnregisterDataChannel(const FTransmissionAddress& InAddress)
 	{
 		return GlobalRouter.UnregisterDataChannel(InAddress);
 	}
 
-	bool FDataTransmissionCenter::UnregisterDataChannelIfUnconnected(const FSendAddress& InAddress)
+	bool FDataTransmissionCenter::UnregisterDataChannelIfUnconnected(const FTransmissionAddress& InAddress)
 	{
 		return GlobalRouter.UnregisterDataChannelIfUnconnected(InAddress);
 	}
 
 	bool FDataTransmissionCenter::PushLiteral(FName DataType, FName GlobalChannelName, const FLiteral& InParam)
 	{
+PRAGMA_DISABLE_DEPRECATION_WARNINGS
 		return GlobalRouter.PushLiteral(DataType, GlobalChannelName, InParam);
+PRAGMA_ENABLE_DEPRECATION_WARNINGS
 	}
 
+	FAddressRouter::FAddressRouter(const FAddressRouter& Other)
+		: DataChannelMap(Other.DataChannelMap)
+	{
+	}
 
-	TSharedPtr<IDataChannel, ESPMode::ThreadSafe> FAddressRouter::FindDataChannel(const FSendAddress& InAddress)
+	TSharedPtr<IDataChannel, ESPMode::ThreadSafe> FAddressRouter::FindDataChannel(const FTransmissionAddress& InAddress)
 	{
 		TSharedPtr<IDataChannel, ESPMode::ThreadSafe> Channel;
 
@@ -174,7 +146,7 @@ namespace Metasound
 		return Channel;
 	}
 
-	TSharedPtr<IDataChannel, ESPMode::ThreadSafe> FAddressRouter::GetDataChannel(const FSendAddress& InAddress, const FOperatorSettings& InOperatorSettings)
+	TSharedPtr<IDataChannel, ESPMode::ThreadSafe> FAddressRouter::GetDataChannel(const FTransmissionAddress& InAddress, const FOperatorSettings& InOperatorSettings)
 	{
 		METASOUND_LLM_SCOPE;
 		TSharedPtr<IDataChannel, ESPMode::ThreadSafe> DataChannel = FindDataChannel(InAddress);
@@ -194,7 +166,7 @@ namespace Metasound
 		return DataChannel;
 	}
 
-	TUniquePtr<ISender> FAddressRouter::RegisterNewSender(const FSendAddress& InAddress, const FSenderInitParams& InitParams)
+	TUniquePtr<ISender> FAddressRouter::RegisterNewSender(const FTransmissionAddress& InAddress, const FSenderInitParams& InitParams)
 	{
 		METASOUND_LLM_SCOPE;
 		TSharedPtr<IDataChannel, ESPMode::ThreadSafe> DataChannel = GetDataChannel(InAddress, InitParams.OperatorSettings);
@@ -209,7 +181,22 @@ namespace Metasound
 		}
 	}
 
-	bool FAddressRouter::UnregisterDataChannel(const FSendAddress& InAddress)
+	bool FAddressRouter::PushLiteral(const FName& InDataTypeName, const FName& InChannelName, const FLiteral& InParam)
+	{
+		FSendAddress Address{ InChannelName, InDataTypeName };
+
+		TSharedPtr<IDataChannel, ESPMode::ThreadSafe> Channel = FindDataChannel(Address);
+		if (Channel.IsValid())
+		{
+			return Channel->PushLiteral(InParam);
+		}
+		else
+		{
+			return false;
+		}
+	}
+
+	bool FAddressRouter::UnregisterDataChannel(const FTransmissionAddress& InAddress)
 	{
 		METASOUND_LLM_SCOPE;
 		FScopeLock ScopeLock(&DataChannelMapMutationLock);
@@ -230,7 +217,7 @@ namespace Metasound
 		return DataChannelMap.Remove(InAddress) > 0;
 	}
 
-	bool FAddressRouter::UnregisterDataChannelIfUnconnected(const FSendAddress& InAddress)
+	bool FAddressRouter::UnregisterDataChannelIfUnconnected(const FTransmissionAddress& InAddress)
 	{
 		METASOUND_LLM_SCOPE;
 		FScopeLock ScopeLock(&DataChannelMapMutationLock);
@@ -249,7 +236,7 @@ namespace Metasound
 		return false;
 	}
 
-	TUniquePtr<IReceiver> FAddressRouter::RegisterNewReceiver(const FSendAddress& InAddress, const FReceiverInitParams& InitParams)
+	TUniquePtr<IReceiver> FAddressRouter::RegisterNewReceiver(const FTransmissionAddress& InAddress, const FReceiverInitParams& InitParams)
 	{
 		METASOUND_LLM_SCOPE;
 		TSharedPtr<IDataChannel, ESPMode::ThreadSafe> DataChannel = GetDataChannel(InAddress, InitParams.OperatorSettings);
@@ -263,4 +250,53 @@ namespace Metasound
 			return TUniquePtr<IReceiver>(nullptr);
 		}
 	}
+
+	FAddressRouter::FTransmissionAddressKey::FTransmissionAddressKey(const FTransmissionAddress& InCopy)
+	: PImpl(InCopy.Clone())
+	{}
+
+	FAddressRouter::FTransmissionAddressKey::FTransmissionAddressKey(const FTransmissionAddressKey& InCopy)
+	: PImpl(InCopy.PImpl->Clone())
+	{}
+
+	FAddressRouter::FTransmissionAddressKey::FTransmissionAddressKey(TUniquePtr<FTransmissionAddress>&& InImpl)
+	: PImpl(MoveTemp(InImpl))
+	{
+	}
+
+	FAddressRouter::FTransmissionAddressKey& FAddressRouter::FTransmissionAddressKey::operator=(const FTransmissionAddress& InOther)
+	{
+		PImpl = InOther.Clone();
+		return *this;
+	}
+
+	FName FAddressRouter::FTransmissionAddressKey::GetAddressType() const
+	{
+		return PImpl->GetAddressType();
+	}
+
+	FName FAddressRouter::FTransmissionAddressKey::GetDataType() const
+	{
+		return PImpl->GetDataType();
+	}
+
+	TUniquePtr<FTransmissionAddress> FAddressRouter::FTransmissionAddressKey::Clone() const
+	{
+		return PImpl->Clone();
+	}
+
+	FString FAddressRouter::FTransmissionAddressKey::ToString() const
+	{
+		return PImpl->ToString();
+	}
+
+	uint32 FAddressRouter::FTransmissionAddressKey::GetHash() const
+	{
+		return GetTypeHash(*PImpl);
+	}
+
+	bool FAddressRouter::FTransmissionAddressKey::IsEqual(const FTransmissionAddress& InOther) const
+	{
+		return InOther == *PImpl;
+	}
 }
diff --git a/Engine/Plugins/Runtime/Metasound/Source/MetasoundFrontend/Public/MetasoundRouter.h b/Engine/Plugins/Runtime/Metasound/Source/MetasoundFrontend/Public/MetasoundRouter.h
--- a/Engine/Plugins/Runtime/Metasound/Source/MetasoundFrontend/Public/MetasoundRouter.h
+++ b/Engine/Plugins/Runtime/Metasound/Source/MetasoundFrontend/Public/MetasoundRouter.h
@@ -3,6 +3,7 @@
 
 #include "DSP/Dsp.h"
 #include "DSP/MultithreadedPatching.h"
+#include "HAL/Platform.h"
 #include "MetasoundAudioBuffer.h"
 #include "MetasoundDataFactory.h"
 #include "MetasoundDataReference.h"
@@ -21,8 +22,7 @@
 
 /**
  * METASOUND TRANSMISSION SYSTEM
- * This allows for transmission of arbitrary datatypes using addresses set by FNames.
- * This system is used by Send and Receive nodes in the Metasound graph, but can also be used in C++ for arbitrary message passing.
+ * This allows for transmission of arbitrary data types using transmission addresses.
  *
  * Typical use case:
  * TSenderPtr<float> FloatSender FDataTransmissionCenter::Get().RegisterNewSender<float>(ExampleAddress, InitSettings);
@@ -34,67 +34,138 @@
  * //... elsewhere on a different thread:
  * float ReceivedFloat = FloatReceiver->Pop();
  *
- * It's that easy! (hopefully!)
- *
  * In general, some performance considerations for this system:
- * 1) a circular buffer is allocated for each address a sender is set up for, and is sized to the max Delay of any connected sender.
- * 2) Senders and receivers are meant to have long lifecycles, and requesting them from the FDataTransmissionCenter may be expensive if done every tick.
- * 3) Most non-audio data types cannot be mixed, so multiple senders to the same address will cause only one of the senders to be effective.
- * 4) while the overall system is thread safe, individual TSender and TReceiver objects should only be used by a single thread.
+ * 1) Senders and receivers are meant to have long lifecycles, and requesting them from the FDataTransmissionCenter may be expensive if done every tick.
+ * 2) Most non-audio data types cannot be mixed, so multiple senders to the same address will cause only one of the senders to be effective.
+ * 3) while the overall system is thread safe, individual TSender and TReceiver objects should only be used by a single thread.
  */
 
-
 namespace Metasound
 {
-	class METASOUNDFRONTEND_API FSendAddress
+	//Forward declare
+	class FTransmissionAddress;
+	uint32 GetTypeHash(const Metasound::FTransmissionAddress&);
+
+	/** A FTranmissionAddress acts as a identifier for a data channel shared by 
+	 * one or more senders and or receivers. 
+	 *
+	 * Different addressing schemes are supported by deriving a system specific 
+	 * address type from a FTransmissionAddress. This allows different systems 
+	 * to construct and store address information in their desired format while
+	 * all utilizing the same system for transmission.
+	 */
+	class METASOUNDFRONTEND_API FTransmissionAddress
 	{
 	public:
-		FSendAddress() = default;
+		FTransmissionAddress() = default;
+		virtual ~FTransmissionAddress() = default;
 
-		// Create an address without a data type or instance ID
-		FSendAddress(const FString& InChannelName);
+		/** Return the type of address. 
+		 *
+		 * The AddressType must be unique per a subclass of FTransmissionAddress
+		 * as it is used to determine the validity of a static_cast from FTransmissionAddress
+		 * to the derived address type.
+		 */
+		virtual FName GetAddressType() const = 0;
 
-		// Create an address with a channel name, data type and optionally an instance ID. 
-		FSendAddress(const FName& InChannelName, const FName& InDataType, uint64 InInstanceID=INDEX_NONE);
+		bool operator==(const FTransmissionAddress& InOther) const;
+		bool operator!=(const FTransmissionAddress& InOther) const;
 
-		friend uint32 GetTypeHash(const FSendAddress& InAddress)
-		{
-			uint32 HashedChannel = HashCombineFast(::GetTypeHash(InAddress.DataType), ::GetTypeHash(InAddress.ChannelName));
-			HashedChannel = HashCombineFast(HashedChannel, ::GetTypeHash(InAddress.InstanceID));
-			return HashedChannel;
-		}
+		/** Return the data type. */
+		virtual FName GetDataType() const = 0;
 
-		FORCEINLINE bool operator==(const FSendAddress& Other) const
-		{
-			return (InstanceID == Other.InstanceID) && (ChannelName == Other.ChannelName) && (DataType == Other.DataType); /*&& (Subsystem == Other.Subsystem);*/
-		}
+		/** Clone the address. */
+		virtual TUniquePtr<FTransmissionAddress> Clone() const = 0;
 
-		FORCEINLINE bool operator!=(const FSendAddress& Other) const
-		{
-			return !(*this == Other);
-		}
+		/** Return a string representation of the address for debugging. */
+		virtual FString ToString() const = 0;
 
-		const FName& GetDataType() const
-		{
-			return DataType;
-		}
+	private:
+		friend uint32 GetTypeHash(const ::Metasound::FTransmissionAddress&);
 
-		const FName& GetChannelName() const
-		{
-			return ChannelName;
-		}
+		template<typename DerivedAddressType>
+		friend const DerivedAddressType* CastAddressType(const FTransmissionAddress&);
 
-		uint64 GetInstanceID() const
+		template<typename DerivedAddressType>
+		struct THasStaticAddressType
 		{
-			return InstanceID;
-		}
+		private:
+			
+			// SFINAE: If the static member `T::AddressType` exists, then this template is used.
+			template<
+				typename T, 
+				typename = typename std::enable_if_t<!std::is_member_pointer_v<decltype(&T::AddressType)>>
+			>
+			static std::true_type CheckIfStaticMemberExists(int32);
+
+			// SFINAE: If the static member `T::AddressType` does not exist, then this template is used.
+			template<typename>
+			static std::false_type CheckIfStaticMemberExists(...);
+
+		public:
+			// True if the static member T::AddressType exists. 
+			static constexpr bool Value = decltype(CheckIfStaticMemberExists<DerivedAddressType>(0))::value;
+		};
+
+		/** Create a hash of the address. */
+		virtual uint32 GetHash() const = 0;
+
+		/** Determine if the given address is equal to this address. */
+		virtual bool IsEqual(const FTransmissionAddress& InOther) const = 0;
+	};
+
+	/** Attempt to cast a FTransmissionAddress to a derived address type. 
+	 *
+	 * This function uses the FTransmissionAddress::GetAddressType() method to 
+	 * determine whether the cast is legal. Each subclass of FTransmissionAddress
+	 * *MUST* provide a unique address type to protect against invalid casts.
+	 */
+	template<typename DerivedAddressType>
+	const DerivedAddressType* CastAddressType(const FTransmissionAddress& InOther)
+	{
+		// Helper struct to determine if the DerivedAddressType class has a `DerivedAddressType::AddressType` static member
+
+		// Static assert on invalid cast attempts
+		static_assert(std::is_base_of<FTransmissionAddress, DerivedAddressType>::value, "The derived address type must be derived from FTransmissionAddress");
+		static_assert(FTransmissionAddress::THasStaticAddressType<DerivedAddressType>::Value, "The derived address type must have a static member `AddressType` to perform casts");
 
-		FString ToString() const
+		// Check if AddressType is equal
+		if (InOther.GetAddressType() == DerivedAddressType::AddressType)
 		{
-			return FString::Format(TEXT("SendAddress {0}:{1}[Type={2}]"), {ChannelName.ToString(), InstanceID, DataType.ToString()});
+			return static_cast<const DerivedAddressType*>(&InOther);
 		}
+		return nullptr;
+	}
+
+	/** FSendAddress is used as a addressing system for MetaSound Send nodes, Recieve
+	 * nodes and MetaSound script communication
+	 */
+	class METASOUNDFRONTEND_API FSendAddress : public FTransmissionAddress
+	{
+	public:
+		static const FLazyName AddressType;
+
+		FSendAddress() = default;
+
+		// Create an address without a data type or instance ID
+		FSendAddress(const FString& InChannelName);
+
+		// Create an address with a channel name, data type and optionally an instance ID. 
+		FSendAddress(const FName& InChannelName, const FName& InDataType, uint64 InInstanceID=INDEX_NONE);
+
+		virtual FName GetAddressType() const override;
+		virtual FName GetDataType() const override;
+		virtual TUniquePtr<FTransmissionAddress> Clone() const override;
+		virtual FString ToString() const override;
+
+		const FName& GetChannelName() const;
+		uint64 GetInstanceID() const;
 
 	private:
+
+		virtual uint32 GetHash() const override;
+		virtual bool IsEqual(const FTransmissionAddress& InOther) const override;
+
 		FName ChannelName;
 		FName DataType;
 
@@ -102,6 +173,11 @@ namespace Metasound
 		uint64 InstanceID = INDEX_NONE;
 	};
 
+	FORCEINLINE uint32 GetTypeHash(const Metasound::FTransmissionAddress& InAddress)
+	{
+		return InAddress.GetHash();
+	}
+
 	// Base class that defines shared state and utilities for senders, receivers, and the shared channel state.
 	class METASOUNDFRONTEND_API IDataTransmissionBase
 	{
@@ -383,17 +459,6 @@ namespace Metasound
 		bool SetDelay(float InSeconds)
 		{
 			Params.DelayTimeInSeconds = InSeconds;
-			// TODO: LOL. I was wondering why the having a delay didn't change the behavior.
-			// Probably for the best since it does not make much sense to delay
-			// sent values. The FMath::Min() results in it being either 1 or 0. 
-			//NumElementsToDelayBy = FMath::Min(InSeconds * Params.OperatorSettings.GetActualBlockRate(), 1.f);
-
-			/*
-			if (NumElementsToDelayBy >= SenderBuffer.GetCapacity())
-			{
-				SenderBuffer.SetCapacity(NumElementsToDelayBy * BufferResizeThreshold);
-			}
-			*/
 			SenderBuffer.SetCapacity(2);
 			return true;
 		}
@@ -942,77 +1007,76 @@ namespace Metasound
 
 	public:
 
-		TSharedPtr<IDataChannel, ESPMode::ThreadSafe> FindDataChannel(const FSendAddress& InAddress);
-		TSharedPtr<IDataChannel, ESPMode::ThreadSafe> GetDataChannel(const FSendAddress& InAddress, const FOperatorSettings& InOperatorSettings);
+		TSharedPtr<IDataChannel, ESPMode::ThreadSafe> FindDataChannel(const FTransmissionAddress& InAddress);
+		TSharedPtr<IDataChannel, ESPMode::ThreadSafe> GetDataChannel(const FTransmissionAddress& InAddress, const FOperatorSettings& InOperatorSettings);
 
-		TUniquePtr<ISender> RegisterNewSender(const FSendAddress& InAddress, const FSenderInitParams& InitParams);
-		TUniquePtr<IReceiver> RegisterNewReceiver(const FSendAddress& InAddress, const FReceiverInitParams& InitParams);
+		TUniquePtr<ISender> RegisterNewSender(const FTransmissionAddress& InAddress, const FSenderInitParams& InitParams);
+		TUniquePtr<IReceiver> RegisterNewReceiver(const FTransmissionAddress& InAddress, const FReceiverInitParams& InitParams);
 
 		template<typename TDataType>
-		TReceiverPtr<TDataType> RegisterNewReceiver(const FSendAddress& InAddress, const FReceiverInitParams& InitParams)
+		TReceiverPtr<TDataType> RegisterNewReceiver(const FTransmissionAddress& InAddress, const FReceiverInitParams& InitParams)
 		{
 			TReceiverPtr<TDataType> Receiver;
 
-			const bool bIsMatchingDataType = GetMetasoundDataTypeName<TDataType>() == InAddress.GetDataType();
-			if (ensure(bIsMatchingDataType))
+			TUniquePtr<IReceiver> ReceiverBase = RegisterNewReceiver(InAddress, InitParams);
+			if (ReceiverBase.IsValid())
 			{
-				TUniquePtr<IReceiver> ReceiverBase = RegisterNewReceiver(InAddress, InitParams);
-				if (ReceiverBase.IsValid())
-				{
-					Receiver = Downcast<TDataType>(MoveTemp(ReceiverBase));
-				}
+				Receiver = Downcast<TDataType>(MoveTemp(ReceiverBase));
 			}
 
-			return MoveTemp(Receiver);
+			return Receiver;
 		}
 
 		template<typename TDataType>
-		TSenderPtr<TDataType> RegisterNewSender(const FSendAddress& InAddress, const FSenderInitParams& InitParams)
+		TSenderPtr<TDataType> RegisterNewSender(const FTransmissionAddress& InAddress, const FSenderInitParams& InitParams)
 		{
 			TSenderPtr<TDataType> Sender;
+			TUniquePtr<ISender> SenderBase = RegisterNewSender(InAddress, InitParams);
 
-			const bool bIsMatchingDataType = GetMetasoundDataTypeName<TDataType>() == InAddress.GetDataType();
-			if (ensure(bIsMatchingDataType))
+			if (SenderBase.IsValid())
 			{
-				TUniquePtr<ISender> SenderBase = RegisterNewSender(InAddress, InitParams);
-				if (SenderBase.IsValid())
-				{
-					Sender = Downcast<TDataType>(MoveTemp(SenderBase));
-				}
+				Sender = Downcast<TDataType>(MoveTemp(SenderBase));
 			}
 
-			return MoveTemp(Sender);
+			return Sender;
 		}
 
-		bool PushLiteral(const FName& InDataTypeName, const FName& InChannelName, const FLiteral& InParam)
-		{
-			FSendAddress Address{ InChannelName, InDataTypeName };
+		UE_DEPRECATED(5.1, "Pushing literals from the Metasound::FAddressRouter will no longer be supported in future released.")
+		bool PushLiteral(const FName& InDataTypeName, const FName& InChannelName, const FLiteral& InParam);
 
-			TSharedPtr<IDataChannel, ESPMode::ThreadSafe> Channel = FindDataChannel(Address);
-			if (Channel.IsValid())
-			{
-				return Channel->PushLiteral(InParam);
-			}
-			else
-			{
-				return false;
-			}
-		}
+		bool UnregisterDataChannel(const FTransmissionAddress& InAddress);
+		bool UnregisterDataChannelIfUnconnected(const FTransmissionAddress& InAddress);
 
-		bool UnregisterDataChannel(const FSendAddress& InAddress);
-		bool UnregisterDataChannelIfUnconnected(const FSendAddress& InAddress);
+		FAddressRouter(const FAddressRouter& Other);
+		FAddressRouter() = default;
+
+	private:
 
-		FAddressRouter(const FAddressRouter& Other)
-			: DataChannelMap(Other.DataChannelMap)
+		// Wrapper class for having an FTransmissionAddress work in a TMap<>
+		// This implements the copy constructor and copy assignment operator
+		// by calling the virtual Clone() method. 
+		class FTransmissionAddressKey : public FTransmissionAddress
 		{
-		}
+		public:
+			FTransmissionAddressKey(const FTransmissionAddress& InCopy);
+			FTransmissionAddressKey(const FTransmissionAddressKey& InCopy);
+			FTransmissionAddressKey(TUniquePtr<FTransmissionAddress>&& InImpl);
+			virtual ~FTransmissionAddressKey() = default;
 
-		FAddressRouter()
-		{}
+			FTransmissionAddressKey& operator=(const FTransmissionAddress& InOther);
+			virtual FName GetAddressType() const override;
+			virtual FName GetDataType() const override;
+			virtual TUniquePtr<FTransmissionAddress> Clone() const override;
+			virtual FString ToString() const override;
 
-	private:
+		private:
+			virtual uint32 GetHash() const override;
+			virtual bool IsEqual(const FTransmissionAddress& InOther) const override;
 
-		TMap<FSendAddress, TSharedRef<IDataChannel, ESPMode::ThreadSafe>> DataChannelMap;
+			TUniquePtr<FTransmissionAddress> PImpl;
+		};
+
+		TMap<FTransmissionAddressKey, TSharedRef<IDataChannel, ESPMode::ThreadSafe>> DataChannelMap;
 		FCriticalSection DataChannelMapMutationLock;
 	};
 
@@ -1028,24 +1092,19 @@ namespace Metasound
 
 		// Creates a new object to push data to an address.
 		// Returns a new sender, or nullptr if registration failed.
-		TUniquePtr<ISender> RegisterNewSender(const FSendAddress& InAddress, const FSenderInitParams& InitParams);
+		TUniquePtr<ISender> RegisterNewSender(const FTransmissionAddress& InAddress, const FSenderInitParams& InitParams);
 
 		// Creates a new object to push data to an address.
 		// Returns a new sender, or nullptr if registration failed.
 		template<typename TDataType>
-		TSenderPtr<TDataType> RegisterNewSender(const FSendAddress& InAddress, const FSenderInitParams& InitParams)
+		TSenderPtr<TDataType> RegisterNewSender(const FTransmissionAddress& InAddress, const FSenderInitParams& InitParams)
 		{
 			TSenderPtr<TDataType> Sender;
-			
-			const bool bIsMatchingDataType = GetMetasoundDataTypeName<TDataType>() == InAddress.GetDataType();
-			if (ensure(bIsMatchingDataType))
-			{
-				TUniquePtr<ISender> SenderBase = RegisterNewSender(InAddress, InitParams);
+			TUniquePtr<ISender> SenderBase = RegisterNewSender(InAddress, InitParams);
 
-				if (SenderBase.IsValid())
-				{
-					Sender = Downcast<TDataType>(MoveTemp(SenderBase));
-				}
+			if (SenderBase.IsValid())
+			{
+				Sender = Downcast<TDataType>(MoveTemp(SenderBase));
 			}
 
 			return Sender;
@@ -1053,37 +1112,33 @@ namespace Metasound
 
 		// Registers a new object to poll data from an address.
 		// Returns a new receiver, or nullptr if registration failed.
-		TUniquePtr<IReceiver> RegisterNewReceiver(const FSendAddress& InAddress, const FReceiverInitParams& InitParams);
+		TUniquePtr<IReceiver> RegisterNewReceiver(const FTransmissionAddress& InAddress, const FReceiverInitParams& InitParams);
 
 		// Registers a new object to poll data from an address.
 		// Returns a new receiver, or nullptr if registration failed.
 		template<typename TDataType>
-		TReceiverPtr<TDataType> RegisterNewReceiver(const FSendAddress& InAddress, const FReceiverInitParams& InitParams)
+		TReceiverPtr<TDataType> RegisterNewReceiver(const FTransmissionAddress& InAddress, const FReceiverInitParams& InitParams)
 		{
 			TReceiverPtr<TDataType> Receiver;
+			TUniquePtr<IReceiver> ReceiverBase = RegisterNewReceiver(InAddress, InitParams);
 
-			const bool bIsMatchingDataType = GetMetasoundDataTypeName<TDataType>() == InAddress.GetDataType();
-			if (ensure(bIsMatchingDataType))
+			if (ReceiverBase.IsValid())
 			{
-				TUniquePtr<IReceiver> ReceiverBase = RegisterNewReceiver(InAddress, InitParams);
-
-				if (ReceiverBase.IsValid())
-				{
-					Receiver = Downcast<TDataType>(MoveTemp(ReceiverBase));
-				}
+				Receiver = Downcast<TDataType>(MoveTemp(ReceiverBase));
 			}
 
 			return Receiver;
 		}
 
 		// Unregisters DataChannel irrespective of number of receivers or senders still active.
-		bool UnregisterDataChannel(const FSendAddress& InAddress);
+		bool UnregisterDataChannel(const FTransmissionAddress& InAddress);
 
 		// Unregister a data channel if there are no senders or receivers
-		bool UnregisterDataChannelIfUnconnected(const FSendAddress& InAddress);
+		bool UnregisterDataChannelIfUnconnected(const FTransmissionAddress& InAddress);
 
 		// Pushes a literal parameter to a specific data channel in the global router.
 		// returns false if the literal type isn't supported.
+		UE_DEPRECATED(5.1, "Pushing literals from the Metasound::FDataTransmissionCenter will no longer be supported in future released.")
 		bool PushLiteral(FName DataTypeName, FName GlobalChannelName, const FLiteral& InParam);
 
 	private:
@@ -1095,3 +1150,5 @@ namespace Metasound
 
 	DECLARE_METASOUND_DATA_REFERENCE_TYPES(FSendAddress, METASOUNDFRONTEND_API, FSendAddressTypeInfo, FSendAddressReadRef, FSendAddressWriteRef)
 }
+
+
