Split MetaSound DataType "IsArrayType()" function into bIsArrayParseable & bIsArrayType members, deprecating function in favor of caching booleans on DataType registration
Enforce single UObject class --> MetaSound Frontend class & fix-up callsites relying on ArrayObject type/literals using cached proxy type. This lead to confusion WRT what registered type is directly associated with a given UObject & possible bugs therein.
#rnx

diff --git a/Engine/Plugins/Runtime/Metasound/Source/MetasoundEditor/Private/MetasoundEditorModule.cpp b/Engine/Plugins/Runtime/Metasound/Source/MetasoundEditor/Private/MetasoundEditorModule.cpp
--- a/Engine/Plugins/Runtime/Metasound/Source/MetasoundEditor/Private/MetasoundEditorModule.cpp
+++ b/Engine/Plugins/Runtime/Metasound/Source/MetasoundEditor/Private/MetasoundEditorModule.cpp
@@ -527,8 +527,7 @@ namespace Metasound
 				// Default to object as most calls to this outside of the MetaSound Editor will be for custom UObject types
 				const FName PinCategory = InPinCategory.IsNone() ? FGraphBuilder::PinCategoryObject : InPinCategory;
 
-				const bool bIsArray = DataTypeInfo.IsArrayType();
-				const EPinContainerType ContainerType = bIsArray ? EPinContainerType::Array : EPinContainerType::None;
+				const EPinContainerType ContainerType = DataTypeInfo.bIsArrayType ? EPinContainerType::Array : EPinContainerType::None;
 				FEdGraphPinType PinType(PinCategory, InPinSubCategory, nullptr, ContainerType, false, FEdGraphTerminalType());
 				UClass* ClassToUse = IDataTypeRegistry::Get().GetUClassForDataType(InDataTypeName);
 				PinType.PinSubCategoryObject = Cast<UObject>(ClassToUse);
@@ -614,17 +613,12 @@ namespace Metasound
 
 			virtual const FSlateBrush* GetIconBrush(FName InDataType, const bool bIsConstructorType) const override
 			{
-				bool bIsArrayType = false;
-
 				Frontend::FDataTypeRegistryInfo Info;
-				if (Frontend::IDataTypeRegistry::Get().GetDataTypeInfo(InDataType, Info))
-				{
-					bIsArrayType = Info.IsArrayType();
-				}
+				Frontend::IDataTypeRegistry::Get().GetDataTypeInfo(InDataType, Info);
 
 				if (const ISlateStyle* MetasoundStyle = FSlateStyleRegistry::FindSlateStyle("MetaSoundStyle"))
 				{
-					if (bIsArrayType)
+					if (Info.bIsArrayType)
 					{
 						return bIsConstructorType ? MetasoundStyle->GetBrush(TEXT("MetasoundEditor.Graph.ConstructorPinArray")) : MetasoundStyle->GetBrush(TEXT("MetasoundEditor.Graph.ArrayPin"));
 					}
@@ -635,7 +629,7 @@ namespace Metasound
 				}
 				else
 				{
-					return bIsArrayType ? FAppStyle::GetBrush("Graph.ArrayPin.Connected") : FAppStyle::GetBrush("Icons.BulletPoint");
+					return Info.bIsArrayType ? FAppStyle::GetBrush("Graph.ArrayPin.Connected") : FAppStyle::GetBrush("Icons.BulletPoint");
 				}
 			}
 
diff --git a/Engine/Plugins/Runtime/Metasound/Source/MetasoundEditor/Private/MetasoundNodeDetailCustomization.cpp b/Engine/Plugins/Runtime/Metasound/Source/MetasoundEditor/Private/MetasoundNodeDetailCustomization.cpp
--- a/Engine/Plugins/Runtime/Metasound/Source/MetasoundEditor/Private/MetasoundNodeDetailCustomization.cpp
+++ b/Engine/Plugins/Runtime/Metasound/Source/MetasoundEditor/Private/MetasoundNodeDetailCustomization.cpp
@@ -74,7 +74,18 @@ namespace Metasound
 					{
 						if (const UMetasoundEditorGraphMember* Member = Cast<UMetasoundEditorGraphMember>(DefaultLiteral->GetOuter()))
 						{
-							ensure(IDataTypeRegistry::Get().GetDataTypeInfo(Member->GetDataType(), OutDataTypeInfo));
+							FName DataTypeName = Member->GetDataType();
+							ensure(IDataTypeRegistry::Get().GetDataTypeInfo(DataTypeName, OutDataTypeInfo));
+							if (OutDataTypeInfo.bIsArrayType)
+							{
+								DataTypeName = CreateElementTypeNameFromArrayTypeName(DataTypeName);
+								const bool bIsHiddenType = HiddenInputTypeNames.Contains(DataTypeName);
+								OutDataTypeInfo = { };
+								if (!bIsHiddenType)
+								{
+									ensure(IDataTypeRegistry::Get().GetDataTypeInfo(DataTypeName, OutDataTypeInfo));
+								}
+							}
 						}
 					}
 				}
@@ -84,7 +95,7 @@ namespace Metasound
 			// element type. Otherwise, returns this type's DataTypeName.
 			FName GetPrimitiveTypeName(const Frontend::FDataTypeRegistryInfo& InDataTypeInfo)
 			{
-				return InDataTypeInfo.IsArrayType()
+				return InDataTypeInfo.bIsArrayType
 					? CreateElementTypeNameFromArrayTypeName(InDataTypeInfo.DataTypeName)
 					: InDataTypeInfo.DataTypeName;
 			}
@@ -227,10 +238,12 @@ namespace Metasound
 					MemberCustomizationPrivate::GetDataTypeFromElementPropertyHandle(DefaultValueHandle, DataTypeInfo);
 
 					const IMetasoundEditorModule& EditorModule = FModuleManager::GetModuleChecked<IMetasoundEditorModule>("MetaSoundEditor");
-					bool bCanDrop = true;
-					for (const FAssetData& AssetData : InAssets)
+					bool bCanDrop = false;
+
+					if (UClass* ProxyGenClass = DataTypeInfo.ProxyGeneratorClass)
 					{
-						if (DataTypeInfo.ProxyGeneratorClass)
+						bCanDrop = true;
+						for (const FAssetData& AssetData : InAssets)
 						{
 							if (UClass* Class = AssetData.GetClass())
 							{
@@ -245,7 +258,8 @@ namespace Metasound
 							}
 						}
 					}
-					return true;
+
+					return bCanDrop;
 				})
 				.OnAssetsDropped_Lambda([this, DefaultValueHandle](const FDragDropEvent& DragDropEvent, TArrayView<FAssetData> InAssets)
 				{
@@ -551,6 +565,7 @@ namespace Metasound
 					}
 				}
 			}
+
 			MemberCustomizationPrivate::GetDataTypeFromElementPropertyHandle(ElementPropertyHandle, DataTypeInfo);
 
 			TSharedRef<SWidget> ValueWidget = CreateValueWidget(ParentPropertyHandleArray, StructPropertyHandle);
@@ -683,8 +698,7 @@ namespace Metasound
 				return;
 			}
 
-			const bool bIsArrayType = DataTypeInfo.IsArrayType();
-			if (bIsArrayType)
+			if (DataTypeInfo.bIsArrayType)
 			{
 				ArrayTypeName = GraphMember->GetDataType();
 				BaseTypeName = CreateElementTypeNameFromArrayTypeName(InGraphMember->GetDataType());
@@ -699,17 +713,15 @@ namespace Metasound
 
 			// Not all types have an equivalent array type. Base types without array
 			// types should have the "Is Array" checkbox disabled.
-			const bool bIsArrayTypeRegistered = bIsArrayType || IDataTypeRegistry::Get().IsRegistered(ArrayTypeName);
+			const bool bIsArrayTypeRegistered = IDataTypeRegistry::Get().IsRegistered(ArrayTypeName);
 			const bool bIsArrayTypeRegisteredHidden = MemberCustomizationPrivate::HiddenInputTypeNames.Contains(ArrayTypeName);
 
 			TArray<FName> BaseDataTypes;
 			IDataTypeRegistry::Get().IterateDataTypeInfo([&BaseDataTypes](const FDataTypeRegistryInfo& RegistryInfo)
 			{
-				// Hide the type from the combo selector if any of the following is true;
-				const bool bIsArrayType = RegistryInfo.IsArrayType();
-				const bool bIsVariable = RegistryInfo.bIsVariable;
+				// Hide the type from the combo selector if any of the following is true
 				const bool bIsHiddenType = MemberCustomizationPrivate::HiddenInputTypeNames.Contains(RegistryInfo.DataTypeName);
-				const bool bHideBaseType = bIsArrayType || bIsVariable || bIsHiddenType;
+				const bool bHideBaseType = RegistryInfo.bIsArrayType || RegistryInfo.bIsVariable || bIsHiddenType;
 				if (!bHideBaseType)
 				{
 					BaseDataTypes.Add(RegistryInfo.DataTypeName);
@@ -804,7 +816,7 @@ namespace Metasound
 				FDataTypeRegistryInfo DataTypeInfo;
 				if (ensure(IDataTypeRegistry::Get().GetDataTypeInfo(InGraphMember->GetDataType(), DataTypeInfo)))
 				{
-					return DataTypeInfo.IsArrayType() ? ECheckBoxState::Checked : ECheckBoxState::Unchecked;
+					return DataTypeInfo.bIsArrayType ? ECheckBoxState::Checked : ECheckBoxState::Unchecked;
 				}
 			}
 
@@ -1258,6 +1270,7 @@ namespace Metasound
 				}
 			}
 		}
+
 		void FMetasoundInputDetailCustomization::CustomizeGeneralCategory(IDetailLayoutBuilder& InDetailLayout)
 		{
 			FMetasoundVertexDetailCustomization::CustomizeGeneralCategory(InDetailLayout);
diff --git a/Engine/Plugins/Runtime/Metasound/Source/MetasoundEditor/Private/SMetasoundActionMenu.cpp b/Engine/Plugins/Runtime/Metasound/Source/MetasoundEditor/Private/SMetasoundActionMenu.cpp
--- a/Engine/Plugins/Runtime/Metasound/Source/MetasoundEditor/Private/SMetasoundActionMenu.cpp
+++ b/Engine/Plugins/Runtime/Metasound/Source/MetasoundEditor/Private/SMetasoundActionMenu.cpp
@@ -103,7 +103,7 @@ namespace Metasound
 
 					TypeColor = FGraphBuilder::GetPinCategoryColor(Pin->PinType);
 
-					if (RegistryInfo.IsArrayType())
+					if (RegistryInfo.bIsArrayType)
 					{
 						PinBrush = FAppStyle::GetBrush("Graph.ArrayPin.Connected");
 					}
diff --git a/Engine/Plugins/Runtime/Metasound/Source/MetasoundEngine/Private/MetasoundEngineModule.cpp b/Engine/Plugins/Runtime/Metasound/Source/MetasoundEngine/Private/MetasoundEngineModule.cpp
--- a/Engine/Plugins/Runtime/Metasound/Source/MetasoundEngine/Private/MetasoundEngineModule.cpp
+++ b/Engine/Plugins/Runtime/Metasound/Source/MetasoundEngine/Private/MetasoundEngineModule.cpp
@@ -6,6 +6,7 @@
 #include "MetasoundDataReference.h"
 #include "MetasoundDataTypeRegistrationMacro.h"
 #include "MetasoundEngineArchetypes.h"
+#include "MetasoundFrontendDataTypeTraits.h"
 #include "MetasoundInterface.h"
 #include "MetasoundSource.h"
 #include "MetasoundTrace.h"
@@ -18,10 +19,12 @@
 
 DEFINE_LOG_CATEGORY(LogMetasoundEngine);
 
+
+REGISTER_METASOUND_DATATYPE(Metasound::FAudioBusAsset, "AudioBusAsset", Metasound::ELiteralType::UObjectProxy, UAudioBus);
 REGISTER_METASOUND_DATATYPE(Metasound::FWaveAsset, "WaveAsset", Metasound::ELiteralType::UObjectProxy, USoundWave);
 REGISTER_METASOUND_DATATYPE(WaveTable::FWaveTable, "WaveTable", Metasound::ELiteralType::FloatArray)
 REGISTER_METASOUND_DATATYPE(Metasound::FWaveTableBankAsset, "WaveTableBankAsset", Metasound::ELiteralType::UObjectProxy, UWaveTableBank);
-REGISTER_METASOUND_DATATYPE(Metasound::FAudioBusAsset, "AudioBusAsset", Metasound::ELiteralType::UObjectProxy, UAudioBus);
+
 
 class FMetasoundEngineModule : public IMetasoundEngineModule
 {
diff --git a/Engine/Plugins/Runtime/Metasound/Source/MetasoundEngine/Public/MetasoundWaveTable.h b/Engine/Plugins/Runtime/Metasound/Source/MetasoundEngine/Public/MetasoundWaveTable.h
--- a/Engine/Plugins/Runtime/Metasound/Source/MetasoundEngine/Public/MetasoundWaveTable.h
+++ b/Engine/Plugins/Runtime/Metasound/Source/MetasoundEngine/Public/MetasoundWaveTable.h
@@ -77,6 +77,13 @@ namespace Metasound
 		}
 	};
 
+	// Disable arrays of WaveTables (as of addition, not supported by UX). Must be defined prior to DataType declaration macro below.
+	template<>
+	struct TEnableAutoArrayTypeRegistration<WaveTable::FWaveTable>
+	{
+		static constexpr bool Value = false;
+	};
+
 	DECLARE_METASOUND_DATA_REFERENCE_TYPES(WaveTable::FWaveTable, METASOUNDENGINE_API, FWaveTableTypeInfo, FWaveTableReadRef, FWaveTableWriteRef)
 	DECLARE_METASOUND_DATA_REFERENCE_TYPES(FWaveTableBankAsset, METASOUNDENGINE_API, FWaveTableBankAssetTypeInfo, FWaveTableBankAssetReadRef, FWaveTableBankAssetWriteRef)
 } // namespace Metasound
diff --git a/Engine/Plugins/Runtime/Metasound/Source/MetasoundFrontend/Private/MetasoundFrontendDataTypeRegistry.cpp b/Engine/Plugins/Runtime/Metasound/Source/MetasoundFrontend/Private/MetasoundFrontendDataTypeRegistry.cpp
--- a/Engine/Plugins/Runtime/Metasound/Source/MetasoundFrontend/Private/MetasoundFrontendDataTypeRegistry.cpp
+++ b/Engine/Plugins/Runtime/Metasound/Source/MetasoundFrontend/Private/MetasoundFrontendDataTypeRegistry.cpp
@@ -582,12 +582,9 @@ namespace Metasound
 					}
 
 					const FDataTypeRegistryInfo& RegistryInfo = Entry->GetDataTypeInfo();
-					if (!RegistryInfo.IsArrayType())
+					if (const UClass* Class = RegistryInfo.ProxyGeneratorClass)
 					{
-						if (const UClass* Class = RegistryInfo.ProxyGeneratorClass)
-						{
-							RegisteredObjectClasses.Add(Class, Name);
-						}
+						RegisteredObjectClasses.Add(Class, Name);
 					}
 
 					RegisteredDataTypes.Add(Name, Entry);
@@ -825,9 +822,46 @@ namespace Metasound
 			FLiteral FDataTypeRegistry::CreateLiteralFromUObjectArray(const FName& InDataType, const TArray<UObject*>& InObjectArray) const
 			{
 				TArray<Audio::IProxyDataPtr> ProxyArray;
+				const IDataTypeRegistryEntry* DataTypeEntry = FindDataTypeEntry(InDataType);
+				if (!DataTypeEntry)
+				{
+					UE_LOG(LogMetaSound, Error, TEXT("Failed to create a valid proxy from Array DataType '%s': Type is not registered."), *InDataType.ToString());
+					return Metasound::FLiteral(MoveTemp(ProxyArray));
+				}
+
+				const FDataTypeRegistryInfo& DataTypeInfo = DataTypeEntry->GetDataTypeInfo();
+
+				const bool bIsArrayType = DataTypeInfo.bIsProxyArrayParsable;
+				if (!bIsArrayType)
+				{
+					UE_LOG(LogMetaSound, Error, TEXT("Failed to create a valid proxy from DataType '%s': Type is not 'ArrayType'."), *InDataType.ToString());
+					return Metasound::FLiteral(MoveTemp(ProxyArray));
+				}
+
+				const bool bIsProxyArrayParseable = DataTypeInfo.bIsProxyArrayParsable;
+				if (!bIsProxyArrayParseable)
+				{
+					UE_LOG(LogMetaSound, Error, TEXT("Failed to create a valid proxy from DataType '%s': Type is not proxy parseable."), *InDataType.ToString());
+					return Metasound::FLiteral(MoveTemp(ProxyArray));
+				}
+
+				const FName ElementDataType = CreateElementTypeNameFromArrayTypeName(InDataType);
+				const IDataTypeRegistryEntry* ElementEntry = FindDataTypeEntry(ElementDataType);
+				if (!ElementEntry)
+				{
+					UE_LOG(LogMetaSound, Error, TEXT("Failed to create a valid proxy from DataType '%s': ElementType '%s' is not registered."), *ElementDataType.ToString());
+					return Metasound::FLiteral(MoveTemp(ProxyArray));
+				}
+
 				for (UObject* InObject : InObjectArray)
 				{
-					Audio::IProxyDataPtr ProxyPtr = CreateProxyFromUObject(InDataType, InObject);
+					Audio::IProxyDataPtr ProxyPtr = CreateProxyFromUObject(ElementDataType, InObject);
+					ProxyPtr = ElementEntry->CreateProxy(InObject);
+					if (!ProxyPtr && InObject)
+					{
+						UE_LOG(LogMetaSound, Error, TEXT("Failed to create a valid proxy from UObject '%s'."), *InObject->GetName());
+					}
+
 					ProxyArray.Emplace(MoveTemp(ProxyPtr));
 				}
 
diff --git a/Engine/Plugins/Runtime/Metasound/Source/MetasoundFrontend/Public/MetasoundDataTypeRegistrationMacro.h b/Engine/Plugins/Runtime/Metasound/Source/MetasoundFrontend/Public/MetasoundDataTypeRegistrationMacro.h
--- a/Engine/Plugins/Runtime/Metasound/Source/MetasoundFrontend/Public/MetasoundDataTypeRegistrationMacro.h
+++ b/Engine/Plugins/Runtime/Metasound/Source/MetasoundFrontend/Public/MetasoundDataTypeRegistrationMacro.h
@@ -162,6 +162,10 @@ namespace Metasound
 			RegistryInfo.PreferredLiteralType = PreferredArgType;
 
 			RegistryInfo.bIsParsable = TLiteralTraits<TDataType>::bIsParsableFromAnyLiteralType;
+			RegistryInfo.bIsArrayParseable = TLiteralTraits<TDataType>::bIsParseableFromAnyArrayLiteralType;
+
+			RegistryInfo.bIsArrayType = TIsArrayType<TDataType>::Value;
+
 			RegistryInfo.bIsDefaultParsable = TIsParsable<TDataType, FLiteral::FNone>::Value;
 			RegistryInfo.bIsBoolParsable = TIsParsable<TDataType, bool>::Value;
 			RegistryInfo.bIsIntParsable = TIsParsable<TDataType, int32>::Value;
@@ -185,7 +189,7 @@ namespace Metasound
 			{
 				RegistryInfo.ProxyGeneratorClass = UClassToUse::StaticClass();
 			}
-			else 
+			else
 			{
 				static_assert(std::is_same<UClassToUse, void>::value, "Only UObject derived classes can supply proxy interfaces.");
 				RegistryInfo.ProxyGeneratorClass = nullptr;
@@ -576,7 +580,7 @@ namespace Metasound
 		 *
 		 * @return True on success, false on failure.
 		 */
-		template<typename TDataType, ELiteralType PreferredArgType, typename UClassToUse>
+		template<typename TDataType, ELiteralType PreferredArgType>
 		bool RegisterDataTypeArrayWithFrontend()
 		{
 			using namespace MetasoundDataTypeRegistrationPrivate;
@@ -584,7 +588,8 @@ namespace Metasound
 
 			if (TEnableAutoArrayTypeRegistration<TDataType>::Value)
 			{
-				bool bSuccess = RegisterDataTypeWithFrontendInternal<TArrayType, TLiteralArrayEnum<PreferredArgType>::Value, UClassToUse>();
+				constexpr bool bIsArrayType = true;
+				bool bSuccess = RegisterDataTypeWithFrontendInternal<TArrayType, TLiteralArrayEnum<PreferredArgType>::Value>();
 				bSuccess = bSuccess && RegisterArrayNodes<TArrayType>();
 				bSuccess = bSuccess && RegisterDataTypeWithFrontendInternal<TVariable<TArrayType>>();
 				return bSuccess;
@@ -617,7 +622,7 @@ namespace Metasound
 		ensure(bSuccess);
 
 		// Register TArray<TDataType> as a metasound data type.
-		bSuccess = bSuccess && RegisterDataTypeArrayWithFrontend<TDataType, PreferredArgType, UClassToUse>();
+		bSuccess = bSuccess && RegisterDataTypeArrayWithFrontend<TDataType, PreferredArgType>();
 		ensure(bSuccess);
 
 		return bSuccess;
diff --git a/Engine/Plugins/Runtime/Metasound/Source/MetasoundFrontend/Public/MetasoundFrontendDataTypeRegistry.h b/Engine/Plugins/Runtime/Metasound/Source/MetasoundFrontend/Public/MetasoundFrontendDataTypeRegistry.h
--- a/Engine/Plugins/Runtime/Metasound/Source/MetasoundFrontend/Public/MetasoundFrontendDataTypeRegistry.h
+++ b/Engine/Plugins/Runtime/Metasound/Source/MetasoundFrontend/Public/MetasoundFrontendDataTypeRegistry.h
@@ -46,30 +46,33 @@ namespace Metasound
 			bool bIsStringArrayParsable = false;
 			bool bIsProxyArrayParsable = false;
 
-			// Is an TEnum wrapped enum
+			// Is a TEnum wrapped enum
 			bool bIsEnum = false;
 
 			// Determines whether the type can be used with send/receive transmitters
 			bool bIsTransmittable = false;
 
-			// Is true if is a Variable type
+			// Returns if DataType is a Variable type
 			bool bIsVariable = false;
 
-			// Is true if the data type can be used in on constructor vertices. 
+			// Returns if DataType can be used for constructor vertices.
 			bool bIsConstructorType = false;
 
-			// If provided in registration call, UClass this datatype was registered with
+			// Returns if DataType represents an array type (ex. TArray<float>, TArray<int32>, etc.).
+			bool bIsArrayType = false;
+
+			// Returns if DataType supports array parsing and passing array of base type to constructor.
+			bool bIsArrayParseable = false;
+
+			// If provided in registration call, UClass this datatype was registered with.
 			UClass* ProxyGeneratorClass = nullptr;
 
+
 			// Returns if DataType supports array parsing.
+			UE_DEPRECATED(5.1, "Use bIsArrayType or bIsArrayParseable members instead")
 			FORCEINLINE bool IsArrayType() const
 			{
-				return bIsDefaultArrayParsable
-				|| bIsBoolArrayParsable
-				|| bIsIntArrayParsable
-				|| bIsFloatArrayParsable
-				|| bIsStringArrayParsable
-				|| bIsProxyArrayParsable;
+				return bIsArrayType;
 			}
 		};
 
diff --git a/Engine/Plugins/Runtime/Metasound/Source/MetasoundGraphCore/Public/MetasoundDataFactory.h b/Engine/Plugins/Runtime/Metasound/Source/MetasoundGraphCore/Public/MetasoundDataFactory.h
--- a/Engine/Plugins/Runtime/Metasound/Source/MetasoundGraphCore/Public/MetasoundDataFactory.h
+++ b/Engine/Plugins/Runtime/Metasound/Source/MetasoundGraphCore/Public/MetasoundDataFactory.h
@@ -754,6 +754,20 @@ namespace Metasound
 		static constexpr bool Value = FElementConstructorTraits::bIsConstructibleWithArgs || FElementConstructorTraits::bIsConstructibleWithSettingsAndArgs ;
 	};
 
+	/* Determines whether a DataType is an array type. */
+	template<typename DataType>
+	struct TIsArrayType
+	{
+		static constexpr bool Value = false;
+	};
+
+	/* Specialization for Array DataTypes. */
+	template<typename DataElementType>
+	struct TIsArrayType<TArray<DataElementType>>
+	{
+		static constexpr bool Value = true;
+	};
+
 	/** Determines whether a DataType supports construction using the given literal.
 	 *
 	 * @tparam DataType - The registered Metasound Data Type.
@@ -763,25 +777,28 @@ namespace Metasound
 	{
 		using TVariantConstructorTraits = DataFactoryPrivate::TDataTypeVariantConstructorTraits<DataType, FLiteral::FVariantType>;
 
-		static constexpr bool bIsParsableFromAnyLiteralType = 
-			TVariantConstructorTraits::bIsDefaultConstructible || 
-			TVariantConstructorTraits::bIsConstructibleWithSettings || 
-			TVariantConstructorTraits::bIsConstructibleWithArgs || 
-			TVariantConstructorTraits::bIsConstructibleWithSettingsAndArgs || 
-			TVariantConstructorTraits::bIsArrayDefaultConstructible || 
-			TVariantConstructorTraits::bIsArrayConstructibleWithSettings || 
-			TVariantConstructorTraits::bIsArrayConstructibleWithArgs || 
+		static constexpr bool bIsParseableFromAnyArrayLiteralType =
+			TVariantConstructorTraits::bIsArrayDefaultConstructible ||
+			TVariantConstructorTraits::bIsArrayConstructibleWithSettings ||
+			TVariantConstructorTraits::bIsArrayConstructibleWithArgs ||
 			TVariantConstructorTraits::bIsArrayConstructibleWithSettingsAndArgs;
 
+		static constexpr bool bIsParsableFromAnyLiteralType =
+			TVariantConstructorTraits::bIsDefaultConstructible ||
+			TVariantConstructorTraits::bIsConstructibleWithSettings ||
+			TVariantConstructorTraits::bIsConstructibleWithArgs ||
+			TVariantConstructorTraits::bIsConstructibleWithSettingsAndArgs ||
+			bIsParseableFromAnyArrayLiteralType;
+
 		/** Determines if a constructor for the DataType exists which accepts 
 		 * an FOperatorSettings with the literals constructor arg type, and/or one that
-		 * accpets the literal constructor arg type. 
+		 * accepts the literal constructor arg type. 
 		 *
 		 * @param InLiteral - The literal containing the constructor argument.
 		 *
 		 * @return True if a constructor for the DataType exists which accepts 
 		 * an FOperatorSettings with the literals constructor arg type, or one that
-		 * accpets the literal constructor arg type. It returns False otherwise.
+		 * accepts the literal constructor arg type. It returns False otherwise.
 		 */
 		static const bool IsParsable(const FLiteral& InLiteral)
 		{
