MetaSound Constructor pin editor-frontend integration checkpoint
- Both inputs and outputs can be constructor vertices, not just inputs
- Enforce connection and data type constraints for access types
- Prevent default value constructor value from being changed while playing
- Added cvar (default disabled) for enabling constructor pin option in editor

diff --git a/Engine/Plugins/Runtime/Metasound/Source/MetasoundEditor/Private/MetasoundEditor.cpp b/Engine/Plugins/Runtime/Metasound/Source/MetasoundEditor/Private/MetasoundEditor.cpp
--- a/Engine/Plugins/Runtime/Metasound/Source/MetasoundEditor/Private/MetasoundEditor.cpp
+++ b/Engine/Plugins/Runtime/Metasound/Source/MetasoundEditor/Private/MetasoundEditor.cpp
@@ -517,9 +517,10 @@ namespace Metasound
 				{
 
 					UMetasoundEditorGraphMember* GraphMember = GraphMemberAction->GetGraphMember();
-					if (const UMetasoundEditorGraphInput* Input = Cast<UMetasoundEditorGraphInput>(GraphMember))
+					if (const UMetasoundEditorGraphVertex* Vertex = Cast<UMetasoundEditorGraphVertex>(GraphMember))
 					{
-						bIsConstructorPin = Input->ConstructorPin;
+						EMetasoundFrontendVertexAccessType AccessType = Vertex->GetVertexAccessType();
+						bIsConstructorPin = AccessType == EMetasoundFrontendVertexAccessType::Value;
 					}
 					FName DataTypeName = GraphMember->GetDataType(); 
 
@@ -3321,7 +3322,10 @@ namespace Metasound
 					Metasound->Modify();
 
 					const FName DataTypeName = GetMetasoundDataTypeName<float>();
-					Frontend::FNodeHandle NodeHandle = FGraphBuilder::AddInputNodeHandle(*Metasound, DataTypeName);
+					FCreateNodeVertexParams VertexParams;
+					VertexParams.DataType = DataTypeName;
+
+					Frontend::FNodeHandle NodeHandle = FGraphBuilder::AddInputNodeHandle(*Metasound, VertexParams);
 					if (ensure(NodeHandle->IsValid()))
 					{
 						NameToSelect = NodeHandle->GetNodeName();
@@ -3347,7 +3351,10 @@ namespace Metasound
 					Metasound->Modify();
 
 					const FName DataTypeName = GetMetasoundDataTypeName<float>();
-					Frontend::FNodeHandle NodeHandle = FGraphBuilder::AddOutputNodeHandle(*Metasound, DataTypeName);
+					FCreateNodeVertexParams VertexParams;
+					VertexParams.DataType = DataTypeName;
+
+					Frontend::FNodeHandle NodeHandle = FGraphBuilder::AddOutputNodeHandle(*Metasound, VertexParams);
 					if (ensure(NodeHandle->IsValid()))
 					{
 						NameToSelect = NodeHandle->GetNodeName();
diff --git a/Engine/Plugins/Runtime/Metasound/Source/MetasoundEditor/Private/MetasoundEditorGraph.cpp b/Engine/Plugins/Runtime/Metasound/Source/MetasoundEditor/Private/MetasoundEditorGraph.cpp
--- a/Engine/Plugins/Runtime/Metasound/Source/MetasoundEditor/Private/MetasoundEditorGraph.cpp
+++ b/Engine/Plugins/Runtime/Metasound/Source/MetasoundEditor/Private/MetasoundEditorGraph.cpp
@@ -7,7 +7,6 @@
 #include "EdGraph/EdGraphNode.h"
 #include "Interfaces/ITargetPlatform.h"
 #include "MetasoundAssetBase.h"
-#include "MetasoundEditorGraphBuilder.h"
 #include "MetasoundEditorGraphInputNode.h"
 #include "MetasoundEditorGraphMemberDefaults.h"
 #include "MetasoundEditorGraphNode.h"
@@ -343,7 +342,11 @@ void UMetasoundEditorGraphVertex::SetDataType(FName InNewType, bool bPostTransac
 
 	// 4. Add the new input node with the same identifier data but new datatype.
 	UObject& Metasound = Graph->GetMetasoundChecked();
-	FNodeHandle NewNodeHandle = AddNodeHandle(NodeName, InNewType);
+
+	FCreateNodeVertexParams VertexParams;
+	VertexParams.DataType = InNewType;
+
+	FNodeHandle NewNodeHandle = AddNodeHandle(NodeName, VertexParams);
 	NewNodeHandle->SetNodeName(NodeName);
 	NewNodeHandle->SetDisplayName(NodeDisplayName);
 
@@ -369,6 +372,101 @@ void UMetasoundEditorGraphVertex::SetDataType(FName InNewType, bool bPostTransac
 	NameChanged.Broadcast(NodeID);
 }
 
+void UMetasoundEditorGraphVertex::SetVertexAccessType(EMetasoundFrontendVertexAccessType InNewAccessType, bool bPostTransaction)
+{
+	using namespace Metasound::Editor;
+	using namespace Metasound::Frontend;
+
+	if (InNewAccessType == GetVertexAccessType())
+	{
+		return;
+	}
+
+	UMetasoundEditorGraph* Graph = GetOwningGraph();
+	if (!ensure(Graph))
+	{
+		return;
+	}
+
+	const FScopedTransaction Transaction(LOCTEXT("SetGraphVertexType", "Set MetaSound GraphVertex Type"), bPostTransaction);
+	Graph->GetMetasoundChecked().Modify();
+	Graph->Modify();
+	Modify();
+
+	// Logic is based on UMetasoundEditorGraphMember::SetDataType
+	// 1. Cache current editor node reference positions & linked pins
+	TArray<UMetasoundEditorGraphMemberNode*> Nodes = GetNodes();
+	TArray<FVector2D> NodeLocations;
+	// Map of node index within NodeLocations to pins linked to that node
+	TMap<int32, TArray<UEdGraphPin*>> LinkedPins;
+	for (int32 i = 0; i < Nodes.Num(); ++i)
+	{
+		const UMetasoundEditorGraphMemberNode* Node = Nodes[i];
+		if (ensure(Node))
+		{
+			NodeLocations.Add(FVector2D(Node->NodePosX, Node->NodePosY));
+
+			ensure(Node->Pins.Num() == 1);
+			if (UEdGraphPin* NodePin = Node->Pins.Last()) 
+			{
+				if (NodePin->LinkedTo.Num() > 0)
+				{
+					LinkedPins.Add(i, NodePin->LinkedTo);
+				}
+			}
+		}
+	}
+
+	// 2. Cache the old version's Frontend data.
+	FNodeHandle NodeHandle = GetNodeHandle();
+	const FName NodeName = NodeHandle->GetNodeName();
+	const FText NodeDisplayName = NodeHandle->GetDisplayName();
+	const FName NodeType = GetDataType();
+	
+	// 3. Remove the current nodes and vertex
+	Graph->RemoveMemberNodes(*this);
+	Graph->RemoveFrontendMember(*this);
+
+	// 4. Add the new input node with the same identifier data but new access type.
+	FCreateNodeVertexParams VertexParams;
+	VertexParams.DataType = NodeType;
+	VertexParams.AccessType = InNewAccessType;
+	UObject& Metasound = Graph->GetMetasoundChecked();
+	
+	FNodeHandle NewNodeHandle = AddNodeHandle(NodeName, VertexParams);
+	NewNodeHandle->SetNodeName(NodeName);
+	NewNodeHandle->SetDisplayName(NodeDisplayName);
+
+	if (!ensure(NewNodeHandle->IsValid()))
+	{
+		return;
+	}
+
+	ClassName = NewNodeHandle->GetClassMetadata().GetClassName();
+	NodeID = NewNodeHandle->GetID();
+
+	InitializeLiteral();
+
+	// 6. Create new node references in the same locations as the old locations
+	// and reconnect linked pins 
+	for (int32 i = 0; i < NodeLocations.Num(); ++i)
+	{
+		UMetasoundEditorGraphNode* NewEdGraphNode = FGraphBuilder::AddNode(Metasound, NewNodeHandle, NodeLocations[i], false /* bInSelectNewNode */);
+		
+		if (LinkedPins.Contains(i))
+		{
+			for (UEdGraphPin* LinkedPin : LinkedPins[i])
+			{
+				NewEdGraphNode->AutowireNewNode(LinkedPin);
+			}
+		}
+	}
+
+	// Notify now that the node has a new ID (doing so before creating & syncing Frontend Node &
+	// EdGraph variable can result in refreshing editors while in a desync'ed state)
+	NameChanged.Broadcast(NodeID);
+}
+
 Metasound::Frontend::FNodeHandle UMetasoundEditorGraphVertex::GetNodeHandle()
 {
 	using namespace Metasound;
@@ -496,6 +594,13 @@ Metasound::Editor::ENodeSection UMetasoundEditorGraphInput::GetSectionID() const
 }
 
 Metasound::Frontend::FNodeHandle UMetasoundEditorGraphInput::AddNodeHandle(const FName& InName, FName InDataType)
+{
+	UE_LOG(LogMetaSound, Error, TEXT("UMetasoundEditorGraphInput::AddNodeHandle with these parameters is no longer supported and should not be called. Use the one with FCreateNodeVertexParams instead."));
+
+	return Metasound::Frontend::INodeController::GetInvalidHandle();
+}
+
+Metasound::Frontend::FNodeHandle UMetasoundEditorGraphInput::AddNodeHandle(const FName& InName, const Metasound::Editor::FCreateNodeVertexParams& InParams)
 {
 	using namespace Metasound::Editor;
 	using namespace Metasound::Frontend;
@@ -507,7 +612,7 @@ Metasound::Frontend::FNodeHandle UMetasoundEditorGraphInput::AddNodeHandle(const
 	}
 
 	UObject& Metasound = Graph->GetMetasoundChecked();
-	Metasound::Frontend::FNodeHandle NewNodeHandle = FGraphBuilder::AddInputNodeHandle(Metasound, InDataType, nullptr /* DefaultValue */, &InName);
+	Metasound::Frontend::FNodeHandle NewNodeHandle = FGraphBuilder::AddInputNodeHandle(Metasound, InParams, nullptr /* DefaultValue */, &InName);
 	return NewNodeHandle;
 }
 
@@ -640,7 +745,32 @@ void UMetasoundEditorGraphInput::UpdateFrontendDefaultLiteral(bool bPostTransact
 	}
 }
 
+EMetasoundFrontendVertexAccessType UMetasoundEditorGraphInput::GetVertexAccessType() const
+{
+	using namespace Metasound::Frontend;
+
+	const FName MemberName = GetMemberName();
+	const UMetasoundEditorGraph* MetaSoundGraph = GetOwningGraph();
+	FConstGraphHandle GraphHandle = MetaSoundGraph->GetGraphHandle();
+
+	FConstClassInputAccessPtr ClassInputPtr = GraphHandle->FindClassInputWithName(MemberName);
+
+	if (const FMetasoundFrontendClassInput* InputPtr = ClassInputPtr.Get())
+	{
+		return InputPtr->AccessType;
+	}
+
+	return EMetasoundFrontendVertexAccessType::Reference;
+}
+
 Metasound::Frontend::FNodeHandle UMetasoundEditorGraphOutput::AddNodeHandle(const FName& InName, FName InDataType)
+{
+	UE_LOG(LogMetaSound, Error, TEXT("UMetasoundEditorGraphOutput::AddNodeHandle with these parameters is no longer supported and should not be called. Use the one with FCreateNodeVertexParams instead."));
+
+	return Metasound::Frontend::INodeController::GetInvalidHandle();
+}
+
+Metasound::Frontend::FNodeHandle UMetasoundEditorGraphOutput::AddNodeHandle(const FName& InName, const Metasound::Editor::FCreateNodeVertexParams& InParams)
 {
 	using namespace Metasound::Editor;
 	using namespace Metasound::Frontend;
@@ -649,10 +779,10 @@ Metasound::Frontend::FNodeHandle UMetasoundEditorGraphOutput::AddNodeHandle(cons
 	if (!ensure(Graph))
 	{
 		return Metasound::Frontend::INodeController::GetInvalidHandle();
-	}
+	} 
 
 	UObject& Metasound = Graph->GetMetasoundChecked();
-	FNodeHandle NewNodeHandle = FGraphBuilder::AddOutputNodeHandle(Metasound, InDataType, &InName);
+	FNodeHandle NewNodeHandle = FGraphBuilder::AddOutputNodeHandle(Metasound, InParams, &InName);
 	return NewNodeHandle;
 }
 
@@ -756,6 +886,24 @@ void UMetasoundEditorGraphOutput::UpdateFrontendDefaultLiteral(bool bPostTransac
 	}
 }
 
+EMetasoundFrontendVertexAccessType UMetasoundEditorGraphOutput::GetVertexAccessType() const
+{
+	using namespace Metasound::Frontend;
+
+	const FName MemberName = GetMemberName();
+	const UMetasoundEditorGraph* MetaSoundGraph = GetOwningGraph();
+	FConstGraphHandle GraphHandle = MetaSoundGraph->GetGraphHandle();
+
+	FConstClassOutputAccessPtr ClassOutputPtr = GraphHandle->FindClassOutputWithName(MemberName);
+
+	if (const FMetasoundFrontendClassOutput* OutputPtr = ClassOutputPtr.Get())
+	{
+		return OutputPtr->AccessType;
+	}
+
+	return EMetasoundFrontendVertexAccessType::Reference;
+}
+
 Metasound::Editor::ENodeSection UMetasoundEditorGraphOutput::GetSectionID() const 
 {
 	return Metasound::Editor::ENodeSection::Outputs;
diff --git a/Engine/Plugins/Runtime/Metasound/Source/MetasoundEditor/Private/MetasoundEditorGraphBuilder.cpp b/Engine/Plugins/Runtime/Metasound/Source/MetasoundEditor/Private/MetasoundEditorGraphBuilder.cpp
--- a/Engine/Plugins/Runtime/Metasound/Source/MetasoundEditor/Private/MetasoundEditorGraphBuilder.cpp
+++ b/Engine/Plugins/Runtime/Metasound/Source/MetasoundEditor/Private/MetasoundEditorGraphBuilder.cpp
@@ -993,7 +993,10 @@ namespace Metasound
 					if (ensure(Input))
 					{
 						const FName PinName = Pin->GetFName();
-						NodeHandle = AddInputNodeHandle(InMetaSound, Input->GetDataType(), nullptr, &PinName);
+						FCreateNodeVertexParams VertexParams;
+						VertexParams.DataType = Input->GetDataType();
+
+						NodeHandle = AddInputNodeHandle(InMetaSound, VertexParams, nullptr, &PinName);
 						NodeHandle->SetDescription(InGraphNode.GetTooltipText());
 					}
 				}
@@ -1009,7 +1012,10 @@ namespace Metasound
 					if (ensure(Output))
 					{
 						const FName PinName = Pin->GetFName();
-						NodeHandle = FGraphBuilder::AddOutputNodeHandle(InMetaSound, Output->GetDataType(), &PinName);
+						FCreateNodeVertexParams VertexParams;
+						VertexParams.DataType = Output->GetDataType();
+
+						NodeHandle = FGraphBuilder::AddOutputNodeHandle(InMetaSound, VertexParams, &PinName);
 						NodeHandle->SetDescription(InGraphNode.GetTooltipText());
 					}
 				}
@@ -1045,6 +1051,12 @@ namespace Metasound
 		}
 
 		Frontend::FNodeHandle FGraphBuilder::AddInputNodeHandle(UObject& InMetaSound, const FName InTypeName, const FMetasoundFrontendLiteral* InDefaultValue, const FName* InNameBase)
+		{
+			UE_LOG(LogMetaSound, Error, TEXT("FGraphBuilder::AddInputNodeHandle with these parameters is no longer supported and should not be called. Use the one with FCreateNodeVertexParams instead."));
+			return Frontend::INodeController::GetInvalidHandle();
+		}
+
+		Frontend::FNodeHandle FGraphBuilder::AddInputNodeHandle(UObject& InMetaSound, const FCreateNodeVertexParams& InParams, const FMetasoundFrontendLiteral* InDefaultValue, const FName* InNameBase)
 		{
 			FMetasoundAssetBase* MetaSoundAsset = IMetasoundUObjectRegistry::Get().GetObjectAsAssetBase(&InMetaSound);
 			check(MetaSoundAsset);
@@ -1052,8 +1064,9 @@ namespace Metasound
 
 			FMetasoundFrontendClassInput ClassInput;
 			ClassInput.Name = GenerateUniqueNameByClassType(InMetaSound, EMetasoundFrontendClassType::Input, InNameBase ? InNameBase->ToString() : TEXT("Input"));
-			ClassInput.TypeName = InTypeName;
+			ClassInput.TypeName = InParams.DataType;
 			ClassInput.VertexID = VertexID;
+			ClassInput.AccessType = InParams.AccessType;
 
 			if (nullptr != InDefaultValue)
 			{
@@ -1061,20 +1074,34 @@ namespace Metasound
 			}
 			else
 			{
-				Metasound::FLiteral Literal = Frontend::IDataTypeRegistry::Get().CreateDefaultLiteral(InTypeName);
+				Metasound::FLiteral Literal = Frontend::IDataTypeRegistry::Get().CreateDefaultLiteral(InParams.DataType);
 				ClassInput.DefaultLiteral.SetFromLiteral(Literal);
 			}
 
 			return MetaSoundAsset->GetRootGraphHandle()->AddInputVertex(ClassInput);
 		}
-
+		
 		Frontend::FNodeHandle FGraphBuilder::AddOutputNodeHandle(UObject& InMetaSound, const FName InTypeName, const FName* InNameBase)
+		{
+			UE_LOG(LogMetaSound, Error, TEXT("FGraphBuilder::AddOutputNodeHandle with these parameters is no longer supported and should not be called. Use the one with FCreateNodeVertexParams instead."));
+			return Frontend::INodeController::GetInvalidHandle();
+		}
+
+		Frontend::FNodeHandle FGraphBuilder::AddOutputNodeHandle(UObject& InMetaSound, const FCreateNodeVertexParams& InParams, const FName* InNameBase)
 		{
 			FMetasoundAssetBase* MetaSoundAsset = IMetasoundUObjectRegistry::Get().GetObjectAsAssetBase(&InMetaSound);
 			check(MetaSoundAsset);
+			const FGuid VertexID = FGuid::NewGuid();
 
 			const FName NewName = GenerateUniqueNameByClassType(InMetaSound, EMetasoundFrontendClassType::Output, InNameBase ? InNameBase->ToString() : TEXT("Output"));
-			return MetaSoundAsset->GetRootGraphHandle()->AddOutputVertex(NewName, InTypeName);
+
+			FMetasoundFrontendClassOutput ClassOutput;
+			ClassOutput.Name = NewName;
+			ClassOutput.TypeName = InParams.DataType;
+			ClassOutput.VertexID = VertexID;
+			ClassOutput.AccessType = InParams.AccessType;
+
+			return MetaSoundAsset->GetRootGraphHandle()->AddOutputVertex(ClassOutput);
 		}
 
 		FName FGraphBuilder::GenerateUniqueVariableName(const Frontend::FConstGraphHandle& InFrontendGraph, const FString& InBaseName)
diff --git a/Engine/Plugins/Runtime/Metasound/Source/MetasoundEditor/Private/MetasoundEditorGraphBuilder.h b/Engine/Plugins/Runtime/Metasound/Source/MetasoundEditor/Private/MetasoundEditorGraphBuilder.h
--- a/Engine/Plugins/Runtime/Metasound/Source/MetasoundEditor/Private/MetasoundEditorGraphBuilder.h
+++ b/Engine/Plugins/Runtime/Metasound/Source/MetasoundEditor/Private/MetasoundEditorGraphBuilder.h
@@ -37,6 +37,7 @@ namespace Metasound
 		// Forward Declarations
 		class FEditor;
 		class IMetasoundEditorModule;
+		struct FCreateNodeVertexParams;
 
 		class FGraphBuilder
 		{
@@ -118,13 +119,20 @@ namespace Metasound
 			// Adds a corresponding UMetasoundEditorGraphInputNode for the provided node handle.
 			static UMetasoundEditorGraphInputNode* AddInputNode(UObject& InMetaSound, Frontend::FNodeHandle InNodeHandle, FVector2D InLocation, bool bInSelectNewNode = true);
 
-			// Generates FNodeHandle for the given external node data. Does not bind or create EdGraph representation of given node.
+			UE_DEPRECATED(5.1, "Use AddInputNodeHandle with FCreateNodeVertexParams instead.")
 			static Frontend::FNodeHandle AddInputNodeHandle(
 				UObject& InMetaSound,
 				const FName InTypeName,
 				const FMetasoundFrontendLiteral* InDefaultValue = nullptr,
 				const FName* InNameBase = nullptr);
 
+			// Generates FNodeHandle for the given external node data. Does not bind or create EdGraph representation of given node.
+			static Frontend::FNodeHandle AddInputNodeHandle(
+				UObject& InMetaSound,
+				const FCreateNodeVertexParams& InVertexParams,
+				const FMetasoundFrontendLiteral* InDefaultValue = nullptr,
+				const FName* InNameBase = nullptr);
+
 			// Adds a corresponding UMetasoundEditorGraphExternalNode for the provided node handle.
 			static UMetasoundEditorGraphExternalNode* AddExternalNode(UObject& InMetaSound, Frontend::FNodeHandle& InNodeHandle, FVector2D InLocation, bool bInSelectNewNode = true);
 
@@ -138,9 +146,12 @@ namespace Metasound
 			// Adds an output node to the editor graph that corresponds to the provided node handle.
 			static UMetasoundEditorGraphOutputNode* AddOutputNode(UObject& InMetaSound, Frontend::FNodeHandle& InNodeHandle, FVector2D InLocation, bool bInSelectNewNode = true);
 
-			// Generates analogous FNodeHandle for the given internal node data. Does not bind nor create EdGraph representation of given node.
+			UE_DEPRECATED(5.1, "Use AddOutputNodeHandle with FCreateNodeVertexParams instead.")
 			static Frontend::FNodeHandle AddOutputNodeHandle(UObject& InMetaSound, const FName InTypeName, const FName* InNameBase = nullptr);
 
+			// Generates analogous FNodeHandle for the given internal node data. Does not bind nor create EdGraph representation of given node.
+			static Frontend::FNodeHandle AddOutputNodeHandle(UObject& InMetaSound, const FCreateNodeVertexParams& InParams, const FName* InNameBase = nullptr);
+
 			// Create a unique name for the variable.
 			static FName GenerateUniqueVariableName(const Frontend::FConstGraphHandle& InFrontendGraph, const FString& InBaseName);
 
diff --git a/Engine/Plugins/Runtime/Metasound/Source/MetasoundEditor/Private/MetasoundEditorGraphSchema.cpp b/Engine/Plugins/Runtime/Metasound/Source/MetasoundEditor/Private/MetasoundEditorGraphSchema.cpp
--- a/Engine/Plugins/Runtime/Metasound/Source/MetasoundEditor/Private/MetasoundEditorGraphSchema.cpp
+++ b/Engine/Plugins/Runtime/Metasound/Source/MetasoundEditor/Private/MetasoundEditorGraphSchema.cpp
@@ -154,7 +154,8 @@ namespace Metasound
 						if (Pin->Direction == EGPD_Output)
 						{
 							FOutputHandle OutputHandle = FGraphBuilder::GetOutputHandleFromPin(Pin);
-							if (OutputHandle->IsValid() && OutputHandle->GetDataType() == InputHandle->GetDataType())
+							if (OutputHandle->IsValid() && 
+								InputHandle->CanConnectTo(*OutputHandle).Connectable == FConnectability::EConnectable::Yes)
 							{
 								if (ensure(FGraphBuilder::ConnectNodes(*FromPin, *Pin, true /* bConnectEdPins */)))
 								{
@@ -173,7 +174,8 @@ namespace Metasound
 						if (Pin->Direction == EGPD_Input)
 						{
 							FInputHandle InputHandle = FGraphBuilder::GetInputHandleFromPin(Pin);
-							if (InputHandle->IsValid() && InputHandle->GetDataType() == OutputHandle->GetDataType())
+							if (InputHandle->IsValid() && 
+								InputHandle->CanConnectTo(*OutputHandle).Connectable == FConnectability::EConnectable::Yes)
 							{
 								if (ensure(FGraphBuilder::ConnectNodes(*Pin, *FromPin, true /* bConnectEdPins */)))
 								{
@@ -548,7 +550,10 @@ UEdGraphNode* FMetasoundGraphSchemaAction_PromoteToInput::PerformAction(UEdGraph
 	FGraphBuilder::GetPinLiteral(*FromPin, DefaultValue);
 
 	const FName InputName = InputHandle->GetName();
-	FNodeHandle NodeHandle = FGraphBuilder::AddInputNodeHandle(ParentMetasound, InputHandle->GetDataType(), &DefaultValue, &InputName);
+	// The promoted input must have the same vertex access type in order for it to be connectable 
+	const FCreateNodeVertexParams VertexParams = { InputHandle->GetDataType(), InputHandle->GetVertexAccessType() };
+
+	FNodeHandle NodeHandle = FGraphBuilder::AddInputNodeHandle(ParentMetasound, VertexParams, &DefaultValue, &InputName);
 	if (ensure(NodeHandle->IsValid()))
 	{
 		UMetasoundEditorGraphInput* Input = MetasoundGraph->FindOrAddInput(NodeHandle);
@@ -766,7 +771,9 @@ UEdGraphNode* FMetasoundGraphSchemaAction_PromoteToOutput::PerformAction(UEdGrap
 
 	const FString OutputName = OutputHandle->GetName().ToString();
 	const FVertexName NewNodeName = FGraphBuilder::GenerateUniqueNameByClassType(ParentMetasound, EMetasoundFrontendClassType::Output, OutputName);
-	FNodeHandle NodeHandle = FGraphBuilder::AddOutputNodeHandle(ParentMetasound, OutputHandle->GetDataType(), &NewNodeName);
+	const FCreateNodeVertexParams VertexParams = { OutputHandle->GetDataType(), OutputHandle->GetVertexAccessType() };
+
+	FNodeHandle NodeHandle = FGraphBuilder::AddOutputNodeHandle(ParentMetasound, VertexParams, &NewNodeName);
 	if (ensure(NodeHandle->IsValid()))
 	{
 		UMetasoundEditorGraphOutput* Output = MetasoundGraph->FindOrAddOutput(NodeHandle);
@@ -1011,32 +1018,40 @@ void UMetasoundEditorGraphSchema::GetGraphContextActions(FGraphContextMenuBuilde
 
 	FActionClassFilters ClassFilters;
 	FConstGraphHandle GraphHandle = IGraphController::GetInvalidHandle();
+	bool bIsConstructorOutput = false;
 	if (const UEdGraphPin* FromPin = ContextMenuBuilder.FromPin)
 	{
 		if (FromPin->Direction == EGPD_Input)
 		{
 			FConstInputHandle InputHandle = FGraphBuilder::GetConstInputHandleFromPin(ContextMenuBuilder.FromPin);
+			bIsConstructorOutput = InputHandle->GetVertexAccessType() == EMetasoundFrontendVertexAccessType::Value;
+
 			ClassFilters.OutputFilterFunction = [InputHandle](const FMetasoundFrontendClassOutput& InOutput)
 			{
-				return InOutput.TypeName == InputHandle->GetDataType();
+				return InOutput.TypeName == InputHandle->GetDataType() && 
+					FMetasoundFrontendClassVertex::CanConnectVertexAccessTypes(InOutput.AccessType, InputHandle->GetVertexAccessType());
 			};
 
 			// Show only input nodes as output nodes can only connected if FromPin is input
 			GraphHandle = InputHandle->GetOwningNode()->GetOwningGraph();
 			GetDataTypeInputNodeActions(ContextMenuBuilder, GraphHandle, [InputHandle](FConstNodeHandle NodeHandle)
 			{
-				bool bHasOutputOfType = false;
+				bool bHasConnectableOutput = false;
 				NodeHandle->IterateConstOutputs([&](FConstOutputHandle PotentialOutputHandle)
 				{
-					bHasOutputOfType |= PotentialOutputHandle->GetDataType() == InputHandle->GetDataType();
+					bHasConnectableOutput |= (InputHandle->CanConnectTo(*PotentialOutputHandle).Connectable == FConnectability::EConnectable::Yes);
 				});
-				return bHasOutputOfType;
+				return bHasConnectableOutput;
 			});
 
 			FGraphActionMenuBuilder& ActionMenuBuilder = static_cast<FGraphActionMenuBuilder&>(ContextMenuBuilder);
 			ActionMenuBuilder.AddAction(MakeShared<FMetasoundGraphSchemaAction_PromoteToInput>());
-			ActionMenuBuilder.AddAction(MakeShared<FMetasoundGraphSchemaAction_PromoteToVariable_AccessorNode>());
-			ActionMenuBuilder.AddAction(MakeShared<FMetasoundGraphSchemaAction_PromoteToVariable_DeferredAccessorNode>());
+			// Constructor outputs cannot be promoted to variables
+			if (!bIsConstructorOutput)
+			{
+				ActionMenuBuilder.AddAction(MakeShared<FMetasoundGraphSchemaAction_PromoteToVariable_AccessorNode>());
+				ActionMenuBuilder.AddAction(MakeShared<FMetasoundGraphSchemaAction_PromoteToVariable_DeferredAccessorNode>());
+			}
 		}
 
 		if (FromPin->Direction == EGPD_Output)
@@ -1044,19 +1059,20 @@ void UMetasoundEditorGraphSchema::GetGraphContextActions(FGraphContextMenuBuilde
 			FConstOutputHandle OutputHandle = FGraphBuilder::GetConstOutputHandleFromPin(FromPin);
 			ClassFilters.InputFilterFunction = [OutputHandle](const FMetasoundFrontendClassInput& InInput)
 			{
-				return InInput.TypeName == OutputHandle->GetDataType();
+				return InInput.TypeName == OutputHandle->GetDataType() &&
+					FMetasoundFrontendClassVertex::CanConnectVertexAccessTypes(OutputHandle->GetVertexAccessType(), InInput.AccessType);
 			};
 
 			// Show only output nodes as input nodes can only connected if FromPin is output
 			GraphHandle = OutputHandle->GetOwningNode()->GetOwningGraph();
 			GetDataTypeOutputNodeActions(ContextMenuBuilder, GraphHandle, [OutputHandle](FConstNodeHandle NodeHandle)
 			{
-				bool bHasInputOfType = false;
+				bool bHasConnectableInput = false;
 				NodeHandle->IterateConstInputs([&](FConstInputHandle PotentialInputHandle)
 				{
-					bHasInputOfType |= PotentialInputHandle->GetDataType() == OutputHandle->GetDataType();
+					bHasConnectableInput |= (PotentialInputHandle->CanConnectTo(*OutputHandle).Connectable == FConnectability::EConnectable::Yes);
 				});
-				return bHasInputOfType;
+				return bHasConnectableInput;
 			});
 
 			FGraphActionMenuBuilder& ActionMenuBuilder = static_cast<FGraphActionMenuBuilder&>(ContextMenuBuilder);
@@ -1086,8 +1102,13 @@ void UMetasoundEditorGraphSchema::GetGraphContextActions(FGraphContextMenuBuilde
 	}
 
 	GetFunctionActions(ContextMenuBuilder, ClassFilters, true /* bShowSelectedActions */, GraphHandle);
-	GetVariableActions(ContextMenuBuilder, ClassFilters, true /* bShowSelectedActions */, GraphHandle);
-	GetConversionActions(ContextMenuBuilder, ClassFilters);
+
+	// Variable and conversion actions are always by reference so are incompatible with constructor outputs 
+	if (!bIsConstructorOutput)
+	{
+		GetVariableActions(ContextMenuBuilder, ClassFilters, true /* bShowSelectedActions */, GraphHandle);
+		GetConversionActions(ContextMenuBuilder, ClassFilters);
+	}
 }
 
 void UMetasoundEditorGraphSchema::GetContextMenuActions(class UToolMenu* Menu, class UGraphNodeContextMenuContext* Context) const
@@ -1214,6 +1235,10 @@ const FPinConnectionResponse UMetasoundEditorGraphSchema::CanCreateConnection(co
 			{
 				return FPinConnectionResponse(CONNECT_RESPONSE_DISALLOW, LOCTEXT("ConnectionLoop1", "Connection causes loop"));
 			}
+			else if (Frontend::FConnectability::EReason::IncompatibleAccessTypes == Connectability.Reason)
+			{
+				return FPinConnectionResponse(CONNECT_RESPONSE_DISALLOW, LOCTEXT("ConnectionIncompatibleAccessTypes", "Cannot create connection between incompatible access types. Constructor pin outputs can only be connected to constructor pin inputs."));
+			}
 			else
 			{
 				return FPinConnectionResponse(CONNECT_RESPONSE_DISALLOW, LOCTEXT("ConnectionNotAllowed", "Output pin cannot be connected to input pin."));
diff --git a/Engine/Plugins/Runtime/Metasound/Source/MetasoundEditor/Private/MetasoundNodeDetailCustomization.cpp b/Engine/Plugins/Runtime/Metasound/Source/MetasoundEditor/Private/MetasoundNodeDetailCustomization.cpp
--- a/Engine/Plugins/Runtime/Metasound/Source/MetasoundEditor/Private/MetasoundNodeDetailCustomization.cpp
+++ b/Engine/Plugins/Runtime/Metasound/Source/MetasoundEditor/Private/MetasoundNodeDetailCustomization.cpp
@@ -10,6 +10,7 @@
 #include "Internationalization/Text.h"
 #include "MetasoundDataReference.h"
 #include "MetasoundDataReferenceMacro.h"
+#include "MetasoundEditorGraphBuilder.h"
 #include "MetasoundEditorGraphInputNode.h"
 #include "MetasoundEditorGraphSchema.h"
 #include "MetasoundFrontend.h"
@@ -45,6 +46,15 @@ namespace Metasound
 	{
 		namespace MemberCustomizationPrivate
 		{
+			static int32 EnableMetasoundConstructorPins = 0;
+
+			FAutoConsoleVariableRef CVarMetaSoundEditorEnableConstructorPins(
+				TEXT("au.MetaSounds.Editor.ConstructorPinsEnabled"),
+				EnableMetasoundConstructorPins,
+				TEXT("Whether option to make inputs and outputs constructor pins is available in editor.\n")
+				TEXT("0 (default): Disabled, !0: Enabled"),
+				ECVF_Default);
+
 			/** Set of input types which are valid registered types, but should
 			 * not show up as an input type option in the MetaSound editor. */
 			static const TSet<FName> HiddenInputTypeNames =
@@ -59,7 +69,7 @@ namespace Metasound
 
 			static const FText ConstructorPinText = LOCTEXT("ConstructorPinText", "Is Constructor Pin");
 			static const FText ConstructorPinTooltip = LOCTEXT("ConstructorPinTooltip",
-				"Whether this input is a constructor pin (value is only read on construction, and is not updated at runtime).");
+				"Whether this input or output is a constructor pin (value is only read on construction, and is not dynamically updated at runtime).");
 
 			void GetDataTypeFromElementPropertyHandle(TSharedPtr<IPropertyHandle> ElementPropertyHandle, Frontend::FDataTypeRegistryInfo& OutDataTypeInfo)
 			{
@@ -189,7 +199,8 @@ namespace Metasound
 			{
 				if (const UMetasoundEditorGraph* OwningGraph = ParentMember->GetOwningGraph())
 				{
-					bShowWidgetOptions = OwningGraph->IsEditable();
+					// Disable widget options for constructor inputs for now to prevent changing default value via widget while playing 
+					bShowWidgetOptions = OwningGraph->IsEditable() && ParentMember->GetVertexAccessType() == EMetasoundFrontendVertexAccessType::Reference;
 				}
 			}
 
@@ -1136,62 +1147,107 @@ namespace Metasound
 			bIsNameInvalid = false;
 		}
 
+		void FMetasoundVertexDetailCustomization::AddConstructorPinRow(IDetailLayoutBuilder& InDetailLayout)
+		{
+			if (UMetasoundEditorGraphVertex* Vertex = Cast<UMetasoundEditorGraphVertex>(GraphMember.Get()))
+			{
+				InDetailLayout.EditCategory("General").AddCustomRow(MemberCustomizationPrivate::ConstructorPinText)
+				.IsEnabled(IsGraphEditable() && !IsInterfaceMember())
+				.NameContent()
+				[
+					SNew(STextBlock)
+					.Text(MemberCustomizationPrivate::ConstructorPinText)
+					.ToolTipText(MemberCustomizationPrivate::ConstructorPinTooltip)
+					.Font(IDetailLayoutBuilder::GetDetailFontBold())
+				]
+				.ValueContent()
+				[
+					SAssignNew(ConstructorPinCheckbox, SCheckBox)
+					.IsChecked_Lambda([this, Vertex]()
+					{
+						return OnGetConstructorPinCheckboxState(Vertex);
+					})
+					.OnCheckStateChanged_Lambda([this, Vertex](ECheckBoxState InNewState)
+					{
+						OnConstructorPinStateChanged(Vertex, InNewState);
+					})
+					[
+						SNew(STextBlock)
+						.Font(IDetailLayoutBuilder::GetDetailFont())
+					]
+				];
+			}
+		}
+		
 		void FMetasoundVertexDetailCustomization::CustomizeGeneralCategory(IDetailLayoutBuilder& InDetailLayout)
 		{
 			FMetasoundMemberDetailCustomization::CustomizeGeneralCategory(InDetailLayout);
+			UMetasoundEditorGraphVertex* Vertex = Cast<UMetasoundEditorGraphVertex>(GraphMember.Get());
+			if (!ensure(Vertex))
+			{
+				return;
+			}
 
+			// Constructor pin 
+			if (MemberCustomizationPrivate::EnableMetasoundConstructorPins)
+			{
+				Frontend::FDataTypeRegistryInfo DataTypeInfo;
+				Frontend::IDataTypeRegistry::Get().GetDataTypeInfo(Vertex->GetDataType(), DataTypeInfo);
+				if (DataTypeInfo.bIsConstructorType)
+				{
+					AddConstructorPinRow(InDetailLayout);
+				}
+			}
+
+			// Sort order
 			IDetailCategoryBuilder& CategoryBuilder = FMetasoundMemberDetailCustomization::GetGeneralCategoryBuilder(InDetailLayout);
-			UMetasoundEditorGraphVertex* Vertex = Cast<UMetasoundEditorGraphVertex>(GraphMember.Get());
-			if (ensure(Vertex))
-			{
-				TWeakObjectPtr<UMetasoundEditorGraphVertex> VertexPtr = Vertex;
-				static const FText SortOrderText = LOCTEXT("Vertex_SortOrderPropertyName", "Sort Order");
-				static const FText SortOrderToolTipFormat = LOCTEXT("Vertex_SortOrderToolTipFormat", "Sort Order for {0}. Used to organize pins in node view. The higher the number, the lower in the list.");
-				const FText SortOrderToolTipText = FText::Format(SortOrderToolTipFormat, GraphMember->GetGraphMemberLabel());
-				CategoryBuilder.AddCustomRow(SortOrderText)
-					.EditCondition(IsGraphEditable(), nullptr)
-					.NameContent()
-					[
-						SNew(STextBlock)
-						.Font(IDetailLayoutBuilder::GetDetailFontBold())
-						.Text(SortOrderText)
-						.ToolTipText(SortOrderToolTipText)
-					]
-					.ValueContent()
-					[
-						SNew(SNumericEntryBox<int32>)
-						.Value_Lambda([VertexPtr]() { return VertexPtr->GetSortOrderIndex(); })
-						.AllowSpin(false)
-						.UndeterminedString(LOCTEXT("Vertex_SortOrder_MultipleValues", "Multiple"))
-						.OnValueCommitted_Lambda([VertexPtr](int32 NewValue, ETextCommit::Type CommitInfo)
+			TWeakObjectPtr<UMetasoundEditorGraphVertex> VertexPtr = Vertex;
+			static const FText SortOrderText = LOCTEXT("Vertex_SortOrderPropertyName", "Sort Order");
+			static const FText SortOrderToolTipFormat = LOCTEXT("Vertex_SortOrderToolTipFormat", "Sort Order for {0}. Used to organize pins in node view. The higher the number, the lower in the list.");
+			const FText SortOrderToolTipText = FText::Format(SortOrderToolTipFormat, GraphMember->GetGraphMemberLabel());
+			CategoryBuilder.AddCustomRow(SortOrderText)
+				.EditCondition(IsGraphEditable(), nullptr)
+				.NameContent()
+				[
+					SNew(STextBlock)
+					.Font(IDetailLayoutBuilder::GetDetailFontBold())
+					.Text(SortOrderText)
+					.ToolTipText(SortOrderToolTipText)
+				]
+				.ValueContent()
+				[
+					SNew(SNumericEntryBox<int32>)
+					.Value_Lambda([VertexPtr]() { return VertexPtr->GetSortOrderIndex(); })
+					.AllowSpin(false)
+					.UndeterminedString(LOCTEXT("Vertex_SortOrder_MultipleValues", "Multiple"))
+					.OnValueCommitted_Lambda([VertexPtr](int32 NewValue, ETextCommit::Type CommitInfo)
+					{
+						if (!VertexPtr.IsValid())
 						{
-							if (!VertexPtr.IsValid())
-							{
-								return;
-							}
+							return;
+						}
 
-							const FText TransactionTitle = FText::Format(LOCTEXT("SetVertexSortOrderFormat", "Set MetaSound Graph {0} '{1}' SortOrder to {2}"),
-								VertexPtr->GetGraphMemberLabel(),
-								VertexPtr->GetDisplayName(),
-								FText::AsNumber(NewValue));
-							FScopedTransaction Transaction(TransactionTitle);
+						const FText TransactionTitle = FText::Format(LOCTEXT("SetVertexSortOrderFormat", "Set MetaSound Graph {0} '{1}' SortOrder to {2}"),
+							VertexPtr->GetGraphMemberLabel(),
+							VertexPtr->GetDisplayName(),
+							FText::AsNumber(NewValue));
+						FScopedTransaction Transaction(TransactionTitle);
 
-							UObject* MetaSoundObject = VertexPtr->GetOutermostObject();
-							FMetasoundAssetBase* MetaSoundAsset = Metasound::IMetasoundUObjectRegistry::Get().GetObjectAsAssetBase(MetaSoundObject);
-							check(MetaSoundAsset);
+						UObject* MetaSoundObject = VertexPtr->GetOutermostObject();
+						FMetasoundAssetBase* MetaSoundAsset = Metasound::IMetasoundUObjectRegistry::Get().GetObjectAsAssetBase(MetaSoundObject);
+						check(MetaSoundAsset);
 
-							MetaSoundObject->Modify();
-							MetaSoundAsset->GetGraphChecked().Modify();
-							VertexPtr->Modify();
+						MetaSoundObject->Modify();
+						MetaSoundAsset->GetGraphChecked().Modify();
+						VertexPtr->Modify();
 
-							VertexPtr->SetSortOrderIndex(NewValue);
+						VertexPtr->SetSortOrderIndex(NewValue);
 
-							constexpr bool bInForceViewSynchronization = true;
-							FGraphBuilder::RegisterGraphWithFrontend(*MetaSoundObject, bInForceViewSynchronization);
-						})
-						.Font(IDetailLayoutBuilder::GetDetailFont())
-					];
-			}
+						constexpr bool bInForceViewSynchronization = true;
+						FGraphBuilder::RegisterGraphWithFrontend(*MetaSoundObject, bInForceViewSynchronization);
+					})
+					.Font(IDetailLayoutBuilder::GetDetailFont())
+				];
 		}
 
 		bool FMetasoundVertexDetailCustomization::IsInterfaceMember() const
@@ -1271,17 +1327,38 @@ namespace Metasound
 			}
 		}
 
-		void FMetasoundInputDetailCustomization::CustomizeGeneralCategory(IDetailLayoutBuilder& InDetailLayout)
+		ECheckBoxState FMetasoundVertexDetailCustomization::OnGetConstructorPinCheckboxState(TWeakObjectPtr<UMetasoundEditorGraphVertex> InGraphVertex) const
 		{
-			FMetasoundVertexDetailCustomization::CustomizeGeneralCategory(InDetailLayout);
+			if (InGraphVertex.IsValid())
+			{
+				return InGraphVertex->GetVertexAccessType() == EMetasoundFrontendVertexAccessType::Value ? ECheckBoxState::Checked : ECheckBoxState::Unchecked;
+			}
 
-			if (UMetasoundEditorGraphInput* Input = Cast<UMetasoundEditorGraphInput>(GraphMember.Get()))
+			return ECheckBoxState::Undetermined;
+		}
+
+		void FMetasoundVertexDetailCustomization::OnConstructorPinStateChanged(TWeakObjectPtr<UMetasoundEditorGraphVertex> InGraphVertex, ECheckBoxState InNewState)
+		{
+			if (InGraphVertex.IsValid() && ConstructorPinCheckbox.IsValid())
 			{
-				IDetailCategoryBuilder& CategoryBuilder = FMetasoundMemberDetailCustomization::GetGeneralCategoryBuilder(InDetailLayout);
-				// TODO: Remove ConstructorPin bool from UMetasoundEditorGraphInput and add constructor pin bool property from frontend here 
-				/*IDetailPropertyRow* Row = CategoryBuilder.AddExternalObjectProperty(TArray<UObject*>({ Input }), GET_MEMBER_NAME_CHECKED(UMetasoundEditorGraphInput, ConstructorPin));
-				auto PropertyEnabled = TAttribute<bool>::CreateLambda([this] { return !GetDocumentHandle()->GetRootGraphClass().PresetOptions.bIsPreset; });
-				Row->EditCondition(PropertyEnabled, { });*/
+				EMetasoundFrontendVertexAccessType NewAccessType = InNewState == ECheckBoxState::Checked ? 
+					EMetasoundFrontendVertexAccessType::Value : EMetasoundFrontendVertexAccessType::Reference;
+
+				if (InGraphVertex->GetVertexAccessType() == NewAccessType)
+				{
+					return;
+				}
+
+				// Have to stop playback to avoid attempting to change live edit data on invalid input type.
+				check(GEditor);
+				GEditor->ResetPreviewAudioComponent();
+
+				InGraphVertex->SetVertexAccessType(NewAccessType);
+				
+				if (FMetasoundAssetBase* MetasoundAsset = Metasound::IMetasoundUObjectRegistry::Get().GetObjectAsAssetBase(GraphMember->GetOutermostObject()))
+				{
+					MetasoundAsset->SetUpdateDetailsOnSynchronization();
+				}
 			}
 		}
 
@@ -1359,6 +1436,30 @@ namespace Metasound
 						DefaultPropertyRow->OverrideResetToDefault(ResetOverride);
 					}
 				}
+			} 
+			else if (!bIsPreset)
+			{
+				// Make default value uneditable while playing for constructor inputs
+				const UMetasoundEditorGraphInput* Input = Cast<UMetasoundEditorGraphInput>(MemberDefaultLiteral->GetParentMember());
+				if (ensure(Input))
+				{
+					auto PropertyEnabled = TAttribute<bool>::CreateLambda([this, Input]
+					{
+						if (Input->GetVertexAccessType() == EMetasoundFrontendVertexAccessType::Value)
+						{
+							UObject* MetaSoundObject = Input->GetOutermostObject();
+							if (TSharedPtr<FEditor> MetaSoundEditor = FGraphBuilder::GetEditorForMetasound(*MetaSoundObject))
+							{
+								return !MetaSoundEditor->IsPlaying();
+							}
+						}
+						return true;
+					});
+					for (IDetailPropertyRow* DefaultPropertyRow : DefaultPropertyRows)
+					{
+						DefaultPropertyRow->EditCondition(PropertyEnabled, { });
+					}
+				}
 			}
 		}
 
diff --git a/Engine/Plugins/Runtime/Metasound/Source/MetasoundEditor/Private/MetasoundNodeDetailCustomization.h b/Engine/Plugins/Runtime/Metasound/Source/MetasoundEditor/Private/MetasoundNodeDetailCustomization.h
--- a/Engine/Plugins/Runtime/Metasound/Source/MetasoundEditor/Private/MetasoundNodeDetailCustomization.h
+++ b/Engine/Plugins/Runtime/Metasound/Source/MetasoundEditor/Private/MetasoundNodeDetailCustomization.h
@@ -355,6 +355,12 @@ namespace Metasound
 			virtual void CustomizeGeneralCategory(IDetailLayoutBuilder& InDetailLayout) override;
 			virtual EVisibility GetDefaultVisibility() const;
 			virtual bool IsInterfaceMember() const override;
+
+			void AddConstructorPinRow(IDetailLayoutBuilder& InDetailLayout);
+			ECheckBoxState OnGetConstructorPinCheckboxState(TWeakObjectPtr<UMetasoundEditorGraphVertex> InGraphMember) const;
+			void OnConstructorPinStateChanged(TWeakObjectPtr<UMetasoundEditorGraphVertex> InGraphMember, ECheckBoxState InNewState);
+
+			TSharedPtr<SCheckBox> ConstructorPinCheckbox;
 		};
 
 		class FMetasoundInputDetailCustomization : public FMetasoundVertexDetailCustomization
@@ -365,9 +371,6 @@ namespace Metasound
 			virtual void CustomizeDetails(IDetailLayoutBuilder& InDetailLayout) override;
 			virtual bool IsDefaultEditable() const override;
 
-		protected: 
-			virtual void CustomizeGeneralCategory(IDetailLayoutBuilder& InDetailLayout) override;
-
 		private:
 			bool GetInputInheritsDefault() const;
 			void SetInputInheritsDefault();
diff --git a/Engine/Plugins/Runtime/Metasound/Source/MetasoundEditor/Private/SMetasoundGraphNode.cpp b/Engine/Plugins/Runtime/Metasound/Source/MetasoundEditor/Private/SMetasoundGraphNode.cpp
--- a/Engine/Plugins/Runtime/Metasound/Source/MetasoundEditor/Private/SMetasoundGraphNode.cpp
+++ b/Engine/Plugins/Runtime/Metasound/Source/MetasoundEditor/Private/SMetasoundGraphNode.cpp
@@ -630,12 +630,12 @@ namespace Metasound
 			TSharedPtr<SHorizontalBox> ContentBox = SNew(SHorizontalBox);
 			TSharedPtr<SWidget> OuterContentBox; // currently only used for input float nodes to accommodate the input widget
 
-			// If editable float input node, check if custom widget required
+			// If editable float input node and not constructor input, check if custom widget required
 			bool bShowInputWidget = false;
 			if (UMetasoundEditorGraphInput* GraphMember = Cast<UMetasoundEditorGraphInput>(GetMetaSoundMember()))
 			{
 				const UMetasoundEditorGraph* OwningGraph = GraphMember->GetOwningGraph();
-				if (OwningGraph && OwningGraph->IsEditable())
+				if (OwningGraph && OwningGraph->IsEditable() && GraphMember->GetVertexAccessType() == EMetasoundFrontendVertexAccessType::Reference)
 				{
 					UMetasoundEditorGraphMemberDefaultFloat* DefaultFloat = Cast<UMetasoundEditorGraphMemberDefaultFloat>(GraphMember->GetLiteral());
 					if (DefaultFloat && DefaultFloat->WidgetType != EMetasoundMemberDefaultWidget::None)
diff --git a/Engine/Plugins/Runtime/Metasound/Source/MetasoundEditor/Private/SMetasoundGraphPin.h b/Engine/Plugins/Runtime/Metasound/Source/MetasoundEditor/Private/SMetasoundGraphPin.h
--- a/Engine/Plugins/Runtime/Metasound/Source/MetasoundEditor/Private/SMetasoundGraphPin.h
+++ b/Engine/Plugins/Runtime/Metasound/Source/MetasoundEditor/Private/SMetasoundGraphPin.h
@@ -211,22 +211,37 @@ namespace Metasound
 			virtual const FSlateBrush* GetPinIcon() const
 			{
 				const bool bIsConnected = ParentPinType::IsConnected();
-				bool bIsConstructorPin = false;
+				EMetasoundFrontendVertexAccessType AccessType = EMetasoundFrontendVertexAccessType::Reference;
+
 				if (const UEdGraphPin* Pin = ParentPinType::GetPinObj())
 				{
 					if (const UMetasoundEditorGraphNode* Node = Cast<UMetasoundEditorGraphNode>(Pin->GetOwningNode()))
 					{
 						if (const UMetasoundEditorGraphMemberNode* MemberNode = Cast<UMetasoundEditorGraphMemberNode>(Node))
 						{
-							if (const UMetasoundEditorGraphInput* Input = Cast<UMetasoundEditorGraphInput>(MemberNode->GetMember()))
+							if (const UMetasoundEditorGraphVertex* Vertex = Cast<UMetasoundEditorGraphVertex>(MemberNode->GetMember()))
+							{
+								AccessType = Vertex->GetVertexAccessType();
+							}
+						}
+						else if (const UMetasoundEditorGraphExternalNode* ExternalNode = Cast<UMetasoundEditorGraphExternalNode>(Node))
+						{
+							if (Pin->Direction == EGPD_Input)
+							{
+								Frontend::FInputHandle InputHandle = FGraphBuilder::GetInputHandleFromPin(Pin);
+								AccessType = InputHandle->GetVertexAccessType();
+							}
+							else if (Pin->Direction == EGPD_Output)
 							{
-								bIsConstructorPin = Input->ConstructorPin; 
+								Frontend::FOutputHandle OutputHandle = FGraphBuilder::GetOutputHandleFromPin(Pin);
+								AccessType = OutputHandle->GetVertexAccessType();
 							}
 						}
 					}
 				}
 
-				if (bIsConstructorPin)
+				// Is constructor pin 
+				if (AccessType == EMetasoundFrontendVertexAccessType::Value)
 				{
 					if (const ISlateStyle* MetasoundStyle = FSlateStyleRegistry::FindSlateStyle("MetaSoundStyle"))
 					{
diff --git a/Engine/Plugins/Runtime/Metasound/Source/MetasoundEditor/Public/MetasoundEditorGraph.h b/Engine/Plugins/Runtime/Metasound/Source/MetasoundEditor/Public/MetasoundEditorGraph.h
--- a/Engine/Plugins/Runtime/Metasound/Source/MetasoundEditor/Public/MetasoundEditorGraph.h
+++ b/Engine/Plugins/Runtime/Metasound/Source/MetasoundEditor/Public/MetasoundEditorGraph.h
@@ -34,6 +34,12 @@ namespace Metasound
 	namespace Editor
 	{
 		struct FGraphValidationResults;
+
+		struct FCreateNodeVertexParams
+		{
+			FName DataType;
+			EMetasoundFrontendVertexAccessType AccessType = EMetasoundFrontendVertexAccessType::Reference;
+		};
 	} // namespace Editor
 } // namespace Metasound
 
@@ -192,9 +198,12 @@ class METASOUNDEDITOR_API UMetasoundEditorGraphVertex : public UMetasoundEditorG
 	GENERATED_BODY()
 
 protected:
-	/** Adds the node handle for a newly created vertex. */
+	UE_DEPRECATED(5.1, "Use AddNodeHandle with FCreateNodeVertexParams instead.")
 	virtual Metasound::Frontend::FNodeHandle AddNodeHandle(const FName& InNodeName, FName InDataType) PURE_VIRTUAL(UMetasoundEditorGraphVertex::AddNodeHandle, return Metasound::Frontend::INodeController::GetInvalidHandle(); )
 
+	/** Adds the node handle for a newly created vertex. */
+	virtual Metasound::Frontend::FNodeHandle AddNodeHandle(const FName& InNodeName, const Metasound::Editor::FCreateNodeVertexParams& InParams) PURE_VIRTUAL(UMetasoundEditorGraphVertex::AddNodeHandle, return Metasound::Frontend::INodeController::GetInvalidHandle(); )
+
 public:
 	/** Initializes all properties with the given parameters required to identify the frontend member from this editor graph member. */
 	void InitMember(FName InDataType, const FMetasoundFrontendLiteral& InDefaultLiteral, FGuid InNodeID, FMetasoundFrontendClassName&& InClassName);
@@ -235,6 +244,8 @@ public:
 
 	/** Sets the SortOrderIndex assigned to this member. */
 	virtual void SetSortOrderIndex(int32 InSortOrderIndex) PURE_VIRTUAL(UMetasoundEditorGraphMember::SetSortOrderIndex, )
+	
+	void SetVertexAccessType(EMetasoundFrontendVertexAccessType InNewAccessType, bool bPostTransaction = true);
 
 	/** Returns the node handle associated with the vertex. */
 	Metasound::Frontend::FNodeHandle GetNodeHandle();
@@ -242,6 +253,8 @@ public:
 	/** Returns the node handle associated with the vertex. */
 	Metasound::Frontend::FConstNodeHandle GetConstNodeHandle() const;
 
+	virtual EMetasoundFrontendVertexAccessType GetVertexAccessType() const PURE_VIRTUAL(UMetasoundEditorGraphVertex::GetVertexAccessType, return EMetasoundFrontendVertexAccessType::Reference; );
+
 	virtual bool CanRename() const override;
 };
 
@@ -251,10 +264,6 @@ class METASOUNDEDITOR_API UMetasoundEditorGraphInput : public UMetasoundEditorGr
 	GENERATED_BODY()
 
 public:
-	/** Whether this input is a constructor pin (value is only read on construction, and is not updated at runtime). */
-	UPROPERTY(Transient, EditAnywhere, Category = General)
-	bool ConstructorPin = false;
-
 	virtual int32 GetSortOrderIndex() const override;
 	virtual void SetSortOrderIndex(int32 InSortOrderIndex) override;
 
@@ -262,9 +271,11 @@ public:
 	virtual void ResetToClassDefault() override;
 	virtual void SetMemberName(const FName& InNewName, bool bPostTransaction) override;
 	virtual void UpdateFrontendDefaultLiteral(bool bPostTransaction) override;
+	virtual EMetasoundFrontendVertexAccessType GetVertexAccessType() const override;
 
 protected:
 	virtual Metasound::Frontend::FNodeHandle AddNodeHandle(const FName& InNodeName, FName InDataType) override;
+	virtual Metasound::Frontend::FNodeHandle AddNodeHandle(const FName& InNodeName, const Metasound::Editor::FCreateNodeVertexParams& InParams) override;
 	virtual EMetasoundFrontendClassType GetClassType() const override { return EMetasoundFrontendClassType::Input; }
 	virtual Metasound::Editor::ENodeSection GetSectionID() const override;
 };
@@ -281,9 +292,11 @@ public:
 	virtual const FText& GetGraphMemberLabel() const override;
 	virtual void ResetToClassDefault() override;
 	virtual void UpdateFrontendDefaultLiteral(bool bPostTransaction) override;
+	virtual EMetasoundFrontendVertexAccessType GetVertexAccessType() const override;
 
 protected:
 	virtual Metasound::Frontend::FNodeHandle AddNodeHandle(const FName& InNodeName, FName InDataType) override;
+	virtual Metasound::Frontend::FNodeHandle AddNodeHandle(const FName& InNodeName, const Metasound::Editor::FCreateNodeVertexParams& InParams) override;
 	virtual EMetasoundFrontendClassType GetClassType() const override { return EMetasoundFrontendClassType::Output; }
 	virtual Metasound::Editor::ENodeSection GetSectionID() const override;
 };
diff --git a/Engine/Plugins/Runtime/Metasound/Source/MetasoundFrontend/Private/MetasoundAssetBase.cpp b/Engine/Plugins/Runtime/Metasound/Source/MetasoundFrontend/Private/MetasoundAssetBase.cpp
--- a/Engine/Plugins/Runtime/Metasound/Source/MetasoundFrontend/Private/MetasoundAssetBase.cpp
+++ b/Engine/Plugins/Runtime/Metasound/Source/MetasoundFrontend/Private/MetasoundAssetBase.cpp
@@ -861,12 +861,16 @@ TArray<FMetasoundFrontendClassInput> FMetasoundAssetBase::GetTransmittableClassI
 	IDataTypeRegistry& Registry = IDataTypeRegistry::Get();
 	auto IsTransmittable = [&Registry, &NonTransmittableInputs](const FMetasoundFrontendClassVertex& InVertex)
 	{
-		if (!NonTransmittableInputs.Contains(InVertex.Name))
+		// Don't transmit constructor inputs 
+		if (InVertex.AccessType == EMetasoundFrontendVertexAccessType::Reference)
 		{
-			FDataTypeRegistryInfo Info;
-			if (Registry.GetDataTypeInfo(InVertex.TypeName, Info))
+			if (!NonTransmittableInputs.Contains(InVertex.Name))
 			{
-				return Info.bIsTransmittable;
+				FDataTypeRegistryInfo Info;
+				if (Registry.GetDataTypeInfo(InVertex.TypeName, Info))
+				{
+					return Info.bIsTransmittable;
+				}
 			}
 		}
 
diff --git a/Engine/Plugins/Runtime/Metasound/Source/MetasoundFrontend/Private/MetasoundFrontendDocument.cpp b/Engine/Plugins/Runtime/Metasound/Source/MetasoundFrontend/Private/MetasoundFrontendDocument.cpp
--- a/Engine/Plugins/Runtime/Metasound/Source/MetasoundFrontend/Private/MetasoundFrontendDocument.cpp
+++ b/Engine/Plugins/Runtime/Metasound/Source/MetasoundFrontend/Private/MetasoundFrontendDocument.cpp
@@ -164,6 +164,19 @@ bool FMetasoundFrontendClassVertex::IsFunctionalEquivalent(const FMetasoundFront
 	return FMetasoundFrontendVertex::IsFunctionalEquivalent(InLHS, InRHS) && (InLHS.AccessType == InRHS.AccessType);
 }
 
+bool FMetasoundFrontendClassVertex::CanConnectVertexAccessTypes(EMetasoundFrontendVertexAccessType InFromType, EMetasoundFrontendVertexAccessType InToType)
+{
+	if (EMetasoundFrontendVertexAccessType::Value == InToType)
+	{
+		// If the input vertex accesses by "Value" then the output vertex 
+		// must also access by "Value" to enforce unexpected consequences 
+		// of connecting data which varies over time to an input which only
+		// evaluates the data during operator initialization.
+		return (EMetasoundFrontendVertexAccessType::Value == InFromType);
+	}
+	return true;
+}
+
 FMetasoundFrontendClassName::FMetasoundFrontendClassName(const FName& InNamespace, const FName& InName, const FName& InVariant)
 : Namespace(InNamespace)
 , Name(InName)
diff --git a/Engine/Plugins/Runtime/Metasound/Source/MetasoundFrontend/Private/MetasoundFrontendGraph.cpp b/Engine/Plugins/Runtime/Metasound/Source/MetasoundFrontend/Private/MetasoundFrontendGraph.cpp
--- a/Engine/Plugins/Runtime/Metasound/Source/MetasoundFrontend/Private/MetasoundFrontendGraph.cpp
+++ b/Engine/Plugins/Runtime/Metasound/Source/MetasoundFrontend/Private/MetasoundFrontendGraph.cpp
@@ -221,7 +221,14 @@ namespace Metasound
 						MoveTemp(Literal)
 					};
 
-					return DataTypeRegistry.CreateInputNode(InputVertex.TypeName, MoveTemp(InitParams));
+					if (InOwningGraphClassInput.AccessType == EMetasoundFrontendVertexAccessType::Reference)
+					{
+						return DataTypeRegistry.CreateInputNode(InputVertex.TypeName, MoveTemp(InitParams));
+					}
+					else // InOwningGraphClassInput.AccessType == EMetasoundFrontendVertexAccessType::Value
+					{
+						return DataTypeRegistry.CreateConstructorInputNode(InputVertex.TypeName, MoveTemp(InitParams));
+					}
 				}
 				else
 				{
@@ -264,7 +271,15 @@ namespace Metasound
 				}
 			}
 
-			return IDataTypeRegistry::Get().CreateOutputNode(OutputVertex.TypeName, MoveTemp(InitParams));
+			ensure(InClass.Interface.Outputs.Num() == 1);
+			if (InClass.Interface.Outputs[0].AccessType == EMetasoundFrontendVertexAccessType::Reference)
+			{
+				return IDataTypeRegistry::Get().CreateOutputNode(OutputVertex.TypeName, MoveTemp(InitParams));
+			}
+			else // InClass.Interface.Outputs[0].AccessType == EMetasoundFrontendVertexAccessType::Value
+			{
+				return IDataTypeRegistry::Get().CreateConstructorOutputNode(OutputVertex.TypeName, MoveTemp(InitParams));
+			}
 		}
 		return TUniquePtr<INode>(nullptr);
 	}
diff --git a/Engine/Plugins/Runtime/Metasound/Source/MetasoundFrontend/Private/MetasoundFrontendGraphController.cpp b/Engine/Plugins/Runtime/Metasound/Source/MetasoundFrontend/Private/MetasoundFrontendGraphController.cpp
--- a/Engine/Plugins/Runtime/Metasound/Source/MetasoundFrontend/Private/MetasoundFrontendGraphController.cpp
+++ b/Engine/Plugins/Runtime/Metasound/Source/MetasoundFrontend/Private/MetasoundFrontendGraphController.cpp
@@ -949,7 +949,7 @@ namespace Metasound
 				if (Algo::NoneOf(GraphClass->Interface.Inputs, IsInputWithSameName))
 				{
 					FNodeRegistryKey Key;
-					if (FRegistry::GetInputNodeRegistryKeyForDataType(InClassInput.TypeName, Key))
+					if (FRegistry::GetInputNodeRegistryKeyForDataType(InClassInput.TypeName, InClassInput.AccessType, Key))
 					{
 						FConstClassAccessPtr InputClassPtr = OwningDocument->FindOrAddClass(Key);
 						if (const FMetasoundFrontendClass* InputClass = InputClassPtr.Get())
@@ -1084,7 +1084,7 @@ namespace Metasound
 				if (Algo::NoneOf(GraphClass->Interface.Outputs, IsOutputWithSameName))
 				{
 					FNodeRegistryKey Key;
-					if (FRegistry::GetOutputNodeRegistryKeyForDataType(InClassOutput.TypeName, Key))
+					if (FRegistry::GetOutputNodeRegistryKeyForDataType(InClassOutput.TypeName, InClassOutput.AccessType, Key))
 					{
 						FConstClassAccessPtr OutputClassPtr = OwningDocument->FindOrAddClass(Key);
 						if (const FMetasoundFrontendClass* OutputClass = OutputClassPtr.Get())
diff --git a/Engine/Plugins/Runtime/Metasound/Source/MetasoundFrontend/Private/MetasoundFrontendInputController.cpp b/Engine/Plugins/Runtime/Metasound/Source/MetasoundFrontend/Private/MetasoundFrontendInputController.cpp
--- a/Engine/Plugins/Runtime/Metasound/Source/MetasoundFrontend/Private/MetasoundFrontendInputController.cpp
+++ b/Engine/Plugins/Runtime/Metasound/Source/MetasoundFrontend/Private/MetasoundFrontendInputController.cpp
@@ -203,7 +203,7 @@ namespace Metasound
 				OutConnectability.Connectable = FConnectability::EConnectable::No;
 				OutConnectability.Reason = FConnectability::EReason::IncompatibleDataTypes;
 			}
-			else if (!IsCompatibleAccessType(InController))
+			else if (!FMetasoundFrontendClassVertex::CanConnectVertexAccessTypes(InController.GetVertexAccessType(), GetVertexAccessType()))
 			{
 				OutConnectability.Connectable = FConnectability::EConnectable::No;
 				OutConnectability.Reason = FConnectability::EReason::IncompatibleAccessTypes;
@@ -254,7 +254,7 @@ namespace Metasound
 			{
 				if (OtherDataType == DataType)
 				{
-					if (IsCompatibleAccessType(InController))
+					if (FMetasoundFrontendClassVertex::CanConnectVertexAccessTypes(InController.GetVertexAccessType(), GetVertexAccessType()))
 					{
 						// Overwrite an existing connection if it exists.
 						FMetasoundFrontendEdge* Edge = FindEdge();
@@ -365,19 +365,6 @@ namespace Metasound
 			return false;
 		}
 
-		bool FBaseInputController::IsCompatibleAccessType(const IOutputController& InOutputController) const
-		{
-			if (EMetasoundFrontendVertexAccessType::Value == GetVertexAccessType())
-			{
-				// If the input vertex accesses by "Value" then the output vertex 
-				// must also access by "Value" to enforce unexpected consequences 
-				// of connecting data which varies over time to an input which only
-				// evaluates the data during operator initialization.
-				return (EMetasoundFrontendVertexAccessType::Value == InOutputController.GetVertexAccessType());
-			}
-			return true;
-		}
-
 		const FMetasoundFrontendEdge* FBaseInputController::FindEdge() const
 		{
 			if (const FMetasoundFrontendGraph* Graph = GraphPtr.Get())
@@ -502,6 +489,11 @@ namespace Metasound
 			}
 		}
 
+		EMetasoundFrontendVertexAccessType FOutputNodeInputController::GetVertexAccessType() const
+		{
+			return OwningGraphClassOutputPtr.Get()->AccessType;
+		}
+
 		FDocumentAccess FOutputNodeInputController::ShareAccess() 
 		{
 			FDocumentAccess Access = FBaseInputController::ShareAccess();
@@ -574,6 +566,11 @@ namespace Metasound
 				Vertex->Name = InName;
 			}
 		}
+		
+		EMetasoundFrontendVertexAccessType FInputNodeInputController::GetVertexAccessType() const
+		{
+			return OwningGraphClassInputPtr.Get()->AccessType;
+		}
 
 		bool FInputNodeInputController::IsConnectionUserModifiable() const
 		{
diff --git a/Engine/Plugins/Runtime/Metasound/Source/MetasoundFrontend/Private/MetasoundFrontendInputController.h b/Engine/Plugins/Runtime/Metasound/Source/MetasoundFrontend/Private/MetasoundFrontendInputController.h
--- a/Engine/Plugins/Runtime/Metasound/Source/MetasoundFrontend/Private/MetasoundFrontendInputController.h
+++ b/Engine/Plugins/Runtime/Metasound/Source/MetasoundFrontend/Private/MetasoundFrontendInputController.h
@@ -92,11 +92,6 @@ namespace Metasound
 			FConstClassInputAccessPtr ClassInputPtr;
 			FGraphAccessPtr GraphPtr;
 			FNodeHandle OwningNode;
-
-		private:
-
-			bool IsCompatibleAccessType(const IOutputController& InOutputController) const;
-
 		};
 
 		/** FOutputNodeInputController represents the input vertex of an output 
@@ -133,6 +128,7 @@ namespace Metasound
 #endif // WITH_EDITOR
 
 			virtual void SetName(const FVertexName& InName) override;
+			virtual EMetasoundFrontendVertexAccessType GetVertexAccessType() const override;
 
 		protected:
 
@@ -180,6 +176,7 @@ namespace Metasound
 #endif // WITH_EDITOR
 
 			virtual void SetName(const FVertexName& InName) override;
+			virtual EMetasoundFrontendVertexAccessType GetVertexAccessType() const override;
 
 			virtual bool IsConnectionUserModifiable() const override;
 			virtual FConnectability CanConnectTo(const IOutputController& InController) const override;
diff --git a/Engine/Plugins/Runtime/Metasound/Source/MetasoundFrontend/Private/MetasoundFrontendNodeController.cpp b/Engine/Plugins/Runtime/Metasound/Source/MetasoundFrontend/Private/MetasoundFrontendNodeController.cpp
--- a/Engine/Plugins/Runtime/Metasound/Source/MetasoundFrontend/Private/MetasoundFrontendNodeController.cpp
+++ b/Engine/Plugins/Runtime/Metasound/Source/MetasoundFrontend/Private/MetasoundFrontendNodeController.cpp
@@ -807,19 +807,22 @@ namespace Metasound
 				const FVertexNameAndType ConnectionKey(InputHandle->GetName(), InputHandle->GetDataType());
 				if (FInputConnectionInfo* ConnectionInfo = InputConnections.Find(ConnectionKey))
 				{
-					if (ConnectionInfo->bLiteralSet)
+					if (InputHandle->CanConnectTo(*ConnectionInfo->ConnectedOutput).Connectable == FConnectability::EConnectable::Yes)
 					{
-						InputHandle->SetLiteral(ConnectionInfo->DefaultValue);
-					}
+						if (ConnectionInfo->bLiteralSet)
+						{
+							InputHandle->SetLiteral(ConnectionInfo->DefaultValue);
+						}
 
-					if (ConnectionInfo->ConnectedOutput->IsValid())
-					{
-						ensure(InputHandle->Connect(*ConnectionInfo->ConnectedOutput));
-					}
+						if (ConnectionInfo->ConnectedOutput->IsValid())
+						{
+							ensure(InputHandle->Connect(*ConnectionInfo->ConnectedOutput));
+						}
 
-					// Remove connection to track missing connections between 
-					// node versions.
-					InputConnections.Remove(ConnectionKey);
+						// Remove connection to track missing connections between 
+						// node versions.
+						InputConnections.Remove(ConnectionKey);
+					}
 				}
 			});
 
@@ -841,17 +844,24 @@ namespace Metasound
 				const FVertexNameAndType ConnectionKey(OutputHandle->GetName(), OutputHandle->GetDataType());
 				if (FOutputConnectionInfo* ConnectionInfo = OutputConnections.Find(ConnectionKey))
 				{
+					bool bConnectionSuccess = false;
 					for (FInputHandle InputHandle : ConnectionInfo->ConnectedInputs)
 					{
-						if (InputHandle->IsValid())
+						if (InputHandle->CanConnectTo(*OutputHandle).Connectable == FConnectability::EConnectable::Yes)
 						{
-							ensure(InputHandle->Connect(*OutputHandle));
+							if (InputHandle->IsValid())
+							{
+								ensure(InputHandle->Connect(*OutputHandle));
+								bConnectionSuccess = true;
+							}
 						}
 					}
-
 					// Remove connection to track missing connections between 
 					// node versions.
-					OutputConnections.Remove(ConnectionKey);
+					if (bConnectionSuccess)
+					{
+						OutputConnections.Remove(ConnectionKey);
+					}
 				}
 			});
 
diff --git a/Engine/Plugins/Runtime/Metasound/Source/MetasoundFrontend/Private/MetasoundFrontendNodeRegistryPrivate.h b/Engine/Plugins/Runtime/Metasound/Source/MetasoundFrontend/Private/MetasoundFrontendNodeRegistryPrivate.h
--- a/Engine/Plugins/Runtime/Metasound/Source/MetasoundFrontend/Private/MetasoundFrontendNodeRegistryPrivate.h
+++ b/Engine/Plugins/Runtime/Metasound/Source/MetasoundFrontend/Private/MetasoundFrontendNodeRegistryPrivate.h
@@ -66,9 +66,13 @@ namespace Metasound
 			// Find Frontend Document data.
 			virtual bool FindFrontendClassFromRegistered(const FNodeRegistryKey& InKey, FMetasoundFrontendClass& OutClass) override;
 			virtual bool FindNodeClassInfoFromRegistered(const Metasound::Frontend::FNodeRegistryKey& InKey, FNodeClassInfo& OutInfo) override;
+			UE_DEPRECATED(5.1, "Use FindInputNodeRegistryKeyForDataType with EMetasoundFrontendVertexAccessType instead.")
 			virtual bool FindInputNodeRegistryKeyForDataType(const FName& InDataTypeName, FNodeRegistryKey& OutKey) override;
+			virtual bool FindInputNodeRegistryKeyForDataType(const FName& InDataTypeName, const EMetasoundFrontendVertexAccessType InAccessType, FNodeRegistryKey& OutKey) override;
 			virtual bool FindVariableNodeRegistryKeyForDataType(const FName& InDataTypeName, FNodeRegistryKey& OutKey) override;
+			UE_DEPRECATED(5.1, "Use FindOutputNodeRegistryKeyForDataType with EMetasoundFrontendVertexAccessType instead.")
 			virtual bool FindOutputNodeRegistryKeyForDataType(const FName& InDataTypeName, FNodeRegistryKey& OutKey) override;
+			virtual bool FindOutputNodeRegistryKeyForDataType(const FName& InDataTypeName, const EMetasoundFrontendVertexAccessType InAccessType, FNodeRegistryKey& OutKey) override;
 
 			// Create a new instance of a C++ implemented node from the registry.
 			virtual TUniquePtr<Metasound::INode> CreateNode(const FNodeRegistryKey& InKey, const Metasound::FNodeInitData& InInitData) const override;
diff --git a/Engine/Plugins/Runtime/Metasound/Source/MetasoundFrontend/Private/MetasoundFrontendOutputController.cpp b/Engine/Plugins/Runtime/Metasound/Source/MetasoundFrontend/Private/MetasoundFrontendOutputController.cpp
--- a/Engine/Plugins/Runtime/Metasound/Source/MetasoundFrontend/Private/MetasoundFrontendOutputController.cpp
+++ b/Engine/Plugins/Runtime/Metasound/Source/MetasoundFrontend/Private/MetasoundFrontendOutputController.cpp
@@ -293,6 +293,11 @@ namespace Metasound
 				Vertex->Name = InName;
 			}
 		}
+		
+		EMetasoundFrontendVertexAccessType FInputNodeOutputController::GetVertexAccessType() const
+		{
+			return OwningGraphClassInputPtr.Get()->AccessType;
+		}
 
 		FDocumentAccess FInputNodeOutputController::ShareAccess() 
 		{
@@ -366,6 +371,11 @@ namespace Metasound
 			}
 		}
 
+		EMetasoundFrontendVertexAccessType FOutputNodeOutputController::GetVertexAccessType() const
+		{
+			return OwningGraphClassOutputPtr.Get()->AccessType;
+		}
+
 		bool FOutputNodeOutputController::IsConnectionUserModifiable() const 
 		{
 			return false;
diff --git a/Engine/Plugins/Runtime/Metasound/Source/MetasoundFrontend/Private/MetasoundFrontendOutputController.h b/Engine/Plugins/Runtime/Metasound/Source/MetasoundFrontend/Private/MetasoundFrontendOutputController.h
--- a/Engine/Plugins/Runtime/Metasound/Source/MetasoundFrontend/Private/MetasoundFrontendOutputController.h
+++ b/Engine/Plugins/Runtime/Metasound/Source/MetasoundFrontend/Private/MetasoundFrontendOutputController.h
@@ -127,6 +127,7 @@ namespace Metasound
 #endif // WITH_EDITOR
 
 			virtual void SetName(const FVertexName& InName) override;
+			virtual EMetasoundFrontendVertexAccessType GetVertexAccessType() const override;
 
 		protected:
 			virtual FDocumentAccess ShareAccess() override;
@@ -177,6 +178,7 @@ namespace Metasound
 #endif // WITH_EDITOR
 
 			virtual void SetName(const FVertexName& InName) override;
+			virtual EMetasoundFrontendVertexAccessType GetVertexAccessType() const override;
 
 			virtual bool IsConnectionUserModifiable() const override;
 			virtual FConnectability CanConnectTo(const IInputController& InController) const override;
diff --git a/Engine/Plugins/Runtime/Metasound/Source/MetasoundFrontend/Private/MetasoundFrontendRegistries.cpp b/Engine/Plugins/Runtime/Metasound/Source/MetasoundFrontend/Private/MetasoundFrontendRegistries.cpp
--- a/Engine/Plugins/Runtime/Metasound/Source/MetasoundFrontend/Private/MetasoundFrontendRegistries.cpp
+++ b/Engine/Plugins/Runtime/Metasound/Source/MetasoundFrontend/Private/MetasoundFrontendRegistries.cpp
@@ -314,12 +314,29 @@ namespace Metasound
 		}
 
 		bool FRegistryContainerImpl::FindInputNodeRegistryKeyForDataType(const FName& InDataTypeName, FNodeRegistryKey& OutKey)
+		{
+			UE_LOG(LogMetaSound, Error, TEXT("FRegistryContainerImpl::FindInputNodeRegistryKeyForDataType with these parameters is no longer supported and should not be called. Use the one with EMetasoundFrontendVertexAccessType instead."));
+			return false;
+		}
+
+		bool FRegistryContainerImpl::FindInputNodeRegistryKeyForDataType(const FName& InDataTypeName, const EMetasoundFrontendVertexAccessType InAccessType, FNodeRegistryKey& OutKey)
 		{
 			FMetasoundFrontendClass Class;
-			if (IDataTypeRegistry::Get().GetFrontendInputClass(InDataTypeName, Class))
+			if (InAccessType == EMetasoundFrontendVertexAccessType::Reference)
 			{
-				OutKey = NodeRegistryKey::CreateKey(Class.Metadata);
-				return true;
+				if (IDataTypeRegistry::Get().GetFrontendInputClass(InDataTypeName, Class))
+				{
+					OutKey = NodeRegistryKey::CreateKey(Class.Metadata);
+					return true;
+				}
+			}
+			else // InAccessType == EMetasoundFrontendVertexAccessType::Value
+			{
+				if (IDataTypeRegistry::Get().GetFrontendConstructorInputClass(InDataTypeName, Class))
+				{
+					OutKey = NodeRegistryKey::CreateKey(Class.Metadata);
+					return true;
+				}
 			}
 			return false;
 		}
@@ -336,13 +353,31 @@ namespace Metasound
 		}
 
 		bool FRegistryContainerImpl::FindOutputNodeRegistryKeyForDataType(const FName& InDataTypeName, FNodeRegistryKey& OutKey)
+		{
+			UE_LOG(LogMetaSound, Error, TEXT("FRegistryContainerImpl::FindInputNodeRegistryKeyForDataType with these parameters is no longer supported and should not be called. Use the one with EMetasoundFrontendVertexAccessType instead."));
+			return false;
+		}
+
+		bool FRegistryContainerImpl::FindOutputNodeRegistryKeyForDataType(const FName& InDataTypeName, const EMetasoundFrontendVertexAccessType InAccessType, FNodeRegistryKey& OutKey)
 		{
 			FMetasoundFrontendClass Class;
-			if (IDataTypeRegistry::Get().GetFrontendOutputClass(InDataTypeName, Class))
+			if (InAccessType == EMetasoundFrontendVertexAccessType::Reference) 
 			{
-				OutKey = NodeRegistryKey::CreateKey(Class.Metadata);
-				return true;
+				if (IDataTypeRegistry::Get().GetFrontendOutputClass(InDataTypeName, Class))
+				{
+					OutKey = NodeRegistryKey::CreateKey(Class.Metadata);
+					return true;
+				}
+			}
+			else // InAccessType == EMetasoundFrontendVertexAccessType::Value
+			{
+				if (IDataTypeRegistry::Get().GetFrontendConstructorOutputClass(InDataTypeName, Class))
+				{
+					OutKey = NodeRegistryKey::CreateKey(Class.Metadata);
+					return true;
+				}
 			}
+
 			return false;
 		}
 
@@ -562,12 +597,17 @@ bool FMetasoundFrontendRegistryContainer::GetNodeClassInfoFromRegistered(const F
 	return false;
 }
 
-
 bool FMetasoundFrontendRegistryContainer::GetInputNodeRegistryKeyForDataType(const FName& InDataTypeName, FNodeRegistryKey& OutKey)
+{
+	UE_LOG(LogMetaSound, Error, TEXT("FMetasoundFrontendRegistryContainer::GetInputNodeRegistryKeyForDataType with these parameters is no longer supported and should not be called. Use the one with EMetasoundFrontendVertexAccessType instead."));
+	return false;
+}
+
+bool FMetasoundFrontendRegistryContainer::GetInputNodeRegistryKeyForDataType(const FName& InDataTypeName, const EMetasoundFrontendVertexAccessType InAccessType, FNodeRegistryKey& OutKey)
 {
 	if (FMetasoundFrontendRegistryContainer* Registry = FMetasoundFrontendRegistryContainer::Get())
 	{
-		return Registry->FindInputNodeRegistryKeyForDataType(InDataTypeName, OutKey);
+		return Registry->FindInputNodeRegistryKeyForDataType(InDataTypeName, InAccessType, OutKey);
 	}
 	return false;
 }
@@ -582,12 +622,16 @@ bool FMetasoundFrontendRegistryContainer::GetVariableNodeRegistryKeyForDataType(
 }
 
 bool FMetasoundFrontendRegistryContainer::GetOutputNodeRegistryKeyForDataType(const FName& InDataTypeName, FNodeRegistryKey& OutKey)
+{
+	UE_LOG(LogMetaSound, Error, TEXT("FMetasoundFrontendRegistryContainer::GetOutputNodeRegistryKeyForDataType with these parameters is no longer supported and should not be called. Use the one with EMetasoundFrontendVertexAccessType instead."));
+	return false;
+}
+
+bool FMetasoundFrontendRegistryContainer::GetOutputNodeRegistryKeyForDataType(const FName& InDataTypeName, const EMetasoundFrontendVertexAccessType InVertexAccessType, FNodeRegistryKey& OutKey)
 {
 	if (FMetasoundFrontendRegistryContainer* Registry = FMetasoundFrontendRegistryContainer::Get())
 	{
-		return Registry->FindOutputNodeRegistryKeyForDataType(InDataTypeName, OutKey);
+		return Registry->FindOutputNodeRegistryKeyForDataType(InDataTypeName, InVertexAccessType, OutKey);
 	}
 	return false;
 }
-
-
diff --git a/Engine/Plugins/Runtime/Metasound/Source/MetasoundFrontend/Private/MetasoundFrontendTransform.cpp b/Engine/Plugins/Runtime/Metasound/Source/MetasoundFrontend/Private/MetasoundFrontendTransform.cpp
--- a/Engine/Plugins/Runtime/Metasound/Source/MetasoundFrontend/Private/MetasoundFrontendTransform.cpp
+++ b/Engine/Plugins/Runtime/Metasound/Source/MetasoundFrontend/Private/MetasoundFrontendTransform.cpp
@@ -832,6 +832,7 @@ namespace Metasound
 
 					ClassInput.Name = NodeName;
 					ClassInput.TypeName = Input->GetDataType();
+					ClassInput.AccessType = Input->GetVertexAccessType();
 
 #if WITH_EDITOR
 					ClassInput.Metadata.SetDescription(InputNode->GetDescription());
@@ -907,6 +908,7 @@ namespace Metasound
 
 					ClassOutput.Name = NodeName;
 					ClassOutput.TypeName = Output->GetDataType();
+					ClassOutput.AccessType = Output->GetVertexAccessType();
 
 #if WITH_EDITOR
 					ClassOutput.Metadata.SetDescription(OutputNode->GetDescription());
diff --git a/Engine/Plugins/Runtime/Metasound/Source/MetasoundFrontend/Public/MetasoundFrontendDocument.h b/Engine/Plugins/Runtime/Metasound/Source/MetasoundFrontend/Public/MetasoundFrontendDocument.h
--- a/Engine/Plugins/Runtime/Metasound/Source/MetasoundFrontend/Public/MetasoundFrontendDocument.h
+++ b/Engine/Plugins/Runtime/Metasound/Source/MetasoundFrontend/Public/MetasoundFrontendDocument.h
@@ -694,6 +694,8 @@ struct METASOUNDFRONTEND_API FMetasoundFrontendClassVertex : public FMetasoundFr
 	void SplitName(FName& OutNamespace, FName& OutParameterName) const;
 
 	static bool IsFunctionalEquivalent(const FMetasoundFrontendClassVertex& InLHS, const FMetasoundFrontendClassVertex& InRHS);
+	// Whether vertex access types are compatible when connecting from an output to an input 
+	static bool CanConnectVertexAccessTypes(EMetasoundFrontendVertexAccessType InFromType, EMetasoundFrontendVertexAccessType InToType);
 };
 
 
diff --git a/Engine/Plugins/Runtime/Metasound/Source/MetasoundFrontend/Public/MetasoundFrontendRegistries.h b/Engine/Plugins/Runtime/Metasound/Source/MetasoundFrontend/Public/MetasoundFrontendRegistries.h
--- a/Engine/Plugins/Runtime/Metasound/Source/MetasoundFrontend/Public/MetasoundFrontendRegistries.h
+++ b/Engine/Plugins/Runtime/Metasound/Source/MetasoundFrontend/Public/MetasoundFrontendRegistries.h
@@ -287,9 +287,13 @@ public:
 
 	static bool GetFrontendClassFromRegistered(const FNodeRegistryKey& InKey, FMetasoundFrontendClass& OutClass);
 	static bool GetNodeClassInfoFromRegistered(const FNodeRegistryKey& InKey, FNodeClassInfo& OutInfo);
+	UE_DEPRECATED(5.1, "Use GetInputNodeRegistryKeyForDataType with EMetasoundFrontendVertexAccessType instead.")
 	static bool GetInputNodeRegistryKeyForDataType(const FName& InDataTypeName, FNodeRegistryKey& OutKey);
+	static bool GetInputNodeRegistryKeyForDataType(const FName& InDataTypeName, const EMetasoundFrontendVertexAccessType InAccessType, FNodeRegistryKey& OutKey);
 	static bool GetVariableNodeRegistryKeyForDataType(const FName& InDataTypeName, FNodeRegistryKey& OutKey);
+	UE_DEPRECATED(5.1, "Use GetOutputNodeRegistryKeyForDataType with EMetasoundFrontendVertexAccessType instead.")
 	static bool GetOutputNodeRegistryKeyForDataType(const FName& InDataTypeName, FNodeRegistryKey& OutKey);
+	static bool GetOutputNodeRegistryKeyForDataType(const FName& InDataTypeName, const EMetasoundFrontendVertexAccessType InAccessType, FNodeRegistryKey& OutKey);
 
 
 	FMetasoundFrontendRegistryContainer() = default;
@@ -340,10 +344,13 @@ public:
 	// Query for MetaSound Frontend document objects.
 	virtual bool FindFrontendClassFromRegistered(const Metasound::Frontend::FNodeRegistryKey& InKey, FMetasoundFrontendClass& OutClass) = 0;
 	virtual bool FindNodeClassInfoFromRegistered(const Metasound::Frontend::FNodeRegistryKey& InKey, FNodeClassInfo& OutInfo) = 0;
-
+	UE_DEPRECATED(5.1, "Use FindInputNodeRegistryKeyForDataType with EMetasoundFrontendVertexAccessType instead.")
 	virtual bool FindInputNodeRegistryKeyForDataType(const FName& InDataTypeName, FNodeRegistryKey& OutKey) = 0;
+	virtual bool FindInputNodeRegistryKeyForDataType(const FName& InDataTypeName, const EMetasoundFrontendVertexAccessType InAccessType, FNodeRegistryKey& OutKey) = 0;
 	virtual bool FindVariableNodeRegistryKeyForDataType(const FName& InDataTypeName, FNodeRegistryKey& OutKey) = 0;
+	UE_DEPRECATED(5.1, "Use FindOutputNodeRegistryKeyForDataType with EMetasoundFrontendVertexAccessType instead.")
 	virtual bool FindOutputNodeRegistryKeyForDataType(const FName& InDataTypeName, FNodeRegistryKey& OutKey) = 0;
+	virtual bool FindOutputNodeRegistryKeyForDataType(const FName& InDataTypeName, const EMetasoundFrontendVertexAccessType InAccessType, FNodeRegistryKey& OutKey) = 0;
 
 	virtual TUniquePtr<Metasound::INode> CreateNode(const FNodeRegistryKey& InKey, const Metasound::FNodeInitData&) const = 0;
 	virtual TUniquePtr<Metasound::INode> CreateNode(const FNodeRegistryKey& InKey, Metasound::FDefaultLiteralNodeConstructorParams&&) const = 0;
diff --git a/Engine/Plugins/Runtime/Metasound/Source/MetasoundGraphCore/Public/MetasoundDataReferenceCollection.h b/Engine/Plugins/Runtime/Metasound/Source/MetasoundGraphCore/Public/MetasoundDataReferenceCollection.h
--- a/Engine/Plugins/Runtime/Metasound/Source/MetasoundGraphCore/Public/MetasoundDataReferenceCollection.h
+++ b/Engine/Plugins/Runtime/Metasound/Source/MetasoundGraphCore/Public/MetasoundDataReferenceCollection.h
@@ -120,7 +120,7 @@ namespace Metasound
 					if (IsDataReferenceOfType<DataType>(*DataRef))
 					{
 						const EDataReferenceAccessType AccessType = DataRef->GetAccessType();
-						return (AccessType == EDataReferenceAccessType::Read) || (AccessType == EDataReferenceAccessType::Write);
+						return (AccessType == EDataReferenceAccessType::Read) || (AccessType == EDataReferenceAccessType::Write) || (AccessType == EDataReferenceAccessType::Value);
 					}
 				}
 
