Minor comment fix
#rnx

diff --git a/Engine/Plugins/Runtime/Metasound/Source/MetasoundEditor/Private/MetasoundEditor.cpp b/Engine/Plugins/Runtime/Metasound/Source/MetasoundEditor/Private/MetasoundEditor.cpp
--- a/Engine/Plugins/Runtime/Metasound/Source/MetasoundEditor/Private/MetasoundEditor.cpp
+++ b/Engine/Plugins/Runtime/Metasound/Source/MetasoundEditor/Private/MetasoundEditor.cpp
@@ -1205,7 +1205,7 @@ namespace Metasound
 				GetToolkitCommands(),
 				FToolBarExtensionDelegate::CreateLambda([this](FToolBarBuilder& ToolbarBuilder)
 				{
-				// TODO: Add OS SVD and clean this up post UE5.0 - Early Access
+					// TODO: Clean-up json importer/exporter and re-enable this
  					ToolbarBuilder.BeginSection("Utilities");
  					{
 // 						ToolbarBuilder.AddToolBarButton
