Fix conditional for reconnecting connections when updating a MetaSound node version

#rnx

diff --git a/Engine/Plugins/Runtime/Metasound/Source/MetasoundFrontend/Private/MetasoundFrontendNodeController.cpp b/Engine/Plugins/Runtime/Metasound/Source/MetasoundFrontend/Private/MetasoundFrontendNodeController.cpp
--- a/Engine/Plugins/Runtime/Metasound/Source/MetasoundFrontend/Private/MetasoundFrontendNodeController.cpp
+++ b/Engine/Plugins/Runtime/Metasound/Source/MetasoundFrontend/Private/MetasoundFrontendNodeController.cpp
@@ -807,18 +807,16 @@ namespace Metasound
 				const FVertexNameAndType ConnectionKey(InputHandle->GetName(), InputHandle->GetDataType());
 				if (FInputConnectionInfo* ConnectionInfo = InputConnections.Find(ConnectionKey))
 				{
-					if (InputHandle->CanConnectTo(*ConnectionInfo->ConnectedOutput).Connectable == FConnectability::EConnectable::Yes)
+					if (ConnectionInfo->bLiteralSet)
 					{
-						if (ConnectionInfo->bLiteralSet)
-						{
-							InputHandle->SetLiteral(ConnectionInfo->DefaultValue);
-						}
-
-						if (ConnectionInfo->ConnectedOutput->IsValid())
-						{
-							ensure(InputHandle->Connect(*ConnectionInfo->ConnectedOutput));
-						}
+						InputHandle->SetLiteral(ConnectionInfo->DefaultValue);
+					}
 
+					if (ConnectionInfo->ConnectedOutput->IsValid() && 
+						InputHandle->CanConnectTo(*ConnectionInfo->ConnectedOutput).Connectable == FConnectability::EConnectable::Yes)
+					{
+						ensure(InputHandle->Connect(*ConnectionInfo->ConnectedOutput));
+						
 						// Remove connection to track missing connections between 
 						// node versions.
 						InputConnections.Remove(ConnectionKey);
@@ -847,13 +845,11 @@ namespace Metasound
 					bool bConnectionSuccess = false;
 					for (FInputHandle InputHandle : ConnectionInfo->ConnectedInputs)
 					{
-						if (InputHandle->CanConnectTo(*OutputHandle).Connectable == FConnectability::EConnectable::Yes)
+						if (InputHandle->IsValid() && 
+							InputHandle->CanConnectTo(*OutputHandle).Connectable == FConnectability::EConnectable::Yes)
 						{
-							if (InputHandle->IsValid())
-							{
-								ensure(InputHandle->Connect(*OutputHandle));
-								bConnectionSuccess = true;
-							}
+							ensure(InputHandle->Connect(*OutputHandle));
+							bConnectionSuccess = true;
 						}
 					}
 					// Remove connection to track missing connections between 
