Fix Delay Pitch Shifter Node's internal delay length

diff --git a/Engine/Plugins/Runtime/Metasound/Source/MetasoundStandardNodes/Private/MetasoundDopplerPitchShiftNode.cpp b/Engine/Plugins/Runtime/Metasound/Source/MetasoundStandardNodes/Private/MetasoundDopplerPitchShiftNode.cpp
--- a/Engine/Plugins/Runtime/Metasound/Source/MetasoundStandardNodes/Private/MetasoundDopplerPitchShiftNode.cpp
+++ b/Engine/Plugins/Runtime/Metasound/Source/MetasoundStandardNodes/Private/MetasoundDopplerPitchShiftNode.cpp
@@ -55,7 +55,7 @@ namespace Metasound
 		, DelayLength(InDelayLength)
 		, AudioOutput(FAudioBufferWriteRef::CreateNew(InSettings))
 	{
-		DelayBuffer.Init(InSettings.GetSampleRate(), Audio::FTapDelayPitchShifter::MaxDelayLength);
+		DelayBuffer.Init(InSettings.GetSampleRate(), 0.001f * Audio::FTapDelayPitchShifter::MaxDelayLength);
 		DelayPitchShifter.Init(InSettings.GetSampleRate(), *PitchShift, *DelayLength);
 	}
 
