Minor FrontendRegistry Clean-up
#rnx

diff --git a/Engine/Plugins/Runtime/Metasound/Source/MetasoundFrontend/Public/MetasoundFrontendRegistries.h b/Engine/Plugins/Runtime/Metasound/Source/MetasoundFrontend/Public/MetasoundFrontendRegistries.h
--- a/Engine/Plugins/Runtime/Metasound/Source/MetasoundFrontend/Public/MetasoundFrontendRegistries.h
+++ b/Engine/Plugins/Runtime/Metasound/Source/MetasoundFrontend/Public/MetasoundFrontendRegistries.h
@@ -103,7 +103,7 @@ namespace Metasound
 		/** INodeRegistryEntry declares the interface for a node registry entry.
 		 * Each node class in the registry must satisfy this interface. 
 		 */
-		class INodeRegistryEntry
+		class METASOUNDFRONTEND_API INodeRegistryEntry
 		{
 		public:
 			virtual ~INodeRegistryEntry() = default;
@@ -154,9 +154,6 @@ namespace Metasound
 			virtual bool IsNative() const = 0;
 		};
 
-
-
-
 		struct METASOUNDFRONTEND_API FConverterNodeRegistryKey
 		{
 			// The datatype one would like to convert from.
@@ -243,12 +240,6 @@ namespace Metasound
 			// Returns true if both keys represent the same entry in the node registry.
 			METASOUNDFRONTEND_API bool IsEqual(const FNodeRegistryKey& InLHS, const FNodeRegistryKey& InRHS);
 
-			// Returns true if the class metadata and key represent the same entry in the node registry.
-			METASOUNDFRONTEND_API bool IsEqual(const FMetasoundFrontendClassMetadata& InLHS, const FNodeRegistryKey& InRHS);
-
-			// Returns true if the class info and key represent the same entry in the node registry.
-			METASOUNDFRONTEND_API bool IsEqual(const FNodeClassInfo& InLHS, const FNodeRegistryKey& InRHS);
-
 			// Returns true if the class metadata represent the same entry in the node registry.
 			METASOUNDFRONTEND_API bool IsEqual(const FMetasoundFrontendClassMetadata& InLHS, const FMetasoundFrontendClassMetadata& InRHS);
 
@@ -313,7 +304,6 @@ public:
 	UE_DEPRECATED(5.1, "ForEachNodeRegistryTransactionSince is no longer be supported")
 	virtual void ForEachNodeRegistryTransactionSince(Metasound::Frontend::FRegistryTransactionID InSince, Metasound::Frontend::FRegistryTransactionID* OutCurrentRegistryTransactionID, TFunctionRef<void(const Metasound::Frontend::FNodeRegistryTransaction&)> InFunc) const = 0;
 
-
 	/** Register an external node with the frontend.
 	 *
 	 * @param InCreateNode - Function for creating node from FNodeInitData.
@@ -322,7 +312,7 @@ public:
 	 * @return A node registration key. If the registration failed, then the registry 
 	 *         key will be invalid.
 	 */
-	virtual FNodeRegistryKey RegisterNode(TUniquePtr<Metasound::Frontend::INodeRegistryEntry>&&) = 0;
+	virtual FNodeRegistryKey RegisterNode(TUniquePtr<Metasound::Frontend::INodeRegistryEntry>&& InEntry) = 0;
 
 	/** Unregister an external node from the frontend.
 	 *
