Metasound Absoulte Value math node

diff --git a/Engine/Plugins/Runtime/Metasound/Source/MetasoundStandardNodes/Private/MetasoundAbsoluteValueNodes.cpp b/Engine/Plugins/Runtime/Metasound/Source/MetasoundStandardNodes/Private/MetasoundAbsoluteValueNodes.cpp
new file mode 100644
--- /dev/null
+++ b/Engine/Plugins/Runtime/Metasound/Source/MetasoundStandardNodes/Private/MetasoundAbsoluteValueNodes.cpp
@@ -0,0 +1,251 @@
+// Copyright Epic Games, Inc. All Rights Reserved.
+
+#include "MetasoundNodeRegistrationMacro.h"
+#include "CoreMinimal.h"
+#include "MetasoundFacade.h"
+#include "MetasoundExecutableOperator.h"
+#include "MetasoundPrimitives.h"
+#include "MetasoundStandardNodesCategories.h"
+#include "MetasoundStandardNodesNames.h"
+#include "MetasoundAudioBuffer.h"
+#include "Internationalization/Text.h"
+#include "MetasoundParamHelper.h"
+
+#define LOCTEXT_NAMESPACE "MetasoundStandardNodes_AbsNode"
+
+namespace Metasound
+{
+	namespace AbsVertexNames
+	{
+		METASOUND_PARAM(InputValue, "Input", "Input value.");
+		METASOUND_PARAM(OutputValue, "Value", "The Absolute Value of the input.");
+	}
+
+	namespace MetasoundAbsNodePrivate
+	{
+		FNodeClassMetadata CreateNodeClassMetadata(const FName& InDataTypeName, const FName& InOperatorName, const FText& InDisplayName, const FText& InDescription, const FVertexInterface& InDefaultInterface)
+		{
+			FNodeClassMetadata Metadata
+			{
+				FNodeClassName { "AbsoluteValue", InOperatorName, InDataTypeName },
+				1, // Major Version
+				0, // Minor Version
+				InDisplayName,
+				InDescription,
+				PluginAuthor,
+				PluginNodeMissingPrompt,
+				InDefaultInterface,
+				{ NodeCategories::Math },
+				{ },
+				FNodeDisplayStyle{}
+			};
+
+			return Metadata;
+		}
+
+		template<typename ValueType>
+		struct TAbs
+		{
+			bool bSupported = false;
+		};
+
+		template<>
+		struct TAbs<int32>
+		{
+			static void GetAbs(const int32 In, int32& OutAbs)
+			{
+				OutAbs = FMath::Abs(In);
+			}
+
+			static TDataReadReference<int32> CreateInRef(const FCreateOperatorParams& InParams)
+			{
+				using namespace AbsVertexNames;
+				const FInputVertexInterface& InputInterface = InParams.Node.GetVertexInterface().GetInputInterface();
+				return InParams.InputDataReferences.GetDataReadReferenceOrConstructWithVertexDefault<int32>(InputInterface, METASOUND_GET_PARAM_NAME(InputValue), InParams.OperatorSettings);
+			}
+		};
+
+		template<>
+		struct TAbs<float>
+		{
+			static void GetAbs(const float In, float& OutAbs)
+			{
+				OutAbs = FMath::Abs(In);
+			}
+
+			static TDataReadReference<float> CreateInRef(const FCreateOperatorParams& InParams)
+			{
+				using namespace AbsVertexNames;
+				const FInputVertexInterface& InputInterface = InParams.Node.GetVertexInterface().GetInputInterface();
+				return InParams.InputDataReferences.GetDataReadReferenceOrConstructWithVertexDefault<float>(InputInterface, METASOUND_GET_PARAM_NAME(InputValue), InParams.OperatorSettings);
+			}
+
+			static bool IsAudioBuffer() { return false; }
+		};
+		
+		template<>
+		struct TAbs<FTime>
+		{
+			static void GetAbs(const FTime& In, FTime& OutAbs)
+			{
+				OutAbs = FTime(FMath::Abs(In.GetSeconds()));
+			}
+
+			static TDataReadReference<FTime> CreateInRef(const FCreateOperatorParams& InParams)
+			{
+				using namespace AbsVertexNames;
+				const FInputVertexInterface& InputInterface = InParams.Node.GetVertexInterface().GetInputInterface();
+				return InParams.InputDataReferences.GetDataReadReferenceOrConstructWithVertexDefault<FTime>(InputInterface, METASOUND_GET_PARAM_NAME(InputValue), InParams.OperatorSettings);
+			}
+
+			static bool IsAudioBuffer() { return false; }
+		};
+
+		template<>
+		struct TAbs<FAudioBuffer>
+		{
+			static void GetAbs(const FAudioBuffer& In, FAudioBuffer& OutAbs)
+			{
+				TArrayView<float> OutAbsView(OutAbs.GetData(), OutAbs.Num());
+				TArrayView<const float> InView(In.GetData(), OutAbs.Num());
+				
+				Audio::ArrayAbs(InView, OutAbsView);
+			}
+
+			static TDataReadReference<FAudioBuffer> CreateInRef(const FCreateOperatorParams& InParams)
+			{
+				using namespace AbsVertexNames;
+				return InParams.InputDataReferences.GetDataReadReferenceOrConstruct<FAudioBuffer>(METASOUND_GET_PARAM_NAME(InputValue), InParams.OperatorSettings);
+			}
+		};
+	}
+
+	template<typename ValueType>
+	class TAbsNodeOperator : public TExecutableOperator<TAbsNodeOperator<ValueType>>
+	{
+	public:
+		static const FVertexInterface& GetDefaultInterface()
+		{
+			using namespace AbsVertexNames;
+			using namespace MetasoundAbsNodePrivate;
+
+			static const FVertexInterface DefaultInterface(
+				FInputVertexInterface(
+					TInputDataVertex<ValueType>(METASOUND_GET_PARAM_NAME_AND_METADATA(InputValue))
+				),
+				FOutputVertexInterface(
+					TOutputDataVertex<ValueType>(METASOUND_GET_PARAM_NAME_AND_METADATA(OutputValue))
+				)
+			);
+
+			return DefaultInterface;
+		}
+
+		static const FNodeClassMetadata& GetNodeInfo()
+		{
+			auto CreateNodeClassMetadata = []() -> FNodeClassMetadata
+			{
+				const FName DataTypeName = GetMetasoundDataTypeName<ValueType>();
+				const FName OperatorName = TEXT("Abs");
+				const FText NodeDisplayName = METASOUND_LOCTEXT_FORMAT("AbsDisplayNamePattern", "Abs ({0})", GetMetasoundDataTypeDisplayText<ValueType>());
+				const FText NodeDescription = METASOUND_LOCTEXT("AbsDesc", "Returns the Absolute Value of the input.");
+				const FVertexInterface NodeInterface = GetDefaultInterface();
+
+				return MetasoundAbsNodePrivate::CreateNodeClassMetadata(DataTypeName, OperatorName, NodeDisplayName, NodeDescription, NodeInterface);
+			};
+
+			static const FNodeClassMetadata Metadata = CreateNodeClassMetadata();
+			return Metadata;
+		}
+
+		static TUniquePtr<IOperator> CreateOperator(const FCreateOperatorParams& InParams, TArray<TUniquePtr<IOperatorBuildError>>& OutErrors)
+		{
+			using namespace AbsVertexNames;
+			using namespace MetasoundAbsNodePrivate;
+
+			TDataReadReference<ValueType> Input = TAbs<ValueType>::CreateInRef(InParams);
+
+			return MakeUnique<TAbsNodeOperator<ValueType>>(InParams.OperatorSettings, Input);
+		}
+
+
+		TAbsNodeOperator(const FOperatorSettings& InSettings,
+			const TDataReadReference<ValueType>& InInput)
+			: Input(InInput)
+			, OutputValue(TDataWriteReferenceFactory<ValueType>::CreateAny(InSettings))
+		{
+			GetAbs();
+		}
+
+		virtual ~TAbsNodeOperator() = default;
+
+		virtual FDataReferenceCollection GetInputs() const override
+		{
+			using namespace AbsVertexNames;
+
+			FDataReferenceCollection Inputs;
+			Inputs.AddDataReadReference(METASOUND_GET_PARAM_NAME(InputValue), Input);
+
+			return Inputs;
+		}
+
+		virtual FDataReferenceCollection GetOutputs() const override
+		{
+			using namespace AbsVertexNames;
+
+			FDataReferenceCollection Outputs;
+			Outputs.AddDataReadReference(METASOUND_GET_PARAM_NAME(OutputValue), OutputValue);
+
+			return Outputs;
+		}
+
+		void GetAbs()
+		{
+			using namespace MetasoundAbsNodePrivate;
+
+			TAbs<ValueType>::GetAbs(*Input, *OutputValue);
+		}
+
+		void Execute()
+		{
+			GetAbs();
+		}
+
+	private:
+
+		TDataReadReference<ValueType> Input;
+		TDataWriteReference<ValueType> OutputValue;
+	};
+
+	/** TAbsNode
+	 *
+	 *  Returns the Aof the input
+	 */
+	template<typename ValueType>
+	class METASOUNDSTANDARDNODES_API TAbsNode : public FNodeFacade
+	{
+	public:
+		/**
+		 * Constructor used by the Metasound Frontend.
+		 */
+		TAbsNode(const FNodeInitData& InInitData)
+			: FNodeFacade(InInitData.InstanceName, InInitData.InstanceID, TFacadeOperatorClass<TAbsNodeOperator<ValueType>>())
+		{}
+
+		virtual ~TAbsNode() = default;
+	};
+
+	using FAbsNodeInt32 = TAbsNode<int32>;
+ 	METASOUND_REGISTER_NODE(FAbsNodeInt32)
+
+	using FAbsNodeFloat = TAbsNode<float>;
+	METASOUND_REGISTER_NODE(FAbsNodeFloat)
+	
+	using FAbsNodeTime = TAbsNode<FTime>;
+    METASOUND_REGISTER_NODE(FAbsNodeTime)
+
+	using FAbsNodeAudioBuffer = TAbsNode<FAudioBuffer>;
+	METASOUND_REGISTER_NODE(FAbsNodeAudioBuffer)
+}
+
+#undef LOCTEXT_NAMESPACE
diff --git a/Engine/Source/Runtime/SignalProcessing/Private/FloatArrayMath.cpp b/Engine/Source/Runtime/SignalProcessing/Private/FloatArrayMath.cpp
--- a/Engine/Source/Runtime/SignalProcessing/Private/FloatArrayMath.cpp
+++ b/Engine/Source/Runtime/SignalProcessing/Private/FloatArrayMath.cpp
@@ -559,9 +559,24 @@ namespace Audio
 		}
 		else
 		{
-			for (int32 i = 0; i < Num; i++)
+			const int32 NumToSimd = Num & MathIntrinsics::SimdMask;
+			const int32 NumNotToSimd = Num & MathIntrinsics::NotSimdMask;
+
+			if (NumToSimd)
+			{
+				for (int32 i = 0; i < NumToSimd; i += AUDIO_NUM_FLOATS_PER_VECTOR_REGISTER)
+				{
+					VectorRegister4Float Input = VectorLoad(&InData[i]);
+					VectorStore(VectorAbs(Input), &OutData[i]);
+				}
+			}
+
+			if (NumNotToSimd)
 			{
-				OutData[i] = FMath::Abs(InData[i]);
+				for (int32 i = NumToSimd; i < Num; ++i)
+				{
+					OutData[i] = FMath::Abs(InData[i]);
+				}
 			}
 		}
 	}
@@ -581,9 +596,24 @@ namespace Audio
 		}
 		else
 		{
-			for (int32 i = 0; i < Num; i++)
+			const int32 NumToSimd = Num & MathIntrinsics::SimdMask;
+			const int32 NumNotToSimd = Num & MathIntrinsics::NotSimdMask;
+
+			if (NumToSimd)
 			{
-				Data[i] = FMath::Abs(Data[i]);
+				for (int32 i = 0; i < NumToSimd; i += AUDIO_NUM_FLOATS_PER_VECTOR_REGISTER)
+				{
+					VectorRegister4Float Input = VectorLoad(&Data[i]);
+					VectorStore(VectorAbs(Input), &Data[i]);
+				}
+			}
+
+			if (NumNotToSimd)
+			{
+				for (int32 i = NumToSimd; i < Num; ++i)
+				{
+					Data[i] = FMath::Abs(Data[i]);
+				}
 			}
 		}
 	}
