Metasound QToBandwidth node

diff --git a/Engine/Plugins/Runtime/Metasound/Source/MetasoundStandardNodes/Private/MetasoundFilterQToBandwidthNode.cpp b/Engine/Plugins/Runtime/Metasound/Source/MetasoundStandardNodes/Private/MetasoundFilterQToBandwidthNode.cpp
new file mode 100644
--- /dev/null
+++ b/Engine/Plugins/Runtime/Metasound/Source/MetasoundStandardNodes/Private/MetasoundFilterQToBandwidthNode.cpp
@@ -0,0 +1,134 @@
+// Copyright Epic Games, Inc. All Rights Reserved.
+
+#include "DSP/Dsp.h"
+#include "MetasoundDataTypeRegistrationMacro.h"
+#include "MetasoundDataReference.h"
+#include "MetasoundDataFactory.h"
+#include "MetasoundExecutableOperator.h"
+#include "MetasoundFacade.h"
+#include "MetasoundNodeRegistrationMacro.h"
+#include "MetasoundParamHelper.h"
+#include "MetasoundPrimitives.h"
+#include "MetasoundStandardNodesCategories.h"
+#include "MetasoundStandardNodesNames.h"
+
+#define LOCTEXT_NAMESPACE "Metasound_ConvertQToBandwidth"
+
+namespace Metasound
+{
+	namespace QToBandwidth
+	{
+		METASOUND_PARAM(InputQ, "Q", "Q parameter for filter control.");
+		METASOUND_PARAM(OutputBandwidth, "Bandwidth", "Resulting bandwidth value.");
+	}
+
+	class FConvertQToBandwidthNodeOperator : public TExecutableOperator<FConvertQToBandwidthNodeOperator>
+	{
+	public:
+		static const FVertexInterface& GetDefaultInterface()
+		{
+			using namespace QToBandwidth;
+
+			static const FVertexInterface DefaultInterface(
+				FInputVertexInterface(
+					TInputDataVertex<float>(METASOUND_GET_PARAM_NAME_AND_METADATA(InputQ), 1.f)
+				),
+				FOutputVertexInterface(
+					TOutputDataVertex<float>(METASOUND_GET_PARAM_NAME_AND_METADATA(OutputBandwidth))
+				)
+			);
+
+			return DefaultInterface;
+		}
+
+		static const FNodeClassMetadata& GetNodeInfo()
+		{
+			auto CreateNodeClassMetadata = []() -> FNodeClassMetadata
+			{
+				FNodeClassMetadata Metadata
+				{
+					FNodeClassName { StandardNodes::Namespace, "Convert Filter Q To Bandwidth", "" },
+					1, // Major Version
+					0, // Minor Version
+					METASOUND_LOCTEXT("ConvertQToBandwidthNode_Name", "Filter Q To Bandwidth"),
+					METASOUND_LOCTEXT("ConvertQToBandwidthNode_Description", "Converts Filter Q To Bandwidth."),
+					PluginAuthor,
+					PluginNodeMissingPrompt,
+					GetDefaultInterface(),
+					{ NodeCategories::Math },
+					{ },
+					{ }
+				};
+
+				return Metadata;
+			};
+
+			static const FNodeClassMetadata Metadata = CreateNodeClassMetadata();
+			return Metadata;
+		}
+
+		static TUniquePtr<IOperator> CreateOperator(const FCreateOperatorParams& InParams, TArray<TUniquePtr<IOperatorBuildError>>& OutErrors)
+		{
+			using namespace QToBandwidth;
+
+			const FDataReferenceCollection& Inputs = InParams.InputDataReferences;
+
+			FFloatReadRef FloatInput = Inputs.GetDataReadReferenceOrConstruct<float>(METASOUND_GET_PARAM_NAME(InputQ));
+
+			return MakeUnique<FConvertQToBandwidthNodeOperator>(InParams.OperatorSettings, FloatInput);
+		}
+
+		FConvertQToBandwidthNodeOperator(const FOperatorSettings& InSettings, const FFloatReadRef& InInValue)
+			: InValue(InInValue)
+			, OutValue(TDataWriteReferenceFactory<float>::CreateAny(InSettings))
+		{
+			Execute();
+		}
+
+		virtual ~FConvertQToBandwidthNodeOperator() = default;
+
+		virtual FDataReferenceCollection GetInputs() const override
+		{
+			using namespace QToBandwidth;
+
+			FDataReferenceCollection Inputs;
+			Inputs.AddDataReadReference(METASOUND_GET_PARAM_NAME(InputQ), InValue);
+
+			return Inputs;
+		}
+
+		virtual FDataReferenceCollection GetOutputs() const override
+		{
+			using namespace QToBandwidth;
+
+			FDataReferenceCollection Outputs;
+			Outputs.AddDataReadReference(METASOUND_GET_PARAM_NAME(OutputBandwidth), OutValue);
+
+			return Outputs;
+		}
+
+		void Execute()
+		{
+			*OutValue = Audio::GetBandwidthFromQ(FMath::Max(*InValue, KINDA_SMALL_NUMBER));
+		}
+
+	private:
+		FFloatReadRef InValue;
+		FFloatWriteRef OutValue;
+	};
+
+	class FConvertQToBandwidthNode : public FNodeFacade
+	{
+	public:
+		FConvertQToBandwidthNode(const FNodeInitData& InInitData)
+			: FNodeFacade(InInitData.InstanceName, InInitData.InstanceID, TFacadeOperatorClass<FConvertQToBandwidthNodeOperator>())
+		{
+		}
+
+		virtual ~FConvertQToBandwidthNode() = default;
+	};
+
+	METASOUND_REGISTER_NODE(FConvertQToBandwidthNode)
+} 
+
+#undef LOCTEXT_NAMESPACE
