MetaSound Template nodes

diff --git a/Engine/Plugins/Runtime/Metasound/Source/MetasoundEditor/Private/MetasoundEditorGraphBuilder.cpp b/Engine/Plugins/Runtime/Metasound/Source/MetasoundEditor/Private/MetasoundEditorGraphBuilder.cpp
--- a/Engine/Plugins/Runtime/Metasound/Source/MetasoundEditor/Private/MetasoundEditorGraphBuilder.cpp
+++ b/Engine/Plugins/Runtime/Metasound/Source/MetasoundEditor/Private/MetasoundEditorGraphBuilder.cpp
@@ -1208,6 +1208,7 @@ namespace Metasound
 				break;
 
 				case EMetasoundFrontendClassType::External:
+				case EMetasoundFrontendClassType::Template:
 				{
 					return CastChecked<UMetasoundEditorGraphNode>(AddExternalNode(InMetaSound, InNodeHandle, InLocation, bInSelectNewNode));
 				}
@@ -1236,7 +1237,7 @@ namespace Metasound
 				default:
 				{
 					checkNoEntry();
-					static_assert(static_cast<int32>(EMetasoundFrontendClassType::Invalid) == 9, "Possible missing FMetasoundFrontendClassType case coverage");
+					static_assert(static_cast<int32>(EMetasoundFrontendClassType::Invalid) == 10, "Possible missing FMetasoundFrontendClassType case coverage");
 				}
 				break;
 			}
@@ -1514,9 +1515,10 @@ namespace Metasound
 					case EMetasoundFrontendClassType::VariableMutator:
 					case EMetasoundFrontendClassType::Variable:
 					case EMetasoundFrontendClassType::External:
+					case EMetasoundFrontendClassType::Template:
 					default:
 					{
-						static_assert(static_cast<int32>(EMetasoundFrontendClassType::Invalid) == 9, "Possible missing MetasoundFrontendClassType switch case coverage.");
+						static_assert(static_cast<int32>(EMetasoundFrontendClassType::Invalid) == 10, "Possible missing MetasoundFrontendClassType switch case coverage.");
 						RemoveNodeHandle(GraphHandle, NodeHandle);
 					}
 					break;
diff --git a/Engine/Plugins/Runtime/Metasound/Source/MetasoundEditor/Private/MetasoundEditorGraphSchema.cpp b/Engine/Plugins/Runtime/Metasound/Source/MetasoundEditor/Private/MetasoundEditorGraphSchema.cpp
--- a/Engine/Plugins/Runtime/Metasound/Source/MetasoundEditor/Private/MetasoundEditorGraphSchema.cpp
+++ b/Engine/Plugins/Runtime/Metasound/Source/MetasoundEditor/Private/MetasoundEditorGraphSchema.cpp
@@ -780,7 +780,7 @@ UEdGraphNode* FMetasoundGraphSchemaAction_PromoteToOutput::PerformAction(UEdGrap
 		if (ensure(Output))
 		{
 			if (UMetasoundEditorGraphOutputNode* NewGraphNode = FGraphBuilder::AddOutputNode(ParentMetasound, NodeHandle, InLocation))
-			{
+			{ 
 				UEdGraphNode* EdGraphNode = CastChecked<UEdGraphNode>(NewGraphNode);
 
 				if (ensure(SchemaPrivate::TryConnectNewNodeToMatchingDataTypePin(*EdGraphNode, FromPin)))
@@ -900,7 +900,7 @@ Metasound::Frontend::FNodeHandle FMetasoundGraphSchemaAction_NewVariableMutatorN
 	return INodeController::GetInvalidHandle();
 }
 
-UEdGraphNode* FMetasoundGraphSchemaAction_NewFromSelected::PerformAction(UEdGraph* ParentGraph, UEdGraphPin* FromPin, const FVector2D Location, bool bSelectNewNode/* = true*/)
+UEdGraphNode* FMetasoundGraphSchemaAction_NewFromSelected::PerformAction(UEdGraph* ParentGraph, UEdGraphPin* FromPin, const FVector2D Location, bool bSelectNewNode /* = true*/)
 {
 	// TODO: Implement
 	return nullptr;
@@ -1046,6 +1046,7 @@ void UMetasoundEditorGraphSchema::GetGraphContextActions(FGraphContextMenuBuilde
 
 			FGraphActionMenuBuilder& ActionMenuBuilder = static_cast<FGraphActionMenuBuilder&>(ContextMenuBuilder);
 			ActionMenuBuilder.AddAction(MakeShared<FMetasoundGraphSchemaAction_PromoteToInput>());
+
 			// Constructor outputs cannot be promoted to variables
 			if (!bIsConstructorOutput)
 			{
@@ -1390,6 +1391,7 @@ FText UMetasoundEditorGraphSchema::GetPinDisplayName(const UEdGraphPin* Pin) con
 
 		case EMetasoundFrontendClassType::Literal:
 		case EMetasoundFrontendClassType::External:
+		case EMetasoundFrontendClassType::Template:
 		{
 			if (Pin->Direction == EGPD_Input)
 			{
@@ -1415,7 +1417,7 @@ FText UMetasoundEditorGraphSchema::GetPinDisplayName(const UEdGraphPin* Pin) con
 		case EMetasoundFrontendClassType::Invalid:
 		default:
 		{
-			static_assert(static_cast<int32>(EMetasoundFrontendClassType::Invalid) == 9, "Possible missing EMetasoundFrontendClassType case coverage");
+			static_assert(static_cast<int32>(EMetasoundFrontendClassType::Invalid) == 10, "Possible missing EMetasoundFrontendClassType case coverage");
 			return Super::GetPinDisplayName(Pin);
 		}
 	}
diff --git a/Engine/Plugins/Runtime/Metasound/Source/MetasoundEngine/Private/MetasoundSource.cpp b/Engine/Plugins/Runtime/Metasound/Source/MetasoundEngine/Private/MetasoundSource.cpp
--- a/Engine/Plugins/Runtime/Metasound/Source/MetasoundEngine/Private/MetasoundSource.cpp
+++ b/Engine/Plugins/Runtime/Metasound/Source/MetasoundEngine/Private/MetasoundSource.cpp
@@ -7,6 +7,7 @@
 #include "IAudioParameterInterfaceRegistry.h"
 #include "Internationalization/Text.h"
 #include "MetasoundAssetBase.h"
+#include "MetasoundAssetManager.h"
 #include "MetasoundAudioFormats.h"
 #include "MetasoundEngineArchetypes.h"
 #include "MetasoundEngineEnvironment.h"
@@ -70,6 +71,18 @@ namespace Metasound
 			static const FFormatOutputVertexKeyMap Map = CreateVertexKeyMap();
 			return Map;
 		}
+
+		Frontend::FMetaSoundAssetRegistrationOptions GetInitRegistrationOptions()
+		{
+			Frontend::FMetaSoundAssetRegistrationOptions RegOptions;
+			RegOptions.bForceReregister = false;
+			if (const UMetaSoundSettings* Settings = GetDefault<UMetaSoundSettings>())
+			{
+				RegOptions.bAutoUpdateLogWarningOnDroppedConnection = Settings->bAutoUpdateLogWarningOnDroppedConnection;
+			}
+
+			return RegOptions;
+		}
 	} // namespace SourcePrivate
 } // namespace Metasound
 
@@ -268,11 +281,14 @@ void UMetaSoundSource::SetRegistryAssetClassInfo(const Metasound::Frontend::FNod
 
 void UMetaSoundSource::InitParameters(TArray<FAudioParameter>& InParametersToInit, FName InFeatureName)
 {
+	using namespace Metasound::SourcePrivate;
+
 	METASOUND_LLM_SCOPE;
 
-	// Have to call cache vs a simple get as the source may have yet to start playing/has not been registered
+	// Have to call register vs a simple get as the source may have yet to start playing/has not been registered
 	// via InitResources. If it has, this call is fast and returns the already cached RuntimeData.
-	const FRuntimeData& RuntimeData = CacheRuntimeData();
+	RegisterGraphWithFrontend(GetInitRegistrationOptions());
+	const FRuntimeData& RuntimeData = GetRuntimeData();
 	const TArray<FMetasoundFrontendClassInput>& TransmittableInputs = RuntimeData.TransmittableInputs;
 
 	// Removes values that are not explicitly defined by the ParamType and returns
@@ -460,19 +476,12 @@ void UMetaSoundSource::InitParameters(TArray<FAudioParameter>& InParametersToIni
 void UMetaSoundSource::InitResources()
 {
 	using namespace Metasound::Frontend;
+	using namespace Metasound::SourcePrivate;
 
 	METASOUND_LLM_SCOPE;
 	METASOUND_TRACE_CPUPROFILER_EVENT_SCOPE(UMetaSoundSource::InitResources);
 
-	FMetaSoundAssetRegistrationOptions RegOptions;
-	RegOptions.bForceReregister = false;
-	if (const UMetaSoundSettings* Settings = GetDefault<UMetaSoundSettings>())
-	{
-		RegOptions.bAutoUpdateLogWarningOnDroppedConnection = Settings->bAutoUpdateLogWarningOnDroppedConnection;
-	}
-
-	RegisterGraphWithFrontend(RegOptions);
-	CacheRuntimeData();
+	RegisterGraphWithFrontend(GetInitRegistrationOptions());
 }
 
 Metasound::Frontend::FNodeClassInfo UMetaSoundSource::GetAssetClassInfo() const
diff --git a/Engine/Plugins/Runtime/Metasound/Source/MetasoundFrontend/Private/MetasoundAssetBase.cpp b/Engine/Plugins/Runtime/Metasound/Source/MetasoundFrontend/Private/MetasoundAssetBase.cpp
--- a/Engine/Plugins/Runtime/Metasound/Source/MetasoundFrontend/Private/MetasoundAssetBase.cpp
+++ b/Engine/Plugins/Runtime/Metasound/Source/MetasoundFrontend/Private/MetasoundAssetBase.cpp
@@ -18,6 +18,7 @@
 #include "MetasoundFrontendDocument.h"
 #include "MetasoundFrontendGraph.h"
 #include "MetasoundFrontendInjectReceiveNodes.h"
+#include "MetasoundFrontendNodeTemplateRegistry.h"
 #include "MetasoundFrontendRegistries.h"
 #include "MetasoundFrontendSearchEngine.h"
 #include "MetasoundFrontendTransform.h"
@@ -27,6 +28,7 @@
 #include "MetasoundTrace.h"
 #include "MetasoundVertex.h"
 #include "StructSerializer.h"
+#include "Templates/SharedPointer.h"
 #include "UObject/MetaData.h"
 
 #define LOCTEXT_NAMESPACE "MetaSound"
@@ -57,6 +59,90 @@ namespace Metasound
 					}
 				}
 			}
+
+			// Takes in document from asset and builds all template nodes and removes template classes from document dependencies list.
+			void PreBuildDocument(FMetasoundFrontendDocument& InOutDoc)
+			{
+				using namespace Frontend;
+
+				METASOUND_TRACE_CPUPROFILER_EVENT_SCOPE(Metasound::Frontend::AssetBase::PreBuildDocument);
+
+				FMetasoundFrontendGraph& Graph = InOutDoc.RootGraph.Graph;
+				TArray<FMetasoundFrontendClass>& Dependencies = InOutDoc.Dependencies;
+				const TArray<FMetasoundFrontendEdge>& GraphEdges = Graph.Edges;
+
+				TMap<FGuid, const FMetasoundFrontendClass*> TemplateClassMap;
+				Algo::TransformIf(Dependencies, TemplateClassMap,
+					[](const FMetasoundFrontendClass& Class) { return Class.Metadata.GetType() == EMetasoundFrontendClassType::Template; },
+					[](const FMetasoundFrontendClass& Class) { return TPair<FGuid, const FMetasoundFrontendClass*> { Class.ID, & Class }; }
+				);
+
+				TSet<const FMetasoundFrontendNode*> NodesToRemove;
+				TSet<TPair<FGuid, FGuid>> VerticesRemoved;
+
+				INodeTemplateRegistry& TemplateReg = INodeTemplateRegistry::Get();
+				for (FMetasoundFrontendNode& Node : Graph.Nodes)
+				{
+					const FMetasoundFrontendClass* NodeClass = TemplateClassMap.FindRef(Node.ClassID);
+					if (NodeClass)
+					{
+						if (NodeClass->Metadata.GetType() == EMetasoundFrontendClassType::Template)
+						{
+							const FNodeRegistryKey Key = NodeRegistryKey::CreateKey(NodeClass->Metadata);
+							const INodeTemplate* Template = TemplateReg.FindTemplate(Key);
+							if (ensure(Template))
+							{
+								const bool bTemplateBuilt = Template->BuildTemplate(InOutDoc, Graph, Node);
+								if (bTemplateBuilt)
+								{
+									NodesToRemove.Add(&Node);
+
+									auto GetNodeVertexGuidPair = [NodeID = Node.GetID()](const FMetasoundFrontendVertex& Vertex) { return TPair<FGuid, FGuid> { NodeID, Vertex.VertexID }; };
+									Algo::Transform(Node.Interface.Inputs, VerticesRemoved, GetNodeVertexGuidPair);
+									Algo::Transform(Node.Interface.Outputs, VerticesRemoved, GetNodeVertexGuidPair);
+								}
+							}
+						}
+					}
+				}
+
+				constexpr bool bAllowShrinking = false;
+				for (int32 i = Dependencies.Num() - 1; i >= 0; --i)
+				{
+					if (TemplateClassMap.Contains(Dependencies[i].ID))
+					{
+						Dependencies.RemoveAtSwap(i, 1, bAllowShrinking);
+					}
+				}
+				Dependencies.Shrink();
+
+				for (int32 i = Graph.Nodes.Num() - 1; i >= 0; --i)
+				{
+					if (NodesToRemove.Contains(&Graph.Nodes[i]))
+					{
+						Graph.Nodes.RemoveAtSwap(i, 1, bAllowShrinking);
+					}
+				}
+				Graph.Nodes.Shrink();
+
+				for (int32 i = Graph.Edges.Num() - 1; i >= 0; --i)
+				{
+					const TPair<FGuid, FGuid> FromNodeVertexPair{ Graph.Edges[i].FromNodeID, Graph.Edges[i].FromVertexID };
+					if (VerticesRemoved.Contains(FromNodeVertexPair))
+					{
+						Graph.Edges.RemoveAtSwap(i, 1, bAllowShrinking);
+					}
+					else
+					{
+						const TPair<FGuid, FGuid> ToNodeVertexPair{ Graph.Edges[i].ToNodeID, Graph.Edges[i].ToVertexID };
+						if (VerticesRemoved.Contains(ToNodeVertexPair))
+						{
+							Graph.Edges.RemoveAtSwap(i, 1, bAllowShrinking);
+						}
+					}
+				}
+				Graph.Edges.Shrink();
+			}
 		} // namespace AssetBasePrivate
 
 		FAutoConsoleVariableRef CVarMetaSoundBlockRate(
@@ -89,9 +175,6 @@ void FMetasoundAssetBase::RegisterGraphWithFrontend(Metasound::Frontend::FMetaSo
 		}
 	}
 
-	// Triggers the existing runtime data to be out-of-date.
-	CurrentCachedRuntimeDataChangeID = FGuid::NewGuid();
-
 	if (InRegistrationOptions.bRegisterDependencies)
 	{
 		// Must be called in case register is called prior to asset scan being completed
@@ -157,14 +240,15 @@ void FMetasoundAssetBase::RegisterGraphWithFrontend(Metasound::Frontend::FMetaSo
 	class FNodeRegistryEntry : public INodeRegistryEntry
 	{
 	public:
-		FNodeRegistryEntry(const FString& InName, const FMetasoundFrontendDocument& InDocument, FName InAssetPath)
+		FNodeRegistryEntry(const FString& InName, TSharedPtr<FMetasoundFrontendDocument> InPreBuiltDocument, FName InAssetPath)
 		: Name(InName)
-		, Document(InDocument)
+		, PreBuiltDocument(InPreBuiltDocument)
 		{
-			// Copy frontend class to preserve original document.
-			FrontendClass = Document.RootGraph;
+			// Copy FrontendClass to preserve original document.
+			FrontendClass = PreBuiltDocument->RootGraph;
 			FrontendClass.Metadata.SetType(EMetasoundFrontendClassType::External);
-			ClassInfo = FNodeClassInfo(Document.RootGraph, InAssetPath);
+
+			ClassInfo = FNodeClassInfo(PreBuiltDocument->RootGraph, InAssetPath);
 		}
 
 		virtual ~FNodeRegistryEntry() = default;
@@ -176,7 +260,7 @@ void FMetasoundAssetBase::RegisterGraphWithFrontend(Metasound::Frontend::FMetaSo
 
 		virtual TUniquePtr<INode> CreateNode(const FNodeInitData&) const override
 		{
-			return FFrontendGraphBuilder().CreateGraph(Document, Name);
+			return FFrontendGraphBuilder().CreateGraph(*PreBuiltDocument, Name);
 		}
 
 		virtual TUniquePtr<INode> CreateNode(FDefaultLiteralNodeConstructorParams&&) const override
@@ -201,7 +285,7 @@ void FMetasoundAssetBase::RegisterGraphWithFrontend(Metasound::Frontend::FMetaSo
 
 		virtual TUniquePtr<INodeRegistryEntry> Clone() const override
 		{
-			return MakeUnique<FNodeRegistryEntry>(Name, Document, ClassInfo.AssetPath);
+			return MakeUnique<FNodeRegistryEntry>(Name, PreBuiltDocument, ClassInfo.AssetPath);
 		}
 
 		virtual bool IsNative() const override
@@ -210,9 +294,8 @@ void FMetasoundAssetBase::RegisterGraphWithFrontend(Metasound::Frontend::FMetaSo
 		}
 
 	private:
-		
 		FString Name;
-		FMetasoundFrontendDocument Document;
+		TSharedPtr<FMetasoundFrontendDocument> PreBuiltDocument;
 		FMetasoundFrontendClass FrontendClass;
 		FNodeClassInfo ClassInfo;
 	};
@@ -229,11 +312,11 @@ void FMetasoundAssetBase::RegisterGraphWithFrontend(Metasound::Frontend::FMetaSo
 	}
 
 	FNodeClassInfo AssetClassInfo = GetAssetClassInfo();
-	const FMetasoundFrontendDocument* Doc = GetDocument().Get();
-	if (Doc)
-	{
-		RegistryKey = FMetasoundFrontendRegistryContainer::Get()->RegisterNode(MakeUnique<FNodeRegistryEntry>(AssetName, *Doc, AssetClassInfo.AssetPath));
-	}
+	TSharedPtr<FMetasoundFrontendDocument> PreBuiltDoc = MakeShared<FMetasoundFrontendDocument>(GetDocumentChecked());
+	AssetBasePrivate::PreBuildDocument(*PreBuiltDoc);
+
+	TUniquePtr<INodeRegistryEntry> RegistryEntry = MakeUnique<FNodeRegistryEntry>(AssetName, PreBuiltDoc, AssetClassInfo.AssetPath);
+	RegistryKey = FMetasoundFrontendRegistryContainer::Get()->RegisterNode(MoveTemp(RegistryEntry));
 
 	if (NodeRegistryKey::IsValid(RegistryKey) && FMetasoundFrontendRegistryContainer::Get()->IsNodeRegistered(RegistryKey))
 	{
@@ -268,6 +351,10 @@ void FMetasoundAssetBase::RegisterGraphWithFrontend(Metasound::Frontend::FMetaSo
 		}
 		UE_LOG(LogMetaSound, Error, TEXT("Registration failed for MetaSound node class '%s' of UObject class '%s'"), *AssetName, *ClassName);
 	}
+
+	// Triggers the existing runtime data to be out-of-date.
+	CurrentCachedRuntimeDataChangeID = FGuid::NewGuid();
+	CacheRuntimeData(*PreBuiltDoc);
 }
 
 void FMetasoundAssetBase::UnregisterGraphWithFrontend()
@@ -555,7 +642,7 @@ void FMetasoundAssetBase::ResetSynchronizationState()
 }
 #endif // WITH_EDITOR
 
-TSharedPtr<Metasound::IGraph, ESPMode::ThreadSafe> FMetasoundAssetBase::BuildMetasoundDocument(const TArray<FMetasoundFrontendClassInput>& InTransmittableInputs) const
+TSharedPtr<Metasound::IGraph, ESPMode::ThreadSafe> FMetasoundAssetBase::BuildMetasoundDocument(const FMetasoundFrontendDocument& InPreBuiltDoc, const TArray<FMetasoundFrontendClassInput>& InTransmittableInputs) const
 {
 	using namespace Metasound;
 	using namespace Metasound::Frontend;
@@ -563,7 +650,7 @@ TSharedPtr<Metasound::IGraph, ESPMode::ThreadSafe> FMetasoundAssetBase::BuildMet
 	METASOUND_TRACE_CPUPROFILER_EVENT_SCOPE(MetaSoundAssetBase::BuildMetasoundDocument);
 
 	// Create graph which can spawn instances. 
-	TUniquePtr<FFrontendGraph> FrontendGraph = FFrontendGraphBuilder::CreateGraph(GetDocumentChecked(), GetOwningAssetName());
+	TUniquePtr<FFrontendGraph> FrontendGraph = FFrontendGraphBuilder::CreateGraph(InPreBuiltDoc, GetOwningAssetName());
 	if (FrontendGraph.IsValid())
 	{
 		TSet<FVertexName> TransmittableInputNames;
@@ -889,7 +976,7 @@ void FMetasoundAssetBase::RebuildReferencedAssetClassKeys()
 	SetReferencedAssetClassKeys(MoveTemp(ReferencedKeys));
 }
 
-const FMetasoundAssetBase::FRuntimeData& FMetasoundAssetBase::CacheRuntimeData()
+const FMetasoundAssetBase::FRuntimeData& FMetasoundAssetBase::CacheRuntimeData(const FMetasoundFrontendDocument& InPreBuiltDoc)
 {
 	METASOUND_TRACE_CPUPROFILER_EVENT_SCOPE(MetaSoundAssetBase::CacheRuntimeData);
 
@@ -904,7 +991,7 @@ const FMetasoundAssetBase::FRuntimeData& FMetasoundAssetBase::CacheRuntimeData()
 	{
 		// Update CachedRuntimeData.
 		TArray<FMetasoundFrontendClassInput> ClassInputs = GetTransmittableClassInputs();
-		TSharedPtr<Metasound::IGraph, ESPMode::ThreadSafe> Graph = BuildMetasoundDocument(ClassInputs);
+		TSharedPtr<Metasound::IGraph, ESPMode::ThreadSafe> Graph = BuildMetasoundDocument(InPreBuiltDoc, ClassInputs);
 
 		CachedRuntimeData =
 		{
diff --git a/Engine/Plugins/Runtime/Metasound/Source/MetasoundFrontend/Private/MetasoundFrontendDocumentController.cpp b/Engine/Plugins/Runtime/Metasound/Source/MetasoundFrontend/Private/MetasoundFrontendDocumentController.cpp
--- a/Engine/Plugins/Runtime/Metasound/Source/MetasoundFrontend/Private/MetasoundFrontendDocumentController.cpp
+++ b/Engine/Plugins/Runtime/Metasound/Source/MetasoundFrontend/Private/MetasoundFrontendDocumentController.cpp
@@ -280,7 +280,8 @@ namespace Metasound
 				if (FMetasoundFrontendClass* MetasoundClass = ClassPtr.Get())
 				{
 					// External node classes must match version to return shared definition.
-					if (MetasoundClass->Metadata.GetType() == EMetasoundFrontendClassType::External)
+					if (MetasoundClass->Metadata.GetType() == EMetasoundFrontendClassType::Template
+						|| MetasoundClass->Metadata.GetType() == EMetasoundFrontendClassType::External)
 					{
 						// TODO: Assuming we want to recheck classes when they add another
 						// node, this should be replace with a call to synchronize a 
@@ -323,8 +324,9 @@ namespace Metasound
 			
 			if (const FMetasoundFrontendClass* Class = ClassPtr.Get())
 			{
-				// External node classes must match major version to return shared definition.
-				if (EMetasoundFrontendClassType::External == InMetadata.GetType())
+				// External & Template node classes must match major version to return shared definition.
+				if (EMetasoundFrontendClassType::External == InMetadata.GetType()
+					|| EMetasoundFrontendClassType::Template == InMetadata.GetType())
 				{
 					if (InMetadata.GetVersion().Major != Class->Metadata.GetVersion().Major)
 					{
@@ -344,6 +346,7 @@ namespace Metasound
 					switch (InMetadata.GetType())
 					{
 						case EMetasoundFrontendClassType::External:
+						case EMetasoundFrontendClassType::Template:
 						case EMetasoundFrontendClassType::Input:
 						case EMetasoundFrontendClassType::Output:
 						{
diff --git a/Engine/Plugins/Runtime/Metasound/Source/MetasoundFrontend/Private/MetasoundFrontendGraph.cpp b/Engine/Plugins/Runtime/Metasound/Source/MetasoundFrontend/Private/MetasoundFrontendGraph.cpp
--- a/Engine/Plugins/Runtime/Metasound/Source/MetasoundFrontend/Private/MetasoundFrontendGraph.cpp
+++ b/Engine/Plugins/Runtime/Metasound/Source/MetasoundFrontend/Private/MetasoundFrontendGraph.cpp
@@ -9,6 +9,8 @@
 #include "CoreMinimal.h"
 #include "MetasoundFrontend.h"
 #include "MetasoundFrontendDataTypeRegistry.h"
+#include "MetasoundFrontendNodeTemplateRegistry.h"
+#include "MetasoundFrontendRegistries.h"
 #include "MetasoundGraph.h"
 #include "MetasoundLiteralNode.h"
 #include "MetasoundLog.h"
@@ -286,6 +288,8 @@ namespace Metasound
 
 	TUniquePtr<INode> FFrontendGraphBuilder::CreateExternalNode(const FMetasoundFrontendNode& InNode, const FMetasoundFrontendClass& InClass, FBuildGraphContext& InGraphContext, const TSet<FNodeIDVertexID>& InEdgeDestinations)
 	{
+		using namespace Frontend;
+
 		check(InNode.ClassID == InClass.ID);
 
 		const FNodeInitData InitData = FrontendGraphPrivate::CreateNodeInitData(InNode);
@@ -300,8 +304,7 @@ namespace Metasound
 		// TODO: handle check to see if node interface conforms to class interface here. 
 		// TODO: check to see if external object supports class interface.
 
-		Metasound::Frontend::FNodeRegistryKey Key = FMetasoundFrontendRegistryContainer::GetRegistryKey(InClass.Metadata);
-
+		const FNodeRegistryKey Key = NodeRegistryKey::CreateKey(InClass.Metadata);
 		return FMetasoundFrontendRegistryContainer::Get()->CreateNode(Key, InitData);
 	}
 
@@ -492,8 +495,9 @@ namespace Metasound
 					}
 					break;
 
-					// Variable accessors and mutators are constructed with the same
-					// parameters as external nodes. 
+					// Templates, variable accessors and mutators are
+					// constructed with the same parameters as external nodes.
+					case EMetasoundFrontendClassType::Template:
 					case EMetasoundFrontendClassType::VariableAccessor:
 					case EMetasoundFrontendClassType::VariableDeferredAccessor:
 					case EMetasoundFrontendClassType::VariableMutator:
@@ -720,8 +724,9 @@ namespace Metasound
 	{
 		// All dependencies are external dependencies in a flat graph
 		auto IsClassExternal = [&](const FMetasoundFrontendClass& InDesc) 
-		{ 
-			bool bIsExternal = (InDesc.Metadata.GetType() == EMetasoundFrontendClassType::External) ||
+		{
+			const bool bIsExternal = (InDesc.Metadata.GetType() == EMetasoundFrontendClassType::External) ||
+				(InDesc.Metadata.GetType() == EMetasoundFrontendClassType::Template) ||
 				(InDesc.Metadata.GetType() == EMetasoundFrontendClassType::Input) ||
 				(InDesc.Metadata.GetType() == EMetasoundFrontendClassType::Output);
 			return bIsExternal;
diff --git a/Engine/Plugins/Runtime/Metasound/Source/MetasoundFrontend/Private/MetasoundFrontendGraphController.cpp b/Engine/Plugins/Runtime/Metasound/Source/MetasoundFrontend/Private/MetasoundFrontendGraphController.cpp
--- a/Engine/Plugins/Runtime/Metasound/Source/MetasoundFrontend/Private/MetasoundFrontendGraphController.cpp
+++ b/Engine/Plugins/Runtime/Metasound/Source/MetasoundFrontend/Private/MetasoundFrontendGraphController.cpp
@@ -8,6 +8,7 @@
 #include "MetasoundFrontendDocumentAccessPtr.h"
 #include "MetasoundFrontendGraph.h"
 #include "MetasoundFrontendNodeController.h"
+#include "MetasoundFrontendNodeTemplateRegistry.h"
 #include "MetasoundFrontendSubgraphNodeController.h"
 #include "MetasoundFrontendInvalidController.h"
 #include "MetasoundFrontendVariableController.h"
@@ -1378,6 +1379,38 @@ namespace Metasound
 			return AddNode(NodeRegistryKey::CreateKey(InClassMetadata), InNodeGuid);
 		}
 
+		FNodeHandle FGraphController::AddTemplateNode(const FNodeRegistryKey& InKey, FMetasoundFrontendNodeInterface&& InNodeInterface, FGuid InNodeGuid)
+		{
+			if (const INodeTemplate* Template = INodeTemplateRegistry::Get().FindTemplate(InKey))
+			{
+				const bool bIsValidInterface = Template->IsValidNodeInterface(InNodeInterface);
+				if (ensureAlwaysMsgf(bIsValidInterface, TEXT("Cannot implement interface when attempting to add node using template with key '%s'"), *InKey))
+				{
+					// Construct a FNodeClassInfo from this lookup key.
+					FConstClassAccessPtr Class = OwningDocument->FindOrAddClass(InKey);
+					const bool bIsValidClass = (nullptr != Class.Get());
+
+					if (bIsValidClass)
+					{
+						if (FMetasoundFrontendGraphClass* GraphClass = GraphClassPtr.Get())
+						{
+							if (const FMetasoundFrontendClass* NodeClass = Class.Get())
+							{
+								FMetasoundFrontendNode& Node = GraphClass->Graph.Nodes.Emplace_GetRef(*NodeClass);
+								Node.UpdateID(InNodeGuid);
+								Node.Interface = InNodeInterface;
+								FNodeAccessPtr NodePtr = GraphClassPtr.GetNodeWithNodeID(Node.GetID());
+								return GetNodeHandle(FGraphController::FNodeAndClass{ NodePtr, Class });
+							}
+						}
+					}
+				}
+			}
+
+			UE_LOG(LogMetaSound, Warning, TEXT("Failed to find or add node template class info with registry key [Key:%s]"), *InKey);
+			return INodeController::GetInvalidHandle();
+		}
+
 		FNodeHandle FGraphController::AddDuplicateNode(const INodeController& InNode)
 		{
 			// TODO: will need to copy node interface when dynamic pins exist.
@@ -1454,6 +1487,7 @@ namespace Metasound
 							// TODO: remove node from variable.7
 						case EMetasoundFrontendClassType::Literal:
 						case EMetasoundFrontendClassType::External:
+						case EMetasoundFrontendClassType::Template:
 						case EMetasoundFrontendClassType::Graph:
 						{
 							return RemoveNode(*FrontendNode);
@@ -1462,7 +1496,7 @@ namespace Metasound
 						default:
 						case EMetasoundFrontendClassType::Invalid:
 						{
-							static_assert(static_cast<int32>(EMetasoundFrontendClassType::Invalid) == 9, "Possible missing switch case coverage for EMetasoundFrontendClassType.");
+							static_assert(static_cast<int32>(EMetasoundFrontendClassType::Invalid) == 10, "Possible missing switch case coverage for EMetasoundFrontendClassType.");
 							checkNoEntry();
 						}
 					}
@@ -1950,6 +1984,7 @@ namespace Metasound
 						break;
 
 					case EMetasoundFrontendClassType::External:
+					case EMetasoundFrontendClassType::Template:
 						{
 							FNodeController::FInitParams InitParams
 							{
@@ -1977,7 +2012,7 @@ namespace Metasound
 
 					default:
 						checkNoEntry();
-						static_assert(static_cast<int32>(EMetasoundFrontendClassType::Invalid) == 9, "Possible missing switch case coverage for EMetasoundFrontendClassType.");
+						static_assert(static_cast<int32>(EMetasoundFrontendClassType::Invalid) == 10, "Possible missing switch case coverage for EMetasoundFrontendClassType.");
 				}
 			}
 
@@ -2049,6 +2084,7 @@ namespace Metasound
 						break;
 
 					case EMetasoundFrontendClassType::External:
+					case EMetasoundFrontendClassType::Template:
 						{
 							FNodeController::FInitParams InitParams
 							{
@@ -2076,7 +2112,7 @@ namespace Metasound
 
 					default:
 						checkNoEntry();
-						static_assert(static_cast<int32>(EMetasoundFrontendClassType::Invalid) == 9, "Possible missing switch case coverage for EMetasoundFrontendClassType.");
+						static_assert(static_cast<int32>(EMetasoundFrontendClassType::Invalid) == 10, "Possible missing switch case coverage for EMetasoundFrontendClassType.");
 				}
 			}
 
diff --git a/Engine/Plugins/Runtime/Metasound/Source/MetasoundFrontend/Private/MetasoundFrontendGraphController.h b/Engine/Plugins/Runtime/Metasound/Source/MetasoundFrontend/Private/MetasoundFrontendGraphController.h
--- a/Engine/Plugins/Runtime/Metasound/Source/MetasoundFrontend/Private/MetasoundFrontendGraphController.h
+++ b/Engine/Plugins/Runtime/Metasound/Source/MetasoundFrontend/Private/MetasoundFrontendGraphController.h
@@ -172,6 +172,7 @@ namespace Metasound
 			FNodeHandle AddNode(const FNodeRegistryKey& InNodeClass, FGuid InNodeGuid) override;
 			FNodeHandle AddNode(const FMetasoundFrontendClassMetadata& InClassMetadata, FGuid InNodeGuid) override;
 			FNodeHandle AddDuplicateNode(const INodeController& InNode) override;
+			FNodeHandle AddTemplateNode(const FNodeRegistryKey& InKey, FMetasoundFrontendNodeInterface&& InNodeInterface, FGuid InNodeGuid) override;
 
 			// Remove the node corresponding to this node handle.
 			// On success, invalidates the received node handle.
diff --git a/Engine/Plugins/Runtime/Metasound/Source/MetasoundFrontend/Private/MetasoundFrontendInvalidController.cpp b/Engine/Plugins/Runtime/Metasound/Source/MetasoundFrontend/Private/MetasoundFrontendInvalidController.cpp
--- a/Engine/Plugins/Runtime/Metasound/Source/MetasoundFrontend/Private/MetasoundFrontendInvalidController.cpp
+++ b/Engine/Plugins/Runtime/Metasound/Source/MetasoundFrontend/Private/MetasoundFrontendInvalidController.cpp
@@ -50,6 +50,12 @@ namespace Metasound
 				return Invalid;
 			}
 
+			const FMetasoundFrontendNodeInterface& GetInvalidNodeInterface()
+			{
+				static const FMetasoundFrontendNodeInterface Invalid;
+				return Invalid;
+			}
+
 			const FMetasoundFrontendGraphClassPresetOptions& GetInvalidGraphClassPresetOptions()
 			{
 				static const FMetasoundFrontendGraphClassPresetOptions Invalid;
diff --git a/Engine/Plugins/Runtime/Metasound/Source/MetasoundFrontend/Private/MetasoundFrontendInvalidController.h b/Engine/Plugins/Runtime/Metasound/Source/MetasoundFrontend/Private/MetasoundFrontendInvalidController.h
--- a/Engine/Plugins/Runtime/Metasound/Source/MetasoundFrontend/Private/MetasoundFrontendInvalidController.h
+++ b/Engine/Plugins/Runtime/Metasound/Source/MetasoundFrontend/Private/MetasoundFrontendInvalidController.h
@@ -34,6 +34,7 @@ namespace Metasound
 
 			const FMetasoundFrontendLiteral& GetInvalidLiteral();
 			const FMetasoundFrontendClassInterface& GetInvalidClassInterface();
+			const FMetasoundFrontendNodeInterface& GetInvalidNodeInterface();
 			const FMetasoundFrontendClassMetadata& GetInvalidClassMetadata();
 			const FMetasoundFrontendGraphClassPresetOptions& GetInvalidGraphClassPresetOptions();
 			const FMetasoundFrontendGraphClass& GetInvalidGraphClass();
@@ -237,6 +238,8 @@ namespace Metasound
 			virtual const FMetasoundFrontendClassInterface& GetClassInterface() const override { return Invalid::GetInvalidClassInterface(); }
 			virtual const FMetasoundFrontendClassMetadata& GetClassMetadata() const override { return Invalid::GetInvalidClassMetadata(); }
 
+			virtual const FMetasoundFrontendNodeInterface& GetNodeInterface() const override { return Invalid::GetInvalidNodeInterface(); }
+
 #if WITH_EDITOR
 			virtual const FMetasoundFrontendInterfaceStyle& GetInputStyle() const override { return Invalid::GetInvalidInterfaceStyle(); }
 			virtual const FMetasoundFrontendInterfaceStyle& GetOutputStyle() const override { return Invalid::GetInvalidInterfaceStyle(); }
@@ -404,6 +407,7 @@ namespace Metasound
 			virtual FNodeHandle AddNode(const FNodeRegistryKey& InNodeClass, FGuid InNodeGuid) override { return INodeController::GetInvalidHandle(); }
 			virtual FNodeHandle AddNode(const FMetasoundFrontendClassMetadata& InNodeClass, FGuid InNodeGuid) override { return INodeController::GetInvalidHandle(); }
 			virtual FNodeHandle AddDuplicateNode(const INodeController& InNode) override { return INodeController::GetInvalidHandle(); }
+			virtual FNodeHandle AddTemplateNode(const FNodeRegistryKey& InNodeClass, FMetasoundFrontendNodeInterface&& InNodeInterface, FGuid InNodeGuid) override { return INodeController::GetInvalidHandle(); }
 
 			// Remove the node corresponding to this node handle.
 			// On success, invalidates the received node handle.
diff --git a/Engine/Plugins/Runtime/Metasound/Source/MetasoundFrontend/Private/MetasoundFrontendNodeController.cpp b/Engine/Plugins/Runtime/Metasound/Source/MetasoundFrontend/Private/MetasoundFrontendNodeController.cpp
--- a/Engine/Plugins/Runtime/Metasound/Source/MetasoundFrontend/Private/MetasoundFrontendNodeController.cpp
+++ b/Engine/Plugins/Runtime/Metasound/Source/MetasoundFrontend/Private/MetasoundFrontendNodeController.cpp
@@ -166,6 +166,15 @@ namespace Metasound
 			return Invalid::GetInvalidClassMetadata();
 		}
 
+		const FMetasoundFrontendNodeInterface& FBaseNodeController::GetNodeInterface() const
+		{
+			if (const FMetasoundFrontendNode* Node = NodePtr.Get())
+			{
+				return Node->Interface;
+			}
+			return Invalid::GetInvalidNodeInterface();
+		}
+
 #if WITH_EDITOR
 		const FMetasoundFrontendInterfaceStyle& FBaseNodeController::GetInputStyle() const
 		{
diff --git a/Engine/Plugins/Runtime/Metasound/Source/MetasoundFrontend/Private/MetasoundFrontendNodeController.h b/Engine/Plugins/Runtime/Metasound/Source/MetasoundFrontend/Private/MetasoundFrontendNodeController.h
--- a/Engine/Plugins/Runtime/Metasound/Source/MetasoundFrontend/Private/MetasoundFrontendNodeController.h
+++ b/Engine/Plugins/Runtime/Metasound/Source/MetasoundFrontend/Private/MetasoundFrontendNodeController.h
@@ -47,6 +47,8 @@ namespace Metasound
 			virtual const FMetasoundFrontendClassInterface& GetClassInterface() const override;
 			virtual const FMetasoundFrontendClassMetadata& GetClassMetadata() const override;
 
+			virtual const FMetasoundFrontendNodeInterface& GetNodeInterface() const override;
+
 #if WITH_EDITOR
 			virtual const FMetasoundFrontendInterfaceStyle& GetInputStyle() const override;
 			virtual const FMetasoundFrontendInterfaceStyle& GetOutputStyle() const override;
diff --git a/Engine/Plugins/Runtime/Metasound/Source/MetasoundFrontend/Private/MetasoundFrontendNodeRegistryPrivate.h b/Engine/Plugins/Runtime/Metasound/Source/MetasoundFrontend/Private/MetasoundFrontendNodeRegistryPrivate.h
--- a/Engine/Plugins/Runtime/Metasound/Source/MetasoundFrontend/Private/MetasoundFrontendNodeRegistryPrivate.h
+++ b/Engine/Plugins/Runtime/Metasound/Source/MetasoundFrontend/Private/MetasoundFrontendNodeRegistryPrivate.h
@@ -15,6 +15,26 @@ namespace Metasound
 		using FNodeRegistryTransactionBuffer = TTransactionBuffer<FNodeRegistryTransaction>;
 		using FNodeRegistryTransactionStream = TTransactionStream<FNodeRegistryTransaction>; 
 
+		/** INodeTemplateRegistryEntry declares the interface for a node registry entry.
+		 * Each node class in the registry must satisfy this interface.
+		 */
+		class INodeTemplateRegistryEntry
+		{
+		public:
+			virtual ~INodeTemplateRegistryEntry() = default;
+
+			/** Return FNodeClassInfo for the node class.
+			 *
+			 * Implementations of method should avoid any expensive operations
+			 * (e.g. loading from disk, allocating memory) as this method is called
+			 * frequently when querying nodes.
+			 */
+			virtual const FNodeClassInfo& GetClassInfo() const = 0;
+
+			/** Return a FMetasoundFrontendClass which describes the node. */
+			virtual const FMetasoundFrontendClass& GetFrontendClass() const = 0;
+		};
+
 		// Registry container private implementation.
 		class FRegistryContainerImpl : public FMetasoundFrontendRegistryContainer
 		{
@@ -53,7 +73,8 @@ namespace Metasound
 			 *
 			 * @return True on success.
 			 */
-			virtual FNodeRegistryKey RegisterNode(TUniquePtr<Metasound::Frontend::INodeRegistryEntry>&&) override;
+			virtual FNodeRegistryKey RegisterNode(TUniquePtr<Metasound::Frontend::INodeRegistryEntry>&& InEntry) override;
+
 			virtual void ForEachNodeRegistryTransactionSince(Metasound::Frontend::FRegistryTransactionID InSince, Metasound::Frontend::FRegistryTransactionID* OutCurrentRegistryTransactionID, TFunctionRef<void(const Metasound::Frontend::FNodeRegistryTransaction&)> InFunc) const override;
 			virtual bool UnregisterNode(const FNodeRegistryKey& InKey) override;
 			virtual bool IsNodeRegistered(const FNodeRegistryKey& InKey) const override;
@@ -84,6 +105,10 @@ namespace Metasound
 			// Returns an empty array if none are available.
 			virtual TArray<FConverterNodeInfo> GetPossibleConverterNodes(const FName& FromDataType, const FName& ToDataType) override;
 
+			// Private implementation until hardened and used for template nodes other than reroutes.
+			FNodeRegistryKey RegisterNodeTemplate(TUniquePtr<Metasound::Frontend::INodeTemplateRegistryEntry>&& InEntry);
+			bool UnregisterNodeTemplate(const FNodeRegistryKey& InKey);
+
 			// Create a transaction stream for any newly transactions
 			TUniquePtr<FNodeRegistryTransactionStream> CreateTransactionStream();
 
@@ -92,6 +117,7 @@ namespace Metasound
 
 			const INodeRegistryEntry* FindNodeEntry(const FNodeRegistryKey& InKey) const;
 
+			const INodeTemplateRegistryEntry* FindNodeTemplateEntry(const FNodeRegistryKey& InKey) const;
 
 			// This buffer is used to enqueue nodes and datatypes to register when the module has been initialized,
 			// in order to avoid bad behavior with ensures, logs, etc. on static initialization.
@@ -105,6 +131,9 @@ namespace Metasound
 			// Registry in which we keep all information about nodes implemented in C++.
 			TMap<FNodeRegistryKey, TSharedRef<INodeRegistryEntry, ESPMode::ThreadSafe>> RegisteredNodes;
 
+			// Registry in which we keep all information about dynamically-generated templated nodes via in C++.
+			TMap<FNodeRegistryKey, TSharedRef<INodeTemplateRegistryEntry, ESPMode::ThreadSafe>> RegisteredNodeTemplates;
+
 			// Registry in which we keep lists of possible nodes to use to convert between two datatypes
 			TMap<FConverterNodeRegistryKey, FConverterNodeRegistryValue> ConverterNodeRegistry;
 
diff --git a/Engine/Plugins/Runtime/Metasound/Source/MetasoundFrontend/Private/MetasoundFrontendNodeTemplateRegistry.cpp b/Engine/Plugins/Runtime/Metasound/Source/MetasoundFrontend/Private/MetasoundFrontendNodeTemplateRegistry.cpp
new file mode 100644
--- /dev/null
+++ b/Engine/Plugins/Runtime/Metasound/Source/MetasoundFrontend/Private/MetasoundFrontendNodeTemplateRegistry.cpp
@@ -0,0 +1,81 @@
+// Copyright Epic Games, Inc. All Rights Reserved.
+
+#include "MetasoundFrontendNodeTemplateRegistry.h"
+
+#include "MetasoundFrontendDocument.h"
+#include "MetasoundFrontendRegistries.h"
+#include "MetasoundFrontendRegistryTransaction.h"
+#include "MetasoundLog.h"
+#include "MetasoundTrace.h"
+
+
+namespace Metasound
+{
+	namespace Frontend
+	{
+		class FNodeTemplateRegistry : public INodeTemplateRegistry
+		{
+		public:
+			static FNodeTemplateRegistry& Get();
+
+			FNodeTemplateRegistry() = default;
+			virtual ~FNodeTemplateRegistry() = default;
+
+			virtual const INodeTemplate* FindTemplate(const FNodeRegistryKey& InKey) const override;
+
+			void Register(TUniquePtr<INodeTemplate>&& InEntry);
+			void Unregister(const FMetasoundFrontendVersion& InNodeTemplateVersion);
+
+		private:
+			TMap<FNodeRegistryKey, TUniquePtr<INodeTemplate>> Templates;
+		};
+
+		void FNodeTemplateRegistry::Register(TUniquePtr<INodeTemplate>&& InTemplate)
+		{
+			if (ensure(InTemplate.IsValid()))
+			{
+				const FMetasoundFrontendVersion& Version = InTemplate->GetVersion();
+				const FNodeRegistryKey Key = NodeRegistryKey::CreateKey(InTemplate->GetFrontendClass().Metadata);
+				if (ensure(NodeRegistryKey::IsValid(Key)))
+				{
+					Templates.Add(Key, MoveTemp(InTemplate));
+				}
+			}
+		}
+
+		void FNodeTemplateRegistry::Unregister(const FMetasoundFrontendVersion& InVersion)
+		{
+			const FNodeRegistryKey Key = NodeRegistryKey::CreateKey(EMetasoundFrontendClassType::Template, InVersion.Name.ToString(), InVersion.Number.Major, InVersion.Number.Minor);
+			if (ensure(NodeRegistryKey::IsValid(Key)))
+			{
+				Templates.Remove(Key);
+			}
+		}
+
+		const INodeTemplate* FNodeTemplateRegistry::FindTemplate(const FNodeRegistryKey& InKey) const
+		{
+			if (const TUniquePtr<INodeTemplate>* TemplatePtr = Templates.Find(InKey))
+			{
+				return TemplatePtr->Get();
+			}
+
+			return nullptr;
+		}
+
+		INodeTemplateRegistry& INodeTemplateRegistry::Get()
+		{
+			static FNodeTemplateRegistry Registry;
+			return Registry;
+		}
+
+		void RegisterNodeTemplate(TUniquePtr<INodeTemplate>&& InTemplate)
+		{
+			static_cast<FNodeTemplateRegistry&>(INodeTemplateRegistry::Get()).Register(MoveTemp(InTemplate));
+		}
+
+		void UnregisterNodeTemplate(const FMetasoundFrontendVersion& InVersion)
+		{
+			static_cast<FNodeTemplateRegistry&>(INodeTemplateRegistry::Get()).Unregister(InVersion);
+		}
+	} // namespace Frontend
+} // namespace Metasound
diff --git a/Engine/Plugins/Runtime/Metasound/Source/MetasoundFrontend/Private/MetasoundFrontendRegistries.cpp b/Engine/Plugins/Runtime/Metasound/Source/MetasoundFrontend/Private/MetasoundFrontendRegistries.cpp
--- a/Engine/Plugins/Runtime/Metasound/Source/MetasoundFrontend/Private/MetasoundFrontendRegistries.cpp
+++ b/Engine/Plugins/Runtime/Metasound/Source/MetasoundFrontend/Private/MetasoundFrontendRegistries.cpp
@@ -32,6 +32,7 @@ namespace Metasound
 				static const FString VariableType(TEXT("Variable"));
 				static const FString LiteralType(TEXT("Literal"));
 				static const FString GraphType(TEXT("Graph"));
+				static const FString TemplateType(TEXT("Template"));
 				static const FString InvalidType(TEXT("Invalid"));
 
 				switch (InType)
@@ -60,11 +61,14 @@ namespace Metasound
 					case EMetasoundFrontendClassType::VariableMutator:
 						return VariableMutatorType;
 
+					case EMetasoundFrontendClassType::Template:
+						return TemplateType;
+
 					case EMetasoundFrontendClassType::Graph:
 						return GraphType;
 
 					default:
-						static_assert(static_cast<uint8>(EMetasoundFrontendClassType::Invalid) == 9, "Missing EMetasoundFrontendClassType case coverage");
+						static_assert(static_cast<uint8>(EMetasoundFrontendClassType::Invalid) == 10, "Missing EMetasoundFrontendClassType case coverage");
 						return InvalidType;
 				}
 			}
@@ -235,6 +239,38 @@ namespace Metasound
 			return Key;
 		}
 
+		FNodeRegistryKey FRegistryContainerImpl::RegisterNodeTemplate(TUniquePtr<INodeTemplateRegistryEntry>&& InEntry)
+		{
+			METASOUND_LLM_SCOPE;
+
+			FNodeRegistryKey Key;
+
+			if (InEntry.IsValid())
+			{
+				TSharedRef<INodeTemplateRegistryEntry, ESPMode::ThreadSafe> Entry(InEntry.Release());
+
+				FNodeRegistryTransaction::FTimeType Timestamp = FPlatformTime::Cycles64();
+
+				Key = NodeRegistryKey::CreateKey(Entry->GetClassInfo());
+
+				// check to see if an identical node was already registered, and log
+				ensureAlwaysMsgf(
+					!RegisteredNodes.Contains(Key),
+					TEXT("Node template with registry key '%s' already registered. The previously registered node will be overwritten."),
+					*Key);
+
+				// Store update to newly registered node in history so nodes
+				// can be queried by transaction ID
+
+				TransactionBuffer->AddTransaction(FNodeRegistryTransaction(FNodeRegistryTransaction::ETransactionType::NodeRegistration, Entry->GetClassInfo(), Timestamp));
+
+				// Store registry elements in map so nodes can be queried using registry key.
+				RegisteredNodeTemplates.Add(Key, Entry);
+			}
+
+			return Key;
+		}
+
 		bool FRegistryContainerImpl::UnregisterNode(const FNodeRegistryKey& InKey)
 		{
 			METASOUND_LLM_SCOPE;
@@ -255,6 +291,26 @@ namespace Metasound
 			return false;
 		}
 
+		bool FRegistryContainerImpl::UnregisterNodeTemplate(const FNodeRegistryKey& InKey)
+		{
+			METASOUND_LLM_SCOPE;
+
+			if (NodeRegistryKey::IsValid(InKey))
+			{
+				if (const INodeTemplateRegistryEntry* Entry = FindNodeTemplateEntry(InKey))
+				{
+					FNodeRegistryTransaction::FTimeType Timestamp = FPlatformTime::Cycles64();
+
+					TransactionBuffer->AddTransaction(FNodeRegistryTransaction(FNodeRegistryTransaction::ETransactionType::NodeUnregistration, Entry->GetClassInfo(), Timestamp));
+
+					RegisteredNodeTemplates.Remove(InKey);
+					return true;
+				}
+			}
+
+			return false;
+		}
+
 		bool FRegistryContainerImpl::RegisterConversionNode(const FConverterNodeRegistryKey& InNodeKey, const FConverterNodeInfo& InNodeInfo)
 		{
 			if (!ConverterNodeRegistry.Contains(InNodeKey))
@@ -278,7 +334,7 @@ namespace Metasound
 
 		bool FRegistryContainerImpl::IsNodeRegistered(const FNodeRegistryKey& InKey) const
 		{
-			return RegisteredNodes.Contains(InKey);
+			return RegisteredNodes.Contains(InKey) || RegisteredNodeTemplates.Contains(InKey);
 		}
 
 		bool FRegistryContainerImpl::IsNodeNative(const FNodeRegistryKey& InKey) const
@@ -288,10 +344,14 @@ namespace Metasound
 				return Entry->IsNative();
 			}
 
+			if (const INodeTemplateRegistryEntry* TemplateEntry = FindNodeTemplateEntry(InKey))
+			{
+				return true;
+			}
+
 			return false;
 		}
 
-
 		bool FRegistryContainerImpl::FindFrontendClassFromRegistered(const FNodeRegistryKey& InKey, FMetasoundFrontendClass& OutClass)
 		{
 			if (const INodeRegistryEntry* Entry = FindNodeEntry(InKey))
@@ -299,6 +359,13 @@ namespace Metasound
 				OutClass = Entry->GetFrontendClass();
 				return true;
 			}
+
+			if (const INodeTemplateRegistryEntry* Entry = FindNodeTemplateEntry(InKey))
+			{
+				OutClass = Entry->GetFrontendClass();
+				return true;
+			}
+
 			return false;
 		}
 
@@ -310,6 +377,12 @@ namespace Metasound
 				return true;
 			}
 
+			if (const INodeTemplateRegistryEntry* Entry = FindNodeTemplateEntry(InKey))
+			{
+				OutInfo = Entry->GetClassInfo();
+				return true;
+			}
+
 			return false;
 		}
 
@@ -404,7 +477,6 @@ namespace Metasound
 			}
 		}
 
-
 		const INodeRegistryEntry* FRegistryContainerImpl::FindNodeEntry(const FNodeRegistryKey& InKey) const
 		{
 			if (const TSharedRef<INodeRegistryEntry, ESPMode::ThreadSafe>* Entry = RegisteredNodes.Find(InKey))
@@ -415,6 +487,15 @@ namespace Metasound
 			return nullptr;
 		}
 
+		const INodeTemplateRegistryEntry* FRegistryContainerImpl::FindNodeTemplateEntry(const FNodeRegistryKey& InKey) const
+		{
+			if (const TSharedRef<INodeTemplateRegistryEntry, ESPMode::ThreadSafe>* Entry = RegisteredNodeTemplates.Find(InKey))
+			{
+				return &Entry->Get();
+			}
+
+			return nullptr;
+		}
 
 		FNodeRegistryTransaction::FNodeRegistryTransaction(ETransactionType InType, const FNodeClassInfo& InNodeClassInfo, FNodeRegistryTransaction::FTimeType InTimestamp)
 		: Type(InType)
diff --git a/Engine/Plugins/Runtime/Metasound/Source/MetasoundFrontend/Public/MetasoundAssetBase.h b/Engine/Plugins/Runtime/Metasound/Source/MetasoundFrontend/Public/MetasoundAssetBase.h
--- a/Engine/Plugins/Runtime/Metasound/Source/MetasoundFrontend/Public/MetasoundAssetBase.h
+++ b/Engine/Plugins/Runtime/Metasound/Source/MetasoundFrontend/Public/MetasoundAssetBase.h
@@ -11,6 +11,7 @@
 #include "MetasoundLog.h"
 #include "MetasoundParameterTransmitter.h"
 #include "MetasoundVertex.h"
+#include "Templates/SharedPointer.h"
 #include "UObject/SoftObjectPath.h"
 #include "UObject/WeakObjectPtrTemplates.h"
 
@@ -202,7 +203,7 @@ protected:
 	};
 
 	// Returns the cached runtime data. Call updates cached data if out-of-date.
-	const FRuntimeData& CacheRuntimeData();
+	const FRuntimeData& CacheRuntimeData(const FMetasoundFrontendDocument& InPreBuiltDoc);
 
 	// Returns the cached runtime data.
 	const FRuntimeData& GetRuntimeData() const;
@@ -218,7 +219,7 @@ private:
 	FGuid CurrentCachedRuntimeDataChangeID;
 	FRuntimeData CachedRuntimeData;
 
-	TSharedPtr<Metasound::IGraph, ESPMode::ThreadSafe> BuildMetasoundDocument(const TArray<FMetasoundFrontendClassInput>& InTransmittableInputs) const;
+	TSharedPtr<Metasound::IGraph, ESPMode::ThreadSafe> BuildMetasoundDocument(const FMetasoundFrontendDocument& InPreBuildDoc, const TArray<FMetasoundFrontendClassInput>& InTransmittableInputs) const;
 	Metasound::FSendAddress CreateSendAddress(uint64 InInstanceID, const Metasound::FVertexName& InVertexName, const FName& InDataTypeName) const;
 	Metasound::Frontend::FNodeHandle AddInputPinForSendAddress(const Metasound::FMetaSoundParameterTransmitter::FSendInfo& InSendInfo, Metasound::Frontend::FGraphHandle InGraph) const;
 };
diff --git a/Engine/Plugins/Runtime/Metasound/Source/MetasoundFrontend/Public/MetasoundFrontendController.h b/Engine/Plugins/Runtime/Metasound/Source/MetasoundFrontend/Public/MetasoundFrontendController.h
--- a/Engine/Plugins/Runtime/Metasound/Source/MetasoundFrontend/Public/MetasoundFrontendController.h
+++ b/Engine/Plugins/Runtime/Metasound/Source/MetasoundFrontend/Public/MetasoundFrontendController.h
@@ -529,6 +529,11 @@ namespace Metasound
 			virtual const FMetasoundFrontendClassMetadata& GetClassMetadata() const = 0;
 			virtual const FMetasoundFrontendClassInterface& GetClassInterface() const = 0;
 
+			/** Returns the node interface, which may be different than the class interface
+			  * if the class supports dynamic input/output behavior (ex. Templates).
+			  */
+			virtual const FMetasoundFrontendNodeInterface& GetNodeInterface() const = 0;
+			
 #if WITH_EDITOR
 			/** Returns associated node class data */
 			virtual const FMetasoundFrontendInterfaceStyle& GetOutputStyle() const = 0;
@@ -945,6 +950,17 @@ namespace Metasound
 			 */
 			virtual FNodeHandle AddDuplicateNode(const INodeController& InNodeController) = 0;
 
+			/** Add a new template node to this graph, providing the defined interface as expected by the caller.
+			 *
+			 * @param InKey - Class key (must correspond with a class in the registry that was registered as a template).
+			 * @param InNodeInterface - Interface for node class.  Validated by template class registered in the node class registry. If invalid, node is not created/added.
+			 * @param InNodeGuid - (Optional) Explicit guid for the new node. Must be unique within the graph.
+			 * Only useful to specify explicitly when caller is managing or tracking the graph's guids (ex. replacing removed node).
+			 *
+			 * @return Node handle for class. On error, an invalid handle is returned.
+			 */
+			virtual FNodeHandle AddTemplateNode(const FNodeRegistryKey& InKey, FMetasoundFrontendNodeInterface&& InNodeInterface, FGuid InNodeGuid = FGuid::NewGuid()) = 0;
+
 			/** Remove the node corresponding to this node handle.
 			 *
 			 * @return True on success, false on failure. 
diff --git a/Engine/Plugins/Runtime/Metasound/Source/MetasoundFrontend/Public/MetasoundFrontendDocument.h b/Engine/Plugins/Runtime/Metasound/Source/MetasoundFrontend/Public/MetasoundFrontendDocument.h
--- a/Engine/Plugins/Runtime/Metasound/Source/MetasoundFrontend/Public/MetasoundFrontendDocument.h
+++ b/Engine/Plugins/Runtime/Metasound/Source/MetasoundFrontend/Public/MetasoundFrontendDocument.h
@@ -58,22 +58,22 @@ enum class EMetasoundFrontendVertexAccessType
 UENUM()
 enum class EMetasoundFrontendClassType : uint8
 {
-	// The Metasound class is defined externally, in compiled code or in another document.
+	// The MetaSound class is defined externally, in compiled code or in another document.
 	External,
 
-	// The Metasound class is a graph within the containing document.
+	// The MetaSound class is a graph within the containing document.
 	Graph,
 
-	// The Metasound class is an input into a graph in the containing document.
+	// The MetaSound class is an input into a graph in the containing document.
 	Input,
 
-	// The Metasound class is an output from a graph in the containing document.
+	// The MetaSound class is an output from a graph in the containing document.
 	Output,
 
-	// The Metasound class is an literal requiring an literal value to construct.
+	// The MetaSound class is an literal requiring an literal value to construct.
 	Literal,
 
-	// The Metasound class is an variable requiring an literal value to construct.
+	// The MetaSound class is an variable requiring an literal value to construct.
 	Variable,
 
 	// The MetaSound class accesses variables.
@@ -85,6 +85,10 @@ enum class EMetasoundFrontendClassType : uint8
 	// The MetaSound class mutates variables.
 	VariableMutator,
 
+	// The MetaSound class is defined only by the Frontend, and associatively
+	// performs a functional replacement operation in a pre-build step.
+	Template,
+
 	Invalid UMETA(Hidden)
 };
 
@@ -261,8 +265,6 @@ struct METASOUNDFRONTEND_API FMetasoundFrontendVertex
 {
 	GENERATED_BODY()
 
-	virtual ~FMetasoundFrontendVertex() = default;
-
 	// Name of the vertex. Unique amongst other vertices on the same interface.
 	UPROPERTY(VisibleAnywhere, Category = CustomView)
 	FName Name;
@@ -676,8 +678,6 @@ struct METASOUNDFRONTEND_API FMetasoundFrontendClassVertex : public FMetasoundFr
 {
 	GENERATED_BODY()
 
-	virtual ~FMetasoundFrontendClassVertex() = default;
-
 	UPROPERTY()
 	FGuid NodeID = Metasound::FrontendInvalidID;
 
@@ -713,6 +713,7 @@ struct METASOUNDFRONTEND_API FMetasoundFrontendClassStyleDisplay
 	,	bShowName(InDisplayStyle.bShowName)
 	,	bShowInputNames(InDisplayStyle.bShowInputNames)
 	,	bShowOutputNames(InDisplayStyle.bShowOutputNames)
+	,	bShowLiterals(InDisplayStyle.bShowLiterals)
 	{
 	}
 
@@ -727,6 +728,9 @@ struct METASOUNDFRONTEND_API FMetasoundFrontendClassStyleDisplay
 
 	UPROPERTY()
 	bool bShowOutputNames = true;
+
+	UPROPERTY()
+	bool bShowLiterals = true;
 #endif // WITH_EDITORONLY_DATA
 };
 
@@ -741,8 +745,6 @@ struct METASOUNDFRONTEND_API FMetasoundFrontendClassInput : public FMetasoundFro
 
 	FMetasoundFrontendClassInput(const FMetasoundFrontendClassVertex& InOther);
 
-	virtual ~FMetasoundFrontendClassInput() = default;
-
 	// Default value for this input.
 	UPROPERTY(EditAnywhere, Category = Parameters)
 	FMetasoundFrontendLiteral DefaultLiteral;
@@ -758,8 +760,6 @@ struct METASOUNDFRONTEND_API FMetasoundFrontendClassVariable : public FMetasound
 
 	FMetasoundFrontendClassVariable(const FMetasoundFrontendClassVertex& InOther);
 
-	virtual ~FMetasoundFrontendClassVariable() = default;
-
 	// Default value for this variable.
 	UPROPERTY(EditAnywhere, Category = Parameters)
 	FMetasoundFrontendLiteral DefaultLiteral;
@@ -777,8 +777,6 @@ struct METASOUNDFRONTEND_API FMetasoundFrontendClassOutput : public FMetasoundFr
 	:	FMetasoundFrontendClassVertex(InOther)
 	{
 	}
-
-	virtual ~FMetasoundFrontendClassOutput() = default;
 };
 
 USTRUCT()
@@ -786,8 +784,6 @@ struct METASOUNDFRONTEND_API FMetasoundFrontendClassEnvironmentVariable
 {
 	GENERATED_BODY()
 
-	virtual ~FMetasoundFrontendClassEnvironmentVariable() = default;
-
 	// Name of environment variable.
 	UPROPERTY()
 	FName Name;
diff --git a/Engine/Plugins/Runtime/Metasound/Source/MetasoundFrontend/Public/MetasoundFrontendNodeTemplateRegistry.h b/Engine/Plugins/Runtime/Metasound/Source/MetasoundFrontend/Public/MetasoundFrontendNodeTemplateRegistry.h
new file mode 100644
--- /dev/null
+++ b/Engine/Plugins/Runtime/Metasound/Source/MetasoundFrontend/Public/MetasoundFrontendNodeTemplateRegistry.h
@@ -0,0 +1,42 @@
+// Copyright Epic Games, Inc. All Rights Reserved.
+
+#pragma once
+
+#include "MetasoundFrontendController.h"
+#include "MetasoundFrontendDocument.h"
+#include "Templates/SharedPointer.h"
+#include "UObject/NoExportTypes.h"
+
+namespace Metasound
+{
+	namespace Frontend
+	{
+		class METASOUNDFRONTEND_API INodeTemplate
+		{
+		public:
+			virtual ~INodeTemplate() = default;
+
+			virtual bool BuildTemplate(FMetasoundFrontendDocument& InOutDocument, FMetasoundFrontendGraph& InOutGraph, FMetasoundFrontendNode& InOutNode) const = 0;
+			virtual const FMetasoundFrontendClass& GetFrontendClass() const = 0;
+			virtual const FMetasoundFrontendVersion& GetVersion() const = 0;
+			virtual bool IsValidNodeInterface(const FMetasoundFrontendNodeInterface& InNodeInterface) const = 0;
+			virtual bool HasRequiredConnections(FConstNodeHandle InNodeHandle) const = 0;
+		};
+
+		class METASOUNDFRONTEND_API INodeTemplateRegistry
+		{
+		public:
+			static INodeTemplateRegistry& Get();
+
+			virtual ~INodeTemplateRegistry() = default;
+
+			// Find a template with the given key. Returns null if entry not found with given key.
+			virtual const INodeTemplate* FindTemplate(const FNodeRegistryKey& InKey) const = 0;
+		};
+
+		// Register & Unregister are not publicly accessible implementation as the API
+		// is in beta and currently, only to be used by reroute nodes.
+		void RegisterNodeTemplate(TUniquePtr<INodeTemplate>&& InTemplate);
+		void UnregisterNodeTemplate(const FMetasoundFrontendVersion& InNodeTemplateVersion);
+	} // namespace Frontend
+} // namespace Metasound
diff --git a/Engine/Plugins/Runtime/Metasound/Source/MetasoundGraphCore/Public/MetasoundNodeInterface.h b/Engine/Plugins/Runtime/Metasound/Source/MetasoundGraphCore/Public/MetasoundNodeInterface.h
--- a/Engine/Plugins/Runtime/Metasound/Source/MetasoundGraphCore/Public/MetasoundNodeInterface.h
+++ b/Engine/Plugins/Runtime/Metasound/Source/MetasoundGraphCore/Public/MetasoundNodeInterface.h
@@ -32,6 +32,9 @@ namespace Metasound
 
 		/** Whether or not to show output names in visual layout. */
 		bool bShowOutputNames = true;
+
+		/** Whether or not to show input literals in visual layout. */
+		bool bShowLiterals = true;
 	};
 
 	/** Name of a node class.
