Minor error text fix-up (missed in prior CL)
#rnx

diff --git a/Engine/Plugins/Runtime/Metasound/Source/MetasoundFrontend/Private/MetasoundFrontendRegistries.cpp b/Engine/Plugins/Runtime/Metasound/Source/MetasoundFrontend/Private/MetasoundFrontendRegistries.cpp
--- a/Engine/Plugins/Runtime/Metasound/Source/MetasoundFrontend/Private/MetasoundFrontendRegistries.cpp
+++ b/Engine/Plugins/Runtime/Metasound/Source/MetasoundFrontend/Private/MetasoundFrontendRegistries.cpp
@@ -427,7 +427,7 @@ namespace Metasound
 
 		bool FRegistryContainerImpl::FindOutputNodeRegistryKeyForDataType(const FName& InDataTypeName, FNodeRegistryKey& OutKey)
 		{
-			UE_LOG(LogMetaSound, Error, TEXT("FRegistryContainerImpl::FindInputNodeRegistryKeyForDataType with these parameters is no longer supported and should not be called. Use the one with EMetasoundFrontendVertexAccessType instead."));
+			UE_LOG(LogMetaSound, Error, TEXT("FRegistryContainerImpl::FindOutputNodeRegistryKeyForDataType with these parameters is no longer supported and should not be called. Use the one with EMetasoundFrontendVertexAccessType instead."));
 			return false;
 		}
 
