Set MetaSound editor animated connections as on by default and move display name from "Experimental" to "Beta"
#rnx

diff --git a/Engine/Plugins/Runtime/Metasound/Source/MetasoundEditor/Private/MetasoundEditorSettings.cpp b/Engine/Plugins/Runtime/Metasound/Source/MetasoundEditor/Private/MetasoundEditorSettings.cpp
--- a/Engine/Plugins/Runtime/Metasound/Source/MetasoundEditor/Private/MetasoundEditorSettings.cpp
+++ b/Engine/Plugins/Runtime/Metasound/Source/MetasoundEditor/Private/MetasoundEditorSettings.cpp
@@ -1,8 +1,8 @@
 // Copyright Epic Games, Inc. All Rights Reserved.
 #include "MetasoundEditorSettings.h"
 
-#include "UObject/UnrealType.h"
 #include "Styling/AppStyle.h"
+#include "UObject/UnrealType.h"
 
 
 UMetasoundEditorSettings::UMetasoundEditorSettings(const FObjectInitializer& ObjectInitializer)
@@ -28,7 +28,7 @@ UMetasoundEditorSettings::UMetasoundEditorSettings(const FObjectInitializer& Obj
 	OutputNodeTitleColor = FLinearColor(1.0f, 0.878f, 0.1686f);					// yellow
 	VariableNodeTitleColor = FLinearColor(0.211f, 0.513f, 0.035f);				// copper
 
-	bAnimateConnections = false;
+	bAnimateConnections = true;
 	ActiveAnalyzerNumericWireThickness = 5.0f;
 	ActiveAnalyzerEnvelopeWireThickness = 1.0f;
 	ActiveAnalyzerWireScalarMin = 1.0f;
diff --git a/Engine/Plugins/Runtime/Metasound/Source/MetasoundEditor/Private/MetasoundEditorSettings.h b/Engine/Plugins/Runtime/Metasound/Source/MetasoundEditor/Private/MetasoundEditorSettings.h
--- a/Engine/Plugins/Runtime/Metasound/Source/MetasoundEditor/Private/MetasoundEditorSettings.h
+++ b/Engine/Plugins/Runtime/Metasound/Source/MetasoundEditor/Private/MetasoundEditorSettings.h
@@ -1,10 +1,9 @@
 // Copyright Epic Games, Inc. All Rights Reserved.
 #pragma once
 
-#include "CoreMinimal.h"
 #include "UObject/NoExportTypes.h"
-#include "UObject/ObjectMacros.h"
 #include "UObject/Object.h"
+#include "UObject/ObjectMacros.h"
 
 #include "MetasoundEditorSettings.generated.h"
 
@@ -108,7 +107,7 @@ public:
 	FLinearColor VariableNodeTitleColor;
 
 	/** Whether or not animated connections are enabled. */
-	UPROPERTY(EditAnywhere, config, Category = GraphAnimation, meta = (DisplayName = "Animate Connections (Experimental)"))
+	UPROPERTY(EditAnywhere, config, Category = GraphAnimation, meta = (DisplayName = "Animate Connections (Beta)"))
 	bool bAnimateConnections;
 
 	/** Thickness of default envelope analyzer wire thickness when connection analyzer is active. */
