MetaSound Reroute Nodes
- Minor play button/transport bar updates using new icons from RobbieK
- Validation Clean-up to better report final message severity state & accomidate template node validation (i.e. reroute validation & visual reporting/user feedback)
#rnx

diff --git a/Engine/Plugins/Runtime/Metasound/Source/MetasoundEditor/Private/MetasoundEditor.cpp b/Engine/Plugins/Runtime/Metasound/Source/MetasoundEditor/Private/MetasoundEditor.cpp
--- a/Engine/Plugins/Runtime/Metasound/Source/MetasoundEditor/Private/MetasoundEditor.cpp
+++ b/Engine/Plugins/Runtime/Metasound/Source/MetasoundEditor/Private/MetasoundEditor.cpp
@@ -54,6 +54,7 @@
 #include "SMetasoundActionMenu.h"
 #include "SMetasoundPalette.h"
 #include "SNodePanel.h"
+#include "Stats/Stats.h"
 #include "Templates/SharedPointer.h"
 #include "Widgets/Docking/SDockTab.h"
 #include "Widgets/Images/SImage.h"
@@ -62,7 +63,6 @@
 #include "Widgets/SBoxPanel.h"
 #include "Widgets/SOverlay.h"
 #include "Widgets/SWindow.h"
-#include "Stats/Stats.h"
 
 struct FGraphActionNode;
 
@@ -359,11 +359,8 @@ namespace Metasound
 							{
 								bDropTargetValid = true;
 
-								if (const ISlateStyle* MetasoundStyle = FSlateStyleRegistry::FindSlateStyle("MetaSoundStyle"))
-								{
-									PrimarySymbol = MetasoundStyle->GetBrush("MetasoundEditor.Graph.Node.Class.Input");
-									SecondarySymbol = nullptr;
-								}
+								Style::GetSlateBrushSafe("MetasoundEditor.Graph.Node.Class.Input");
+								SecondarySymbol = nullptr;
 
 								if (const UMetasoundEditorSettings* EditorSettings = GetDefault<UMetasoundEditorSettings>())
 								{
@@ -976,10 +973,7 @@ namespace Metasound
 		{
 			if (IsPlaying())
 			{
-				if (const ISlateStyle* MetasoundStyle = FSlateStyleRegistry::FindSlateStyle("MetaSoundStyle"))
-				{
-					return MetasoundStyle->GetBrush("MetasoundEditor.Play");
-				}
+				return &Style::GetSlateBrushSafe("MetasoundEditor.Play");
 			}
 
 			return FAssetEditorToolkit::GetDefaultTabIcon();
@@ -1243,17 +1237,7 @@ namespace Metasound
  					}
  					ToolbarBuilder.EndSection();
 
-					if (Metasound->IsA<USoundBase>())
-					{
-						ToolbarBuilder.BeginSection("Audition");
-						{
-							ToolbarBuilder.AddToolBarButton(FEditorCommands::Get().Play);
-							ToolbarBuilder.AddToolBarButton(FEditorCommands::Get().Stop);
-						}
-						ToolbarBuilder.EndSection();
-					}
-
-					ToolbarBuilder.BeginSection("Utilities");
+					ToolbarBuilder.BeginSection("Settings");
 					{
 						if (Metasound->IsA<USoundBase>())
 						{
@@ -1277,6 +1261,37 @@ namespace Metasound
 						);
 					}
 					ToolbarBuilder.EndSection();
+
+					if (Metasound->IsA<USoundBase>())
+					{
+						ToolbarBuilder.BeginSection("Audition");
+						{
+							ToolbarBuilder.BeginStyleOverride("Toolbar.BackplateLeft");
+							{
+								ToolbarBuilder.AddToolBarButton(
+									FEditorCommands::Get().Play,
+									NAME_None,
+									TAttribute<FText>(),
+									TAttribute<FText>::Create([this] { return GetGraphStatusDescription(); }),
+									TAttribute<FSlateIcon>::Create([this]() { return GetPlayIcon(); })
+								);
+							}
+							ToolbarBuilder.EndStyleOverride();
+
+							ToolbarBuilder.BeginStyleOverride("Toolbar.BackplateRight");
+							{
+								ToolbarBuilder.AddToolBarButton(
+									FEditorCommands::Get().Stop,
+									NAME_None,
+									TAttribute<FText>(),
+									TAttribute<FText>(),
+									TAttribute<FSlateIcon>::Create([this]() { return GetStopIcon(); })
+								);
+							}
+							ToolbarBuilder.EndStyleOverride();
+						}
+						ToolbarBuilder.EndSection();
+					}
 				})
 			);
 
@@ -1425,11 +1440,112 @@ namespace Metasound
 			MetasoundAsset->GetDocumentHandle()->ExportToJSONAsset(Path);
 		}
 
+		FText FEditor::GetGraphStatusDescription() const
+		{
+			switch (HighestMessageSeverity)
+			{
+				case EMessageSeverity::Error:
+				{
+					return LOCTEXT("MetaSoundPlayStateTooltip_Error", "MetaSound contains errors and cannot be played.");
+				}
+
+				case EMessageSeverity::PerformanceWarning:
+				case EMessageSeverity::Warning:
+				{
+					return LOCTEXT("MetaSoundPlayStateTooltip_Warning", "MetaSound contains warnings and playback behavior may be undesired.");
+				}
+				break;
+
+				case EMessageSeverity::Info:
+				default:
+				{
+					return FEditorCommands::Get().Play->GetDescription();
+				}
+				break;
+			}
+		}
+
+		const FSlateIcon& FEditor::GetPlayIcon() const
+		{
+			switch (HighestMessageSeverity)
+			{
+				case EMessageSeverity::Error:
+				{
+					static const FSlateIcon Icon = Style::CreateSlateIcon("MetasoundEditor.Play.Error");
+					return Icon;
+				}
+
+				case EMessageSeverity::PerformanceWarning:
+				case EMessageSeverity::Warning:
+				{
+					if (IsPlaying())
+					{
+						static const FSlateIcon Icon = Style::CreateSlateIcon("MetasoundEditor.Play.Active.Warning");
+						return Icon;
+					}
+					else
+					{
+						static const FSlateIcon Icon = Style::CreateSlateIcon("MetasoundEditor.Play.Inactive.Warning");
+						return Icon;
+					}
+				}
+				break;
+
+				case EMessageSeverity::Info:
+				default:
+				{
+					if (IsPlaying())
+					{
+						static const FSlateIcon Icon = Style::CreateSlateIcon("MetasoundEditor.Play.Active.Valid");
+						return Icon;
+					}
+					else
+					{
+						static const FSlateIcon Icon = Style::CreateSlateIcon("MetasoundEditor.Play.Inactive.Valid");
+						return Icon;
+					}
+				}
+				break;
+			}
+		}
+
+		const FSlateIcon& FEditor::GetStopIcon() const
+		{
+			switch (HighestMessageSeverity)
+			{
+				case EMessageSeverity::Error:
+				{
+					static const FSlateIcon Icon = Style::CreateSlateIcon("MetasoundEditor.Stop.Disabled");
+					return Icon;
+				}
+				break;
+
+				case EMessageSeverity::PerformanceWarning:
+				case EMessageSeverity::Warning:
+				case EMessageSeverity::Info:
+				default:
+				{
+					if (IsPlaying())
+					{
+						static const FSlateIcon Icon = Style::CreateSlateIcon("MetasoundEditor.Stop.Active");
+						return Icon;
+					}
+					else
+					{
+						static const FSlateIcon Icon = Style::CreateSlateIcon("MetasoundEditor.Stop.Inactive");
+						return Icon;
+					}
+				}
+				break;
+			}
+		}
+
 		void FEditor::Play()
 		{
 			if (USoundBase* MetasoundToPlay = Cast<USoundBase>(Metasound))
 			{
-				if (!FGraphBuilder::SynchronizeGraph(*Metasound))
+				HighestMessageSeverity = FGraphBuilder::SynchronizeGraph(*Metasound);
+				if (HighestMessageSeverity == EMessageSeverity::Error)
 				{
 					return;
 				}
@@ -3168,7 +3284,7 @@ namespace Metasound
 
 				// Capture before synchronizing as the flag is cleared therein.
 				const bool bShouldRefreshDetails = MetasoundAsset->GetSynchronizationUpdateDetails();
-				FGraphBuilder::SynchronizeGraph(*Metasound);
+				HighestMessageSeverity = FGraphBuilder::SynchronizeGraph(*Metasound);
 
 				// Presets always update interfaces
 				const FMetasoundFrontendGraphClass& RootGraphClass = MetasoundAsset->GetDocumentHandle()->GetRootGraphClass();
diff --git a/Engine/Plugins/Runtime/Metasound/Source/MetasoundEditor/Private/MetasoundEditor.h b/Engine/Plugins/Runtime/Metasound/Source/MetasoundEditor/Private/MetasoundEditor.h
--- a/Engine/Plugins/Runtime/Metasound/Source/MetasoundEditor/Private/MetasoundEditor.h
+++ b/Engine/Plugins/Runtime/Metasound/Source/MetasoundEditor/Private/MetasoundEditor.h
@@ -9,6 +9,7 @@
 #include "GraphEditor.h"
 #include "IDetailsView.h"
 #include "IMetasoundEditor.h"
+#include "Logging/TokenizedMessage.h"
 #include "Math/UnrealMathUtility.h"
 #include "MetasoundEditorGraphConnectionManager.h"
 #include "MetasoundEditorGraphNode.h"
@@ -189,6 +190,8 @@ namespace Metasound
 			/** Creates analyzers */
 			void CreateAnalyzers();
 
+			void BuildTransport(FToolBarBuilder& InToolBarBuilder);
+
 			/** Destroys analyzers */
 			void DestroyAnalyzers();
 
@@ -256,6 +259,10 @@ namespace Metasound
 
 			UMetasoundEditorGraph& GetMetaSoundGraphChecked();
 
+			FText GetGraphStatusDescription() const;
+			const FSlateIcon& GetPlayIcon() const;
+			const FSlateIcon& GetStopIcon() const;
+
 			/**
 			 * Called when a node's title is committed for a rename
 			 *
@@ -443,6 +450,8 @@ namespace Metasound
 
 			/** Whether or not registry is currently being primed due to assets still loading */
 			bool bPrimingRegistry = false;
+
+			int32 HighestMessageSeverity = EMessageSeverity::Info;
 		};
 	} // namespace Editor
 } // namespace Metasound
diff --git a/Engine/Plugins/Runtime/Metasound/Source/MetasoundEditor/Private/MetasoundEditorGraph.cpp b/Engine/Plugins/Runtime/Metasound/Source/MetasoundEditor/Private/MetasoundEditorGraph.cpp
--- a/Engine/Plugins/Runtime/Metasound/Source/MetasoundEditor/Private/MetasoundEditorGraph.cpp
+++ b/Engine/Plugins/Runtime/Metasound/Source/MetasoundEditor/Private/MetasoundEditorGraph.cpp
@@ -1826,14 +1826,11 @@ void UMetasoundEditorGraph::IterateOutputs(TUniqueFunction<void(UMetasoundEditor
 	}
 }
 
-bool UMetasoundEditorGraph::ValidateInternal(Metasound::Editor::FGraphValidationResults& OutResults)
+void UMetasoundEditorGraph::ValidateInternal(Metasound::Editor::FGraphValidationResults& OutResults)
 {
 	using namespace Metasound::Editor;
 	using namespace Metasound::Frontend;
 
-	bool bMarkDirty = false;
-	bool bIsValid = true;
-
 	OutResults = FGraphValidationResults();
 
 	TArray<UMetasoundEditorGraphNode*> NodesToValidate;
@@ -1841,19 +1838,10 @@ bool UMetasoundEditorGraph::ValidateInternal(Metasound::Editor::FGraphValidation
 	for (UMetasoundEditorGraphNode* Node : NodesToValidate)
 	{
 		FGraphNodeValidationResult NodeResult(*Node);
+		Node->Validate(NodeResult);
 
-		bIsValid &= Node->Validate(NodeResult);
-		bMarkDirty |= NodeResult.bIsDirty;
-
-		OutResults.NodeResults.Add(NodeResult);
-	}
-
-	if (bMarkDirty)
-	{
-		MarkPackageDirty();
+		OutResults.NodeResults.Add(MoveTemp(NodeResult));
 	}
-
-	return bIsValid;
 }
 
 bool UMetasoundEditorGraph::RemoveMember(UMetasoundEditorGraphMember& InGraphMember)
diff --git a/Engine/Plugins/Runtime/Metasound/Source/MetasoundEditor/Private/MetasoundEditorGraphBuilder.cpp b/Engine/Plugins/Runtime/Metasound/Source/MetasoundEditor/Private/MetasoundEditorGraphBuilder.cpp
--- a/Engine/Plugins/Runtime/Metasound/Source/MetasoundEditor/Private/MetasoundEditorGraphBuilder.cpp
+++ b/Engine/Plugins/Runtime/Metasound/Source/MetasoundEditor/Private/MetasoundEditorGraphBuilder.cpp
@@ -1,9 +1,9 @@
 // Copyright Epic Games, Inc. All Rights Reserved.
 #include "MetasoundEditorGraphBuilder.h"
 
+#include "Algo/AnyOf.h"
 #include "Algo/Sort.h"
 #include "Algo/Transform.h"
-#include "Algo/AnyOf.h"
 #include "AssetRegistry/AssetRegistryModule.h"
 #include "EdGraph/EdGraph.h"
 #include "EdGraph/EdGraphNode.h"
@@ -35,6 +35,7 @@
 #include "MetasoundVariableNodes.h"
 #include "MetasoundVertex.h"
 #include "Modules/ModuleManager.h"
+#include "NodeTemplates/MetasoundFrontendNodeTemplateReroute.h"
 #include "Templates/Tuple.h"
 #include "Toolkits/ToolkitManager.h"
 #include "Widgets/Notifications/SNotificationList.h"
@@ -336,7 +337,11 @@ namespace Metasound
 			using namespace Frontend;
 
 			UMetasoundEditorGraphExternalNode* NewGraphNode = nullptr;
-			if (!ensure(InNodeHandle->GetClassMetadata().GetType() == EMetasoundFrontendClassType::External))
+
+			const EMetasoundFrontendClassType ClassType = InNodeHandle->GetClassMetadata().GetType();
+			const bool bIsExternalNode = ClassType == EMetasoundFrontendClassType::External;
+			const bool bIsTemplateNode = ClassType == EMetasoundFrontendClassType::Template;
+			if (!ensure(bIsExternalNode || bIsTemplateNode))
 			{
 				return nullptr;
 			}
@@ -476,7 +481,7 @@ namespace Metasound
 			RebuildNodePins(*NewGraphNode);
 		}
 
-		bool FGraphBuilder::ValidateGraph(UObject& InMetaSound, bool bForceRefreshNodes)
+		EMessageSeverity::Type FGraphBuilder::ValidateGraph(UObject& InMetaSound, bool bForceRefreshNodes)
 		{
 			using namespace Frontend;
 
@@ -494,30 +499,28 @@ namespace Metasound
 			bForceRefreshNodes |= Graph.RequiresForceRefreshNodes();
 			Graph.ClearForceRefreshNodes();
 
-			FGraphValidationResults Results;
-
-			bool bMarkDirty = false;
+			int8 HighestMessageSeverity = EMessageSeverity::Info;
 
+			FGraphValidationResults Results;
 			Graph.ValidateInternal(Results);
 			for (const FGraphNodeValidationResult& Result : Results.GetResults())
 			{
-				bMarkDirty |= Result.bIsDirty;
-				check(Result.Node);
-				const bool bInterfaceChange = Result.Node->ContainsInterfaceChange();
-				const bool bMetadataChange = Result.Node->ContainsMetadataChange();
-				const bool bStyleChange = Result.Node->ContainsStyleChange();
+				UMetasoundEditorGraphNode& Node = Result.GetNodeChecked();
+				const bool bInterfaceChange = Node.ContainsInterfaceChange();
+				const bool bMetadataChange = Node.ContainsMetadataChange();
+				const bool bStyleChange = Node.ContainsStyleChange();
 
-				const FText Title = Result.Node->GetCachedTitle();
-				Result.Node->CacheTitle();
-				const bool bTitleUpdated = !Title.IdenticalTo(Result.Node->GetCachedTitle());
+				const FText Title = Node.GetCachedTitle();
+				Node.CacheTitle();
+				const bool bTitleUpdated = !Title.IdenticalTo(Node.GetCachedTitle());
 
-				if (Result.bIsDirty || bTitleUpdated || bMetadataChange || bInterfaceChange || bStyleChange || bForceRefreshNodes)
+				if (Result.GetHasDirtiedNode() || bTitleUpdated || bMetadataChange || bInterfaceChange || bStyleChange || bForceRefreshNodes)
 				{
-					Result.Node->SyncChangeIDs();
+					Node.SyncChangeIDs();
 
 					if (GraphEditor.IsValid())
 					{
-						GraphEditor->RefreshNode(*Result.Node);
+						GraphEditor->RefreshNode(Node);
 					}
 				}
 			}
@@ -527,12 +530,7 @@ namespace Metasound
 				MetaSoundEditor->RefreshGraphMemberMenu();
 			}
 
-			if (bMarkDirty)
-			{
-				InMetaSound.MarkPackageDirty();
-			}
-
-			return Results.IsValid();
+			return Results.GetHighestMessageSeverity();
 		}
 
 		TArray<FString> FGraphBuilder::GetDataTypeNameCategories(const FName& InDataTypeName)
@@ -696,6 +694,27 @@ namespace Metasound
 			return GetInputHandleFromPin(InPin);
 		}
 
+		FName FGraphBuilder::GetPinDataType(const UEdGraphPin* InPin)
+		{
+			using namespace Frontend;
+
+			if (InPin)
+			{
+				if (InPin->Direction == EGPD_Input)
+				{
+					FConstInputHandle InputHandle = GetConstInputHandleFromPin(InPin);
+					return InputHandle->GetDataType();
+				}
+				else // EGPD_Output
+				{
+					FConstOutputHandle OutputHandle = GetConstOutputHandleFromPin(InPin);
+					return OutputHandle->GetDataType();
+				}
+			}
+
+			return { };
+		}
+
 		Frontend::FOutputHandle FGraphBuilder::GetOutputHandleFromPin(const UEdGraphPin* InPin)
 		{
 			using namespace Frontend;
@@ -743,7 +762,7 @@ namespace Metasound
 		{
 			using namespace Frontend;
 
-			FConstOutputHandle OutputHandle = GetConstOutputHandleFromPin(InGraphPin);
+			FConstOutputHandle OutputHandle = FindReroutedConstOutputHandleFromPin(InGraphPin);
 			return GetOutputEdgeStyle(OutputHandle);
 		}
 
@@ -752,6 +771,96 @@ namespace Metasound
 			return GetOutputHandleFromPin(InPin);
 		}
 
+		UEdGraphPin* FGraphBuilder::FindReroutedOutputPin(UEdGraphPin* OutputPin)
+		{
+			using namespace Frontend;
+
+			if (OutputPin)
+			{
+				if (UMetasoundEditorGraphExternalNode* ExternalNode = Cast<UMetasoundEditorGraphExternalNode>(OutputPin->GetOwningNode()))
+				{
+					if (ExternalNode->GetClassName() == FRerouteNodeTemplate::ClassName)
+					{
+						auto IsInput = [](UEdGraphPin* Pin) { check(Pin); return Pin->Direction == EGPD_Input; };
+						if (UEdGraphPin* RerouteInput = *ExternalNode->Pins.FindByPredicate(IsInput))
+						{
+							TArray<UEdGraphPin*>& LinkedTo = RerouteInput->LinkedTo;
+							if (!LinkedTo.IsEmpty())
+							{
+								UEdGraphPin* ReroutedOutput = LinkedTo.Last();
+								return FindReroutedOutputPin(ReroutedOutput);
+							}
+						}
+					}
+				}
+			}
+
+			return OutputPin;
+		}
+
+		Frontend::FOutputHandle FGraphBuilder::FindReroutedOutputHandleFromPin(const UEdGraphPin* OutputPin)
+		{
+			using namespace Frontend;
+
+			if (OutputPin)
+			{
+				if (UMetasoundEditorGraphExternalNode* ExternalNode = Cast<UMetasoundEditorGraphExternalNode>(OutputPin->GetOwningNode()))
+				{
+					if (ExternalNode->GetClassName() == FRerouteNodeTemplate::ClassName)
+					{
+						auto IsInput = [](UEdGraphPin* Pin) { check(Pin); return Pin->Direction == EGPD_Input; };
+						if (UEdGraphPin* RerouteInput = *ExternalNode->Pins.FindByPredicate(IsInput))
+						{
+							TArray<UEdGraphPin*>& LinkedTo = RerouteInput->LinkedTo;
+							if (!LinkedTo.IsEmpty())
+							{
+								UEdGraphPin* ReroutedOutput = LinkedTo.Last();
+								return FindReroutedOutputHandleFromPin(ReroutedOutput);
+							}
+						}
+					}
+				}
+
+				return GetOutputHandleFromPin(OutputPin);
+			}
+
+			return IOutputController::GetInvalidHandle();
+		}
+
+		Frontend::FConstOutputHandle FGraphBuilder::FindReroutedConstOutputHandleFromPin(const UEdGraphPin* InPin)
+		{
+			return FindReroutedOutputHandleFromPin(InPin);
+		}
+
+		void FGraphBuilder::FindReroutedInputPins(UEdGraphPin* InPinToCheck, TArray<UEdGraphPin*>& InOutInputPins)
+		{
+			using namespace Frontend;
+
+			if (InPinToCheck && InPinToCheck->Direction == EGPD_Input)
+			{
+				if (UMetasoundEditorGraphExternalNode* ExternalNode = Cast<UMetasoundEditorGraphExternalNode>(InPinToCheck->GetOwningNode()))
+				{
+					if (ExternalNode->GetClassName() == FRerouteNodeTemplate::ClassName)
+					{
+						for (UEdGraphPin* Pin : ExternalNode->Pins)
+						{
+							if (Pin->Direction == EGPD_Output)
+							{
+								for (UEdGraphPin* LinkedInput : Pin->LinkedTo)
+								{
+									FindReroutedInputPins(LinkedInput, InOutInputPins);
+								}
+							}
+						}
+
+						return;
+					}
+				}
+
+				InOutInputPins.Add(InPinToCheck);
+			}
+		}
+
 		bool FGraphBuilder::GraphContainsErrors(const UObject& InMetaSound)
 		{
 			const FMetasoundAssetBase* MetaSoundAsset = IMetasoundUObjectRegistry::Get().GetObjectAsAssetBase(&InMetaSound);
@@ -764,15 +873,7 @@ namespace Metasound
 			EditorGraph->GetNodesOfClass(EditorNodes);
 
 			// Do not synchronize with errors present as the graph is expected to be malformed.
-			for (const UMetasoundEditorGraphNode* Node : EditorNodes)
-			{
-				if (Node->ErrorType == EMessageSeverity::Error)
-				{
-					return true;
-				}
-			}
-
-			return false;
+			return Algo::AnyOf(EditorNodes, [](const UMetasoundEditorGraphNode* Node) { return Node->ErrorType == EMessageSeverity::Error; });
 		}
 
 		bool FGraphBuilder::SynchronizeNodeLocation(UMetasoundEditorGraphNode& InNode)
@@ -1256,7 +1357,10 @@ namespace Metasound
 				const UEdGraphSchema* Schema = InInputPin.GetSchema();
 				if (ensure(Schema))
 				{
-					return Schema->TryCreateConnection(&InInputPin, &InOutputPin);
+					if (!Schema->TryCreateConnection(&InInputPin, &InOutputPin))
+					{
+						return false;
+					}
 				}
 				else
 				{
@@ -1901,7 +2005,7 @@ namespace Metasound
 			return bIsGraphDirty;
 		}
 
-		bool FGraphBuilder::SynchronizeGraph(UObject& InMetaSound, bool bForceRefreshNodes)
+		int32 FGraphBuilder::SynchronizeGraph(UObject& InMetaSound, bool bForceRefreshNodes)
 		{
 			FMetasoundAssetBase* MetaSoundAsset = IMetasoundUObjectRegistry::Get().GetObjectAsAssetBase(&InMetaSound);
 			check(MetaSoundAsset);
@@ -1940,11 +2044,11 @@ namespace Metasound
 			}
 
 			bForceRefreshNodes |= bEditorGraphModified;
-			const bool bIsValid = FGraphBuilder::ValidateGraph(InMetaSound, bForceRefreshNodes);
+			const int32 HighestMessageSeverity = FGraphBuilder::ValidateGraph(InMetaSound, bForceRefreshNodes);
 
 			MetaSoundAsset->ResetSynchronizationState();
 
-			return bIsValid;
+			return HighestMessageSeverity;
 		}
 
 		bool FGraphBuilder::SynchronizeNodeMembers(UObject& InMetaSound)
diff --git a/Engine/Plugins/Runtime/Metasound/Source/MetasoundEditor/Private/MetasoundEditorGraphBuilder.h b/Engine/Plugins/Runtime/Metasound/Source/MetasoundEditor/Private/MetasoundEditorGraphBuilder.h
--- a/Engine/Plugins/Runtime/Metasound/Source/MetasoundEditor/Private/MetasoundEditorGraphBuilder.h
+++ b/Engine/Plugins/Runtime/Metasound/Source/MetasoundEditor/Private/MetasoundEditorGraphBuilder.h
@@ -5,6 +5,7 @@
 #include "MetasoundFrontendController.h"
 #include "MetasoundFrontendDocument.h"
 #include "Internationalization/Text.h"
+#include "Logging/TokenizedMessage.h"
 #include "Templates/SharedPointer.h"
 #include "UObject/NameTypes.h"
 #include "UObject/NoExportTypes.h"
@@ -43,9 +44,6 @@ namespace Metasound
 		{
 			static void InitGraphNode(Frontend::FNodeHandle& InNodeHandle, UMetasoundEditorGraphNode* NewGraphNode, UObject& InMetaSound);
 
-			// Validates MetaSound graph.
-			static bool ValidateGraph(UObject& InMetaSound, bool bForceRefreshNodes);
-
 		public:
 			static const FName PinCategoryAudio;
 			static const FName PinCategoryBoolean;
@@ -72,6 +70,9 @@ namespace Metasound
 			static TSharedPtr<FEditor> GetEditorForNode(const UEdGraphNode& InEdNode);
 			static TSharedPtr<FEditor> GetEditorForPin(const UEdGraphPin& InEdPin);
 
+			// Validates MetaSound graph, returning the highest EMessageSeverity integer value
+			static EMessageSeverity::Type ValidateGraph(UObject& InMetaSound, bool bForceRefreshNodes);
+
 			// Wraps RegisterGraphWithFrontend logic in Frontend with any additional logic required to refresh editor & respective editor object state.
 			// @param InMetaSound - MetaSound to register
 			// @param bInForceSynchronize - Forces the synchronize flag for all open graphs being registered by this call (all referenced graphs and
@@ -196,6 +197,8 @@ namespace Metasound
 			static Frontend::FInputHandle GetInputHandleFromPin(const UEdGraphPin* InPin);
 			static Frontend::FConstInputHandle GetConstInputHandleFromPin(const UEdGraphPin* InPin);
 
+			static FName GetPinDataType(const UEdGraphPin* InPin);
+
 			static const FMetasoundFrontendEdgeStyle* GetOutputEdgeStyle(Frontend::FConstOutputHandle InOutputHandle);
 			static const FMetasoundFrontendEdgeStyle* GetOutputEdgeStyle(const UEdGraphPin* InPin);
 
@@ -204,6 +207,17 @@ namespace Metasound
 			static Frontend::FOutputHandle GetOutputHandleFromPin(const UEdGraphPin* InPin);
 			static Frontend::FConstOutputHandle GetConstOutputHandleFromPin(const UEdGraphPin* InPin);
 
+			static UEdGraphPin* FindReroutedOutputPin(UEdGraphPin* InPin);
+
+			// Find the "concrete" output handle associated with an output pin.  If the given output pin is on
+			// a reroute node, will recursively search for the non-rerouted output its representing.
+			static Frontend::FOutputHandle FindReroutedOutputHandleFromPin(const UEdGraphPin* InOutputPin);
+			static Frontend::FConstOutputHandle FindReroutedConstOutputHandleFromPin(const UEdGraphPin* InPin);
+
+			// Find the "concrete" input handle associated with an output pin.  If the given input pin is on
+			// a reroute node, will recursively search for all the non-rerouted input pins its representing.
+			static void FindReroutedInputPins(UEdGraphPin* InPinToCheck, TArray<UEdGraphPin*>& InOutInputPins);
+
 			// Returns the default literal stored on the respective Frontend Node's Input.
 			static bool GetPinLiteral(UEdGraphPin& InInputPin, FMetasoundFrontendLiteral& OutLiteralDefault);
 
@@ -235,8 +249,8 @@ namespace Metasound
 			// Adds and removes nodes, pins and connections so that the UEdGraph of the MetaSound matches the FMetasoundFrontendDocumentModel
 			//
 			// @param bForceRefreshNodes - Refreshes all transient data stored on editor nodes & redraws them
-			// @return True if the UEdGraph is synchronized and is in valid state, false otherwise.
-			static bool SynchronizeGraph(UObject& InMetaSound, bool bForceRefreshNodes = false);
+			// @return Highest MessageSeverity integer value listed on validated nodes post synchronization.
+			static int32 SynchronizeGraph(UObject& InMetaSound, bool bForceRefreshNodes = false);
 
 			// Synchronizes editor nodes with frontend nodes, removing editor nodes that are not represented in the frontend, and adding editor nodes to represent missing frontend nodes.
 			//
diff --git a/Engine/Plugins/Runtime/Metasound/Source/MetasoundEditor/Private/MetasoundEditorGraphConnectionDrawingPolicy.cpp b/Engine/Plugins/Runtime/Metasound/Source/MetasoundEditor/Private/MetasoundEditorGraphConnectionDrawingPolicy.cpp
--- a/Engine/Plugins/Runtime/Metasound/Source/MetasoundEditor/Private/MetasoundEditorGraphConnectionDrawingPolicy.cpp
+++ b/Engine/Plugins/Runtime/Metasound/Source/MetasoundEditor/Private/MetasoundEditorGraphConnectionDrawingPolicy.cpp
@@ -12,6 +12,7 @@
 #include "MetasoundEditorSettings.h"
 #include "MetasoundTrace.h"
 #include "Misc/App.h"
+#include "NodeTemplates/MetasoundFrontendNodeTemplateReroute.h"
 #include "Templates/Function.h"
 
 #define METASOUND_EDITOR_DEBUG_CONNECTIONS 0
@@ -84,7 +85,7 @@ namespace Metasound
 
 				if (MetasoundEditor.IsValid())
 				{
-					FConstOutputHandle OutputHandle = FGraphBuilder::GetConstOutputHandleFromPin(OutputPin);
+					FConstOutputHandle OutputHandle = FGraphBuilder::FindReroutedConstOutputHandleFromPin(OutputPin);
 					const FGuid NodeId = OutputHandle->GetOwningNodeID();
 					FName OutputName = OutputHandle->GetName();
 
@@ -196,7 +197,7 @@ namespace Metasound
 			if (Params.AssociatedPin1)
 			{
 				OutputPin = Params.AssociatedPin1->Direction == EGPD_Output ? Params.AssociatedPin1 : Params.AssociatedPin2;
-				OutputHandle = FGraphBuilder::GetConstOutputHandleFromPin(OutputPin);
+				OutputHandle = FGraphBuilder::FindReroutedConstOutputHandleFromPin(OutputPin);
 			}
 
 			// The curve will include the endpoints but can extend out of a tight bounds because of the tangents
@@ -289,6 +290,10 @@ namespace Metasound
 			: FConnectionDrawingPolicy(InBackLayerID, InFrontLayerID, InZoomFactor, InClippingRect, InDrawElements)
 			, GraphObj(InGraphObj)
 		{
+			// Don't show arrow images in the MetaSound graph
+			ArrowImage = nullptr;
+			ArrowRadius = FVector2D::ZeroVector;
+
 			WireAnimationLayerID = (InFrontLayerID + InBackLayerID) / 2;
 			ActiveWireThickness = Settings->TraceAttackWireThickness;
 			InactiveWireThickness = Settings->TraceReleaseWireThickness;
@@ -608,7 +613,7 @@ namespace Metasound
 				{
 					if (MetasoundEditor.IsValid())
 					{
-						FConstOutputHandle OutputHandle = FGraphBuilder::GetConstOutputHandleFromPin(InParams.ConnectionData.OutputPin);
+						FConstOutputHandle OutputHandle = FGraphBuilder::FindReroutedConstOutputHandleFromPin(InParams.ConnectionData.OutputPin);
 						const FGuid NodeId = OutputHandle->GetOwningNodeID();
 						FName OutputName = OutputHandle->GetName();
 
diff --git a/Engine/Plugins/Runtime/Metasound/Source/MetasoundEditor/Private/MetasoundEditorGraphInputNode.cpp b/Engine/Plugins/Runtime/Metasound/Source/MetasoundEditor/Private/MetasoundEditorGraphInputNode.cpp
--- a/Engine/Plugins/Runtime/Metasound/Source/MetasoundEditor/Private/MetasoundEditorGraphInputNode.cpp
+++ b/Engine/Plugins/Runtime/Metasound/Source/MetasoundEditor/Private/MetasoundEditorGraphInputNode.cpp
@@ -75,13 +75,13 @@ FSlateIcon UMetasoundEditorGraphInputNode::GetNodeTitleIcon() const
 	return FSlateIcon("MetaSoundStyle", NativeIconName);
 }
 
-bool UMetasoundEditorGraphInputNode::Validate(Metasound::Editor::FGraphNodeValidationResult& OutResult)
+void UMetasoundEditorGraphInputNode::Validate(Metasound::Editor::FGraphNodeValidationResult& OutResult)
 {
 #if WITH_EDITOR
 	using namespace Metasound::Editor;
 	using namespace Metasound::Frontend;
 
-	OutResult = CreateNewValidationResult();
+	Super::Validate(OutResult);
 
 	FNodeHandle NodeHandle = GetNodeHandle();
 	const FMetasoundFrontendClassMetadata& Metadata = NodeHandle->GetClassMetadata();
@@ -102,14 +102,12 @@ bool UMetasoundEditorGraphInputNode::Validate(Metasound::Editor::FGraphNodeValid
 				const FConstInputHandle& InputHandle = InputHandles.Last();
 				if (!InputHandle->IsConnected())
 				{
-					GraphNode::SetMessage(*this, EMessageSeverity::Warning, *RequiredText.ToString());
-					return false;
+					OutResult.SetMessage(EMessageSeverity::Warning, *RequiredText.ToString());
 				}
 			}
 		}
 	}
 #endif // #if WITH_EDITOR
-	return true;
 }
 
 void UMetasoundEditorGraphInputNode::SetNodeID(FGuid InNodeID)
diff --git a/Engine/Plugins/Runtime/Metasound/Source/MetasoundEditor/Private/MetasoundEditorGraphNode.cpp b/Engine/Plugins/Runtime/Metasound/Source/MetasoundEditor/Private/MetasoundEditorGraphNode.cpp
--- a/Engine/Plugins/Runtime/Metasound/Source/MetasoundEditor/Private/MetasoundEditorGraphNode.cpp
+++ b/Engine/Plugins/Runtime/Metasound/Source/MetasoundEditor/Private/MetasoundEditorGraphNode.cpp
@@ -19,11 +19,13 @@
 #include "MetasoundEditorSettings.h"
 #include "MetasoundFrontend.h"
 #include "MetasoundFrontendDataTypeRegistry.h"
+#include "MetasoundFrontendNodeTemplateRegistry.h"
 #include "MetasoundFrontendRegistries.h"
 #include "MetasoundFrontendSearchEngine.h"
 #include "MetasoundLiteral.h"
 #include "MetasoundNodeInterface.h"
 #include "MetasoundUObjectRegistry.h"
+#include "NodeTemplates/MetasoundFrontendNodeTemplateReroute.h"
 #include "ScopedTransaction.h"
 #include "ToolMenus.h"
 
@@ -33,40 +35,14 @@ namespace Metasound
 {
 	namespace Editor
 	{
-		void GraphNode::SetMessage(UEdGraphNode& InNode, EMessageSeverity::Type InSeverity, const FString& InMessage)
+		namespace GraphNodePrivate
 		{
-			InNode.bHasCompilerMessage = true;
-			InNode.ErrorMsg = InMessage;
-			InNode.ErrorType = InSeverity;
-
-			if (InSeverity == EMessageSeverity::Error)
-			{
-				UE_LOG(LogMetasoundEditor, Error, TEXT("%s"), *InMessage);
-			}
+			static const FString MissingConcreteOutputConnectionFormat = TEXT(
+				"Reroute connection for pin '{0}' does not provide a concrete output. "
+				"Resulting literal value is undefined and may result in unintended results."
+			);
 		}
-	}
-}
-
-Metasound::Editor::FGraphNodeValidationResult UMetasoundEditorGraphMemberNode::CreateNewValidationResult()
-{
-	using namespace Metasound::Editor;
-
-	FGraphNodeValidationResult OutResult = FGraphNodeValidationResult(*this);
-
-	// Reset the node validation state
-	if (ErrorType != EMessageSeverity::Info)
-	{
-		ErrorType = EMessageSeverity::Info;
-		OutResult.bIsDirty = true;
-	}
-
-	if (bHasCompilerMessage)
-	{
-		bHasCompilerMessage = false;
-		ErrorMsg.Reset();
-		OutResult.bIsDirty = true;
-	}
-	return OutResult;
+	} 
 }
 
 UMetasoundEditorGraphNode::UMetasoundEditorGraphNode(const FObjectInitializer& ObjectInitializer)
@@ -227,11 +203,42 @@ void UMetasoundEditorGraphNode::CacheTitle()
 	CachedTitle = NodeHandle->GetDisplayTitle();
 }
 
-bool UMetasoundEditorGraphNode::Validate(Metasound::Editor::FGraphNodeValidationResult& OutResult)
+void UMetasoundEditorGraphNode::Validate(Metasound::Editor::FGraphNodeValidationResult& OutResult)
 {
-	// Do no validation by default
-	OutResult = Metasound::Editor::FGraphNodeValidationResult(*this);
-	return true;
+	using namespace Metasound::Editor;
+	using namespace Metasound::Frontend;
+
+	OutResult = FGraphNodeValidationResult(*this);
+
+#if WITH_EDITOR
+	// Validate that non-reroute inputs are connected to "real" outputs
+	if (GetClassName() != FRerouteNodeTemplate::ClassName)
+	{
+		for (UEdGraphPin* Pin : Pins)
+		{
+			OutResult.SetPinOrphaned(*Pin, false);
+			if (Pin->Direction == EGPD_Input)
+			{
+				if (!Pin->LinkedTo.IsEmpty())
+				{
+					UEdGraphPin* ReroutedPin = FGraphBuilder::FindReroutedOutputPin(Pin->LinkedTo.Last());
+					if (ReroutedPin)
+					{
+						if (UMetasoundEditorGraphExternalNode* ExternalNode = Cast<UMetasoundEditorGraphExternalNode>(ReroutedPin->GetOwningNode()))
+						{
+							if (ExternalNode->GetClassName() == FRerouteNodeTemplate::ClassName)
+							{
+								FConstInputHandle InputHandle = FGraphBuilder::GetConstInputHandleFromPin(Pin);
+								const FString Msg = FString::Format(*GraphNodePrivate::MissingConcreteOutputConnectionFormat, { InputHandle->GetDisplayName().ToString() });
+								OutResult.SetMessage(EMessageSeverity::Warning, Msg);
+							}
+						}
+					}
+				}
+			}
+		}
+	}
+#endif // WITH_EDITOR
 }
 
 bool UMetasoundEditorGraphNode::ContainsMetadataChange() const
@@ -342,12 +349,20 @@ void UMetasoundEditorGraphNode::GetPinHoverText(const UEdGraphPin& Pin, FString&
 
 	if (Pin.Direction == EGPD_Input)
 	{
+		// Report if connected to reroute network is not connected to concrete output 
 		FConstInputHandle InputHandle = FGraphBuilder::GetConstInputHandleFromPin(&Pin);
-		OutHoverText = InputHandle->GetTooltip().ToString();
+		if (Pin.bOrphanedPin && InputHandle->IsValid())
+		{
+			OutHoverText = FString::Format(*GraphNodePrivate::MissingConcreteOutputConnectionFormat, { InputHandle->GetDisplayName().ToString() });
+		}
+		else
+		{
+			OutHoverText = InputHandle->GetTooltip().ToString();
+		}
 	}
 	else // Pin.Direction == EGPD_Output
 	{
-		FConstOutputHandle OutputHandle = FGraphBuilder::GetConstOutputHandleFromPin(&Pin);
+		FConstOutputHandle OutputHandle = FGraphBuilder::FindReroutedConstOutputHandleFromPin(&Pin);
 		OutHoverText = OutputHandle->GetTooltip().ToString();
 	}
 }
@@ -619,13 +634,13 @@ bool UMetasoundEditorGraphOutputNode::EnableInteractWidgets() const
 	return bEnabled;
 }
 
-bool UMetasoundEditorGraphOutputNode::Validate(Metasound::Editor::FGraphNodeValidationResult& OutResult) 
+void UMetasoundEditorGraphOutputNode::Validate(Metasound::Editor::FGraphNodeValidationResult& OutResult)
 {
 #if WITH_EDITOR
 	using namespace Metasound::Editor;
 	using namespace Metasound::Frontend;
 
-	OutResult = CreateNewValidationResult();
+	Super::Validate(OutResult);
 
 	// 2. Check if node is invalid, version is missing and cache if interface changes exist between the document's records and the registry
 	FNodeHandle NodeHandle = GetNodeHandle();
@@ -645,15 +660,12 @@ bool UMetasoundEditorGraphOutputNode::Validate(Metasound::Editor::FGraphNodeVali
 				bool bIsConnected = InputHandles.Last()->IsConnected();
 				if (!bIsConnected)
 				{
-					GraphNode::SetMessage(*this, EMessageSeverity::Warning, *RequiredText.ToString());
-					return false;
+					OutResult.SetMessage(EMessageSeverity::Warning, *RequiredText.ToString());
 				}
 			}
 		}
 	}
 #endif // #if WITH_EDITOR
-
-	return true;
 }
 
 FMetasoundFrontendClassName UMetasoundEditorGraphOutputNode::GetClassName() const
@@ -750,31 +762,95 @@ void UMetasoundEditorGraphExternalNode::CacheTitle()
 	CachedTitle = FGraphBuilder::GetDisplayName(*NodeHandle, bIncludeNamespace);
 }
 
-bool UMetasoundEditorGraphExternalNode::Validate(Metasound::Editor::FGraphNodeValidationResult& OutResult)
+void UMetasoundEditorGraphExternalNode::GetPinHoverText(const UEdGraphPin& Pin, FString& OutHoverText) const
+{
+	using namespace Metasound::Frontend;
+
+	if (ClassName == FRerouteNodeTemplate::ClassName)
+	{
+		if (!ErrorMsg.IsEmpty())
+		{
+			OutHoverText = ErrorMsg;
+			return;
+		}
+	}
+
+	Super::GetPinHoverText(Pin, OutHoverText);
+}
+
+void UMetasoundEditorGraphExternalNode::Validate(Metasound::Editor::FGraphNodeValidationResult& OutResult)
 {
 	using namespace Metasound::Editor;
 	using namespace Metasound::Frontend;
 
-	OutResult = FGraphNodeValidationResult(*this);
+#if WITH_EDITOR
+	Super::Validate(OutResult);
 
-	// 1. Reset ed node validation state
-	if (ErrorType != EMessageSeverity::Info)
+	FConstNodeHandle NodeHandle = GetNodeHandle();
+	const FMetasoundFrontendClassMetadata& Metadata = NodeHandle->GetClassMetadata();
+
+	// 1. Validate referenced graph recursively if defined as asset node class
+	const FNodeRegistryKey RegistryKey = NodeRegistryKey::CreateKey(Metadata);
+	if (IMetaSoundAssetManager* AssetManager = IMetaSoundAssetManager::Get())
 	{
-		ErrorType = EMessageSeverity::Info;
-		OutResult.bIsDirty = true;
+		if (const FSoftObjectPath* Path = AssetManager->FindObjectPathFromKey(RegistryKey))
+		{
+			if (UObject* AssetObject = Path->ResolveObject())
+			{
+				const EMessageSeverity::Type MaxGraphMsg = FGraphBuilder::ValidateGraph(*AssetObject, false /* bForceRefreshNodes */);
+				switch (MaxGraphMsg)
+				{
+					case EMessageSeverity::Error:
+					{
+						OutResult.SetMessage(MaxGraphMsg, TEXT("Referenced asset class contains error(s). Check implementation for details."));
+					}
+					break;
+
+					case EMessageSeverity::PerformanceWarning:
+					case EMessageSeverity::Warning:
+					{
+						OutResult.SetMessage(MaxGraphMsg, TEXT("Referenced asset class contains warning(s). Check implementation for details."));
+					}
+					break;
+
+					case EMessageSeverity::Info:
+					default:
+					{
+					}
+					break;
+				}
+			}
+		}
 	}
 
-	if (bHasCompilerMessage)
+	// 2. Validate template nodes
+	if (Metadata.GetType() == EMetasoundFrontendClassType::Template)
 	{
-		bHasCompilerMessage = false;
-		ErrorMsg.Reset();
-		OutResult.bIsDirty = true;
+		const FNodeRegistryKey Key = NodeRegistryKey::CreateKey(Metadata);
+		if (const INodeTemplate* Template = INodeTemplateRegistry::Get().FindTemplate(Key))
+		{
+			const bool bIsValidInterface = Template->IsValidNodeInterface(NodeHandle->GetNodeInterface());
+			if (!bIsValidInterface)
+			{
+				OutResult.SetMessage(EMessageSeverity::Error, FString::Format(TEXT("Cannot implement template interface for node class '{0}"), { *Metadata.GetClassName().ToString() }));
+			}
+			else
+			{
+#if WITH_EDITOR
+				if (!Template->HasRequiredConnections(NodeHandle))
+				{
+					OutResult.SetMessage(EMessageSeverity::Warning, TEXT("Reroute node(s) missing non-reroute input connection(s)."));
+				}
+#endif // WITH_EDITOR
+			}
+		}
+		else
+		{
+			OutResult.SetMessage(EMessageSeverity::Error, FString::Format(TEXT("Template node interface missing for node class '{0}'"), { *Metadata.GetClassName().ToString()}));
+		}
 	}
 
-	// 2. Check if node is invalid, version is missing and cache if interface changes exist between the document's records and the registry
-	FNodeHandle NodeHandle = GetNodeHandle();
-	const FMetasoundFrontendClassMetadata& Metadata = NodeHandle->GetClassMetadata();
-
+	// 3. Check if node is invalid, version is missing and cache if interface changes exist between the document's records and the registry
 	FClassInterfaceUpdates InterfaceUpdates;
 	if (!NodeHandle->DiffAgainstRegistryInterface(InterfaceUpdates, false /* bUseHighestMinorVersion */))
 	{
@@ -800,30 +876,26 @@ bool UMetasoundEditorGraphExternalNode::Validate(Metasound::Editor::FGraphNodeVa
 				*PromptIfMissing->ToString()
 			});
 
-			GraphNode::SetMessage(*this, EMessageSeverity::Error, NewErrorMsg);
+			OutResult.SetMessage(EMessageSeverity::Error, NewErrorMsg);
 		}
 		else
 		{
 			if (bIsClassNative)
 			{
-				GraphNode::SetMessage(*this, EMessageSeverity::Error, FString::Format(
+				OutResult.SetMessage(EMessageSeverity::Error, FString::Format(
 					TEXT("Class '{0}' definition missing for last known natively defined node."),
 					{ *ClassName.ToString() }));
 			}
 			else
 			{
-				GraphNode::SetMessage(*this, EMessageSeverity::Error,
+				OutResult.SetMessage(EMessageSeverity::Error,
 					FString::Format(TEXT("Class definition missing for asset with guid '{0}': Asset is either missing or invalid"),
 					{ *ClassName.Name.ToString() }));
 			}
 		}
-
-		OutResult.bIsDirty = true;
-		OutResult.bIsInvalid = true;
 	}
 
-	// 3. Report if node was nativized
-	const FNodeRegistryKey RegistryKey = NodeRegistryKey::CreateKey(Metadata);
+	// 4. Report if node was nativized
 	bool bNewIsClassNative = FMetasoundFrontendRegistryContainer::Get()->IsNodeNative(RegistryKey);
 	if (bIsClassNative != bNewIsClassNative)
 	{
@@ -832,11 +904,10 @@ bool UMetasoundEditorGraphExternalNode::Validate(Metasound::Editor::FGraphNodeVa
 			NodeUpgradeMessage = FText::Format(LOCTEXT("MetaSoundNode_NativizedMessage", "Class '{0}' has been nativized."), Metadata.GetDisplayName());
 		}
 
-		OutResult.bIsDirty = true;
 		bIsClassNative = bNewIsClassNative;
 	}
 
-	// 4. Report if node was auto-updated
+	// 5. Report if node was auto-updated
 	FMetasoundFrontendNodeStyle Style = NodeHandle->GetNodeStyle();
 	if (Style.bMessageNodeUpdated)
 	{
@@ -844,10 +915,9 @@ bool UMetasoundEditorGraphExternalNode::Validate(Metasound::Editor::FGraphNodeVa
 			Metadata.GetDisplayName(),
 			FText::FromString(Metadata.GetVersion().ToString())
 		);
-		OutResult.bIsDirty = true;
 	}
 
-	// 5. Reset pin state (if pin was orphaned or clear if no longer orphaned)
+	// 6. Reset pin state (if pin was orphaned or clear if no longer orphaned)
 	for (UEdGraphPin* Pin : Pins)
 	{
 		bool bWasRemoved = false;
@@ -869,21 +939,17 @@ bool UMetasoundEditorGraphExternalNode::Validate(Metasound::Editor::FGraphNodeVa
 			});
 		}
 
-		if (Pin->bOrphanedPin != bWasRemoved)
-		{
-			Pin->bOrphanedPin = bWasRemoved;
-			OutResult.bIsDirty = true;
-		}
+		OutResult.SetPinOrphaned(*Pin, bWasRemoved);
 	}
 
-	// Report of node class is deprecated
+	// 7. Report if node class is deprecated
 	FMetasoundFrontendClass RegisteredClass;
 	if (FMetasoundFrontendRegistryContainer::Get()->GetFrontendClassFromRegistered(RegistryKey, RegisteredClass))
 	{
 		if (RegisteredClass.Metadata.GetIsDeprecated())
 		{
 			constexpr bool bIncludeNamespace = true;
-			GraphNode::SetMessage(*this, EMessageSeverity::Warning,
+			OutResult.SetMessage(EMessageSeverity::Warning,
 				FString::Format(TEXT("Class '{0} {1}' is deprecated."),
 				{
 					*FGraphBuilder::GetDisplayName(RegisteredClass.Metadata, { }, bIncludeNamespace).ToString(),
@@ -892,19 +958,17 @@ bool UMetasoundEditorGraphExternalNode::Validate(Metasound::Editor::FGraphNodeVa
 		}
 	}
 
-	// Find all available versions & report if upgrade available
+	// 8. Find all available versions & report if upgrade available
 	const Metasound::FNodeClassName NodeClassName = Metadata.GetClassName().ToNodeClassName();
 	const TArray<FMetasoundFrontendClass> SortedClasses = ISearchEngine::Get().FindClassesWithName(NodeClassName, true /* bSortByVersion */);
 	if (SortedClasses.IsEmpty())
 	{
-		GraphNode::SetMessage(*this, EMessageSeverity::Error,
+		OutResult.SetMessage(EMessageSeverity::Error,
 			FString::Format(TEXT("Class '{0} {1}' not registered."),
 			{
 				*Metadata.GetClassName().ToString(),
 				*Metadata.GetVersion().ToString()
 			}));
-		OutResult.bIsDirty = true;
-		OutResult.bIsInvalid = true;
 	}
 	else
 	{
@@ -939,41 +1003,34 @@ bool UMetasoundEditorGraphExternalNode::Validate(Metasound::Editor::FGraphNodeVa
 					*HighestRegistryClass.Metadata.GetVersion().ToString()
 				});
 				Severity = EMessageSeverity::Error;
-				OutResult.bIsInvalid = true;
 			}
 
-			GraphNode::SetMessage(*this, Severity, NodeMsg);
-			OutResult.bIsDirty = true;
+			OutResult.SetMessage(Severity, NodeMsg);
 		}
 		else if (HighestRegistryClass.Metadata.GetVersion() == CurrentVersion)
 		{
 			if (InterfaceUpdates.ContainsChanges())
 			{
-				GraphNode::SetMessage(*this, EMessageSeverity::Error,
-				FString::Format(TEXT("Node & registered class interface mismatch: '{0} {1}'. Class either versioned improperly, class key collision exists, or AutoUpdate disabled in 'MetaSound' Developer Settings."),
-				{
-					*Metadata.GetClassName().ToString(),
-					*Metadata.GetVersion().ToString()
-				}));
-				OutResult.bIsDirty = true;
-				OutResult.bIsInvalid = true;
+				OutResult.SetMessage(EMessageSeverity::Error,
+					FString::Format(TEXT("Node & registered class interface mismatch: '{0} {1}'. Class either versioned improperly, class key collision exists, or AutoUpdate disabled in 'MetaSound' Developer Settings."),
+					{
+						*Metadata.GetClassName().ToString(),
+						*Metadata.GetVersion().ToString()
+					}));
 			}
 		}
 		else
 		{
-			GraphNode::SetMessage(*this, EMessageSeverity::Error,
+			OutResult.SetMessage(EMessageSeverity::Error,
 				FString::Format(TEXT("Node with class '{0} {1}' interface version higher than that of highest minor revision ({2}) in class registry."),
 				{
 					*Metadata.GetClassName().ToString(),
 					*Metadata.GetVersion().ToString(),
 					*HighestRegistryClass.Metadata.GetVersion().ToString()
 				}));
-			OutResult.bIsDirty = true;
-			OutResult.bIsInvalid = true;
 		}
 	}
-
-	return !OutResult.bIsInvalid;
+#endif // WITH_EDITOR
 }
 
 FLinearColor UMetasoundEditorGraphExternalNode::GetNodeTitleColor() const
@@ -1003,6 +1060,20 @@ FSlateIcon UMetasoundEditorGraphExternalNode::GetNodeTitleIcon() const
 	}
 }
 
+bool UMetasoundEditorGraphExternalNode::ShouldDrawNodeAsControlPointOnly(int32& OutInputPinIndex, int32& OutOutputPinIndex) const
+{
+	using namespace Metasound::Frontend;
+
+	if (GetClassName() == FRerouteNodeTemplate::ClassName)
+	{
+		OutInputPinIndex = 0;
+		OutOutputPinIndex = 1;
+		return true;
+	}
+
+	return Super::ShouldDrawNodeAsControlPointOnly(OutInputPinIndex, OutOutputPinIndex);
+}
+
 UMetasoundEditorGraphMember* UMetasoundEditorGraphVariableNode::GetMember() const
 {
 	return Variable;
diff --git a/Engine/Plugins/Runtime/Metasound/Source/MetasoundEditor/Private/MetasoundEditorGraphNodeFactory.h b/Engine/Plugins/Runtime/Metasound/Source/MetasoundEditor/Private/MetasoundEditorGraphNodeFactory.h
--- a/Engine/Plugins/Runtime/Metasound/Source/MetasoundEditor/Private/MetasoundEditorGraphNodeFactory.h
+++ b/Engine/Plugins/Runtime/Metasound/Source/MetasoundEditor/Private/MetasoundEditorGraphNodeFactory.h
@@ -4,6 +4,7 @@
 #include "EdGraphUtilities.h"
 #include "EdGraph/EdGraphNode.h"
 #include "MetasoundEditorGraphNode.h"
+#include "NodeTemplates/MetasoundFrontendNodeTemplateReroute.h"
 #include "SGraphNode.h"
 #include "SMetasoundGraphNode.h"
 #include "SMetasoundGraphNodeComment.h"
@@ -16,9 +17,19 @@ class FMetasoundGraphNodeFactory : public FGraphPanelNodeFactory
 	virtual TSharedPtr<SGraphNode> CreateNode(UEdGraphNode* InNode) const override
 	{
 		using namespace Metasound::Editor;
+		using namespace Metasound::Frontend;
 
 		if (InNode->IsA<UMetasoundEditorGraphNode>())
 		{
+			if (const UMetasoundEditorGraphExternalNode* Node = Cast<UMetasoundEditorGraphExternalNode>(InNode))
+			{
+				FConstNodeHandle NodeHandle = Node->GetConstNodeHandle();
+				const FMetasoundFrontendClassName ClassName = NodeHandle->GetClassMetadata().GetClassName();
+				if (ClassName == FRerouteNodeTemplate::ClassName)
+				{
+					return SNew(SMetaSoundGraphNodeKnot, InNode);
+				}
+			}
 			return SNew(SMetaSoundGraphNode, InNode);
 		}
 		else if (UEdGraphNode_Comment* CommentNode = Cast<UEdGraphNode_Comment>(InNode))
diff --git a/Engine/Plugins/Runtime/Metasound/Source/MetasoundEditor/Private/MetasoundEditorGraphSchema.cpp b/Engine/Plugins/Runtime/Metasound/Source/MetasoundEditor/Private/MetasoundEditorGraphSchema.cpp
--- a/Engine/Plugins/Runtime/Metasound/Source/MetasoundEditor/Private/MetasoundEditorGraphSchema.cpp
+++ b/Engine/Plugins/Runtime/Metasound/Source/MetasoundEditor/Private/MetasoundEditorGraphSchema.cpp
@@ -24,6 +24,7 @@
 #include "MetasoundEditorGraphBuilder.h"
 #include "MetasoundEditorGraphInputNode.h"
 #include "MetasoundEditorGraphNode.h"
+#include "MetasoundEditorGraphValidation.h"
 #include "MetasoundEditorModule.h"
 #include "MetasoundEditorSettings.h"
 #include "MetasoundFrontend.h"
@@ -34,7 +35,9 @@
 #include "MetasoundLiteral.h"
 #include "MetasoundNodeInterface.h"
 #include "MetasoundStandardNodesCategories.h"
+#include "MetasoundUObjectRegistry.h"
 #include "MetasoundVariableNodes.h"
+#include "NodeTemplates/MetasoundFrontendNodeTemplateReroute.h"
 #include "ScopedTransaction.h"
 #include "Styling/SlateStyleRegistry.h"
 #include "Toolkits/ToolkitManager.h"
@@ -394,21 +397,17 @@ UEdGraphNode* FMetasoundGraphSchemaAction_NodeWithMultipleOutputs::PerformAction
 
 const FSlateBrush* FMetasoundGraphSchemaAction_NewNode::GetIconBrush() const
 {
+	using namespace Metasound::Editor;
 	using namespace Metasound::Frontend;
 
-	if (const ISlateStyle* MetasoundStyle = FSlateStyleRegistry::FindSlateStyle("MetaSoundStyle"))
+	const FNodeRegistryKey RegistryKey = NodeRegistryKey::CreateKey(ClassMetadata);
+	const bool bIsClassNative = FMetasoundFrontendRegistryContainer::Get()->IsNodeNative(RegistryKey);
+	if (bIsClassNative)
 	{
-		const FNodeRegistryKey RegistryKey = NodeRegistryKey::CreateKey(ClassMetadata);
-		const bool bIsClassNative = FMetasoundFrontendRegistryContainer::Get()->IsNodeNative(RegistryKey);
-		if (bIsClassNative)
-		{
-			return MetasoundStyle->GetBrush(TEXT("MetasoundEditor.Graph.Node.Class.Native"));
-		}
-
-		return MetasoundStyle->GetBrush(TEXT("MetasoundEditor.Graph.Node.Class.Graph"));
+		return &Style::GetSlateBrushSafe("MetasoundEditor.Graph.Node.Class.Native");
 	}
 
-	return Super::GetIconBrush();
+	return &Style::GetSlateBrushSafe("MetasoundEditor.Graph.Node.Class.Graph");
 }
 
 const FLinearColor& FMetasoundGraphSchemaAction_NewNode::GetIconColor() const
@@ -467,11 +466,8 @@ FMetasoundGraphSchemaAction_NewInput::FMetasoundGraphSchemaAction_NewInput(FText
 
 const FSlateBrush* FMetasoundGraphSchemaAction_NewInput::GetIconBrush() const
 {
-	if (const ISlateStyle* MetasoundStyle = FSlateStyleRegistry::FindSlateStyle("MetaSoundStyle"))
-	{
-		return MetasoundStyle->GetBrush("MetasoundEditor.Graph.Node.Class.Input");
-	}
-	return Super::GetIconBrush();
+	using namespace Metasound::Editor;
+	return &Style::GetSlateBrushSafe("MetasoundEditor.Graph.Node.Class.Input");
 }
 
 const FLinearColor& FMetasoundGraphSchemaAction_NewInput::GetIconColor() const
@@ -843,11 +839,8 @@ UEdGraphNode* FMetasoundGraphSchemaAction_NewVariableNode::PerformAction(UEdGrap
 
 const FSlateBrush* FMetasoundGraphSchemaAction_NewVariableNode::GetIconBrush() const
 {
-	if (const ISlateStyle* MetasoundStyle = FSlateStyleRegistry::FindSlateStyle("MetaSoundStyle"))
-	{
-		return MetasoundStyle->GetBrush("MetasoundEditor.Graph.Node.Class.Variable");
-	}
-	return Super::GetIconBrush();
+	using namespace Metasound::Editor;
+	return &Style::GetSlateBrushSafe("MetasoundEditor.Graph.Node.Class.Variable");
 }
 
 const FLinearColor& FMetasoundGraphSchemaAction_NewVariableNode::GetIconColor() const
@@ -906,6 +899,69 @@ UEdGraphNode* FMetasoundGraphSchemaAction_NewFromSelected::PerformAction(UEdGrap
 	return nullptr;
 }
 
+FMetasoundGraphSchemaAction_NewReroute::FMetasoundGraphSchemaAction_NewReroute(const FLinearColor& InIconColor)
+	: FMetasoundGraphSchemaAction(
+		FText(),
+		LOCTEXT("RerouteName", "Add Reroute Node..."),
+		LOCTEXT("RerouteTooltip", "Reroute Node (reroutes wires)"),
+		Metasound::Editor::EPrimaryContextGroup::Common)
+	, IconColor(InIconColor)
+{
+}
+
+UEdGraphNode* FMetasoundGraphSchemaAction_NewReroute::PerformAction(UEdGraph* ParentGraph, UEdGraphPin* FromPin, const FVector2D Location, bool bSelectNewNode /* = true*/)
+{
+	using namespace Metasound::Editor;
+	using namespace Metasound::Frontend;
+
+	const FName DataType = FGraphBuilder::GetPinDataType(FromPin);
+	if (DataType.IsNone())
+	{
+		return nullptr;
+	}
+
+	const FScopedTransaction Transaction(FText::Format(LOCTEXT("AddNewNode", "Add {0} Reroute Node"), FText::FromName(DataType)));
+	UMetasoundEditorGraph* MetaSoundGraph = CastChecked<UMetasoundEditorGraph>(ParentGraph);
+	UObject& ParentMetasound = MetaSoundGraph->GetMetasoundChecked();
+	ParentMetasound.Modify();
+	ParentGraph->Modify();
+
+	FMetasoundAssetBase* MetaSoundAsset = Metasound::IMetasoundUObjectRegistry::Get().GetObjectAsAssetBase(&ParentMetasound);
+	check(MetaSoundAsset);
+
+	const FNodeRegistryKey& RerouteTemplateKey = FRerouteNodeTemplate::GetRegistryKey();
+	FMetasoundFrontendNodeInterface NodeInterface = FRerouteNodeTemplate::CreateNodeInterfaceFromDataType(DataType);
+	FNodeHandle NodeHandle = MetaSoundAsset->GetRootGraphHandle()->AddTemplateNode(RerouteTemplateKey, MoveTemp(NodeInterface));
+
+	if (UMetasoundEditorGraphExternalNode* NewGraphNode = FGraphBuilder::AddExternalNode(ParentMetasound, NodeHandle, Location, bSelectNewNode))
+	{
+		NewGraphNode->Modify();
+		SchemaPrivate::TryConnectNewNodeToMatchingDataTypePin(*NewGraphNode, FromPin);
+		MetaSoundGraph->SetSynchronizationRequired();
+
+		TSharedPtr<FEditor> ParentEditor = FGraphBuilder::GetEditorForMetasound(ParentMetasound);
+		if (ParentEditor.IsValid() && bSelectNewNode)
+		{
+			ParentEditor->ClearSelectionAndSelectNode(NewGraphNode);
+		}
+
+		return NewGraphNode;
+	}
+
+	return nullptr;
+}
+
+const FSlateBrush* FMetasoundGraphSchemaAction_NewReroute::GetIconBrush() const
+{
+	using namespace Metasound::Editor;
+	return &Style::GetSlateBrushSafe("MetasoundEditor.Graph.Node.Class.Reroute");
+}
+
+const FLinearColor& FMetasoundGraphSchemaAction_NewReroute::GetIconColor() const
+{
+	return IconColor;
+}
+
 UEdGraphNode* FMetasoundGraphSchemaAction_NewComment::PerformAction(UEdGraph* ParentGraph, UEdGraphPin* FromPin, const FVector2D Location, bool bSelectNewNode/* = true*/)
 {
 	using namespace Metasound::Editor;
@@ -1053,6 +1109,9 @@ void UMetasoundEditorGraphSchema::GetGraphContextActions(FGraphContextMenuBuilde
 				ActionMenuBuilder.AddAction(MakeShared<FMetasoundGraphSchemaAction_PromoteToVariable_AccessorNode>());
 				ActionMenuBuilder.AddAction(MakeShared<FMetasoundGraphSchemaAction_PromoteToVariable_DeferredAccessorNode>());
 			}
+
+			const FLinearColor IconColor = GetPinTypeColor(FromPin->PinType);
+			ActionMenuBuilder.AddAction(MakeShared<FMetasoundGraphSchemaAction_NewReroute>(IconColor));
 		}
 
 		if (FromPin->Direction == EGPD_Output)
@@ -1079,6 +1138,9 @@ void UMetasoundEditorGraphSchema::GetGraphContextActions(FGraphContextMenuBuilde
 			FGraphActionMenuBuilder& ActionMenuBuilder = static_cast<FGraphActionMenuBuilder&>(ContextMenuBuilder);
 			ActionMenuBuilder.AddAction(MakeShared<FMetasoundGraphSchemaAction_PromoteToOutput>());
 			ActionMenuBuilder.AddAction(MakeShared<FMetasoundGraphSchemaAction_PromoteToVariable_MutatorNode>());
+
+			const FLinearColor IconColor = GetPinTypeColor(FromPin->PinType);
+			ActionMenuBuilder.AddAction(MakeShared<FMetasoundGraphSchemaAction_NewReroute>(IconColor));
 		}
 	}
 	else
@@ -1190,8 +1252,6 @@ const FPinConnectionResponse UMetasoundEditorGraphSchema::CanCreateConnection(co
 		return FPinConnectionResponse(CONNECT_RESPONSE_DISALLOW, LOCTEXT("ConnectionTypeIncorrect", "Connection pin types do not match"));
 	}
 
-	// Check if nodes being connected are in error/out-of-date version state before checking if can be connected as 
-	// nodes with errors can contain invalid pins/handles).
 	bool bConnectingNodesWithErrors = false;
 	UEdGraphNode* InputNode = InputPin->GetOwningNode();
 	if (ensure(InputNode))
@@ -1209,14 +1269,13 @@ const FPinConnectionResponse UMetasoundEditorGraphSchema::CanCreateConnection(co
 			bConnectingNodesWithErrors = true;
 		}
 	}
-	if (bConnectingNodesWithErrors)
-	{
-		return FPinConnectionResponse(CONNECT_RESPONSE_DISALLOW, LOCTEXT("ConnectionCannotContainErrorNode", "Cannot create new connections with node containing errors."));
-	}
 
 	Frontend::FInputHandle InputHandle = Editor::FGraphBuilder::GetInputHandleFromPin(InputPin);
 	Frontend::FOutputHandle OutputHandle = Editor::FGraphBuilder::GetOutputHandleFromPin(OutputPin);
-	if (InputHandle->IsValid() && OutputHandle->IsValid())
+
+	const bool bInputValid = InputHandle->IsValid();
+	const bool bOutputValid = OutputHandle->IsValid();
+	if (bInputValid && bOutputValid)
 	{
 		// TODO: Implement YesWithConverterNode to provide conversion options
 		Frontend::FConnectability Connectability = InputHandle->CanConnectTo(*OutputHandle);
@@ -1263,8 +1322,51 @@ const FPinConnectionResponse UMetasoundEditorGraphSchema::CanCreateConnection(co
 
 		return FPinConnectionResponse(CONNECT_RESPONSE_MAKE, TEXT(""));
 	}
+	else if (bConnectingNodesWithErrors)
+	{
+		return FPinConnectionResponse(CONNECT_RESPONSE_DISALLOW, LOCTEXT("ConnectionCannotContainErrorNode", "Cannot create new connections with node containing errors."));
+	}
+	else
+	{
+		return FPinConnectionResponse(CONNECT_RESPONSE_DISALLOW, LOCTEXT("ConnectionInternalError", "Internal error. Metasound node vertex handle mismatch."));
+	}
+}
+
+void UMetasoundEditorGraphSchema::OnPinConnectionDoubleCicked(UEdGraphPin* PinA, UEdGraphPin* PinB, const FVector2D& GraphPosition) const
+{
+	if (!PinA || !PinB)
+	{
+		return;
+	}
+
+	//@TODO: This constant is duplicated from inside of SGraphNodeKnot
+	const FVector2D NodeSpacerSize(42.0f, 24.0f);
+	const FVector2D KnotTopLeft = GraphPosition - (NodeSpacerSize * 0.5f);
 
-	return FPinConnectionResponse(CONNECT_RESPONSE_DISALLOW, LOCTEXT("ConnectionInternalError", "Internal error. Metasound node vertex handle mismatch."));
+	UMetasoundEditorGraph* ParentGraph = Cast<UMetasoundEditorGraph>(PinA->GetOwningNode()->GetGraph());
+	if (ParentGraph->IsEditable())
+	{
+		UEdGraphPin* OutputPin = PinA->Direction == EGPD_Output ? PinA : PinB;
+		const FLinearColor IconColor = GetPinTypeColor(OutputPin->PinType);
+		TSharedPtr<FMetasoundGraphSchemaAction_NewReroute> RerouteAction = MakeShared<FMetasoundGraphSchemaAction_NewReroute>(IconColor);
+
+		UEdGraphNode& Node = *PinA->GetOwningNode();
+		UEdGraphNode* NewNode = RerouteAction->PerformAction(Node.GetGraph(), OutputPin, GraphPosition, true);
+
+		if (ensure(NewNode))
+		{
+			UEdGraphPin** RerouteOutputPtr = NewNode->Pins.FindByPredicate([](const UEdGraphPin* Candidate) { return Candidate->Direction == EGPD_Output; });
+
+			if (ensure(*RerouteOutputPtr))
+			{
+				constexpr bool bShouldBreakSingleTransact = false;
+				BreakSinglePinLink(PinA, PinB, bShouldBreakSingleTransact);
+
+				UEdGraphPin* InputPin = PinA->Direction == EGPD_Input ? PinA : PinB;
+				ensure(TryCreateConnection(InputPin, *RerouteOutputPtr));
+			}
+		}
+	}
 }
 
 bool UMetasoundEditorGraphSchema::TryCreateConnection(UEdGraphPin* PinA, UEdGraphPin* PinB) const
@@ -1308,6 +1410,8 @@ bool UMetasoundEditorGraphSchema::TryCreateConnection(UEdGraphPin* PinA, UEdGrap
 		return false;
 	}
 
+	Graph->SetSynchronizationRequired();
+
 	return true;
 }
 
@@ -1353,6 +1457,26 @@ void UMetasoundEditorGraphSchema::TrySetDefaultValue(UEdGraphPin& Pin, const FSt
 
 bool UMetasoundEditorGraphSchema::ShouldHidePinDefaultValue(UEdGraphPin* Pin) const
 {
+	using namespace Metasound::Editor;
+	using namespace Metasound::Frontend;
+
+	if (!Pin)
+	{
+		return true;
+	}
+
+	if (Pin->Direction == EGPD_Input)
+	{
+		FInputHandle InputHandle = FGraphBuilder::GetInputHandleFromPin(Pin);
+		return !InputHandle->GetOwningNode()->GetClassStyle().Display.bShowLiterals;
+	}
+
+	if (Pin->Direction == EGPD_Output)
+	{
+		FOutputHandle OutputHandle = FGraphBuilder::GetOutputHandleFromPin(Pin);
+		return !OutputHandle->GetOwningNode()->GetClassStyle().Display.bShowLiterals;
+	}
+
 	// TODO: Determine if should be hidden from doc data
 	return false;
 }
@@ -1466,6 +1590,11 @@ void UMetasoundEditorGraphSchema::BreakPinLinks(UEdGraphPin& TargetPin, bool bSe
 }
 
 void UMetasoundEditorGraphSchema::BreakSinglePinLink(UEdGraphPin* SourcePin, UEdGraphPin* TargetPin) const
+{
+	BreakSinglePinLink(SourcePin, TargetPin, true);
+}
+
+void UMetasoundEditorGraphSchema::BreakSinglePinLink(UEdGraphPin* SourcePin, UEdGraphPin* TargetPin, bool bShouldTransact) const
 {
 	using namespace Metasound::Editor;
 
@@ -1495,7 +1624,7 @@ void UMetasoundEditorGraphSchema::BreakSinglePinLink(UEdGraphPin* SourcePin, UEd
 		return;
 	}
 
-	const FScopedTransaction Transaction(LOCTEXT("BreakSinglePinLink", "Break Single Pin Link"));
+	const FScopedTransaction Transaction(LOCTEXT("BreakSinglePinLink", "Break Single Pin Link"), bShouldTransact);
 	UMetasoundEditorGraph* Graph = CastChecked<UMetasoundEditorGraph>(OwningNode->GetGraph());
 	Graph->GetMetasoundChecked().Modify();
 	SourcePin->Modify();
@@ -1843,7 +1972,8 @@ void UMetasoundEditorGraphSchema::GetFunctionActions(FGraphActionMenuBuilder& Ac
 			continue;
 		}
 
-		if (FrontendClass.Metadata.GetType() != EMetasoundFrontendClassType::External)
+		if (FrontendClass.Metadata.GetType() != EMetasoundFrontendClassType::External
+			&& FrontendClass.Metadata.GetType() != EMetasoundFrontendClassType::Template)
 		{
 			continue;
 		}
diff --git a/Engine/Plugins/Runtime/Metasound/Source/MetasoundEditor/Private/MetasoundEditorGraphValidation.cpp b/Engine/Plugins/Runtime/Metasound/Source/MetasoundEditor/Private/MetasoundEditorGraphValidation.cpp
new file mode 100644
--- /dev/null
+++ b/Engine/Plugins/Runtime/Metasound/Source/MetasoundEditor/Private/MetasoundEditorGraphValidation.cpp
@@ -0,0 +1,111 @@
+// Copyright Epic Games, Inc. All Rights Reserved.
+
+#include "MetasoundEditorGraphValidation.h"
+
+#include "Algo/NoneOf.h"
+#include "EdGraph/EdGraphNode.h"
+#include "Logging/TokenizedMessage.h"
+#include "MetasoundEditorGraph.h"
+#include "MetasoundEditorGraphNode.h"
+#include "MetasoundLog.h"
+
+
+namespace Metasound
+{
+	namespace Editor
+	{
+		FGraphNodeValidationResult::FGraphNodeValidationResult(UMetasoundEditorGraphNode& InNode)
+			: Node(&InNode)
+		{
+			// 1. Reset ed node validation state
+			if (InNode.ErrorType != EMessageSeverity::Info)
+			{
+				InNode.ErrorType = EMessageSeverity::Info;
+				bHasDirtiedNode = true;
+			}
+
+			if (InNode.bHasCompilerMessage)
+			{
+				InNode.bHasCompilerMessage = false;
+				bHasDirtiedNode = true;
+			}
+
+			if (!InNode.ErrorMsg.IsEmpty())
+			{
+				InNode.ErrorMsg.Reset();
+				bHasDirtiedNode = true;
+			}
+		}
+
+		bool FGraphNodeValidationResult::GetHasDirtiedNode() const
+		{
+			return bHasDirtiedNode;
+		}
+
+		bool FGraphNodeValidationResult::GetIsInvalid() const
+		{
+			return Node->ErrorType == EMessageSeverity::Error;
+		}
+
+		UMetasoundEditorGraphNode& FGraphNodeValidationResult::GetNodeChecked() const
+		{
+			check(Node);
+			return *Node;
+		}
+
+		void FGraphNodeValidationResult::SetMessage(EMessageSeverity::Type InSeverity, const FString& InMessage)
+		{
+			check(Node);
+
+			if (!Node->bHasCompilerMessage)
+			{
+				bHasDirtiedNode = true;
+				Node->bHasCompilerMessage = true;
+			}
+
+			if (Node->ErrorMsg != InMessage)
+			{
+				bHasDirtiedNode = true;
+				Node->ErrorMsg = InMessage;
+			}
+
+			if (Node->ErrorType != InSeverity)
+			{
+				bHasDirtiedNode = true;
+				Node->ErrorType = InSeverity;
+			}
+
+			if (InSeverity == EMessageSeverity::Error)
+			{
+				UE_LOG(LogMetaSound, Error, TEXT("%s"), *InMessage);
+			}
+		}
+
+		void FGraphNodeValidationResult::SetPinOrphaned(UEdGraphPin& Pin, bool bIsOrphaned)
+		{
+			if (Pin.bOrphanedPin != bIsOrphaned)
+			{
+				Pin.bOrphanedPin = bIsOrphaned;
+				bHasDirtiedNode = true;
+			}
+		}
+
+		EMessageSeverity::Type FGraphValidationResults::GetHighestMessageSeverity() const
+		{
+			int32 HighestMessageSeverity = static_cast<int32>(EMessageSeverity::Info);
+
+			for (const FGraphNodeValidationResult& Result : NodeResults)
+			{
+				UMetasoundEditorGraphNode& Node = Result.GetNodeChecked();
+				HighestMessageSeverity = FMath::Min(Node.ErrorType, HighestMessageSeverity);
+			}
+
+			return static_cast<EMessageSeverity::Type>(HighestMessageSeverity);
+		}
+
+		const TArray<FGraphNodeValidationResult>& FGraphValidationResults::GetResults() const
+		{
+			return NodeResults;
+		}
+	} // namespace Editor
+} // namespace Metasound
\ No newline at end of file
diff --git a/Engine/Plugins/Runtime/Metasound/Source/MetasoundEditor/Private/MetasoundEditorModule.cpp b/Engine/Plugins/Runtime/Metasound/Source/MetasoundEditor/Private/MetasoundEditorModule.cpp
--- a/Engine/Plugins/Runtime/Metasound/Source/MetasoundEditor/Private/MetasoundEditorModule.cpp
+++ b/Engine/Plugins/Runtime/Metasound/Source/MetasoundEditor/Private/MetasoundEditorModule.cpp
@@ -67,6 +67,8 @@ namespace Metasound
 {
 	namespace Editor
 	{
+		using FMetasoundGraphPanelPinFactory = FGraphPanelPinFactory;
+
 		static const FName AssetToolName { "AssetTools" };
 
 		template <typename T>
@@ -103,10 +105,22 @@ namespace Metasound
 					Set("MetasoundSource.Color", FColor(103, 214, 66));
 
 					// Actions
-					Set("MetasoundEditor.Play", new FSlateImageBrush(RootToContentDir(TEXT("Icons/play_40x.png")), Icon40x40));
-					Set("MetasoundEditor.Play.Small", new FSlateImageBrush(RootToContentDir(TEXT("Icons/play_40x.png")), Icon20x20));
-					Set("MetasoundEditor.Stop", new FSlateImageBrush(RootToContentDir(TEXT("Icons/stop_40x.png")), Icon40x40));
-					Set("MetasoundEditor.Stop.Small", new FSlateImageBrush(RootToContentDir(TEXT("Icons/stop_40x.png")), Icon20x20));
+					Set("MetasoundEditor.Play", new IMAGE_BRUSH_SVG(TEXT("Icons/play"), Icon40x40));
+					Set("MetasoundEditor.Play.Small", new IMAGE_BRUSH_SVG(TEXT("Icons/play"), Icon20x20));
+
+					Set("MetasoundEditor.Play.Active.Valid", new IMAGE_BRUSH_SVG(TEXT("Icons/play_active_valid"), Icon40x40));
+					Set("MetasoundEditor.Play.Active.Warning", new IMAGE_BRUSH_SVG(TEXT("Icons/play_active_warning"), Icon40x40));
+					Set("MetasoundEditor.Play.Inactive.Valid", new IMAGE_BRUSH_SVG(TEXT("Icons/play_inactive_valid"), Icon40x40));
+					Set("MetasoundEditor.Play.Inactive.Warning", new IMAGE_BRUSH_SVG(TEXT("Icons/play_inactive_warning"), Icon40x40));
+					Set("MetasoundEditor.Play.Error", new IMAGE_BRUSH_SVG(TEXT("Icons/play_error"), Icon40x40));
+
+					Set("MetasoundEditor.Stop", new IMAGE_BRUSH_SVG(TEXT("Icons/stop"), Icon40x40));
+
+					Set("MetasoundEditor.Stop.Disabled", new IMAGE_BRUSH_SVG(TEXT("Icons/stop_disabled"), Icon40x40));
+					Set("MetasoundEditor.Stop.Active", new IMAGE_BRUSH_SVG(TEXT("Icons/stop_active"), Icon40x40));
+					Set("MetasoundEditor.Stop.Inactive", new IMAGE_BRUSH_SVG(TEXT("Icons/stop_inactive"), Icon40x40));
+
+
 					Set("MetasoundEditor.Import", new FSlateImageBrush(RootToContentDir(TEXT("/Icons/build_40x.png")), Icon40x40));
 					Set("MetasoundEditor.Import.Small", new FSlateImageBrush(RootToContentDir(TEXT("/Icons/build_40x.png")), Icon20x20));
 					Set("MetasoundEditor.Export", new FSlateImageBrush(RootToContentDir(TEXT("/Icons/build_40x.png")), Icon40x40));
@@ -126,6 +140,7 @@ namespace Metasound
 					Set("MetasoundEditor.Graph.Node.Class.Graph", new IMAGE_BRUSH_SVG(TEXT("Icons/graph_node"), Icon16));
 					Set("MetasoundEditor.Graph.Node.Class.Input", new IMAGE_BRUSH_SVG(TEXT("Icons/input_node"), FVector2D(16.0f, 13.0f)));
 					Set("MetasoundEditor.Graph.Node.Class.Output", new IMAGE_BRUSH_SVG(TEXT("Icons/output_node"), FVector2D(16.0f, 13.0f)));
+					Set("MetasoundEditor.Graph.Node.Class.Reroute", new IMAGE_BRUSH_SVG(TEXT("Icons/reroute_node"), Icon16));
 					Set("MetasoundEditor.Graph.Node.Class.Variable", new IMAGE_BRUSH_SVG(TEXT("Icons/variable_node"), FVector2D(16.0f, 13.0f)));
 
 					Set("MetasoundEditor.Graph.Node.Math.Add", new FSlateImageBrush(RootToContentDir(TEXT("/Graph/node_math_add_40x.png")), Icon40x40));
@@ -136,7 +151,8 @@ namespace Metasound
 					Set("MetasoundEditor.Graph.Node.Math.Power", new FSlateImageBrush(RootToContentDir(TEXT("/Graph/node_math_power_40x.png")), Icon40x40));
 					Set("MetasoundEditor.Graph.Node.Math.Logarithm", new FSlateImageBrush(RootToContentDir(TEXT("/Graph/node_math_logarithm_40x.png")), Icon40x40));
 					Set("MetasoundEditor.Graph.Node.Conversion", new FSlateImageBrush(RootToContentDir(TEXT("/Graph/node_conversion_40x.png")), Icon40x40));
-					
+
+					Set("MetasoundEditor.Graph.InvalidReroute", new IMAGE_BRUSH_SVG(TEXT("Icons/invalid_reroute"), Icon16));
 					Set("MetasoundEditor.Graph.ConstructorPinArray", new FSlateImageBrush(RootToContentDir(TEXT("/Icons/array_pin_rotated.png")), Icon16));
 					Set("MetasoundEditor.Graph.ConstructorPinArrayDisconnected", new FSlateImageBrush(RootToContentDir(TEXT("/Icons/array_pin_rotated_disconnected.png")), Icon16));
 					Set("MetasoundEditor.Graph.ArrayPin", new FSlateImageBrush(RootToContentDir(TEXT("/Icons/array_pin.png")), Icon16));
@@ -168,9 +184,34 @@ namespace Metasound
 			}
 		};
 
-		class FMetasoundGraphPanelPinFactory : public FGraphPanelPinFactory
+		namespace Style
 		{
-		};
+			FSlateIcon CreateSlateIcon(FName InName)
+			{
+				return { "MetaSoundStyle", InName};
+			}
+
+			const FSlateBrush& GetSlateBrushSafe(FName InName)
+			{
+				const ISlateStyle* MetaSoundStyle = FSlateStyleRegistry::FindSlateStyle("MetaSoundStyle");
+				if (ensureMsgf(MetaSoundStyle, TEXT("Missing slate style 'MetaSoundStyle'")))
+				{
+					const FSlateBrush* Brush = MetaSoundStyle->GetBrush(InName);
+					if (ensureMsgf(Brush, TEXT("Missing brush '%s'"), *InName.ToString()))
+					{
+						return *Brush;
+					}
+				}
+
+				if (const FSlateBrush* NoBrush = FAppStyle::GetBrush("NoBrush"))
+				{
+					return *NoBrush;
+				}
+
+				static const FSlateBrush NullBrush;
+				return NullBrush;
+			}
+		}
 
 		class FModule : public IMetasoundEditorModule
 		{
@@ -616,20 +657,13 @@ namespace Metasound
 				Frontend::FDataTypeRegistryInfo Info;
 				Frontend::IDataTypeRegistry::Get().GetDataTypeInfo(InDataType, Info);
 
-				if (const ISlateStyle* MetasoundStyle = FSlateStyleRegistry::FindSlateStyle("MetaSoundStyle"))
+				if (Info.bIsArrayType)
 				{
-					if (Info.bIsArrayType)
-					{
-						return bIsConstructorType ? MetasoundStyle->GetBrush(TEXT("MetasoundEditor.Graph.ConstructorPinArray")) : MetasoundStyle->GetBrush(TEXT("MetasoundEditor.Graph.ArrayPin"));
-					}
-					else
-					{
-						return bIsConstructorType ? MetasoundStyle->GetBrush(TEXT("MetasoundEditor.Graph.ConstructorPin")) : FAppStyle::GetBrush("Icons.BulletPoint");
-					}
+					return bIsConstructorType ? &Style::GetSlateBrushSafe("MetasoundEditor.Graph.ConstructorPinArray") : &Style::GetSlateBrushSafe("MetasoundEditor.Graph.ArrayPin");
 				}
 				else
 				{
-					return Info.bIsArrayType ? FAppStyle::GetBrush("Graph.ArrayPin.Connected") : FAppStyle::GetBrush("Icons.BulletPoint");
+					return bIsConstructorType ? &Style::GetSlateBrushSafe("MetasoundEditor.Graph.ConstructorPin") : FAppStyle::GetBrush("Icons.BulletPoint");
 				}
 			}
 
diff --git a/Engine/Plugins/Runtime/Metasound/Source/MetasoundEditor/Private/SMetasoundGraphNode.cpp b/Engine/Plugins/Runtime/Metasound/Source/MetasoundEditor/Private/SMetasoundGraphNode.cpp
--- a/Engine/Plugins/Runtime/Metasound/Source/MetasoundEditor/Private/SMetasoundGraphNode.cpp
+++ b/Engine/Plugins/Runtime/Metasound/Source/MetasoundEditor/Private/SMetasoundGraphNode.cpp
@@ -21,9 +21,11 @@
 #include "MetasoundEditorGraphSchema.h"
 #include "MetasoundEditorModule.h"
 #include "MetasoundFrontendArchetypeRegistry.h"
+#include "MetasoundFrontendNodeTemplateRegistry.h"
 #include "MetasoundFrontendRegistries.h"
 #include "MetasoundTrigger.h"
 #include "NodeFactory.h"
+#include "NodeTemplates/MetasoundFrontendNodeTemplateReroute.h"
 #include "PropertyCustomizationHelpers.h"
 #include "SAudioRadialSlider.h"
 #include "SAudioSlider.h"
@@ -34,6 +36,7 @@
 #include "SLevelOfDetailBranchNode.h"
 #include "SMetasoundGraphEnumPin.h"
 #include "SMetasoundGraphPin.h"
+#include "SPinTypeSelector.h"
 #include "Styling/AppStyle.h"
 #include "Styling/SlateColor.h"
 #include "Styling/SlateStyleRegistry.h"
@@ -55,6 +58,22 @@ namespace Metasound
 {
 	namespace Editor
 	{
+		class SMetaSoundGraphPinKnot : public SGraphPinKnot
+		{
+		public:
+			SLATE_BEGIN_ARGS(SMetaSoundGraphPinKnot) {}
+			SLATE_END_ARGS()
+
+			void Construct(const FArguments& InArgs, UEdGraphPin* InPin);
+
+			virtual const FSlateBrush* GetPinIcon() const override;
+			virtual FSlateColor GetPinColor() const override;
+
+		protected:
+			bool HasRequiredConnections() const;
+		};
+
+
 		SMetaSoundGraphNode::~SMetaSoundGraphNode()
 		{
 			UMetasoundEditorGraphNode& Node = GetMetaSoundNode();
@@ -892,6 +911,139 @@ namespace Metasound
 					(bShowInputWidget ? OuterContentBox : ContentBox).ToSharedRef()
 				];
 		}
+
+		TSharedPtr<SGraphPin> SMetaSoundGraphNodeKnot::CreatePinWidget(UEdGraphPin* Pin) const
+		{
+			return SNew(SMetaSoundGraphPinKnot, Pin);
+		}
+
+		void SMetaSoundGraphNodeKnot::Construct(const FArguments& InArgs, class UEdGraphNode* InNode)
+		{
+			GraphNode = InNode;
+			SetCursor(EMouseCursor::CardinalCross);
+			UpdateGraphNode();
+		}
+
+		void SMetaSoundGraphNodeKnot::MoveTo(const FVector2D& NewPosition, FNodeSet& NodeFilter, bool bMarkDirty)
+		{
+			SGraphNode::MoveTo(NewPosition, NodeFilter, bMarkDirty);
+
+			UMetasoundEditorGraphNode& Node = GetMetaSoundNode();
+			Node.GetMetasoundChecked().Modify();
+			Node.SetNodeLocation(NewPosition);
+		}
+
+		UMetasoundEditorGraphNode& SMetaSoundGraphNodeKnot::GetMetaSoundNode()
+		{
+			return *CastChecked<UMetasoundEditorGraphNode>(GraphNode);
+		}
+
+		const UMetasoundEditorGraphNode& SMetaSoundGraphNodeKnot::GetMetaSoundNode() const
+		{
+			check(GraphNode);
+			return *Cast<UMetasoundEditorGraphNode>(GraphNode);
+		}
+
+		void SMetaSoundGraphPinKnot::Construct(const FArguments& InArgs, UEdGraphPin* InPin)
+		{
+			SGraphPinKnot::Construct(SGraphPinKnot::FArguments(), InPin);
+
+			TSharedRef<SWidget> PinWidgetRef = SPinTypeSelector::ConstructPinTypeImage(
+				MakeAttributeSP(this, &SMetaSoundGraphPinKnot::GetPinIcon),
+				MakeAttributeSP(this, &SMetaSoundGraphPinKnot::GetPinColor),
+				MakeAttributeSP(this, &SMetaSoundGraphPinKnot::GetSecondaryPinIcon),
+				MakeAttributeSP(this, &SMetaSoundGraphPinKnot::GetSecondaryPinColor));
+			PinImage = PinWidgetRef;
+		}
+
+		FSlateColor SMetaSoundGraphPinKnot::GetPinColor() const
+		{
+			if (UEdGraphPin* Pin = GetPinObj())
+			{
+				if (Pin->Direction == EGPD_Output)
+				{
+					if (!HasRequiredConnections())
+					{
+						return FLinearColor::Yellow;
+					}
+				}
+			}
+
+			return SGraphPin::GetPinColor();
+		}
+
+		const FSlateBrush* SMetaSoundGraphPinKnot::GetPinIcon() const
+		{
+			using namespace Metasound::Frontend;
+
+			bool bIsConstructorPin = false;
+
+			if (UEdGraphPin* Pin = GetPinObj())
+			{
+				if (!HasRequiredConnections())
+				{
+					if (const ISlateStyle* MetasoundStyle = FSlateStyleRegistry::FindSlateStyle("MetaSoundStyle"))
+					{
+						return MetasoundStyle->GetBrush(TEXT("MetasoundEditor.Graph.InvalidReroute"));
+					}
+				}
+
+				Pin = FGraphBuilder::FindReroutedOutputPin(Pin);
+				if (const UMetasoundEditorGraphNode* Node = Cast<UMetasoundEditorGraphNode>(Pin->GetOwningNode()))
+				{
+					if (const UMetasoundEditorGraphMemberNode* MemberNode = Cast<UMetasoundEditorGraphMemberNode>(Node))
+					{
+						if (const UMetasoundEditorGraphInput* Input = Cast<UMetasoundEditorGraphInput>(MemberNode->GetMember()))
+						{
+							FConstNodeHandle NodeHandle = Input->GetConstNodeHandle();
+							TArray<FConstOutputHandle> OutputHandles = NodeHandle->GetConstOutputs();
+
+							check(!OutputHandles.IsEmpty());
+							bIsConstructorPin = OutputHandles.Last()->GetVertexAccessType() == EMetasoundFrontendVertexAccessType::Value;
+						}
+					}
+				}
+			}
+
+			if (bIsConstructorPin)
+			{
+				if (const ISlateStyle* MetasoundStyle = FSlateStyleRegistry::FindSlateStyle("MetaSoundStyle"))
+				{
+					const bool bIsConnected = IsConnected();
+					if (IsArray())
+					{
+						return bIsConnected ? MetasoundStyle->GetBrush(TEXT("MetasoundEditor.Graph.ConstructorPinArray")) :
+							MetasoundStyle->GetBrush(TEXT("MetasoundEditor.Graph.ConstructorPinArrayDisconnected"));
+					}
+					else
+					{
+						return bIsConnected ? MetasoundStyle->GetBrush(TEXT("MetasoundEditor.Graph.ConstructorPin")) :
+							MetasoundStyle->GetBrush(TEXT("MetasoundEditor.Graph.ConstructorPinDisconnected"));
+					}
+				}
+			}
+
+			return SGraphPin::GetPinIcon();
+		}
+
+		bool SMetaSoundGraphPinKnot::HasRequiredConnections() const
+		{
+			using namespace Frontend;
+
+			if (UEdGraphPin* Pin = GetPinObj())
+			{
+				if (const UMetasoundEditorGraphExternalNode* OwningNode = Cast<UMetasoundEditorGraphExternalNode>(Pin->GetOwningNode()))
+				{
+					if (const INodeTemplate* Template = INodeTemplateRegistry::Get().FindTemplate(FRerouteNodeTemplate::GetRegistryKey()))
+					{
+						FConstNodeHandle NodeHandle = OwningNode->GetConstNodeHandle();
+						return Template->HasRequiredConnections(NodeHandle);
+					}
+				}
+			}
+
+			return false;
+		}
 	} // namespace Editor
 } // namespace Metasound
 #undef LOCTEXT_NAMESPACE // MetasoundGraphNode
diff --git a/Engine/Plugins/Runtime/Metasound/Source/MetasoundEditor/Private/SMetasoundGraphNode.h b/Engine/Plugins/Runtime/Metasound/Source/MetasoundEditor/Private/SMetasoundGraphNode.h
--- a/Engine/Plugins/Runtime/Metasound/Source/MetasoundEditor/Private/SMetasoundGraphNode.h
+++ b/Engine/Plugins/Runtime/Metasound/Source/MetasoundEditor/Private/SMetasoundGraphNode.h
@@ -6,6 +6,7 @@
 #include "Layout/Visibility.h"
 #include "Misc/Attribute.h"
 #include "SGraphNode.h"
+#include "SGraphNodeKnot.h"
 #include "Widgets/DeclarativeSyntaxSupport.h"
 #include "Widgets/SOverlay.h"
 
@@ -82,5 +83,23 @@ namespace Metasound
 			// Handle for on input slider range changed  
 			FDelegateHandle InputSliderOnRangeChangedDelegateHandle;
 		};
+
+		class SMetaSoundGraphNodeKnot : public SGraphNodeKnot
+		{
+		public:
+			SLATE_BEGIN_ARGS(SMetaSoundGraphNode)
+			{
+			}
+
+			SLATE_END_ARGS()
+
+			void Construct(const FArguments& InArgs, class UEdGraphNode* InNode);
+
+			virtual TSharedPtr<SGraphPin> CreatePinWidget(UEdGraphPin* Pin) const override;
+			virtual void MoveTo(const FVector2D& NewPosition, FNodeSet& NodeFilter, bool bMarkDirty) override;
+
+			UMetasoundEditorGraphNode& GetMetaSoundNode();
+			const UMetasoundEditorGraphNode& GetMetaSoundNode() const;
+		};
 	} // namespace Editor
 } // namespace Metasound
diff --git a/Engine/Plugins/Runtime/Metasound/Source/MetasoundEditor/Private/SMetasoundGraphPin.h b/Engine/Plugins/Runtime/Metasound/Source/MetasoundEditor/Private/SMetasoundGraphPin.h
--- a/Engine/Plugins/Runtime/Metasound/Source/MetasoundEditor/Private/SMetasoundGraphPin.h
+++ b/Engine/Plugins/Runtime/Metasound/Source/MetasoundEditor/Private/SMetasoundGraphPin.h
@@ -208,7 +208,7 @@ namespace Metasound
 					];
 			}
 
-			virtual const FSlateBrush* GetPinIcon() const
+			virtual const FSlateBrush* GetPinIcon() const override
 			{
 				const bool bIsConnected = ParentPinType::IsConnected();
 				EMetasoundFrontendVertexAccessType AccessType = EMetasoundFrontendVertexAccessType::Reference;
diff --git a/Engine/Plugins/Runtime/Metasound/Source/MetasoundEditor/Private/SMetasoundPinValueInspector.cpp b/Engine/Plugins/Runtime/Metasound/Source/MetasoundEditor/Private/SMetasoundPinValueInspector.cpp
--- a/Engine/Plugins/Runtime/Metasound/Source/MetasoundEditor/Private/SMetasoundPinValueInspector.cpp
+++ b/Engine/Plugins/Runtime/Metasound/Source/MetasoundEditor/Private/SMetasoundPinValueInspector.cpp
@@ -138,9 +138,9 @@ namespace Metasound
 			return PickerArgs;
 		}
 
-		const UMetasoundEditorGraphNode* FMetasoundNumericDebugLineItem::GetNode() const
+		const UMetasoundEditorGraphNode* FMetasoundNumericDebugLineItem::GetReroutedNode() const
 		{
-			if (GraphPinObj)
+			if (UEdGraphPin* ReroutedOutputPin = FGraphBuilder::FindReroutedOutputPin(GraphPinObj))
 			{
 				return Cast<UMetasoundEditorGraphNode>(GraphPinObj->GetOwningNode());
 			}
@@ -148,21 +148,21 @@ namespace Metasound
 			return nullptr;
 		}
 
-		const UMetasoundEditorGraphNode& FMetasoundNumericDebugLineItem::GetNodeChecked() const
+		const UMetasoundEditorGraphNode& FMetasoundNumericDebugLineItem::GetReroutedNodeChecked() const
 		{
-			const UMetasoundEditorGraphNode* Node = GetNode();
+			const UMetasoundEditorGraphNode* Node = GetReroutedNode();
 			check(Node);
 			return *Node;
 		}
 
 		Frontend::FConstOutputHandle FMetasoundNumericDebugLineItem::GetOutputHandle() const
 		{
-			return FGraphBuilder::GetOutputHandleFromPin(GraphPinObj);
+			return FGraphBuilder::FindReroutedConstOutputHandleFromPin(GraphPinObj);
 		}
 
 		Frontend::FOutputHandle FMetasoundNumericDebugLineItem::GetOutputHandle()
 		{
-			return FGraphBuilder::GetOutputHandleFromPin(GraphPinObj);
+			return FGraphBuilder::FindReroutedOutputHandleFromPin(GraphPinObj);
 		}
 
 		FLinearColor FMetasoundNumericDebugLineItem::GetEdgeStyleColorAtIndex(int32 InIndex) const
@@ -183,14 +183,15 @@ namespace Metasound
 
 			FScopedTransaction Transaction(PinValueInspectorPrivate::DisableValueColorizationText);
 
-			UObject* ParentObject = GetNodeChecked().GetOutermostObject();
+			const UMetasoundEditorGraphNode& Node = GetReroutedNodeChecked();
+			UObject* ParentObject = Node.GetOutermostObject();
 			check(ParentObject);
 			ParentObject->Modify();
 
 			FGraphHandle GraphHandle = GetGraphHandle();
 			FMetasoundFrontendGraphStyle Style = GraphHandle->GetGraphStyle();
 
-			const FGuid NodeID = GetNodeChecked().GetNodeID();
+			const FGuid NodeID = Node.GetNodeID();
 			const FName OutputName = GetOutputHandle()->GetName();
 			Style.EdgeStyles.RemoveAllSwap([&NodeID, &OutputName](FMetasoundFrontendEdgeStyle& EdgeStyle)
 				{
@@ -207,7 +208,8 @@ namespace Metasound
 
 			FScopedTransaction Transaction(PinValueInspectorPrivate::EnableValueColorizationText);
 
-			UObject* ParentObject = GetNodeChecked().GetOutermostObject();
+			const UMetasoundEditorGraphNode& Node = GetReroutedNodeChecked();
+			UObject* ParentObject = Node.GetOutermostObject();
 			check(ParentObject);
 			ParentObject->Modify();
 
@@ -221,7 +223,7 @@ namespace Metasound
 			FMetasoundFrontendGraphStyle Style = GraphHandle->GetGraphStyle();
 			FConstOutputHandle OutputHandle = GetOutputHandle();
 
-			const FGuid NodeID = GetNodeChecked().GetNodeID();
+			const FGuid NodeID = Node.GetNodeID();
 			const FName OutputName = OutputHandle->GetName();
 			FMetasoundFrontendEdgeStyle* EdgeStyle = Style.EdgeStyles.FindByPredicate([&NodeID, &OutputName](FMetasoundFrontendEdgeStyle& EdgeStyle)
 			{
@@ -249,14 +251,15 @@ namespace Metasound
 			const FName Name = GetOutputHandle()->GetName();
 			FScopedTransaction Transaction(FText::Format(SetValueFormat, FText::FromName(Name), FText::AsNumber(InIndex), FText::FromString(InColor.ToString())));
 
-			UObject* ParentObject = GetNodeChecked().GetOutermostObject();
+			const UMetasoundEditorGraphNode& Node = GetReroutedNodeChecked();
+			UObject* ParentObject = Node.GetOutermostObject();
 			check(ParentObject);
 			ParentObject->Modify();
 
 			FGraphHandle GraphHandle = GetGraphHandle();
 			FMetasoundFrontendGraphStyle Style = GraphHandle->GetGraphStyle();
 
-			const FGuid NodeID = GetNodeChecked().GetNodeID();
+			const FGuid NodeID = Node.GetNodeID();
 			const FName OutputName = GetOutputHandle()->GetName();
 			FMetasoundFrontendEdgeStyle* EdgeStyle = Style.EdgeStyles.FindByPredicate([&NodeID, &OutputName](FMetasoundFrontendEdgeStyle& EdgeStyle)
 			{
@@ -278,7 +281,7 @@ namespace Metasound
 		{
 			using namespace Frontend;
 
-			if (const UMetasoundEditorGraphNode* Node = GetNode())
+			if (const UMetasoundEditorGraphNode* Node = GetReroutedNode())
 			{
 				FOutputHandle OutputHandle = GetOutputHandle();
 
@@ -330,7 +333,8 @@ namespace Metasound
 		{
 			using namespace Frontend;
 
-			TSharedPtr<FEditor> Editor = FGraphBuilder::GetEditorForNode(GetNodeChecked());
+			const UMetasoundEditorGraphNode& Node = GetReroutedNodeChecked();
+			TSharedPtr<FEditor> Editor = FGraphBuilder::GetEditorForNode(Node);
 			if (Editor.IsValid())
 			{
 				return &Editor->GetConnectionManager();
@@ -349,7 +353,8 @@ namespace Metasound
 				return nullptr;
 			}
 
-			const FGuid NodeID = GetNodeChecked().GetNodeID();
+			const UMetasoundEditorGraphNode& Node = GetReroutedNodeChecked();
+			const FGuid NodeID = Node.GetNodeID();
 			const FName OutputName = GetOutputHandle()->GetName();
 			const FMetasoundFrontendGraphStyle& Style = OwningGraph->GetGraphStyle();
 
@@ -366,7 +371,9 @@ namespace Metasound
 			bIsValueColorizationEnabled = false;
 
 			const FMetasoundFrontendEdgeStyle* EdgeStyle = GetEdgeStyle();
-			FGuid NodeID = GetNodeChecked().GetNodeID();
+
+ 			const UMetasoundEditorGraphNode& Node = GetReroutedNodeChecked();
+			FGuid NodeID = Node.GetNodeID();
 			FName OutputName = GetOutputHandle()->GetName();
 			if (EdgeStyle)
 			{
diff --git a/Engine/Plugins/Runtime/Metasound/Source/MetasoundEditor/Private/SMetasoundPinValueInspector.h b/Engine/Plugins/Runtime/Metasound/Source/MetasoundEditor/Private/SMetasoundPinValueInspector.h
--- a/Engine/Plugins/Runtime/Metasound/Source/MetasoundEditor/Private/SMetasoundPinValueInspector.h
+++ b/Engine/Plugins/Runtime/Metasound/Source/MetasoundEditor/Private/SMetasoundPinValueInspector.h
@@ -53,9 +53,9 @@ namespace Metasound
 			FMetasoundNumericDebugLineItem(UEdGraphPin* InGraphPinObj, FGetValueStringFunction&& InGetValueStringFunction);
 			virtual ~FMetasoundNumericDebugLineItem() = default;
 
-			const UMetasoundEditorGraphNode* GetNode() const;
+			const UMetasoundEditorGraphNode* GetReroutedNode() const;
 
-			const UMetasoundEditorGraphNode& GetNodeChecked() const;
+			const UMetasoundEditorGraphNode& GetReroutedNodeChecked() const;
 
 			Frontend::FConstOutputHandle GetOutputHandle() const;
 
diff --git a/Engine/Plugins/Runtime/Metasound/Source/MetasoundEditor/Public/MetasoundEditorGraph.h b/Engine/Plugins/Runtime/Metasound/Source/MetasoundEditor/Public/MetasoundEditorGraph.h
--- a/Engine/Plugins/Runtime/Metasound/Source/MetasoundEditor/Public/MetasoundEditorGraph.h
+++ b/Engine/Plugins/Runtime/Metasound/Source/MetasoundEditor/Public/MetasoundEditorGraph.h
@@ -410,7 +410,7 @@ private:
 	bool RemoveFrontendInput(UMetasoundEditorGraphInput& Input);
 	bool RemoveFrontendOutput(UMetasoundEditorGraphOutput& Output);
 	bool RemoveFrontendVariable(UMetasoundEditorGraphVariable& Variable);
-	bool ValidateInternal(Metasound::Editor::FGraphValidationResults& OutResults);
+	void ValidateInternal(Metasound::Editor::FGraphValidationResults& OutResults);
 
 	// Preview ID is the Unique ID provided by the UObject that implements
 	// a sound's ParameterInterface when a sound begins playing.
diff --git a/Engine/Plugins/Runtime/Metasound/Source/MetasoundEditor/Public/MetasoundEditorGraphInputNode.h b/Engine/Plugins/Runtime/Metasound/Source/MetasoundEditor/Public/MetasoundEditorGraphInputNode.h
--- a/Engine/Plugins/Runtime/Metasound/Source/MetasoundEditor/Public/MetasoundEditorGraphInputNode.h
+++ b/Engine/Plugins/Runtime/Metasound/Source/MetasoundEditor/Public/MetasoundEditorGraphInputNode.h
@@ -34,7 +34,7 @@ public:
 	virtual FGuid GetNodeID() const override;
 	virtual FLinearColor GetNodeTitleColor() const override;
 	virtual FSlateIcon GetNodeTitleIcon() const override;
-	virtual bool Validate(Metasound::Editor::FGraphNodeValidationResult& OutResult) override;
+	virtual void Validate(Metasound::Editor::FGraphNodeValidationResult& OutResult) override;
 
 protected:
 	virtual void SetNodeID(FGuid InNodeID) override;
diff --git a/Engine/Plugins/Runtime/Metasound/Source/MetasoundEditor/Public/MetasoundEditorGraphNode.h b/Engine/Plugins/Runtime/Metasound/Source/MetasoundEditor/Public/MetasoundEditorGraphNode.h
--- a/Engine/Plugins/Runtime/Metasound/Source/MetasoundEditor/Public/MetasoundEditorGraphNode.h
+++ b/Engine/Plugins/Runtime/Metasound/Source/MetasoundEditor/Public/MetasoundEditorGraphNode.h
@@ -31,11 +31,6 @@ namespace Metasound
 
 		// Map of class names to sorted array of registered version numbers
 		using FSortedClassVersionMap = TMap<FName, TArray<FMetasoundFrontendVersionNumber>>;
-
-		namespace GraphNode
-		{
-			void SetMessage(UEdGraphNode& InNode, EMessageSeverity::Type InSeverity, const FString& InMessage);
-		}
 	} // namespace Editor
 } // namespace Metasound
 
@@ -106,7 +101,7 @@ public:
 	virtual FGuid GetNodeID() const { return FGuid(); }
 	virtual FText GetDisplayName() const;
 	virtual void CacheTitle();
-	virtual bool Validate(Metasound::Editor::FGraphNodeValidationResult& OutResult);
+	virtual void Validate(Metasound::Editor::FGraphNodeValidationResult& OutResult);
 
 	// Mark node for refresh
 	void SyncChangeIDs();
@@ -152,10 +147,6 @@ public:
 	{
 		return true;
 	}
-
-protected:
-	// Utility to construct a new validaton result for child classes
-	Metasound::Editor::FGraphNodeValidationResult CreateNewValidationResult();
 };
 
 /** Node that represents a graph output */
@@ -183,7 +174,7 @@ public:
 	// Disables interact widgets (ex. sliders, knobs) when input is connected
 	virtual bool EnableInteractWidgets() const override;
 
-	virtual bool Validate(Metasound::Editor::FGraphNodeValidationResult& OutResult) override;
+	virtual void Validate(Metasound::Editor::FGraphNodeValidationResult& OutResult) override;
 
 protected:
 	virtual FLinearColor GetNodeTitleColor() const override;
@@ -216,15 +207,17 @@ public:
 	virtual FGuid GetNodeID() const override { return NodeID; }
 	virtual FLinearColor GetNodeTitleColor() const override;
 	virtual FSlateIcon GetNodeTitleIcon() const override;
+	virtual bool ShouldDrawNodeAsControlPointOnly(int32& OutInputPinIndex, int32& OutOutputPinIndex) const override;
 
 	virtual void ReconstructNode() override;
 	virtual void CacheTitle() override;
+	virtual void GetPinHoverText(const UEdGraphPin& Pin, FString& OutHoverText) const override;
 
 	FMetasoundFrontendVersionNumber FindHighestVersionInRegistry() const;
 	bool CanAutoUpdate() const;
 
 	// Validates node and returns whether or not the node is valid.
-	virtual bool Validate(Metasound::Editor::FGraphNodeValidationResult& OutResult) override;
+	virtual void Validate(Metasound::Editor::FGraphNodeValidationResult& OutResult) override;
 
 
 protected:
diff --git a/Engine/Plugins/Runtime/Metasound/Source/MetasoundEditor/Public/MetasoundEditorGraphSchema.h b/Engine/Plugins/Runtime/Metasound/Source/MetasoundEditor/Public/MetasoundEditorGraphSchema.h
--- a/Engine/Plugins/Runtime/Metasound/Source/MetasoundEditor/Public/MetasoundEditorGraphSchema.h
+++ b/Engine/Plugins/Runtime/Metasound/Source/MetasoundEditor/Public/MetasoundEditorGraphSchema.h
@@ -68,11 +68,13 @@ struct METASOUNDEDITOR_API FMetasoundGraphSchemaAction : public FEdGraphSchemaAc
 			MoveTemp(InNodeCategory),
 			MoveTemp(InMenuDesc),
 			MoveTemp(InToolTip),
-			/*InNodeCategory.IsEmpty() ? */static_cast<int32>(InGroup)/* : 0*/,
+			static_cast<int32>(InGroup),
 			MoveTemp(InKeywords))
 	{
 	}
 
+	virtual ~FMetasoundGraphSchemaAction() = default;
+
 	virtual const FSlateBrush* GetIconBrush() const
 	{
 		return FAppStyle::GetBrush("NoBrush");
@@ -99,6 +101,8 @@ struct METASOUNDEDITOR_API FMetasoundGraphSchemaAction_NodeWithMultipleOutputs :
 		: FMetasoundGraphSchemaAction(InNodeCategory, InMenuDesc, InToolTip, InGroup, InKeywords)
 	{}
 
+	virtual ~FMetasoundGraphSchemaAction_NodeWithMultipleOutputs() = default;
+
 	//~ Begin FEdGraphSchemaAction Interface
 	virtual UEdGraphNode* PerformAction(UEdGraph* ParentGraph, TArray<UEdGraphPin*>& FromPins, const FVector2D Location, bool bSelectNewNode = true) override;
 	virtual UEdGraphNode* PerformAction(UEdGraph* ParentGraph, UEdGraphPin* FromPin, const FVector2D Location, bool bSelectNewNode = true) { return nullptr; }
@@ -120,6 +124,8 @@ struct METASOUNDEDITOR_API FMetasoundGraphSchemaAction_NewInput : public FMetaso
 
 	FMetasoundGraphSchemaAction_NewInput(FText InNodeCategory, FText InDisplayName, FGuid InInputNodeID, FText InToolTip, Metasound::Editor::EPrimaryContextGroup InGroup);
 
+	virtual ~FMetasoundGraphSchemaAction_NewInput() = default;
+
 	//~ Begin FMetasoundGraphSchemaAction Interface
 	virtual const FSlateBrush* GetIconBrush() const override;
 
@@ -139,6 +145,8 @@ struct METASOUNDEDITOR_API FMetasoundGraphSchemaAction_PromoteToInput : public F
 
 	FMetasoundGraphSchemaAction_PromoteToInput();
 
+	virtual ~FMetasoundGraphSchemaAction_PromoteToInput() = default;
+
 	//~ Begin FEdGraphSchemaAction Interface
 	virtual UEdGraphNode* PerformAction(UEdGraph* ParentGraph, UEdGraphPin* FromPin, const FVector2D Location, bool bSelectNewNode = true) override;
 	//~ End FEdGraphSchemaAction Interface
@@ -159,6 +167,8 @@ struct METASOUNDEDITOR_API FMetasoundGraphSchemaAction_NewOutput : public FMetas
 
 	FMetasoundGraphSchemaAction_NewOutput(FText InNodeCategory, FText InDisplayName, FGuid InOutputNodeID, FText InToolTip, Metasound::Editor::EPrimaryContextGroup InGroup);
 
+	virtual ~FMetasoundGraphSchemaAction_NewOutput() = default;
+
 	//~ Begin FMetasoundGraphSchemaAction Interface
 	virtual const FSlateBrush* GetIconBrush() const override;
 
@@ -178,6 +188,8 @@ struct METASOUNDEDITOR_API FMetasoundGraphSchemaAction_PromoteToOutput : public
 
 	FMetasoundGraphSchemaAction_PromoteToOutput();
 
+	virtual ~FMetasoundGraphSchemaAction_PromoteToOutput() = default;
+
 	//~ Begin FEdGraphSchemaAction Interface
 	virtual UEdGraphNode* PerformAction(UEdGraph* ParentGraph, UEdGraphPin* FromPin, const FVector2D Location, bool bSelectNewNode = true) override;
 	//~ End FEdGraphSchemaAction Interface
@@ -198,6 +210,8 @@ struct METASOUNDEDITOR_API FMetasoundGraphSchemaAction_NewVariableNode : public
 
 	FMetasoundGraphSchemaAction_NewVariableNode(FText InNodeCategory, FText InDisplayName, FGuid InVariableID, FText InToolTip);
 
+	virtual ~FMetasoundGraphSchemaAction_NewVariableNode() = default;
+
 	//~ Begin FEdGraphSchemaAction Interface
 	virtual UEdGraphNode* PerformAction(UEdGraph* ParentGraph, UEdGraphPin* FromPin, const FVector2D Location, bool bSelectNewNode = true) override;
 	//~ End FEdGraphSchemaAction Interface
@@ -228,6 +242,8 @@ struct METASOUNDEDITOR_API FMetasoundGraphSchemaAction_NewVariableAccessorNode :
 
 	FMetasoundGraphSchemaAction_NewVariableAccessorNode(FText InNodeCategory, FText InDisplayName, FGuid InVariableID, FText InToolTip);
 
+	virtual ~FMetasoundGraphSchemaAction_NewVariableAccessorNode() = default;
+
 protected:
 	virtual Metasound::Frontend::FNodeHandle CreateFrontendVariableNode(const Metasound::Frontend::FGraphHandle& InFrontendGraph, const FGuid& InVariableID) const override;
 
@@ -243,6 +259,8 @@ struct METASOUNDEDITOR_API FMetasoundGraphSchemaAction_NewVariableDeferredAccess
 		: FMetasoundGraphSchemaAction_NewVariableNode()
 	{}
 
+	virtual ~FMetasoundGraphSchemaAction_NewVariableDeferredAccessorNode() = default;
+
 	FMetasoundGraphSchemaAction_NewVariableDeferredAccessorNode(FText InNodeCategory, FText InDisplayName, FGuid InVariableID, FText InToolTip);
 	
 protected:
@@ -259,6 +277,8 @@ struct METASOUNDEDITOR_API FMetasoundGraphSchemaAction_NewVariableMutatorNode :
 		: FMetasoundGraphSchemaAction_NewVariableNode()
 	{}
 
+	virtual ~FMetasoundGraphSchemaAction_NewVariableMutatorNode() = default;
+
 	FMetasoundGraphSchemaAction_NewVariableMutatorNode(FText InNodeCategory, FText InDisplayName, FGuid InVariableID, FText InToolTip);
 	
 protected:
@@ -273,6 +293,8 @@ struct METASOUNDEDITOR_API FMetasoundGraphSchemaAction_PromoteToVariable_Accesso
 
 	FMetasoundGraphSchemaAction_PromoteToVariable_AccessorNode();
 
+	virtual ~FMetasoundGraphSchemaAction_PromoteToVariable_AccessorNode() = default;
+
 	//~ Begin FEdGraphSchemaAction Interface
 	virtual UEdGraphNode* PerformAction(UEdGraph* ParentGraph, UEdGraphPin* FromPin, const FVector2D Location, bool bSelectNewNode = true) override;
 	//~ End FEdGraphSchemaAction Interface
@@ -286,6 +308,8 @@ struct METASOUNDEDITOR_API FMetasoundGraphSchemaAction_PromoteToVariable_Mutator
 
 	FMetasoundGraphSchemaAction_PromoteToVariable_MutatorNode();
 
+	virtual ~FMetasoundGraphSchemaAction_PromoteToVariable_MutatorNode() = default;
+
 	//~ Begin FEdGraphSchemaAction Interface
 	virtual UEdGraphNode* PerformAction(UEdGraph* ParentGraph, UEdGraphPin* FromPin, const FVector2D Location, bool bSelectNewNode = true) override;
 	//~ End FEdGraphSchemaAction Interface
@@ -299,6 +323,8 @@ struct METASOUNDEDITOR_API FMetasoundGraphSchemaAction_PromoteToVariable_Deferre
 
 	FMetasoundGraphSchemaAction_PromoteToVariable_DeferredAccessorNode();
 
+	virtual ~FMetasoundGraphSchemaAction_PromoteToVariable_DeferredAccessorNode() = default;
+
 	//~ Begin FEdGraphSchemaAction Interface
 	virtual UEdGraphNode* PerformAction(UEdGraph* ParentGraph, UEdGraphPin* FromPin, const FVector2D Location, bool bSelectNewNode = true) override;
 	//~ End FEdGraphSchemaAction Interface
@@ -322,6 +348,8 @@ struct METASOUNDEDITOR_API FMetasoundGraphSchemaAction_NewNode : public FMetasou
 		: FMetasoundGraphSchemaAction_NodeWithMultipleOutputs(InNodeCategory, InMenuDesc, InToolTip, InGroup, InKeywords)
 	{}
 
+	virtual ~FMetasoundGraphSchemaAction_NewNode() = default;
+
 	//~ Begin FMetasoundGraphSchemaAction Interface
 	virtual const FSlateBrush* GetIconBrush() const override;
 
@@ -347,9 +375,37 @@ struct METASOUNDEDITOR_API FMetasoundGraphSchemaAction_NewFromSelected : public
 		: FMetasoundGraphSchemaAction_NewNode(MoveTemp(InNodeCategory), MoveTemp(InMenuDesc), MoveTemp(InToolTip), InGroup)
 	{}
 
+	virtual ~FMetasoundGraphSchemaAction_NewFromSelected() = default;
+
+	//~ Begin FEdGraphSchemaAction Interface
+	virtual UEdGraphNode* PerformAction(UEdGraph* ParentGraph, UEdGraphPin* FromPin, const FVector2D Location, bool bSelectNewNode = true) override;
+	//~ End FEdGraphSchemaAction Interface
+};
+
+/** Action to create new reroute node */
+USTRUCT()
+struct METASOUNDEDITOR_API FMetasoundGraphSchemaAction_NewReroute : public FMetasoundGraphSchemaAction
+{
+	GENERATED_USTRUCT_BODY();
+
+	FMetasoundGraphSchemaAction_NewReroute() = default;
+
+	FMetasoundGraphSchemaAction_NewReroute(const FLinearColor& InIconColor);
+
+	virtual ~FMetasoundGraphSchemaAction_NewReroute() = default;
+
+	//~ Begin FMetasoundGraphSchemaAction Interface
+	virtual const FSlateBrush* GetIconBrush() const override;
+
+	virtual const FLinearColor& GetIconColor() const override;
+	//~ End FMetasoundGraphSchemaAction Interface
+
 	//~ Begin FEdGraphSchemaAction Interface
 	virtual UEdGraphNode* PerformAction(UEdGraph* ParentGraph, UEdGraphPin* FromPin, const FVector2D Location, bool bSelectNewNode = true) override;
 	//~ End FEdGraphSchemaAction Interface
+
+private:
+	FLinearColor IconColor;
 };
 
 /** Action to create new comment */
@@ -429,6 +485,7 @@ public:
 	virtual void DroppedAssetsOnPin(const TArray<FAssetData>& Assets, const FVector2D& GraphPosition, UEdGraphPin* Pin) const override;
 	virtual int32 GetNodeSelectionCount(const UEdGraph* Graph) const override;
 	virtual TSharedPtr<FEdGraphSchemaAction> GetCreateCommentAction() const override;
+	virtual void OnPinConnectionDoubleCicked(UEdGraphPin* PinA, UEdGraphPin* PinB, const FVector2D& GraphPosition) const override;
 	virtual void SetNodePosition(UEdGraphNode* Node, const FVector2D& Position) const;
 	//~ End EdGraphSchema Interface
 
@@ -440,6 +497,8 @@ private:
 	void GetFunctionActions(FGraphActionMenuBuilder& ActionMenuBuilder, Metasound::Editor::FActionClassFilters InFilters = Metasound::Editor::FActionClassFilters(), bool bShowSelectedActions = true, Metasound::Frontend::FConstGraphHandle InGraphHandle = Metasound::Frontend::IGraphController::GetInvalidHandle()) const;
 	void GetVariableActions(FGraphActionMenuBuilder& ActionMenuBuilder, Metasound::Editor::FActionClassFilters InFilters = Metasound::Editor::FActionClassFilters(), bool bShowSelectedActions = true, Metasound::Frontend::FConstGraphHandle InGraphHandle = Metasound::Frontend::IGraphController::GetInvalidHandle()) const;
 
+	void BreakSinglePinLink(UEdGraphPin* SourcePin, UEdGraphPin* TargetPin, bool bShouldTransact) const;
+
 	void GetDataTypeInputNodeActions(FGraphContextMenuBuilder& InMenuBuilder, Metasound::Frontend::FConstGraphHandle InGraphHandle, Metasound::Editor::FInterfaceNodeFilterFunction InFilter = Metasound::Editor::FInterfaceNodeFilterFunction(), bool bShowSelectedActions = true) const;
 	void GetDataTypeOutputNodeActions(FGraphContextMenuBuilder& InMenuBuilder, Metasound::Frontend::FConstGraphHandle InGraphHandle, Metasound::Editor::FInterfaceNodeFilterFunction InFilter = Metasound::Editor::FInterfaceNodeFilterFunction(), bool bShowSelectedActions = true) const;
 
diff --git a/Engine/Plugins/Runtime/Metasound/Source/MetasoundEditor/Public/MetasoundEditorGraphValidation.h b/Engine/Plugins/Runtime/Metasound/Source/MetasoundEditor/Public/MetasoundEditorGraphValidation.h
--- a/Engine/Plugins/Runtime/Metasound/Source/MetasoundEditor/Public/MetasoundEditorGraphValidation.h
+++ b/Engine/Plugins/Runtime/Metasound/Source/MetasoundEditor/Public/MetasoundEditorGraphValidation.h
@@ -1,32 +1,43 @@
 // Copyright Epic Games, Inc. All Rights Reserved.
 #pragma once
 
-#include "Algo/NoneOf.h"
-#include "MetasoundEditorGraph.h"
+#include "Containers/Array.h"
+#include "Logging/TokenizedMessage.h"
 #include "MetasoundEditorGraphNode.h"
+#include "MetasoundEditorGraph.h"
 
 class UMetasoundEditorGraphNode;
 
+
 namespace Metasound
 {
 	namespace Editor
 	{
 		struct METASOUNDEDITOR_API FGraphNodeValidationResult
 		{
+			FGraphNodeValidationResult(UMetasoundEditorGraphNode& InNode);
+
+			// Whether or not validation operations on result have
+			// dirtied the associated Node (Validation doesn't mark
+			// node as dirtied to avoid resave being required at
+			// asset level)
+			bool GetHasDirtiedNode() const;
+
+			UMetasoundEditorGraphNode& GetNodeChecked() const;
+
+			// Whether associated node is in invalid state, i.e.
+			// may fail to build or may result in undefined behavior.
+			bool GetIsInvalid() const;
+
+			void SetMessage(EMessageSeverity::Type InSeverity, const FString& InMessage);
+			void SetPinOrphaned(UEdGraphPin& Pin, bool bIsOrphaned);
+
+		private:
 			// Node associated with validation result
 			UMetasoundEditorGraphNode* Node = nullptr;
 
-			// Whether associated node is in invalid state (document is corrupt and no Frontend representation could be
-			// found for the node)
-			bool bIsInvalid = false;
-
 			// Whether validation made changes to the node and is now in a dirty state
-			bool bIsDirty = false;
-
-			FGraphNodeValidationResult(UMetasoundEditorGraphNode& InNode)
-				: Node(&InNode)
-			{
-			}
+			bool bHasDirtiedNode = false;
 		};
 
 		struct METASOUNDEDITOR_API FGraphValidationResults
@@ -34,19 +45,10 @@ namespace Metasound
 			TArray<FGraphNodeValidationResult> NodeResults;
 
 			// Results corresponding with node validation
-			const TArray<FGraphNodeValidationResult>& GetResults() const
-			{
-				return NodeResults;
-			}
-
-			// Returns whether or not the graph is in a valid state
-			bool IsValid() const
-			{
-				return Algo::NoneOf(NodeResults, [](const FGraphNodeValidationResult& Result)
-				{
-					return Result.bIsInvalid;
-				});
-			}
+			const TArray<FGraphNodeValidationResult>& GetResults() const;
+
+			// Returns highest message severity of validated nodes
+			EMessageSeverity::Type GetHighestMessageSeverity() const;
 		};
 	} // namespace Editor
 } // namespace Metasound
diff --git a/Engine/Plugins/Runtime/Metasound/Source/MetasoundEditor/Public/MetasoundEditorModule.h b/Engine/Plugins/Runtime/Metasound/Source/MetasoundEditor/Public/MetasoundEditorModule.h
--- a/Engine/Plugins/Runtime/Metasound/Source/MetasoundEditor/Public/MetasoundEditorModule.h
+++ b/Engine/Plugins/Runtime/Metasound/Source/MetasoundEditor/Public/MetasoundEditorModule.h
@@ -25,6 +25,12 @@ namespace Metasound
 {
 	namespace Editor
 	{
+		namespace Style
+		{
+			METASOUNDEDITOR_API FSlateIcon CreateSlateIcon(FName InName);
+			METASOUNDEDITOR_API const FSlateBrush& GetSlateBrushSafe(FName InName);
+		} // namespace Style
+
 		// Status of initial asset scan when editor loads up.
 		enum class EAssetScanStatus : uint8
 		{
diff --git a/Engine/Plugins/Runtime/Metasound/Source/MetasoundFrontend/Private/MetasoundFrontend.cpp b/Engine/Plugins/Runtime/Metasound/Source/MetasoundFrontend/Private/MetasoundFrontend.cpp
--- a/Engine/Plugins/Runtime/Metasound/Source/MetasoundFrontend/Private/MetasoundFrontend.cpp
+++ b/Engine/Plugins/Runtime/Metasound/Source/MetasoundFrontend/Private/MetasoundFrontend.cpp
@@ -9,6 +9,8 @@
 #include "MetasoundDataTypeRegistrationMacro.h"
 #include "MetasoundFrontendDocument.h"
 #include "MetasoundFrontendGraph.h"
+#include "MetasoundFrontendNodeRegistryPrivate.h"
+#include "MetasoundFrontendNodeTemplateRegistry.h"
 #include "MetasoundFrontendRegistries.h"
 #include "MetasoundFrontendRegistryTransaction.h"
 #include "MetasoundJsonBackend.h"
@@ -17,9 +19,12 @@
 #include "MetasoundRouter.h"
 #include "MetasoundVertex.h"
 #include "Modules/ModuleManager.h"
+#include "NodeTemplates/MetasoundFrontendNodeTemplateReroute.h"
 #include "StructDeserializer.h"
 #include "StructSerializer.h"
 #include "Serialization/MemoryReader.h"
+#include "Templates/UniquePtr.h"
+
 
 namespace Metasound
 {
@@ -88,12 +93,24 @@ class FMetasoundFrontendModule : public IModuleInterface
 {
 	virtual void StartupModule() override
 	{
+		using namespace Metasound::Frontend;
+
+		TUniquePtr<INodeTemplate> RerouteTemplate = MakeUnique<FRerouteNodeTemplate>();
+		RegisterNodeTemplate(MoveTemp(RerouteTemplate));
+
 		FMetasoundFrontendRegistryContainer* Registry = FMetasoundFrontendRegistryContainer::Get();
 		if (ensure(nullptr != Registry))
 		{
 			Registry->RegisterPendingNodes();
 		}
 	}
+
+	virtual void ShutdownModule() override
+	{
+		using namespace Metasound::Frontend;
+
+		UnregisterNodeTemplate(FRerouteNodeTemplate::Version);
+	}
 };
 
 IMPLEMENT_MODULE(FMetasoundFrontendModule, MetasoundFrontend);
diff --git a/Engine/Plugins/Runtime/Metasound/Source/MetasoundFrontend/Private/MetasoundFrontendNodeRegistryPrivate.h b/Engine/Plugins/Runtime/Metasound/Source/MetasoundFrontend/Private/MetasoundFrontendNodeRegistryPrivate.h
--- a/Engine/Plugins/Runtime/Metasound/Source/MetasoundFrontend/Private/MetasoundFrontendNodeRegistryPrivate.h
+++ b/Engine/Plugins/Runtime/Metasound/Source/MetasoundFrontend/Private/MetasoundFrontendNodeRegistryPrivate.h
@@ -15,13 +15,13 @@ namespace Metasound
 		using FNodeRegistryTransactionBuffer = TTransactionBuffer<FNodeRegistryTransaction>;
 		using FNodeRegistryTransactionStream = TTransactionStream<FNodeRegistryTransaction>; 
 
-		/** INodeTemplateRegistryEntry declares the interface for a node registry entry.
+		/** INodeRegistryTemplateEntry declares the interface for a node registry entry.
 		 * Each node class in the registry must satisfy this interface.
 		 */
-		class INodeTemplateRegistryEntry
+		class INodeRegistryTemplateEntry
 		{
 		public:
-			virtual ~INodeTemplateRegistryEntry() = default;
+			virtual ~INodeRegistryTemplateEntry() = default;
 
 			/** Return FNodeClassInfo for the node class.
 			 *
@@ -106,7 +106,7 @@ namespace Metasound
 			virtual TArray<FConverterNodeInfo> GetPossibleConverterNodes(const FName& FromDataType, const FName& ToDataType) override;
 
 			// Private implementation until hardened and used for template nodes other than reroutes.
-			FNodeRegistryKey RegisterNodeTemplate(TUniquePtr<Metasound::Frontend::INodeTemplateRegistryEntry>&& InEntry);
+			FNodeRegistryKey RegisterNodeTemplate(TUniquePtr<Metasound::Frontend::INodeRegistryTemplateEntry>&& InEntry);
 			bool UnregisterNodeTemplate(const FNodeRegistryKey& InKey);
 
 			// Create a transaction stream for any newly transactions
@@ -117,7 +117,7 @@ namespace Metasound
 
 			const INodeRegistryEntry* FindNodeEntry(const FNodeRegistryKey& InKey) const;
 
-			const INodeTemplateRegistryEntry* FindNodeTemplateEntry(const FNodeRegistryKey& InKey) const;
+			const INodeRegistryTemplateEntry* FindNodeTemplateEntry(const FNodeRegistryKey& InKey) const;
 
 			// This buffer is used to enqueue nodes and datatypes to register when the module has been initialized,
 			// in order to avoid bad behavior with ensures, logs, etc. on static initialization.
@@ -132,7 +132,7 @@ namespace Metasound
 			TMap<FNodeRegistryKey, TSharedRef<INodeRegistryEntry, ESPMode::ThreadSafe>> RegisteredNodes;
 
 			// Registry in which we keep all information about dynamically-generated templated nodes via in C++.
-			TMap<FNodeRegistryKey, TSharedRef<INodeTemplateRegistryEntry, ESPMode::ThreadSafe>> RegisteredNodeTemplates;
+			TMap<FNodeRegistryKey, TSharedRef<INodeRegistryTemplateEntry, ESPMode::ThreadSafe>> RegisteredNodeTemplates;
 
 			// Registry in which we keep lists of possible nodes to use to convert between two datatypes
 			TMap<FConverterNodeRegistryKey, FConverterNodeRegistryValue> ConverterNodeRegistry;
diff --git a/Engine/Plugins/Runtime/Metasound/Source/MetasoundFrontend/Private/MetasoundFrontendNodeTemplateRegistry.cpp b/Engine/Plugins/Runtime/Metasound/Source/MetasoundFrontend/Private/MetasoundFrontendNodeTemplateRegistry.cpp
--- a/Engine/Plugins/Runtime/Metasound/Source/MetasoundFrontend/Private/MetasoundFrontendNodeTemplateRegistry.cpp
+++ b/Engine/Plugins/Runtime/Metasound/Source/MetasoundFrontend/Private/MetasoundFrontendNodeTemplateRegistry.cpp
@@ -3,6 +3,7 @@
 #include "MetasoundFrontendNodeTemplateRegistry.h"
 
 #include "MetasoundFrontendDocument.h"
+#include "MetasoundFrontendNodeRegistryPrivate.h"
 #include "MetasoundFrontendRegistries.h"
 #include "MetasoundFrontendRegistryTransaction.h"
 #include "MetasoundLog.h"
@@ -24,7 +25,7 @@ namespace Metasound
 			virtual const INodeTemplate* FindTemplate(const FNodeRegistryKey& InKey) const override;
 
 			void Register(TUniquePtr<INodeTemplate>&& InEntry);
-			void Unregister(const FMetasoundFrontendVersion& InNodeTemplateVersion);
+			void Unregister(const FNodeRegistryKey& InKey);
 
 		private:
 			TMap<FNodeRegistryKey, TUniquePtr<INodeTemplate>> Templates;
@@ -43,13 +44,9 @@ namespace Metasound
 			}
 		}
 
-		void FNodeTemplateRegistry::Unregister(const FMetasoundFrontendVersion& InVersion)
+		void FNodeTemplateRegistry::Unregister(const FNodeRegistryKey& InKey)
 		{
-			const FNodeRegistryKey Key = NodeRegistryKey::CreateKey(EMetasoundFrontendClassType::Template, InVersion.Name.ToString(), InVersion.Number.Major, InVersion.Number.Minor);
-			if (ensure(NodeRegistryKey::IsValid(Key)))
-			{
-				Templates.Remove(Key);
-			}
+			ensure(Templates.Remove(InKey) > 0);
 		}
 
 		const INodeTemplate* FNodeTemplateRegistry::FindTemplate(const FNodeRegistryKey& InKey) const
@@ -70,12 +67,46 @@ namespace Metasound
 
 		void RegisterNodeTemplate(TUniquePtr<INodeTemplate>&& InTemplate)
 		{
+			class FTemplateRegistryEntry : public INodeRegistryTemplateEntry
+			{
+				const FNodeClassInfo ClassInfo;
+				const FMetasoundFrontendClass FrontendClass;
+
+			public:
+				FTemplateRegistryEntry(const INodeTemplate& InNodeTemplate)
+					: ClassInfo(InNodeTemplate.GetFrontendClass().Metadata)
+					, FrontendClass(InNodeTemplate.GetFrontendClass())
+				{
+				}
+
+				virtual ~FTemplateRegistryEntry() = default;
+
+				virtual const FNodeClassInfo& GetClassInfo() const override
+				{
+					return ClassInfo;
+				}
+
+				/** Return a FMetasoundFrontendClass which describes the node. */
+				virtual const FMetasoundFrontendClass& GetFrontendClass() const override
+				{
+					return FrontendClass;
+				}
+			};
+
+			TUniquePtr<INodeRegistryTemplateEntry> RegEntry = TUniquePtr<INodeRegistryTemplateEntry>(new FTemplateRegistryEntry(*InTemplate.Get()));
+			FRegistryContainerImpl::Get().RegisterNodeTemplate(MoveTemp(RegEntry));
+
 			static_cast<FNodeTemplateRegistry&>(INodeTemplateRegistry::Get()).Register(MoveTemp(InTemplate));
 		}
 
 		void UnregisterNodeTemplate(const FMetasoundFrontendVersion& InVersion)
 		{
-			static_cast<FNodeTemplateRegistry&>(INodeTemplateRegistry::Get()).Unregister(InVersion);
+			const FNodeRegistryKey Key = NodeRegistryKey::CreateKey(EMetasoundFrontendClassType::Template, InVersion.Name.ToString(), InVersion.Number.Major, InVersion.Number.Minor);
+			if (ensure(NodeRegistryKey::IsValid(Key)))
+			{
+				FRegistryContainerImpl::Get().UnregisterNodeTemplate(Key);
+				static_cast<FNodeTemplateRegistry&>(INodeTemplateRegistry::Get()).Unregister(Key);
+			}
 		}
 	} // namespace Frontend
 } // namespace Metasound
diff --git a/Engine/Plugins/Runtime/Metasound/Source/MetasoundFrontend/Private/MetasoundFrontendRegistries.cpp b/Engine/Plugins/Runtime/Metasound/Source/MetasoundFrontend/Private/MetasoundFrontendRegistries.cpp
--- a/Engine/Plugins/Runtime/Metasound/Source/MetasoundFrontend/Private/MetasoundFrontendRegistries.cpp
+++ b/Engine/Plugins/Runtime/Metasound/Source/MetasoundFrontend/Private/MetasoundFrontendRegistries.cpp
@@ -239,7 +239,7 @@ namespace Metasound
 			return Key;
 		}
 
-		FNodeRegistryKey FRegistryContainerImpl::RegisterNodeTemplate(TUniquePtr<INodeTemplateRegistryEntry>&& InEntry)
+		FNodeRegistryKey FRegistryContainerImpl::RegisterNodeTemplate(TUniquePtr<INodeRegistryTemplateEntry>&& InEntry)
 		{
 			METASOUND_LLM_SCOPE;
 
@@ -247,7 +247,7 @@ namespace Metasound
 
 			if (InEntry.IsValid())
 			{
-				TSharedRef<INodeTemplateRegistryEntry, ESPMode::ThreadSafe> Entry(InEntry.Release());
+				TSharedRef<INodeRegistryTemplateEntry, ESPMode::ThreadSafe> Entry(InEntry.Release());
 
 				FNodeRegistryTransaction::FTimeType Timestamp = FPlatformTime::Cycles64();
 
@@ -297,7 +297,7 @@ namespace Metasound
 
 			if (NodeRegistryKey::IsValid(InKey))
 			{
-				if (const INodeTemplateRegistryEntry* Entry = FindNodeTemplateEntry(InKey))
+				if (const INodeRegistryTemplateEntry* Entry = FindNodeTemplateEntry(InKey))
 				{
 					FNodeRegistryTransaction::FTimeType Timestamp = FPlatformTime::Cycles64();
 
@@ -344,7 +344,7 @@ namespace Metasound
 				return Entry->IsNative();
 			}
 
-			if (const INodeTemplateRegistryEntry* TemplateEntry = FindNodeTemplateEntry(InKey))
+			if (const INodeRegistryTemplateEntry* TemplateEntry = FindNodeTemplateEntry(InKey))
 			{
 				return true;
 			}
@@ -360,7 +360,7 @@ namespace Metasound
 				return true;
 			}
 
-			if (const INodeTemplateRegistryEntry* Entry = FindNodeTemplateEntry(InKey))
+			if (const INodeRegistryTemplateEntry* Entry = FindNodeTemplateEntry(InKey))
 			{
 				OutClass = Entry->GetFrontendClass();
 				return true;
@@ -377,7 +377,7 @@ namespace Metasound
 				return true;
 			}
 
-			if (const INodeTemplateRegistryEntry* Entry = FindNodeTemplateEntry(InKey))
+			if (const INodeRegistryTemplateEntry* Entry = FindNodeTemplateEntry(InKey))
 			{
 				OutInfo = Entry->GetClassInfo();
 				return true;
@@ -487,9 +487,9 @@ namespace Metasound
 			return nullptr;
 		}
 
-		const INodeTemplateRegistryEntry* FRegistryContainerImpl::FindNodeTemplateEntry(const FNodeRegistryKey& InKey) const
+		const INodeRegistryTemplateEntry* FRegistryContainerImpl::FindNodeTemplateEntry(const FNodeRegistryKey& InKey) const
 		{
-			if (const TSharedRef<INodeTemplateRegistryEntry, ESPMode::ThreadSafe>* Entry = RegisteredNodeTemplates.Find(InKey))
+			if (const TSharedRef<INodeRegistryTemplateEntry, ESPMode::ThreadSafe>* Entry = RegisteredNodeTemplates.Find(InKey))
 			{
 				return &Entry->Get();
 			}
diff --git a/Engine/Plugins/Runtime/Metasound/Source/MetasoundFrontend/Private/NodeTemplates/MetasoundFrontendNodeTemplateReroute.cpp b/Engine/Plugins/Runtime/Metasound/Source/MetasoundFrontend/Private/NodeTemplates/MetasoundFrontendNodeTemplateReroute.cpp
new file mode 100644
--- /dev/null
+++ b/Engine/Plugins/Runtime/Metasound/Source/MetasoundFrontend/Private/NodeTemplates/MetasoundFrontendNodeTemplateReroute.cpp
@@ -0,0 +1,209 @@
+// Copyright Epic Games, Inc. All Rights Reserved.
+
+#include "NodeTemplates/MetasoundFrontendNodeTemplateReroute.h"
+
+#include "Algo/AllOf.h"
+#include "MetasoundFrontendDataTypeRegistry.h"
+#include "MetasoundFrontendRegistries.h"
+#include "Algo/AnyOf.h"
+
+
+namespace Metasound
+{
+	namespace Frontend
+	{
+		namespace ReroutePrivate
+		{
+			FConstOutputHandle FindReroutedOutput(FConstOutputHandle InOutputHandle)
+			{
+				using namespace Frontend;
+
+				if (InOutputHandle->IsValid())
+				{
+					FConstNodeHandle NodeHandle = InOutputHandle->GetOwningNode();
+					if (NodeHandle->IsValid())
+					{
+						if (NodeHandle->GetClassMetadata().GetClassName() == FRerouteNodeTemplate::ClassName)
+						{
+							TArray<FConstInputHandle> Inputs = NodeHandle->GetConstInputs();
+							if (!Inputs.IsEmpty())
+							{
+								FConstInputHandle RerouteInputHandle = Inputs.Last();
+								if (RerouteInputHandle->IsValid())
+								{
+									FConstOutputHandle ConnectedOutputHandle = RerouteInputHandle->GetConnectedOutput();
+									return FindReroutedOutput(ConnectedOutputHandle);
+								}
+							}
+						}
+					}
+				}
+
+				return InOutputHandle;
+			}
+
+			void FindReroutedInputs(FConstInputHandle InHandleToCheck, TArray<FConstInputHandle>& InOutInputHandles)
+			{
+				using namespace Frontend;
+
+				if (InHandleToCheck->IsValid())
+				{
+					FConstNodeHandle NodeHandle = InHandleToCheck->GetOwningNode();
+					if (NodeHandle->IsValid())
+					{
+						if (NodeHandle->GetClassMetadata().GetClassName() == FRerouteNodeTemplate::ClassName)
+						{
+							TArray<FConstOutputHandle> Outputs = NodeHandle->GetConstOutputs();
+							for (FConstOutputHandle& OutputHandle : Outputs)
+							{
+								TArray<FConstInputHandle> LinkedInputs = OutputHandle->GetConstConnectedInputs();
+								for (FConstInputHandle LinkedInput : LinkedInputs)
+								{
+									FindReroutedInputs(LinkedInput, InOutInputHandles);
+								}
+							}
+
+							return;
+						}
+					}
+
+					InOutInputHandles.Add(InHandleToCheck);
+				}
+			}
+		}
+
+		const FMetasoundFrontendClassName FRerouteNodeTemplate::ClassName { "UE", "Reroute", "" };
+
+		const FMetasoundFrontendVersion FRerouteNodeTemplate::Version { ClassName.GetFullName(), { 1, 0 } };
+
+		bool FRerouteNodeTemplate::BuildTemplate(FMetasoundFrontendDocument& InOutDocument, FMetasoundFrontendGraph& InOutGraph, FMetasoundFrontendNode& InOutNode) const
+		{
+			if (!ensure(InOutNode.Interface.Inputs.Num() == 1))
+			{
+				return false;
+			}
+			const FMetasoundFrontendVertex& InputVertex = InOutNode.Interface.Inputs.Last();
+
+			if (!ensure(InOutNode.Interface.Outputs.Num() == 1))
+			{
+				return false;
+			}
+			const FMetasoundFrontendVertex& OutputVertex = InOutNode.Interface.Outputs.Last();
+
+			int32 SourceIndex = InOutGraph.Edges.IndexOfByPredicate([&](const FMetasoundFrontendEdge& Edge)
+			{
+				return Edge.ToNodeID == InOutNode.GetID() && Edge.ToVertexID == InputVertex.VertexID;
+			});
+
+			if (SourceIndex != INDEX_NONE)
+			{
+				for (FMetasoundFrontendEdge& Edge : InOutGraph.Edges)
+				{
+					if (Edge.FromNodeID == InOutNode.GetID() && Edge.FromVertexID == OutputVertex.VertexID)
+					{
+						Edge.FromNodeID = InOutGraph.Edges[SourceIndex].FromNodeID;
+						Edge.FromVertexID = InOutGraph.Edges[SourceIndex].FromVertexID;
+					}
+				}
+			}
+
+			return true;
+		}
+
+		const FMetasoundFrontendClass& FRerouteNodeTemplate::GetFrontendClass() const
+		{
+			auto CreateFrontendClass = []()
+			{
+				FMetasoundFrontendClass Class;
+				Class.Metadata.SetClassName(ClassName);
+
+#if WITH_EDITOR
+				Class.Metadata.SetSerializeText(false);
+				Class.Metadata.SetAuthor(Metasound::PluginAuthor);
+				Class.Metadata.SetDescription(Metasound::PluginNodeMissingPrompt);
+
+				FMetasoundFrontendClassStyleDisplay& StyleDisplay = Class.Style.Display;
+				StyleDisplay.ImageName = "MetasoundEditor.Graph.Node.Class.Reroute";
+				StyleDisplay.bShowInputNames = false;
+				StyleDisplay.bShowOutputNames = false;
+				StyleDisplay.bShowLiterals = false;
+				StyleDisplay.bShowName = false;
+#endif // WITH_EDITOR
+
+				Class.Metadata.SetType(EMetasoundFrontendClassType::Template);
+				Class.Metadata.SetVersion(Version.Number);
+
+
+				return Class;
+			};
+
+			static const FMetasoundFrontendClass FrontendClass = CreateFrontendClass();
+			return FrontendClass;
+		}
+
+		FMetasoundFrontendNodeInterface FRerouteNodeTemplate::CreateNodeInterfaceFromDataType(FName InDataType)
+		{
+			auto CreateNewVertex = [&] { return FMetasoundFrontendVertex { "Value", InDataType, FGuid::NewGuid() }; };
+
+			FMetasoundFrontendNodeInterface NewInterface;
+			NewInterface.Inputs.Add(CreateNewVertex());
+			NewInterface.Outputs.Add(CreateNewVertex());
+
+			return NewInterface;
+		}
+
+		const FNodeRegistryKey& FRerouteNodeTemplate::GetRegistryKey()
+		{
+			static const FNodeRegistryKey RegistryKey = NodeRegistryKey::CreateKey(
+				EMetasoundFrontendClassType::Template,
+				ClassName.ToString(), 
+				Version.Number.Major, 
+				Version.Number.Minor);
+
+			return RegistryKey;
+		}
+
+		const FMetasoundFrontendVersion& FRerouteNodeTemplate::GetVersion() const
+		{
+			return Version;
+		}
+
+#if WITH_EDITOR
+		bool FRerouteNodeTemplate::HasRequiredConnections(FConstNodeHandle InNodeHandle) const
+		{
+			TArray<FConstOutputHandle> Outputs = InNodeHandle->GetConstOutputs();
+			TArray<FConstInputHandle> Inputs = InNodeHandle->GetConstInputs();
+			const bool bConnectedToNonRerouteOutputs = Algo::AnyOf(Outputs, [](const FConstOutputHandle& OutputHandle) { return ReroutePrivate::FindReroutedOutput(OutputHandle)->IsValid(); });
+			const bool bConnectedToNonRerouteInputs = Algo::AnyOf(Inputs, [](const FConstInputHandle& InputHandle)
+			{
+				TArray<FConstInputHandle> Inputs;
+				ReroutePrivate::FindReroutedInputs(InputHandle, Inputs);
+				return !Inputs.IsEmpty();
+			});
+
+			return bConnectedToNonRerouteOutputs || bConnectedToNonRerouteOutputs == bConnectedToNonRerouteInputs;
+		}
+#endif // WITH_EDITOR
+
+		bool FRerouteNodeTemplate::IsValidNodeInterface(const FMetasoundFrontendNodeInterface& InNodeInterface) const
+		{
+			if (InNodeInterface.Inputs.Num() != 1)
+			{
+				return false;
+			}
+			
+			if (InNodeInterface.Outputs.Num() != 1)
+			{
+				return false;
+			}
+
+			const FName DataType = InNodeInterface.Inputs.Last().TypeName;
+			if (DataType != InNodeInterface.Outputs.Last().TypeName)
+			{
+				return false;
+			}
+
+			return IDataTypeRegistry::Get().IsRegistered(DataType);
+		}
+	} // namespace Frontend
+} // namespace Metasound
\ No newline at end of file
diff --git a/Engine/Plugins/Runtime/Metasound/Source/MetasoundFrontend/Public/MetasoundFrontendNodeTemplateRegistry.h b/Engine/Plugins/Runtime/Metasound/Source/MetasoundFrontend/Public/MetasoundFrontendNodeTemplateRegistry.h
--- a/Engine/Plugins/Runtime/Metasound/Source/MetasoundFrontend/Public/MetasoundFrontendNodeTemplateRegistry.h
+++ b/Engine/Plugins/Runtime/Metasound/Source/MetasoundFrontend/Public/MetasoundFrontendNodeTemplateRegistry.h
@@ -20,7 +20,10 @@ namespace Metasound
 			virtual const FMetasoundFrontendClass& GetFrontendClass() const = 0;
 			virtual const FMetasoundFrontendVersion& GetVersion() const = 0;
 			virtual bool IsValidNodeInterface(const FMetasoundFrontendNodeInterface& InNodeInterface) const = 0;
+
+#if WITH_EDITOR
 			virtual bool HasRequiredConnections(FConstNodeHandle InNodeHandle) const = 0;
+#endif // WITH_EDITOR
 		};
 
 		class METASOUNDFRONTEND_API INodeTemplateRegistry
diff --git a/Engine/Plugins/Runtime/Metasound/Source/MetasoundFrontend/Public/NodeTemplates/MetasoundFrontendNodeTemplateReroute.h b/Engine/Plugins/Runtime/Metasound/Source/MetasoundFrontend/Public/NodeTemplates/MetasoundFrontendNodeTemplateReroute.h
new file mode 100644
--- /dev/null
+++ b/Engine/Plugins/Runtime/Metasound/Source/MetasoundFrontend/Public/NodeTemplates/MetasoundFrontendNodeTemplateReroute.h
@@ -0,0 +1,36 @@
+// Copyright Epic Games, Inc. All Rights Reserved.
+#pragma once
+
+#include "MetasoundFrontendController.h"
+#include "MetasoundFrontendDocument.h"
+#include "MetasoundFrontendNodeTemplateRegistry.h"
+#include "Templates/SharedPointer.h"
+#include "UObject/NoExportTypes.h"
+
+namespace Metasound
+{
+	namespace Frontend
+	{
+		class METASOUNDFRONTEND_API FRerouteNodeTemplate : public INodeTemplate
+		{
+		public:
+			static const FMetasoundFrontendClassName ClassName;
+			static const FMetasoundFrontendVersion Version;
+
+			static const FNodeRegistryKey& GetRegistryKey();
+
+			static FMetasoundFrontendNodeInterface CreateNodeInterfaceFromDataType(FName InDataType);
+
+			virtual ~FRerouteNodeTemplate() = default;
+
+			virtual bool BuildTemplate(FMetasoundFrontendDocument& InOutDocument, FMetasoundFrontendGraph& InOutGraph, FMetasoundFrontendNode& InOutNode) const override;
+			virtual const FMetasoundFrontendClass& GetFrontendClass() const override;
+			virtual const FMetasoundFrontendVersion& GetVersion() const override;
+			virtual bool IsValidNodeInterface(const FMetasoundFrontendNodeInterface& InNodeInterface) const override;
+
+#if WITH_EDITOR
+			virtual bool HasRequiredConnections(FConstNodeHandle InNodeHandle) const override;
+#endif // WITH_EDITOR
+		};
+	} // namespace Frontend
+} // namespace Metasound
