Rename UMetaSound to UMetaSoundPatch

diff --git a/Engine/Plugins/Runtime/Metasound/Config/BaseMetasound.ini b/Engine/Plugins/Runtime/Metasound/Config/BaseMetasound.ini
--- a/Engine/Plugins/Runtime/Metasound/Config/BaseMetasound.ini
+++ b/Engine/Plugins/Runtime/Metasound/Config/BaseMetasound.ini
@@ -1,4 +1,6 @@
 [CoreRedirects]
+
+; 5.0
 +PropertyRedirects=(OldName="MetaSoundSettings.AutoUpdateBlacklist",NewName="AutoUpdateDenylist")
 +PropertyRedirects=(OldName="MetaSoundSettings.AutoUpdateAssetBlacklist",NewName="AutoUpdateDenylist")
 
@@ -7,7 +9,6 @@
 +StructRedirects=(OldName="MetasoundEditorGraphInputObjectRef",NewName="MetasoundEditorGraphMemberDefaultObjectRef")
  
 +ClassRedirects=(OldName="MetasoundEditorGraphInputLiteral",NewName="MetasoundEditorGraphMemberDefaultLiteral")
- 
 +ClassRedirects=(OldName="MetasoundEditorGraphInputBool",NewName="MetasoundEditorGraphMemberDefaultBool")
 +ClassRedirects=(OldName="MetasoundEditorGraphInputBoolArray",NewName="MetasoundEditorGraphMemberDefaultBoolArray")
 +ClassRedirects=(OldName="MetasoundEditorGraphInputInt",NewName="MetasoundEditorGraphMemberDefaultInt")
@@ -17,4 +18,8 @@
 +ClassRedirects=(OldName="MetasoundEditorGraphInputString",NewName="MetasoundEditorGraphMemberDefaultString")
 +ClassRedirects=(OldName="MetasoundEditorGraphInputStringArray",NewName="MetasoundEditorGraphMemberDefaultStringArray")
 +ClassRedirects=(OldName="MetasoundEditorGraphInputObject",NewName="MetasoundEditorGraphMemberDefaultObject")
-+ClassRedirects=(OldName="MetasoundEditorGraphInputObjectArray",NewName="MetasoundEditorGraphMemberDefaultObjectArray")
\ No newline at end of file
++ClassRedirects=(OldName="MetasoundEditorGraphInputObjectArray",NewName="MetasoundEditorGraphMemberDefaultObjectArray")
+
+
+; 5.1
++ClassRedirects=(OldName="MetaSound",NewName="MetaSoundPatch")
\ No newline at end of file
diff --git a/Engine/Plugins/Runtime/Metasound/Source/MetasoundEditor/Private/MetasoundAssetTypeActions.cpp b/Engine/Plugins/Runtime/Metasound/Source/MetasoundEditor/Private/MetasoundAssetTypeActions.cpp
--- a/Engine/Plugins/Runtime/Metasound/Source/MetasoundEditor/Private/MetasoundAssetTypeActions.cpp
+++ b/Engine/Plugins/Runtime/Metasound/Source/MetasoundEditor/Private/MetasoundAssetTypeActions.cpp
@@ -167,28 +167,28 @@ namespace Metasound
 			}
 		} // namespace AssetTypeActionsPrivate
 
-		UClass* FAssetTypeActions_MetaSound::GetSupportedClass() const
+		UClass* FAssetTypeActions_MetaSoundPatch::GetSupportedClass() const
 		{
-			return UMetaSound::StaticClass();
+			return UMetaSoundPatch::StaticClass();
 		}
 
-		FColor FAssetTypeActions_MetaSound::GetTypeColor() const
+		FColor FAssetTypeActions_MetaSoundPatch::GetTypeColor() const
 		{
 			if (const ISlateStyle* MetasoundStyle = FSlateStyleRegistry::FindSlateStyle("MetaSoundStyle"))
 			{
-				return MetasoundStyle->GetColor("Metasound.Color").ToFColorSRGB();
+				return MetasoundStyle->GetColor("MetaSoundPatch.Color").ToFColorSRGB();
 			}
 
 			return FColor::White;
 		}
 
-		void FAssetTypeActions_MetaSound::OpenAssetEditor(const TArray<UObject*>& InObjects, TSharedPtr<IToolkitHost> ToolkitHost)
+		void FAssetTypeActions_MetaSoundPatch::OpenAssetEditor(const TArray<UObject*>& InObjects, TSharedPtr<IToolkitHost> ToolkitHost)
 		{
 			const EToolkitMode::Type Mode = ToolkitHost.IsValid() ? EToolkitMode::WorldCentric : EToolkitMode::Standalone;
 
 			for (UObject* Object : InObjects)
 			{
-				if (UMetaSound* Metasound = Cast<UMetaSound>(Object))
+				if (UMetaSoundPatch* Metasound = Cast<UMetaSoundPatch>(Object))
 				{
 					TSharedRef<FEditor> NewEditor = MakeShared<FEditor>();
 					NewEditor->InitMetasoundEditor(Mode, ToolkitHost, Metasound);
@@ -196,13 +196,13 @@ namespace Metasound
 			}
 		}
 
-		void FAssetTypeActions_MetaSound::RegisterMenuActions()
+		void FAssetTypeActions_MetaSoundPatch::RegisterMenuActions()
 		{
 			using namespace AssetTypeActionsPrivate;
-			RegisterPresetAction<UMetaSound, UMetaSoundFactory>(PresetLabel, PresetToolTip);
+			RegisterPresetAction<UMetaSoundPatch, UMetaSoundFactory>(PresetLabel, PresetToolTip);
 		}
 
-		const TArray<FText>& FAssetTypeActions_MetaSound::GetSubMenus() const
+		const TArray<FText>& FAssetTypeActions_MetaSoundPatch::GetSubMenus() const
 		{
 			if (GetDefault<UMetasoundEditorSettings>()->bPinMetaSoundInAssetMenu)
 			{
@@ -227,7 +227,7 @@ namespace Metasound
 		{
 			if (const ISlateStyle* MetasoundStyle = FSlateStyleRegistry::FindSlateStyle("MetaSoundStyle"))
 			{
-				return MetasoundStyle->GetColor("MetasoundSource.Color").ToFColorSRGB();
+				return MetasoundStyle->GetColor("MetaSoundSource.Color").ToFColorSRGB();
 			}
 
 			return FColor::White;
@@ -257,7 +257,7 @@ namespace Metasound
 			// that introduces the concept of an interface "interface" (i.e. if the interface one graph subscribes to
 			// or contains all of the necessary inputs/outputs of another), so that the preset can "implement" the
 			// referenced graphs "interface."
-// 			RegisterPresetAction<UMetaSound, UMetaSoundFactory, UMetaSoundSource>(PresetLabel, PresetToolTip);
+// 			RegisterPresetAction<UMetaSoundPatch, UMetaSoundFactory, UMetaSoundSource>(PresetLabel, PresetToolTip);
 		}
 
 		const TArray<FText>& FAssetTypeActions_MetaSoundSource::GetSubMenus() const
diff --git a/Engine/Plugins/Runtime/Metasound/Source/MetasoundEditor/Private/MetasoundAssetTypeActions.h b/Engine/Plugins/Runtime/Metasound/Source/MetasoundEditor/Private/MetasoundAssetTypeActions.h
--- a/Engine/Plugins/Runtime/Metasound/Source/MetasoundEditor/Private/MetasoundAssetTypeActions.h
+++ b/Engine/Plugins/Runtime/Metasound/Source/MetasoundEditor/Private/MetasoundAssetTypeActions.h
@@ -11,11 +11,11 @@ namespace Metasound
 {
 	namespace Editor
 	{
-		class FAssetTypeActions_MetaSound : public FAssetTypeActions_Base
+		class FAssetTypeActions_MetaSoundPatch : public FAssetTypeActions_Base
 		{
 		public:
 			// IAssetTypeActions Implementation
-			virtual FText GetName() const override { return NSLOCTEXT("AssetTypeActions", "AssetTypeActions_MetaSound", "MetaSound"); }
+			virtual FText GetName() const override { return NSLOCTEXT("AssetTypeActions", "AssetTypeActions_MetaSoundPatch", "MetaSound Patch"); }
 			virtual FColor GetTypeColor() const override;
 			virtual UClass* GetSupportedClass() const override;
 			virtual uint32 GetCategories() override { return EAssetTypeCategories::Sounds; }
diff --git a/Engine/Plugins/Runtime/Metasound/Source/MetasoundEditor/Private/MetasoundEditor.cpp b/Engine/Plugins/Runtime/Metasound/Source/MetasoundEditor/Private/MetasoundEditor.cpp
--- a/Engine/Plugins/Runtime/Metasound/Source/MetasoundEditor/Private/MetasoundEditor.cpp
+++ b/Engine/Plugins/Runtime/Metasound/Source/MetasoundEditor/Private/MetasoundEditor.cpp
@@ -957,12 +957,12 @@ namespace Metasound
 			{
 				if (UMetaSoundSource* MetaSoundSource = Cast<UMetaSoundSource>(Metasound))
 				{
-					return MetasoundStyle->GetColor("MetasoundSource.Color");
+					return MetasoundStyle->GetColor("MetaSoundSource.Color");
 				}
 
-				if (UMetaSound* MetaSound = Cast<UMetaSound>(Metasound))
+				if (UMetaSoundPatch* MetaSoundPatch = Cast<UMetaSoundPatch>(Metasound))
 				{
-					return MetasoundStyle->GetColor("Metasound.Color");
+					return MetasoundStyle->GetColor("MetaSoundPatch.Color");
 				}
 			}
 
@@ -987,12 +987,12 @@ namespace Metasound
 				{
 					if (UMetaSoundSource* MetaSoundSource = Cast<UMetaSoundSource>(Metasound))
 					{
-						return MetasoundStyle->GetColor("MetasoundSource.Color");
+						return MetasoundStyle->GetColor("MetaSoundSource.Color");
 					}
 
-					if (UMetaSound* MetaSound = Cast<UMetaSound>(Metasound))
+					if (UMetaSoundPatch* MetaSoundPatch = Cast<UMetaSoundPatch>(Metasound))
 					{
-						return MetasoundStyle->GetColor("Metasound.Color");
+						return MetasoundStyle->GetColor("MetaSoundPatch.Color");
 					}
 				}
 			}
diff --git a/Engine/Plugins/Runtime/Metasound/Source/MetasoundEditor/Private/MetasoundEditor.h b/Engine/Plugins/Runtime/Metasound/Source/MetasoundEditor/Private/MetasoundEditor.h
--- a/Engine/Plugins/Runtime/Metasound/Source/MetasoundEditor/Private/MetasoundEditor.h
+++ b/Engine/Plugins/Runtime/Metasound/Source/MetasoundEditor/Private/MetasoundEditor.h
@@ -46,7 +46,7 @@ class IToolkitHost;
 class SVerticalBox;
 class UAudioComponent;
 class UEdGraphNode;
-class UMetaSound;
+class UMetaSoundPatch;
 class UMetasoundEditorGraph;
 class UMetasoundEditorGraphNode;
 
diff --git a/Engine/Plugins/Runtime/Metasound/Source/MetasoundEditor/Private/MetasoundEditorGraphBuilder.h b/Engine/Plugins/Runtime/Metasound/Source/MetasoundEditor/Private/MetasoundEditorGraphBuilder.h
--- a/Engine/Plugins/Runtime/Metasound/Source/MetasoundEditor/Private/MetasoundEditorGraphBuilder.h
+++ b/Engine/Plugins/Runtime/Metasound/Source/MetasoundEditor/Private/MetasoundEditorGraphBuilder.h
@@ -16,7 +16,7 @@ class SGraphEditor;
 class UEdGraph;
 class UEdGraphNode;
 class UEdGraphPin;
-class UMetaSound;
+class UMetaSoundPatch;
 class UMetasoundEditorGraphExternalNode;
 class UMetasoundEditorGraphNode;
 class UMetasoundEditorGraphInputNode;
diff --git a/Engine/Plugins/Runtime/Metasound/Source/MetasoundEditor/Private/MetasoundEditorGraphNode.cpp b/Engine/Plugins/Runtime/Metasound/Source/MetasoundEditor/Private/MetasoundEditorGraphNode.cpp
--- a/Engine/Plugins/Runtime/Metasound/Source/MetasoundEditor/Private/MetasoundEditorGraphNode.cpp
+++ b/Engine/Plugins/Runtime/Metasound/Source/MetasoundEditor/Private/MetasoundEditorGraphNode.cpp
@@ -99,7 +99,7 @@ void UMetasoundEditorGraphNode::CreateInputPin()
 		return;
 	}
 
-	FString PinName; // get from UMetaSound
+	FString PinName; // get from UMetaSoundPatch
 	UEdGraphPin* NewPin = CreatePin(EGPD_Input, TEXT("MetasoundEditorGraphNode"), *PinName);
 	if (NewPin->PinName.IsNone())
 	{
@@ -583,7 +583,7 @@ FString UMetasoundEditorGraphNode::GetDocumentationExcerptName() const
 {
 	// Default the node to searching for an excerpt named for the C++ node class name, including the U prefix.
 	// This is done so that the excerpt name in the doc file can be found by find-in-files when searching for the full class name.
-	return FString::Printf(TEXT("%s%s"), UMetaSound::StaticClass()->GetPrefixCPP(), *UMetaSound::StaticClass()->GetName());
+	return FString::Printf(TEXT("%s%s"), UMetaSoundPatch::StaticClass()->GetPrefixCPP(), *UMetaSoundPatch::StaticClass()->GetName());
 }
 
 void UMetasoundEditorGraphOutputNode::PinDefaultValueChanged(UEdGraphPin* InPin)
diff --git a/Engine/Plugins/Runtime/Metasound/Source/MetasoundEditor/Private/MetasoundEditorModule.cpp b/Engine/Plugins/Runtime/Metasound/Source/MetasoundEditor/Private/MetasoundEditorModule.cpp
--- a/Engine/Plugins/Runtime/Metasound/Source/MetasoundEditor/Private/MetasoundEditorModule.cpp
+++ b/Engine/Plugins/Runtime/Metasound/Source/MetasoundEditor/Private/MetasoundEditorModule.cpp
@@ -101,8 +101,8 @@ namespace Metasound
 
 				// Metasound Editor
 				{
-					Set("Metasound.Color", FColor(13, 55, 13));
-					Set("MetasoundSource.Color", FColor(103, 214, 66));
+					Set("MetaSoundPatch.Color", FColor(13, 55, 13));
+					Set("MetaSoundSource.Color", FColor(103, 214, 66));
 
 					// Actions
 					Set("MetasoundEditor.Play", new IMAGE_BRUSH_SVG(TEXT("Icons/play"), Icon40x40));
@@ -176,7 +176,7 @@ namespace Metasound
 						Set(*FString::Printf(TEXT("ClassThumbnail.%s"), *ClassName), new IMAGE_BRUSH_SVG(IconFileName, InIcon64));
 					};
 
-					SetClassIcon(TEXT("Metasound"));
+					SetClassIcon(TEXT("MetasoundPatch"));
 					SetClassIcon(TEXT("MetasoundSource"));
 				}
 
@@ -684,14 +684,14 @@ namespace Metasound
 				// Register Metasound asset type actions
 				IAssetTools& AssetTools = FModuleManager::LoadModuleChecked<FAssetToolsModule>(AssetToolName).Get();
 
-				AddAssetAction<FAssetTypeActions_MetaSound>(AssetTools, AssetActions);
+				AddAssetAction<FAssetTypeActions_MetaSoundPatch>(AssetTools, AssetActions);
 				AddAssetAction<FAssetTypeActions_MetaSoundSource>(AssetTools, AssetActions);
 
 				FPropertyEditorModule& PropertyModule = FModuleManager::LoadModuleChecked<FPropertyEditorModule>("PropertyEditor");
 				
 				PropertyModule.RegisterCustomClassLayout(
-					UMetaSound::StaticClass()->GetFName(),
-					FOnGetDetailCustomizationInstance::CreateLambda([]() { return MakeShared<FMetasoundDetailCustomization>(UMetaSound::GetDocumentPropertyName()); }));
+					UMetaSoundPatch::StaticClass()->GetFName(),
+					FOnGetDetailCustomizationInstance::CreateLambda([]() { return MakeShared<FMetasoundDetailCustomization>(UMetaSoundPatch::GetDocumentPropertyName()); }));
 
 				PropertyModule.RegisterCustomClassLayout(
 					UMetaSoundSource::StaticClass()->GetFName(),
@@ -751,10 +751,10 @@ namespace Metasound
 					GetMutableDefault<UMetasoundEditorSettings>()
 				);
 
-				MetaSoundClassNames.Add(UMetaSound::StaticClass()->GetClassPathName());
+				MetaSoundClassNames.Add(UMetaSoundPatch::StaticClass()->GetClassPathName());
 				MetaSoundClassNames.Add(UMetaSoundSource::StaticClass()->GetClassPathName());
 
-				FAssetTypeActions_MetaSound::RegisterMenuActions();
+				FAssetTypeActions_MetaSoundPatch::RegisterMenuActions();
 				FAssetTypeActions_MetaSoundSource::RegisterMenuActions();
 
 				FAssetRegistryModule& AssetRegistryModule = FModuleManager::LoadModuleChecked<FAssetRegistryModule>("AssetRegistry");
diff --git a/Engine/Plugins/Runtime/Metasound/Source/MetasoundEditor/Private/MetasoundFactory.cpp b/Engine/Plugins/Runtime/Metasound/Source/MetasoundEditor/Private/MetasoundFactory.cpp
--- a/Engine/Plugins/Runtime/Metasound/Source/MetasoundEditor/Private/MetasoundFactory.cpp
+++ b/Engine/Plugins/Runtime/Metasound/Source/MetasoundEditor/Private/MetasoundFactory.cpp
@@ -72,13 +72,13 @@ UMetaSoundBaseFactory::UMetaSoundBaseFactory(const FObjectInitializer& ObjectIni
 UMetaSoundFactory::UMetaSoundFactory(const FObjectInitializer& ObjectInitializer)
 	: Super(ObjectInitializer)
 {
-	SupportedClass = UMetaSound::StaticClass();
+	SupportedClass = UMetaSoundPatch::StaticClass();
 }
 
 UObject* UMetaSoundFactory::FactoryCreateNew(UClass* InClass, UObject* InParent, FName InName, EObjectFlags InFlags, UObject* InContext, FFeedbackContext* InFeedbackContext)
 {
 	using namespace Metasound::FactoryPrivate;
-	return CreateNewMetaSoundObject<UMetaSound>(InParent, InName, InFlags, ReferencedMetaSoundObject);
+	return CreateNewMetaSoundObject<UMetaSoundPatch>(InParent, InName, InFlags, ReferencedMetaSoundObject);
 }
 
 UMetaSoundSourceFactory::UMetaSoundSourceFactory(const FObjectInitializer& ObjectInitializer)
diff --git a/Engine/Plugins/Runtime/Metasound/Source/MetasoundEditor/Public/MetasoundEditorGraphNode.h b/Engine/Plugins/Runtime/Metasound/Source/MetasoundEditor/Public/MetasoundEditorGraphNode.h
--- a/Engine/Plugins/Runtime/Metasound/Source/MetasoundEditor/Public/MetasoundEditorGraphNode.h
+++ b/Engine/Plugins/Runtime/Metasound/Source/MetasoundEditor/Public/MetasoundEditorGraphNode.h
@@ -17,7 +17,7 @@
 
 // Forward Declarations
 class UEdGraphPin;
-class UMetaSound;
+class UMetaSoundPatch;
 class UMetasoundEditorGraphOutput;
 class UMetasoundEditorGraphMember;
 class UMetasoundEditorGraphVariable;
diff --git a/Engine/Plugins/Runtime/Metasound/Source/MetasoundEditor/Public/MetasoundEditorGraphSchema.h b/Engine/Plugins/Runtime/Metasound/Source/MetasoundEditor/Public/MetasoundEditorGraphSchema.h
--- a/Engine/Plugins/Runtime/Metasound/Source/MetasoundEditor/Public/MetasoundEditorGraphSchema.h
+++ b/Engine/Plugins/Runtime/Metasound/Source/MetasoundEditor/Public/MetasoundEditorGraphSchema.h
@@ -17,7 +17,7 @@
 // Forward Declarations
 class UEdGraph;
 class UEdGraphNode;
-class UMetaSound;
+class UMetaSoundPatch;
 class UMetasoundEditorGraphNode;
 
 
diff --git a/Engine/Plugins/Runtime/Metasound/Source/MetasoundEditor/Public/MetasoundFactory.h b/Engine/Plugins/Runtime/Metasound/Source/MetasoundEditor/Public/MetasoundFactory.h
--- a/Engine/Plugins/Runtime/Metasound/Source/MetasoundEditor/Public/MetasoundFactory.h
+++ b/Engine/Plugins/Runtime/Metasound/Source/MetasoundEditor/Public/MetasoundFactory.h
@@ -8,7 +8,7 @@
 
 #include "MetasoundFactory.generated.h"
 
-class UMetaSound;
+class UMetaSoundPatch;
 class UMetaSoundSource;
 
 UCLASS(abstract)
diff --git a/Engine/Plugins/Runtime/Metasound/Source/MetasoundEngine/Private/Metasound.cpp b/Engine/Plugins/Runtime/Metasound/Source/MetasoundEngine/Private/Metasound.cpp
--- a/Engine/Plugins/Runtime/Metasound/Source/MetasoundEngine/Private/Metasound.cpp
+++ b/Engine/Plugins/Runtime/Metasound/Source/MetasoundEngine/Private/Metasound.cpp
@@ -31,14 +31,14 @@
 #define LOCTEXT_NAMESPACE "MetaSound"
 
 
-UMetaSound::UMetaSound(const FObjectInitializer& ObjectInitializer)
+UMetaSoundPatch::UMetaSoundPatch(const FObjectInitializer& ObjectInitializer)
 	: Super(ObjectInitializer)
 	, FMetasoundAssetBase()
 {
 }
 
 #if WITH_EDITOR
-void UMetaSound::PostDuplicate(EDuplicateMode::Type InDuplicateMode)
+void UMetaSoundPatch::PostDuplicate(EDuplicateMode::Type InDuplicateMode)
 {
 	Super::PostDuplicate(InDuplicateMode);
 
@@ -51,82 +51,82 @@ void UMetaSound::PostDuplicate(EDuplicateMode::Type InDuplicateMode)
 	}
 }
 
-void UMetaSound::PostEditUndo()
+void UMetaSoundPatch::PostEditUndo()
 {
 	Super::PostEditUndo();
 	Metasound::PostEditUndo(*this);
 }
 #endif // WITHEDITOR
 
-void UMetaSound::BeginDestroy()
+void UMetaSoundPatch::BeginDestroy()
 {
 	UnregisterGraphWithFrontend();
 	Super::BeginDestroy();
 }
 
-void UMetaSound::PreSave(FObjectPreSaveContext InSaveContext)
+void UMetaSoundPatch::PreSave(FObjectPreSaveContext InSaveContext)
 {
 	Super::PreSave(InSaveContext);
 	Metasound::PreSaveAsset(*this, InSaveContext);
 }
 
-void UMetaSound::Serialize(FArchive& InArchive)
+void UMetaSoundPatch::Serialize(FArchive& InArchive)
 {
 	Super::Serialize(InArchive);
 	Metasound::SerializeToArchive(*this, InArchive);
 }
 
 #if WITH_EDITORONLY_DATA
-UEdGraph* UMetaSound::GetGraph()
+UEdGraph* UMetaSoundPatch::GetGraph()
 {
 	return Graph;
 }
 
-const UEdGraph* UMetaSound::GetGraph() const
+const UEdGraph* UMetaSoundPatch::GetGraph() const
 {
 	return Graph;
 }
 
-UEdGraph& UMetaSound::GetGraphChecked()
+UEdGraph& UMetaSoundPatch::GetGraphChecked()
 {
 	check(Graph);
 	return *Graph;
 }
 
-const UEdGraph& UMetaSound::GetGraphChecked() const
+const UEdGraph& UMetaSoundPatch::GetGraphChecked() const
 {
 	check(Graph);
 	return *Graph;
 }
 
-FText UMetaSound::GetDisplayName() const
+FText UMetaSoundPatch::GetDisplayName() const
 {
-	FString TypeName = UMetaSound::StaticClass()->GetName();
+	FString TypeName = UMetaSoundPatch::StaticClass()->GetName();
 	return FMetasoundAssetBase::GetDisplayName(MoveTemp(TypeName));
 }
 
-void UMetaSound::SetRegistryAssetClassInfo(const Metasound::Frontend::FNodeClassInfo& InNodeInfo)
+void UMetaSoundPatch::SetRegistryAssetClassInfo(const Metasound::Frontend::FNodeClassInfo& InNodeInfo)
 {
 	Metasound::SetMetaSoundRegistryAssetClassInfo(*this, InNodeInfo);
 }
 #endif // WITH_EDITORONLY_DATA
 
-Metasound::Frontend::FNodeClassInfo UMetaSound::GetAssetClassInfo() const
+Metasound::Frontend::FNodeClassInfo UMetaSoundPatch::GetAssetClassInfo() const
 {
 	return { GetDocumentChecked().RootGraph, *GetPathName() };
 }
 
-void UMetaSound::SetReferencedAssetClassKeys(TSet<Metasound::Frontend::FNodeRegistryKey>&& InKeys)
+void UMetaSoundPatch::SetReferencedAssetClassKeys(TSet<Metasound::Frontend::FNodeRegistryKey>&& InKeys)
 {
 	ReferencedAssetClassKeys = MoveTemp(InKeys);
 }
 
-TSet<FSoftObjectPath>& UMetaSound::GetReferencedAssetClassCache()
+TSet<FSoftObjectPath>& UMetaSoundPatch::GetReferencedAssetClassCache()
 {
 	return ReferenceAssetClassCache;
 }
 
-const TSet<FSoftObjectPath>& UMetaSound::GetReferencedAssetClassCache() const
+const TSet<FSoftObjectPath>& UMetaSoundPatch::GetReferencedAssetClassCache() const
 {
 	return ReferenceAssetClassCache;
 }
diff --git a/Engine/Plugins/Runtime/Metasound/Source/MetasoundEngine/Private/MetasoundEngineModule.h b/Engine/Plugins/Runtime/Metasound/Source/MetasoundEngine/Private/MetasoundEngineModule.h
--- a/Engine/Plugins/Runtime/Metasound/Source/MetasoundEngine/Private/MetasoundEngineModule.h
+++ b/Engine/Plugins/Runtime/Metasound/Source/MetasoundEngine/Private/MetasoundEngineModule.h
@@ -6,7 +6,7 @@
 #include "Modules/ModuleInterface.h"
 
 // Forward Declarations
-class UMetaSound;
+class UMetaSoundPatch;
 
 DECLARE_LOG_CATEGORY_EXTERN(LogMetasoundEngine, Log, All);
 
diff --git a/Engine/Plugins/Runtime/Metasound/Source/MetasoundEngine/Private/MetasoundInterface.cpp b/Engine/Plugins/Runtime/Metasound/Source/MetasoundEngine/Private/MetasoundInterface.cpp
--- a/Engine/Plugins/Runtime/Metasound/Source/MetasoundEngine/Private/MetasoundInterface.cpp
+++ b/Engine/Plugins/Runtime/Metasound/Source/MetasoundEngine/Private/MetasoundInterface.cpp
@@ -227,7 +227,7 @@ namespace Metasound
 			{
 				constexpr bool bIsDefault = true;
 				constexpr bool bEditorCanAddOrRemove = false;
-				RegisterInterface<UMetaSound>(MetasoundV1_0::GetInterface(), nullptr, bIsDefault, bEditorCanAddOrRemove, IDataReference::RouterName);
+				RegisterInterface<UMetaSoundPatch>(MetasoundV1_0::GetInterface(), nullptr, bIsDefault, bEditorCanAddOrRemove, IDataReference::RouterName);
 
 				RegisterInterface(SourceInterface::CreateInterface(*UMetaSoundSource::StaticClass()), MakeUnique<SourceInterface::FUpdateInterface>(), bIsDefault, bEditorCanAddOrRemove, IDataReference::RouterName);
 				RegisterInterface(OutputFormatMonoInterface::CreateInterface(*UMetaSoundSource::StaticClass()), nullptr, bIsDefault, bEditorCanAddOrRemove, IDataReference::RouterName);
@@ -243,8 +243,8 @@ namespace Metasound
 				// Added for upgrading old metasounds
 				RegisterInterface(SourceInterfaceV1_0::CreateInterface(*UMetaSoundSource::StaticClass()), nullptr, bIsDefault, bEditorCanAddOrRemove, IDataReference::RouterName);
 
-				// Set default interface with unset version to use base UMetaSound class implementation (legacy requirement for 5.0 alpha).
-				RegisterInterface<UMetaSound>(FMetasoundFrontendInterface(), nullptr, bIsDefault, bEditorCanAddOrRemove, IDataReference::RouterName);
+				// Set default interface with unset version to use base UMetaSoundPatch class implementation (legacy requirement for 5.0 alpha).
+				RegisterInterface<UMetaSoundPatch>(FMetasoundFrontendInterface(), nullptr, bIsDefault, bEditorCanAddOrRemove, IDataReference::RouterName);
 
 				RegisterInterface<UMetaSoundSource>(MetasoundOutputFormatStereoV1_0::GetInterface(), nullptr, bIsDefault, bEditorCanAddOrRemove, IDataReference::RouterName);
 				RegisterInterface<UMetaSoundSource>(MetasoundOutputFormatStereoV1_1::GetInterface(), MakeUnique<MetasoundOutputFormatStereoV1_1::FUpdateInterface>(), bIsDefault, bEditorCanAddOrRemove, IDataReference::RouterName);
diff --git a/Engine/Plugins/Runtime/Metasound/Source/MetasoundEngine/Public/Metasound.h b/Engine/Plugins/Runtime/Metasound/Source/MetasoundEngine/Public/Metasound.h
--- a/Engine/Plugins/Runtime/Metasound/Source/MetasoundEngine/Public/Metasound.h
+++ b/Engine/Plugins/Runtime/Metasound/Source/MetasoundEngine/Public/Metasound.h
@@ -144,7 +144,7 @@ namespace Metasound
  * Because of this, they contain no required inputs or outputs.
  */
 UCLASS(hidecategories = object, BlueprintType)
-class METASOUNDENGINE_API UMetaSound : public UObject, public FMetasoundAssetBase
+class METASOUNDENGINE_API UMetaSoundPatch : public UObject, public FMetasoundAssetBase
 {
 	GENERATED_BODY()
 
@@ -158,7 +158,7 @@ protected:
 #endif // WITH_EDITORONLY_DATA
 
 public:
-	UMetaSound(const FObjectInitializer& ObjectInitializer);
+	UMetaSoundPatch(const FObjectInitializer& ObjectInitializer);
 
 	UPROPERTY()
 	TSet<FString> ReferencedAssetClassKeys;
@@ -189,7 +189,7 @@ public:
 	// while allowing editor to reference directly)
 	static FName GetDocumentPropertyName()
 	{
-		return GET_MEMBER_NAME_CHECKED(UMetaSound, RootMetaSoundDocument);
+		return GET_MEMBER_NAME_CHECKED(UMetaSoundPatch, RootMetaSoundDocument);
 	}
 
 	// Name to display in editors
diff --git a/Engine/Plugins/Runtime/Metasound/Source/MetasoundFrontend/Public/MetasoundFrontendController.h b/Engine/Plugins/Runtime/Metasound/Source/MetasoundFrontend/Public/MetasoundFrontendController.h
--- a/Engine/Plugins/Runtime/Metasound/Source/MetasoundFrontend/Public/MetasoundFrontendController.h
+++ b/Engine/Plugins/Runtime/Metasound/Source/MetasoundFrontend/Public/MetasoundFrontendController.h
@@ -24,7 +24,7 @@
  *
  * In general, the workflow for editing a Metasound graph will be:
  * 1) Load or create a metasound asset.
- * 2) Call UMetaSound::GetDocumentHandle() to get a handle to the document for that asset.
+ * 2) Call UMetaSoundPatch::GetDocumentHandle() to get a handle to the document for that asset.
  * 
  * Typically the workflow for creating a Metasound subgraph will be
  * 1) Get a Metasound::Frontend::FGraphHandle (typically through the two steps described in the first paragraph)
