Fix for mac symbol ambiguity
#rnx

diff --git a/Engine/Plugins/Runtime/Metasound/Source/MetasoundEditor/Private/MetasoundEditorGraphSchema.cpp b/Engine/Plugins/Runtime/Metasound/Source/MetasoundEditor/Private/MetasoundEditorGraphSchema.cpp
--- a/Engine/Plugins/Runtime/Metasound/Source/MetasoundEditor/Private/MetasoundEditorGraphSchema.cpp
+++ b/Engine/Plugins/Runtime/Metasound/Source/MetasoundEditor/Private/MetasoundEditorGraphSchema.cpp
@@ -397,17 +397,16 @@ UEdGraphNode* FMetasoundGraphSchemaAction_NodeWithMultipleOutputs::PerformAction
 
 const FSlateBrush* FMetasoundGraphSchemaAction_NewNode::GetIconBrush() const
 {
-	using namespace Metasound::Editor;
 	using namespace Metasound::Frontend;
 
 	const FNodeRegistryKey RegistryKey = NodeRegistryKey::CreateKey(ClassMetadata);
 	const bool bIsClassNative = FMetasoundFrontendRegistryContainer::Get()->IsNodeNative(RegistryKey);
 	if (bIsClassNative)
 	{
-		return &Style::GetSlateBrushSafe("MetasoundEditor.Graph.Node.Class.Native");
+		return &Metasound::Editor::Style::GetSlateBrushSafe("MetasoundEditor.Graph.Node.Class.Native");
 	}
 
-	return &Style::GetSlateBrushSafe("MetasoundEditor.Graph.Node.Class.Graph");
+	return &Metasound::Editor::Style::GetSlateBrushSafe("MetasoundEditor.Graph.Node.Class.Graph");
 }
 
 const FLinearColor& FMetasoundGraphSchemaAction_NewNode::GetIconColor() const
@@ -466,8 +465,7 @@ FMetasoundGraphSchemaAction_NewInput::FMetasoundGraphSchemaAction_NewInput(FText
 
 const FSlateBrush* FMetasoundGraphSchemaAction_NewInput::GetIconBrush() const
 {
-	using namespace Metasound::Editor;
-	return &Style::GetSlateBrushSafe("MetasoundEditor.Graph.Node.Class.Input");
+	return &Metasound::Editor::Style::GetSlateBrushSafe("MetasoundEditor.Graph.Node.Class.Input");
 }
 
 const FLinearColor& FMetasoundGraphSchemaAction_NewInput::GetIconColor() const
