Fix up a couple more spots to use unified safe brush call
#rnx

diff --git a/Engine/Plugins/Runtime/Metasound/Source/MetasoundEditor/Private/MetasoundEditorGraphSchema.cpp b/Engine/Plugins/Runtime/Metasound/Source/MetasoundEditor/Private/MetasoundEditorGraphSchema.cpp
--- a/Engine/Plugins/Runtime/Metasound/Source/MetasoundEditor/Private/MetasoundEditorGraphSchema.cpp
+++ b/Engine/Plugins/Runtime/Metasound/Source/MetasoundEditor/Private/MetasoundEditorGraphSchema.cpp
@@ -685,12 +685,7 @@ FMetasoundGraphSchemaAction_NewOutput::FMetasoundGraphSchemaAction_NewOutput(FTe
 
 const FSlateBrush* FMetasoundGraphSchemaAction_NewOutput::GetIconBrush() const
 {
-	if (const ISlateStyle* MetasoundStyle = FSlateStyleRegistry::FindSlateStyle("MetaSoundStyle"))
-	{
-		return MetasoundStyle->GetBrush("MetasoundEditor.Graph.Node.Class.Output");
-	}
-
-	return Super::GetIconBrush();
+	return &Metasound::Editor::Style::GetSlateBrushSafe("MetasoundEditor.Graph.Node.Class.Output");
 }
 
 const FLinearColor& FMetasoundGraphSchemaAction_NewOutput::GetIconColor() const
@@ -837,8 +832,7 @@ UEdGraphNode* FMetasoundGraphSchemaAction_NewVariableNode::PerformAction(UEdGrap
 
 const FSlateBrush* FMetasoundGraphSchemaAction_NewVariableNode::GetIconBrush() const
 {
-	using namespace Metasound::Editor;
-	return &Style::GetSlateBrushSafe("MetasoundEditor.Graph.Node.Class.Variable");
+	return &Metasound::Editor::Style::GetSlateBrushSafe("MetasoundEditor.Graph.Node.Class.Variable");
 }
 
 const FLinearColor& FMetasoundGraphSchemaAction_NewVariableNode::GetIconColor() const
@@ -951,8 +945,7 @@ UEdGraphNode* FMetasoundGraphSchemaAction_NewReroute::PerformAction(UEdGraph* Pa
 
 const FSlateBrush* FMetasoundGraphSchemaAction_NewReroute::GetIconBrush() const
 {
-	using namespace Metasound::Editor;
-	return &Style::GetSlateBrushSafe("MetasoundEditor.Graph.Node.Class.Reroute");
+	return &Metasound::Editor::Style::GetSlateBrushSafe("MetasoundEditor.Graph.Node.Class.Reroute");
 }
 
 const FLinearColor& FMetasoundGraphSchemaAction_NewReroute::GetIconColor() const
