MetaSound Source Surround Channel Format Support.

diff --git a/Engine/Plugins/Runtime/Metasound/Source/MetasoundEngine/Private/MetasoundInterface.cpp b/Engine/Plugins/Runtime/Metasound/Source/MetasoundEngine/Private/MetasoundInterface.cpp
--- a/Engine/Plugins/Runtime/Metasound/Source/MetasoundEngine/Private/MetasoundInterface.cpp
+++ b/Engine/Plugins/Runtime/Metasound/Source/MetasoundEngine/Private/MetasoundInterface.cpp
@@ -255,6 +255,10 @@ namespace Metasound
 				RegisterInterface<UMetaSoundSource>(MetasoundOutputFormatMonoV1_2::GetInterface(), MakeUnique<MetasoundOutputFormatMonoV1_2::FUpdateInterface>(), bIsDefault, bEditorCanAddOrRemove, IDataReference::RouterName);
 
 				RegisterInterface(OutputFormatStereoInterface::CreateInterface(*UMetaSoundSource::StaticClass()), nullptr, bIsDefault, bEditorCanAddOrRemove, IDataReference::RouterName);			
+				RegisterInterface(OutputFormatQuadInterface::CreateInterface(*UMetaSoundSource::StaticClass()), nullptr, bIsDefault, bEditorCanAddOrRemove, IDataReference::RouterName);			
+				RegisterInterface(OutputFormatFiveDotZeroInterface::CreateInterface(*UMetaSoundSource::StaticClass()), nullptr, bIsDefault, bEditorCanAddOrRemove, IDataReference::RouterName);			
+				RegisterInterface(OutputFormatFiveDotOneInterface::CreateInterface(*UMetaSoundSource::StaticClass()), nullptr, bIsDefault, bEditorCanAddOrRemove, IDataReference::RouterName);			
+				RegisterInterface(OutputFormatSevenDotOneInterface::CreateInterface(*UMetaSoundSource::StaticClass()), nullptr, bIsDefault, bEditorCanAddOrRemove, IDataReference::RouterName);			
 			}
 
 			// Register External Interfaces (Interfaces defined externally & can be managed directly by end-user).
diff --git a/Engine/Plugins/Runtime/Metasound/Source/MetasoundEngine/Private/MetasoundSource.cpp b/Engine/Plugins/Runtime/Metasound/Source/MetasoundEngine/Private/MetasoundSource.cpp
--- a/Engine/Plugins/Runtime/Metasound/Source/MetasoundEngine/Private/MetasoundSource.cpp
+++ b/Engine/Plugins/Runtime/Metasound/Source/MetasoundEngine/Private/MetasoundSource.cpp
@@ -43,35 +43,107 @@ namespace Metasound
 {
 	namespace SourcePrivate
 	{
-		using FFormatOutputVertexKeyMap = TMap<EMetasoundSourceAudioFormat, TArray<Metasound::FVertexName>>;
+		// Contains information on output audio formats
+		struct FOutputFormatInfo
+		{
+			FMetasoundFrontendVersion InterfaceVersion;
+			TArray<Metasound::FVertexName> OutputVertexChannelOrder;
+		};
+
+		using FFormatInfoMap = TMap<EMetasoundSourceAudioFormat, FOutputFormatInfo>;
+		using FFormatInfoPair = TPair<EMetasoundSourceAudioFormat, FOutputFormatInfo>;
 
-		const FFormatOutputVertexKeyMap& GetFormatOutputVertexKeys()
+		// Return a map containing all the supported audio formats for a MetaSound
+		// Source. 
+		const FFormatInfoMap& GetFormatInfoMap()
 		{
-			auto CreateVertexKeyMap = []()
+			auto CreateFormatInfoMap = []()
 			{
 				using namespace Metasound::Frontend;
 
-				return FFormatOutputVertexKeyMap
+				return FFormatInfoMap
 				{
 					{
 						EMetasoundSourceAudioFormat::Mono,
 						{
-							OutputFormatMonoInterface::Outputs::MonoOut
+							OutputFormatMonoInterface::GetVersion(),
+							{
+								OutputFormatMonoInterface::Outputs::MonoOut
+							}
 						}
 					},
 					{
 						EMetasoundSourceAudioFormat::Stereo,
 						{
-							OutputFormatStereoInterface::Outputs::LeftOut,
-							OutputFormatStereoInterface::Outputs::RightOut,
+							OutputFormatStereoInterface::GetVersion(),
+							{
+								OutputFormatStereoInterface::Outputs::LeftOut,
+								OutputFormatStereoInterface::Outputs::RightOut
+							}
+						}
+					},
+					{
+						EMetasoundSourceAudioFormat::Quad,
+						{
+							OutputFormatQuadInterface::GetVersion(),
+							{
+								OutputFormatQuadInterface::Outputs::FrontLeftOut,
+								OutputFormatQuadInterface::Outputs::FrontRightOut,
+								OutputFormatQuadInterface::Outputs::SideLeftOut,
+								OutputFormatQuadInterface::Outputs::SideRightOut
+							}
+						}
+					},
+					{
+						EMetasoundSourceAudioFormat::FiveDotZero,
+						{
+							OutputFormatFiveDotZeroInterface::GetVersion(),
+							{
+								OutputFormatFiveDotZeroInterface::Outputs::FrontLeftOut,
+								OutputFormatFiveDotZeroInterface::Outputs::FrontRightOut,
+								OutputFormatFiveDotZeroInterface::Outputs::FrontCenterOut,
+								OutputFormatFiveDotZeroInterface::Outputs::SideLeftOut,
+								OutputFormatFiveDotZeroInterface::Outputs::SideRightOut
+							}
+						}
+					},
+					{
+						EMetasoundSourceAudioFormat::FiveDotOne,
+						{
+							OutputFormatFiveDotOneInterface::GetVersion(),
+							{
+								OutputFormatFiveDotOneInterface::Outputs::FrontLeftOut,
+								OutputFormatFiveDotOneInterface::Outputs::FrontRightOut,
+								OutputFormatFiveDotOneInterface::Outputs::FrontCenterOut,
+								OutputFormatFiveDotOneInterface::Outputs::LowFrequencyOut,
+								OutputFormatFiveDotOneInterface::Outputs::SideLeftOut,
+								OutputFormatFiveDotOneInterface::Outputs::SideRightOut
+							}
+						}
+					},
+					{
+						EMetasoundSourceAudioFormat::SevenDotOne,
+						{
+							OutputFormatSevenDotOneInterface::GetVersion(),
+							{
+								OutputFormatSevenDotOneInterface::Outputs::FrontLeftOut,
+								OutputFormatSevenDotOneInterface::Outputs::FrontRightOut,
+								OutputFormatSevenDotOneInterface::Outputs::FrontCenterOut,
+								OutputFormatSevenDotOneInterface::Outputs::LowFrequencyOut,
+								OutputFormatSevenDotOneInterface::Outputs::SideLeftOut,
+								OutputFormatSevenDotOneInterface::Outputs::SideRightOut,
+								OutputFormatSevenDotOneInterface::Outputs::BackLeftOut,
+								OutputFormatSevenDotOneInterface::Outputs::BackRightOut
+							}
 						}
 					}
 				};
 			};
-			static const FFormatOutputVertexKeyMap Map = CreateVertexKeyMap();
+
+			static const FFormatInfoMap Map = CreateFormatInfoMap();
 			return Map;
 		}
-
+		
 		Frontend::FMetaSoundAssetRegistrationOptions GetInitRegistrationOptions()
 		{
 			Frontend::FMetaSoundAssetRegistrationOptions RegOptions;
@@ -83,6 +155,25 @@ namespace Metasound
 
 			return RegOptions;
 		}
+
+		// Return an array of all the audio format versions.
+		const TArray<FMetasoundFrontendVersion>& GetFormatInterfaceVersions()
+		{
+			auto CreateFormatInterfaceVersions = []() -> TArray<FMetasoundFrontendVersion>
+			{
+				TArray<FMetasoundFrontendVersion> FormatVersions;
+				const FFormatInfoMap& FormatMap = GetFormatInfoMap();
+				for (const auto& Pair : FormatMap)
+				{
+					FormatVersions.Add(Pair.Value.InterfaceVersion);
+				}
+				return FormatVersions;
+			};
+
+			static const TArray<FMetasoundFrontendVersion> Versions = CreateFormatInterfaceVersions();
+			return Versions;
+		}
+
 	} // namespace SourcePrivate
 } // namespace Metasound
 
@@ -96,7 +187,6 @@ UMetaSoundSource::UMetaSoundSource(const FObjectInitializer& ObjectInitializer)
 
 	// todo: ensure that we have a method so that the audio engine can be authoritative over the sample rate the UMetaSoundSource runs at.
 	SampleRate = 48000.f;
-
 }
 
 #if WITH_EDITOR
@@ -121,85 +211,87 @@ void UMetaSoundSource::PostDuplicate(EDuplicateMode::Type InDuplicateMode)
 
 void UMetaSoundSource::PostEditChangeProperty(FPropertyChangedEvent& InEvent)
 {
-	using namespace Metasound;
-	using namespace Metasound::Engine;
-	using namespace Metasound::Frontend;
-
 	Super::PostEditChangeProperty(InEvent);
 
 	if (InEvent.GetPropertyName() == GET_MEMBER_NAME_CHECKED(UMetaSoundSource, OutputFormat))
 	{
-		ConvertFromPreset();
+		PostEditChangeOutputFormat();
+	}
+}
 
-		bool bDidModifyDocument = false;
-		switch (OutputFormat)
-		{
-			case EMetasoundSourceAudioFormat::Mono:
-			{
-				bDidModifyDocument = FModifyRootGraphInterfaces(
-					{ OutputFormatStereoInterface::GetVersion() },
-					{ OutputFormatMonoInterface::GetVersion() }
-				).Transform(GetDocumentHandle());
-			}
-			break;
+void UMetaSoundSource::PostEditChangeOutputFormat()
+{
+	using namespace Metasound::SourcePrivate;
+	using namespace Metasound::Frontend;
 
-			case EMetasoundSourceAudioFormat::Stereo:
-			{
-				bDidModifyDocument = FModifyRootGraphInterfaces(
-					{ OutputFormatMonoInterface::GetVersion() },
-					{ OutputFormatStereoInterface::GetVersion() }
-				).Transform(GetDocumentHandle());
-			}
-			break;
+	// If this is a preset, convert to normal metasound source since it is being
+	// altered. 
+	ConvertFromPreset();
 
-			default:
+	// Determine which interfaces to add and remove from the document due to the
+	// output format being changed.
+	TArray<FMetasoundFrontendVersion> OutputFormatsToAdd;
+	if (const FOutputFormatInfo* FormatInfo = GetFormatInfoMap().Find(OutputFormat))
+	{
+		OutputFormatsToAdd.Add(FormatInfo->InterfaceVersion);
+	}
+
+	TArray<FMetasoundFrontendVersion> OutputFormatsToRemove;
+	for (const FMetasoundFrontendVersion& FormatVersion : GetFormatInterfaceVersions())
+	{
+		if (RootMetasoundDocument.Interfaces.Contains(FormatVersion))
+		{
+			if (!OutputFormatsToAdd.Contains(FormatVersion))
 			{
-				static_assert(static_cast<int32>(EMetasoundSourceAudioFormat::COUNT) == 2, "Possible missing format switch case coverage.");
+				OutputFormatsToRemove.Add(FormatVersion);
 			}
-			break;
 		}
+	}
 
-		if (bDidModifyDocument)
-		{
-			ConformObjectDataToInterfaces();
+	// Add and/or remove interfaces from the root document.
+	const bool bDidModifyDocument = FModifyRootGraphInterfaces(OutputFormatsToRemove, OutputFormatsToAdd).Transform(GetDocumentHandle());
 
-			// Use the editor form of register to ensure other editors'
-			// MetaSounds are auto-updated if they are referencing this graph.
-			if (Graph)
-			{
-				Graph->RegisterGraphWithFrontend();
-			}
-			MarkMetasoundDocumentDirty();
+	if (bDidModifyDocument)
+	{
+		// Update the data in this UMetaSoundSource to reflect what is in the metasound document.
+		ConformObjectDataToInterfaces();
+
+		// Use the editor form of register to ensure other editors'
+		// MetaSounds are auto-updated if they are referencing this graph.
+		if (Graph)
+		{
+			Graph->RegisterGraphWithFrontend();
 		}
+		MarkMetasoundDocumentDirty();
 	}
 }
 #endif // WITH_EDITOR
 
 bool UMetaSoundSource::ConformObjectDataToInterfaces()
 {
-	using namespace Metasound::Frontend;
+	using namespace Metasound::SourcePrivate;
 
-	if (IsInterfaceDeclared(OutputFormatMonoInterface::GetVersion()))
-	{
-		if (OutputFormat != EMetasoundSourceAudioFormat::Mono || NumChannels != 1)
-		{
-			OutputFormat = EMetasoundSourceAudioFormat::Mono;
-			NumChannels = 1;
-			return true;
-		}
-	}
+	bool bDidAlterObjectData = false;
 
-	if (IsInterfaceDeclared(OutputFormatStereoInterface::GetVersion()))
+	// Update the OutputFormat and NumChannels to match the audio format interface
+	// on the root document.
+	const FFormatInfoMap& FormatInfo = GetFormatInfoMap();
+	for (const FFormatInfoPair& Pair : FormatInfo)
 	{
-		if (OutputFormat != EMetasoundSourceAudioFormat::Stereo || NumChannels != 2)
+		if (RootMetasoundDocument.Interfaces.Contains(Pair.Value.InterfaceVersion))
 		{
-			OutputFormat = EMetasoundSourceAudioFormat::Stereo;
-			NumChannels = 2;
-			return true;
+			if ((OutputFormat != Pair.Key) || (NumChannels != Pair.Value.OutputVertexChannelOrder.Num()))
+			{
+				OutputFormat = Pair.Key;
+				NumChannels = Pair.Value.OutputVertexChannelOrder.Num();
+				bDidAlterObjectData = true;
+			}
+
+			break;
 		}
 	}
 
-	return false;
+	return bDidAlterObjectData;
 }
 
 
@@ -547,7 +639,7 @@ ISoundGeneratorPtr UMetaSoundSource::CreateSoundGenerator(const FSoundGeneratorI
 		MetasoundGraph,
 		Environment,
 		GetName(),
-		GetAudioOutputVertexKeys()
+		GetOutputAudioChannelOrder()
 	};
 
 	return ISoundGeneratorPtr(new FMetasoundGenerator(MoveTemp(InitParams)));
@@ -873,13 +965,13 @@ Metasound::FMetasoundEnvironment UMetaSoundSource::CreateEnvironment(const Audio
 	return Environment;
 }
 
-const TArray<Metasound::FVertexName>& UMetaSoundSource::GetAudioOutputVertexKeys() const
+const TArray<Metasound::FVertexName>& UMetaSoundSource::GetOutputAudioChannelOrder() const
 {
 	using namespace Metasound::SourcePrivate;
 
-	if (const TArray<Metasound::FVertexName>* ArrayKeys = GetFormatOutputVertexKeys().Find(OutputFormat))
+	if (const FOutputFormatInfo* FormatInfo = GetFormatInfoMap().Find(OutputFormat))
 	{
-		return *ArrayKeys;
+		return FormatInfo->OutputVertexChannelOrder;
 	}
 	else
 	{
diff --git a/Engine/Plugins/Runtime/Metasound/Source/MetasoundEngine/Public/MetasoundSource.h b/Engine/Plugins/Runtime/Metasound/Source/MetasoundEngine/Public/MetasoundSource.h
--- a/Engine/Plugins/Runtime/Metasound/Source/MetasoundEngine/Public/MetasoundSource.h
+++ b/Engine/Plugins/Runtime/Metasound/Source/MetasoundEngine/Public/MetasoundSource.h
@@ -22,11 +22,12 @@
 UENUM()
 enum class EMetasoundSourceAudioFormat : uint8
 {
-	// Mono audio output
 	Mono,
-
-	// Stereo audio output
 	Stereo,
+	Quad,
+	FiveDotZero UMETA(DisplayName="5.0"),
+	FiveDotOne UMETA(DisplayName="5.1"),
+	SevenDotOne UMETA(DisplayName="7.1"),
 
 	COUNT UMETA(Hidden)
 };
@@ -132,6 +133,11 @@ public:
 	virtual void PostDuplicate(EDuplicateMode::Type InDuplicateMode) override;
 
 	virtual void PostEditChangeProperty(FPropertyChangedEvent& InEvent) override;
+
+private:
+	void PostEditChangeOutputFormat();
+public:
+
 #endif // WITH_EDITOR
 
 	virtual const TSet<FString>& GetReferencedAssetClassKeys() const override
@@ -200,9 +206,10 @@ protected:
 	}
 
 private:
+
 	Metasound::FOperatorSettings GetOperatorSettings(Metasound::FSampleRate InSampleRate) const;
 	Metasound::FMetasoundEnvironment CreateEnvironment() const;
 	Metasound::FMetasoundEnvironment CreateEnvironment(const FSoundGeneratorInitParams& InParams) const;
 	Metasound::FMetasoundEnvironment CreateEnvironment(const Audio::FParameterTransmitterInitParams& InParams) const;
-	const TArray<Metasound::FVertexName>& GetAudioOutputVertexKeys() const;
+	const TArray<Metasound::FVertexName>& GetOutputAudioChannelOrder() const;
 };
diff --git a/Engine/Plugins/Runtime/Metasound/Source/MetasoundFrontend/Private/MetasoundOutputFormatInterfaces.cpp b/Engine/Plugins/Runtime/Metasound/Source/MetasoundFrontend/Private/MetasoundOutputFormatInterfaces.cpp
--- a/Engine/Plugins/Runtime/Metasound/Source/MetasoundFrontend/Private/MetasoundOutputFormatInterfaces.cpp
+++ b/Engine/Plugins/Runtime/Metasound/Source/MetasoundFrontend/Private/MetasoundOutputFormatInterfaces.cpp
@@ -13,6 +13,113 @@ namespace Metasound
 {
 	namespace Frontend
 	{
+		namespace MetasoundOutputFormatInterfacesPrivate
+		{
+			Audio::FParameterInterface::FOutput GetFrontLeftOutput(const FName& InVertexName)  
+			{
+				return {
+					LOCTEXT("OutputFormatSurroundInterface_GeneratedFrontLeftDisplayName", "Out Front Left"),
+					LOCTEXT("OutputFormatSurroundInterface_GeneratedFrontLeftDescription", "The resulting front left channel output audio from this source."),
+					GetMetasoundDataTypeName<FAudioBuffer>(),
+					InVertexName,
+					FText::GetEmpty(), // RequiredText
+					EAudioParameterType::None,
+					100 // SortOrder
+				};
+			}
+
+			Audio::FParameterInterface::FOutput GetFrontRightOutput(const FName& InVertexName)  
+			{
+				return {
+					LOCTEXT("OutputFormatSurroundInterface_GeneratedFrontRightDisplayName", "Out Front Right"),
+					LOCTEXT("OutputFormatSurroundInterface_GeneratedFrontRightDescription", "The resulting front right channel output audio from this source."),
+					GetMetasoundDataTypeName<FAudioBuffer>(),
+					InVertexName,
+					FText::GetEmpty(), // RequiredText
+					EAudioParameterType::None,
+					101
+				};
+			}
+
+			Audio::FParameterInterface::FOutput GetFrontCenterOutput(const FName& InVertexName)  
+			{
+				return {
+					LOCTEXT("OutputFormatSurroundInterface_GeneratedCenterDisplayName", "Out Front Center"),
+					LOCTEXT("OutputFormatSurroundInterface_GeneratedCenterDescription", "The resulting front center channel output audio from this source."),
+					GetMetasoundDataTypeName<FAudioBuffer>(),
+					InVertexName,
+					FText::GetEmpty(), // RequiredText
+					EAudioParameterType::None,
+					102
+				};
+			}
+
+			Audio::FParameterInterface::FOutput GetLowFrequencyOutput(const FName& InVertexName)  
+			{
+				return {
+					LOCTEXT("OutputFormatSurroundInterface_GeneratedLowFrequencyDisplayName", "Out Low Frequency"),
+					LOCTEXT("OutputFormatSurroundInterface_GeneratedLowFrequencyDescription", "The resulting low frequency channel output audio from this source."),
+					GetMetasoundDataTypeName<FAudioBuffer>(),
+					InVertexName,
+					FText::GetEmpty(), // RequiredText
+					EAudioParameterType::None,
+					103
+				};
+			}
+
+			Audio::FParameterInterface::FOutput GetSideLeftOutput(const FName& InVertexName)  
+			{
+				return {
+					LOCTEXT("OutputFormatSurroundInterface_GeneratedSideLeftDisplayName", "Out Side Left"),
+					LOCTEXT("OutputFormatSurroundInterface_GeneratedSideLeftDescription", "The resulting side left channel output audio from this source."),
+					GetMetasoundDataTypeName<FAudioBuffer>(),
+					InVertexName,
+					FText::GetEmpty(), // RequiredText
+					EAudioParameterType::None,
+					104
+				};
+			}
+
+			Audio::FParameterInterface::FOutput GetSideRightOutput(const FName& InVertexName)  
+			{
+				return {
+					LOCTEXT("OutputFormatSurroundInterface_GeneratedSideRightDisplayName", "Out Side Right"),
+					LOCTEXT("OutputFormatSurroundInterface_GeneratedSideRightDescription", "The resulting side right channel output audio from this source."),
+					GetMetasoundDataTypeName<FAudioBuffer>(),
+					InVertexName,
+					FText::GetEmpty(), // RequiredText
+					EAudioParameterType::None,
+					105
+				};
+			}
+
+			Audio::FParameterInterface::FOutput GetBackLeftOutput(const FName& InVertexName)  
+			{
+				return {
+					LOCTEXT("OutputFormatSurroundInterface_GeneratedBackLeftDisplayName", "Out Back Left"),
+					LOCTEXT("OutputFormatSurroundInterface_GeneratedBackLeftDescription", "The resulting back left channel output audio from this source."),
+					GetMetasoundDataTypeName<FAudioBuffer>(),
+					InVertexName,
+					FText::GetEmpty(), // RequiredText
+					EAudioParameterType::None,
+					106
+				};
+			}
+
+			Audio::FParameterInterface::FOutput GetBackRightOutput(const FName& InVertexName)  
+			{
+				return {
+					LOCTEXT("OutputFormatSurroundInterface_GeneratedBackRightDisplayName", "Out Back Right"),
+					LOCTEXT("OutputFormatSurroundInterface_GeneratedBackRightDescription", "The resulting back right channel output audio from this source."),
+					GetMetasoundDataTypeName<FAudioBuffer>(),
+					InVertexName,
+					FText::GetEmpty(), // RequiredText
+					EAudioParameterType::None,
+					107
+				};
+			}
+		} // namespace MetasoundOutputFormatInterfacesPrivate
+
 #define AUDIO_PARAMETER_INTERFACE_NAMESPACE "UE.OutputFormat.Mono"
 		namespace OutputFormatMonoInterface
 		{
@@ -85,7 +192,7 @@ namespace Metasound
 								Outputs::LeftOut,
 								FText::GetEmpty(), // RequiredText
 								EAudioParameterType::None,
-								100
+								100 // Sort order
 							},
 							{
 								LOCTEXT("OutputFormatStereoInterface_GeneratedRightDisplayName", "Out Right"),
@@ -94,7 +201,7 @@ namespace Metasound
 								Outputs::RightOut,
 								FText::GetEmpty(), // RequiredText
 								EAudioParameterType::None,
-								101
+								101 // Sort order
 							}
 						};
 					}
@@ -104,6 +211,181 @@ namespace Metasound
 			}
 		} // namespace OutputFormatStereoInterface
 #undef AUDIO_PARAMETER_INTERFACE_NAMESPACE
+
+#define AUDIO_PARAMETER_INTERFACE_NAMESPACE "UE.OutputFormat.Quad"
+		namespace OutputFormatQuadInterface
+		{
+			const FMetasoundFrontendVersion& GetVersion()
+			{
+				static const FMetasoundFrontendVersion Version = { AUDIO_PARAMETER_INTERFACE_NAMESPACE, { 1, 0 } };
+				return Version;
+			}
+
+			namespace Outputs
+			{
+				const FName FrontLeftOut = AUDIO_PARAMETER_INTERFACE_MEMBER_DEFINE("Audio:0");
+				const FName FrontRightOut = AUDIO_PARAMETER_INTERFACE_MEMBER_DEFINE("Audio:1");
+				const FName SideLeftOut = AUDIO_PARAMETER_INTERFACE_MEMBER_DEFINE("Audio:4");
+				const FName SideRightOut = AUDIO_PARAMETER_INTERFACE_MEMBER_DEFINE("Audio:5");
+			}
+
+			Audio::FParameterInterfacePtr CreateInterface(const UClass& InUClass)
+			{
+				struct FInterface : public Audio::FParameterInterface
+				{
+					FInterface(const UClass& InAssetClass)
+						: FParameterInterface(OutputFormatQuadInterface::GetVersion().Name, OutputFormatQuadInterface::GetVersion().Number.ToInterfaceVersion(), InAssetClass)
+					{
+						using namespace MetasoundOutputFormatInterfacesPrivate;
+
+						Outputs =
+						{
+							GetFrontLeftOutput(Outputs::FrontLeftOut),
+							GetFrontRightOutput(Outputs::FrontRightOut),
+							GetSideLeftOutput(Outputs::SideLeftOut),
+							GetSideRightOutput(Outputs::SideRightOut)
+						};
+					}
+				};
+
+				return MakeShared<FInterface>(InUClass);
+			}
+		} // namespace OutputFormatQuadInterface
+#undef AUDIO_PARAMETER_INTERFACE_NAMESPACE
+
+#define AUDIO_PARAMETER_INTERFACE_NAMESPACE "UE.OutputFormat.5dot0"
+		namespace OutputFormatFiveDotZeroInterface
+		{
+			const FMetasoundFrontendVersion& GetVersion()
+			{
+				static const FMetasoundFrontendVersion Version = { AUDIO_PARAMETER_INTERFACE_NAMESPACE, { 1, 0 } };
+				return Version;
+			}
+
+			namespace Outputs
+			{
+				const FName FrontLeftOut = AUDIO_PARAMETER_INTERFACE_MEMBER_DEFINE("Audio:0");
+				const FName FrontRightOut = AUDIO_PARAMETER_INTERFACE_MEMBER_DEFINE("Audio:1");
+				const FName FrontCenterOut = AUDIO_PARAMETER_INTERFACE_MEMBER_DEFINE("Audio:2");
+				const FName SideLeftOut = AUDIO_PARAMETER_INTERFACE_MEMBER_DEFINE("Audio:4");
+				const FName SideRightOut = AUDIO_PARAMETER_INTERFACE_MEMBER_DEFINE("Audio:5");
+			}
+
+			Audio::FParameterInterfacePtr CreateInterface(const UClass& InUClass)
+			{
+				struct FInterface : public Audio::FParameterInterface
+				{
+					FInterface(const UClass& InAssetClass)
+						: FParameterInterface(OutputFormatFiveDotZeroInterface::GetVersion().Name, OutputFormatFiveDotZeroInterface::GetVersion().Number.ToInterfaceVersion(), InAssetClass)
+					{
+						using namespace MetasoundOutputFormatInterfacesPrivate;
+						Outputs =
+						{
+							GetFrontLeftOutput(Outputs::FrontLeftOut),
+							GetFrontRightOutput(Outputs::FrontRightOut),
+							GetFrontCenterOutput(Outputs::FrontCenterOut),
+							GetSideLeftOutput(Outputs::SideLeftOut),
+							GetSideRightOutput(Outputs::SideRightOut)
+						};
+					}
+				};
+
+				return MakeShared<FInterface>(InUClass);
+			}
+		} // namespace OutputFormatFiveDotZeroInterface
+#undef AUDIO_PARAMETER_INTERFACE_NAMESPACE
+
+#define AUDIO_PARAMETER_INTERFACE_NAMESPACE "UE.OutputFormat.5dot1"
+		namespace OutputFormatFiveDotOneInterface
+		{
+			const FMetasoundFrontendVersion& GetVersion()
+			{
+				static const FMetasoundFrontendVersion Version = { AUDIO_PARAMETER_INTERFACE_NAMESPACE, { 1, 0 } };
+				return Version;
+			}
+
+			namespace Outputs
+			{
+				const FName FrontLeftOut = AUDIO_PARAMETER_INTERFACE_MEMBER_DEFINE("Audio:0");
+				const FName FrontRightOut = AUDIO_PARAMETER_INTERFACE_MEMBER_DEFINE("Audio:1");
+				const FName FrontCenterOut = AUDIO_PARAMETER_INTERFACE_MEMBER_DEFINE("Audio:2");
+				const FName LowFrequencyOut = AUDIO_PARAMETER_INTERFACE_MEMBER_DEFINE("Audio:3");
+				const FName SideLeftOut = AUDIO_PARAMETER_INTERFACE_MEMBER_DEFINE("Audio:4");
+				const FName SideRightOut = AUDIO_PARAMETER_INTERFACE_MEMBER_DEFINE("Audio:5");
+			}
+
+			Audio::FParameterInterfacePtr CreateInterface(const UClass& InUClass)
+			{
+				struct FInterface : public Audio::FParameterInterface
+				{
+					FInterface(const UClass& InAssetClass)
+						: FParameterInterface(OutputFormatFiveDotOneInterface::GetVersion().Name, OutputFormatFiveDotOneInterface::GetVersion().Number.ToInterfaceVersion(), InAssetClass)
+					{
+						using namespace MetasoundOutputFormatInterfacesPrivate;
+						Outputs =
+						{
+							GetFrontLeftOutput(Outputs::FrontLeftOut),
+							GetFrontRightOutput(Outputs::FrontRightOut),
+							GetFrontCenterOutput(Outputs::FrontCenterOut),
+							GetLowFrequencyOutput(Outputs::LowFrequencyOut),
+							GetSideLeftOutput(Outputs::SideLeftOut),
+							GetSideRightOutput(Outputs::SideRightOut)
+						};
+					}
+				};
+
+				return MakeShared<FInterface>(InUClass);
+			}
+		} // namespace OutputFormatFiveDotOneInterface
+#undef AUDIO_PARAMETER_INTERFACE_NAMESPACE
+
+#define AUDIO_PARAMETER_INTERFACE_NAMESPACE "UE.OutputFormat.7dot1"
+		namespace OutputFormatSevenDotOneInterface
+		{
+			const FMetasoundFrontendVersion& GetVersion()
+			{
+				static const FMetasoundFrontendVersion Version = { AUDIO_PARAMETER_INTERFACE_NAMESPACE, { 1, 0 } };
+				return Version;
+			}
+
+			namespace Outputs
+			{
+				const FName FrontLeftOut = AUDIO_PARAMETER_INTERFACE_MEMBER_DEFINE("Audio:0");
+				const FName FrontRightOut = AUDIO_PARAMETER_INTERFACE_MEMBER_DEFINE("Audio:1");
+				const FName FrontCenterOut = AUDIO_PARAMETER_INTERFACE_MEMBER_DEFINE("Audio:2");
+				const FName LowFrequencyOut = AUDIO_PARAMETER_INTERFACE_MEMBER_DEFINE("Audio:3");
+				const FName SideLeftOut = AUDIO_PARAMETER_INTERFACE_MEMBER_DEFINE("Audio:4");
+				const FName SideRightOut = AUDIO_PARAMETER_INTERFACE_MEMBER_DEFINE("Audio:5");
+				const FName BackLeftOut = AUDIO_PARAMETER_INTERFACE_MEMBER_DEFINE("Audio:6");
+				const FName BackRightOut = AUDIO_PARAMETER_INTERFACE_MEMBER_DEFINE("Audio:7");
+			}
+
+			Audio::FParameterInterfacePtr CreateInterface(const UClass& InUClass)
+			{
+				struct FInterface : public Audio::FParameterInterface
+				{
+					FInterface(const UClass& InAssetClass)
+						: FParameterInterface(OutputFormatSevenDotOneInterface::GetVersion().Name, OutputFormatSevenDotOneInterface::GetVersion().Number.ToInterfaceVersion(), InAssetClass)
+					{
+						using namespace MetasoundOutputFormatInterfacesPrivate;
+						Outputs =
+						{
+							GetFrontLeftOutput(Outputs::FrontLeftOut),
+							GetFrontRightOutput(Outputs::FrontRightOut),
+							GetFrontCenterOutput(Outputs::FrontCenterOut),
+							GetLowFrequencyOutput(Outputs::LowFrequencyOut),
+							GetSideLeftOutput(Outputs::SideLeftOut),
+							GetSideRightOutput(Outputs::SideRightOut),
+							GetBackLeftOutput(Outputs::BackLeftOut),
+							GetBackRightOutput(Outputs::BackRightOut)
+						};
+					}
+				};
+
+				return MakeShared<FInterface>(InUClass);
+			}
+		} // namespace OutputFormatFiveDotOneInterface
+#undef AUDIO_PARAMETER_INTERFACE_NAMESPACE
 	}
 }
-#undef LOCTEXT_NAMESPACE
\ No newline at end of file
+#undef LOCTEXT_NAMESPACE
diff --git a/Engine/Plugins/Runtime/Metasound/Source/MetasoundFrontend/Public/MetasoundOutputFormatInterfaces.h b/Engine/Plugins/Runtime/Metasound/Source/MetasoundFrontend/Public/MetasoundOutputFormatInterfaces.h
--- a/Engine/Plugins/Runtime/Metasound/Source/MetasoundFrontend/Public/MetasoundOutputFormatInterfaces.h
+++ b/Engine/Plugins/Runtime/Metasound/Source/MetasoundFrontend/Public/MetasoundOutputFormatInterfaces.h
@@ -33,6 +33,69 @@ namespace Metasound
 
 			METASOUNDFRONTEND_API const FMetasoundFrontendVersion& GetVersion();
 			METASOUNDFRONTEND_API Audio::FParameterInterfacePtr CreateInterface(const UClass& InClass);
-		} // namespace OutputFormatMonoInterface
+		} // namespace OutputFormatStereoInterface
+
+		namespace OutputFormatQuadInterface
+		{
+			namespace Outputs
+			{
+				METASOUNDFRONTEND_API const extern FName FrontLeftOut;
+				METASOUNDFRONTEND_API const extern FName FrontRightOut;
+				METASOUNDFRONTEND_API const extern FName SideLeftOut;
+				METASOUNDFRONTEND_API const extern FName SideRightOut;
+			} // namespace Outputs
+
+			METASOUNDFRONTEND_API const FMetasoundFrontendVersion& GetVersion();
+			METASOUNDFRONTEND_API Audio::FParameterInterfacePtr CreateInterface(const UClass& InClass);
+		} // namespace OutputFormatQuadInterface
+
+		namespace OutputFormatFiveDotZeroInterface
+		{
+			namespace Outputs
+			{
+				METASOUNDFRONTEND_API const extern FName FrontLeftOut;
+				METASOUNDFRONTEND_API const extern FName FrontRightOut;
+				METASOUNDFRONTEND_API const extern FName FrontCenterOut;
+				METASOUNDFRONTEND_API const extern FName SideLeftOut;
+				METASOUNDFRONTEND_API const extern FName SideRightOut;
+			} // namespace Outputs
+
+			METASOUNDFRONTEND_API const FMetasoundFrontendVersion& GetVersion();
+			METASOUNDFRONTEND_API Audio::FParameterInterfacePtr CreateInterface(const UClass& InClass);
+		} // namespace OutputFormatFiveDotZeroInterface
+
+		namespace OutputFormatFiveDotOneInterface
+		{
+			namespace Outputs
+			{
+				METASOUNDFRONTEND_API const extern FName FrontLeftOut;
+				METASOUNDFRONTEND_API const extern FName FrontRightOut;
+				METASOUNDFRONTEND_API const extern FName FrontCenterOut;
+				METASOUNDFRONTEND_API const extern FName LowFrequencyOut;
+				METASOUNDFRONTEND_API const extern FName SideLeftOut;
+				METASOUNDFRONTEND_API const extern FName SideRightOut;
+			} // namespace Outputs
+
+			METASOUNDFRONTEND_API const FMetasoundFrontendVersion& GetVersion();
+			METASOUNDFRONTEND_API Audio::FParameterInterfacePtr CreateInterface(const UClass& InClass);
+		} // namespace OutputFormatFiveDotOneInterface
+
+		namespace OutputFormatSevenDotOneInterface
+		{
+			namespace Outputs
+			{
+				METASOUNDFRONTEND_API const extern FName FrontLeftOut;
+				METASOUNDFRONTEND_API const extern FName FrontRightOut;
+				METASOUNDFRONTEND_API const extern FName FrontCenterOut;
+				METASOUNDFRONTEND_API const extern FName LowFrequencyOut;
+				METASOUNDFRONTEND_API const extern FName SideLeftOut;
+				METASOUNDFRONTEND_API const extern FName SideRightOut;
+				METASOUNDFRONTEND_API const extern FName BackLeftOut;
+				METASOUNDFRONTEND_API const extern FName BackRightOut;
+			} // namespace Outputs
+
+			METASOUNDFRONTEND_API const FMetasoundFrontendVersion& GetVersion();
+			METASOUNDFRONTEND_API Audio::FParameterInterfacePtr CreateInterface(const UClass& InClass);
+		} // namespace OutputFormatSevenDotOneInterface
 	} // namespace Frontend
 } // namespace Metasound
diff --git a/Engine/Plugins/Runtime/Metasound/Source/MetasoundGenerator/Private/MetasoundGenerator.cpp b/Engine/Plugins/Runtime/Metasound/Source/MetasoundGenerator/Private/MetasoundGenerator.cpp
--- a/Engine/Plugins/Runtime/Metasound/Source/MetasoundGenerator/Private/MetasoundGenerator.cpp
+++ b/Engine/Plugins/Runtime/Metasound/Source/MetasoundGenerator/Private/MetasoundGenerator.cpp
@@ -262,7 +262,7 @@ namespace Metasound
 
 	int32 FMetasoundGenerator::GetNumChannels() const
 	{
-		return GraphOutputAudio.Num();
+		return NumChannels;
 	}
 
 	int32 FMetasoundGenerator::OnGenerateAudio(float* OutAudio, int32 NumSamplesRemaining)
@@ -351,8 +351,7 @@ namespace Metasound
 
 	int32 FMetasoundGenerator::GetDesiredNumSamplesToRenderPerCallback() const
 	{
-		// TODO: may improve performance if this number is increased. 
-		return NumFramesPerExecute;
+		return NumFramesPerExecute * NumChannels;
 	}
 
 	bool FMetasoundGenerator::IsFinished() const
