Fix for duplicate LOC text strings
#rnx

diff --git a/Engine/Plugins/Runtime/Metasound/Source/MetasoundEditor/Private/MetasoundEditorGraphSchema.cpp b/Engine/Plugins/Runtime/Metasound/Source/MetasoundEditor/Private/MetasoundEditorGraphSchema.cpp
--- a/Engine/Plugins/Runtime/Metasound/Source/MetasoundEditor/Private/MetasoundEditorGraphSchema.cpp
+++ b/Engine/Plugins/Runtime/Metasound/Source/MetasoundEditor/Private/MetasoundEditorGraphSchema.cpp
@@ -912,7 +912,7 @@ UEdGraphNode* FMetasoundGraphSchemaAction_NewReroute::PerformAction(UEdGraph* Pa
 		return nullptr;
 	}
 
-	const FScopedTransaction Transaction(FText::Format(LOCTEXT("AddNewNode", "Add {0} Reroute Node"), FText::FromName(DataType)));
+	const FScopedTransaction Transaction(FText::Format(LOCTEXT("AddNewRerouteNode", "Add {0} Reroute Node"), FText::FromName(DataType)));
 	UMetasoundEditorGraph* MetaSoundGraph = CastChecked<UMetasoundEditorGraph>(ParentGraph);
 	UObject& ParentMetasound = MetaSoundGraph->GetMetasoundChecked();
 	ParentMetasound.Modify();
