Fix issue with MetaSound array constructable types that are not array types to use object pin color in editor
#rnx

diff --git a/Engine/Plugins/Runtime/Metasound/Source/MetasoundEditor/Private/MetasoundEditorModule.cpp b/Engine/Plugins/Runtime/Metasound/Source/MetasoundEditor/Private/MetasoundEditorModule.cpp
--- a/Engine/Plugins/Runtime/Metasound/Source/MetasoundEditor/Private/MetasoundEditorModule.cpp
+++ b/Engine/Plugins/Runtime/Metasound/Source/MetasoundEditor/Private/MetasoundEditorModule.cpp
@@ -474,15 +474,11 @@ namespace Metasound
 						{
 							PinCategory = FGraphBuilder::PinCategoryTrigger;
 						}
-
-						// GraphEditor by default designates specialized connection
-						// specification for Int64, so use it even though literal is
-						// boiled down to int32
-						//else if (DataTypeName == Frontend::GetDataTypeName<int64>())
-						//{
-						//	PinCategory = FGraphBuilder::PinCategoryInt64;
-						//}
-
+						// Audio type is ubiquitous, so connections are stylized (i.e. wire color & animation)
+						else if (DataTypeName == GetMetasoundDataTypeName<FAudioBuffer>())
+						{
+							PinCategory = FGraphBuilder::PinCategoryAudio;
+						}
 						// Primitives
 						else
 						{
@@ -496,19 +492,10 @@ namespace Metasound
 								break;
 
 								case ELiteralType::Float:
-								case ELiteralType::FloatArray:
 								{
 									PinCategory = FGraphBuilder::PinCategoryFloat;
 
-									// Doubles use the same preferred literal
-									// but different colorization
-									//if (DataTypeName == Frontend::GetDataTypeName<double>())
-									//{
-									//	PinCategory = FGraphBuilder::PinCategoryDouble;
-									//}
-
-									// Differentiate stronger numeric types associated with audio
-									if (DataTypeName == GetMetasoundDataTypeName<FTime>() || 
+									if (DataTypeName == GetMetasoundDataTypeName<FTime>() ||
 										DataTypeName == CreateArrayTypeNameFromElementTypeName(GetMetasoundDataTypeName<FTime>()))
 									{
 										PinSubCategory = FGraphBuilder::PinSubCategoryTime;
@@ -516,37 +503,65 @@ namespace Metasound
 								}
 								break;
 
+								case ELiteralType::FloatArray:
+								{
+									if (RegistryInfo.bIsArrayType)
+									{
+										PinCategory = FGraphBuilder::PinCategoryFloat;
+									}
+									else
+									{
+										PinCategory = FGraphBuilder::PinCategoryObject;
+									}
+								}
+								break;
+
 								case ELiteralType::Integer:
-								case ELiteralType::IntegerArray:
 								{
 									PinCategory = FGraphBuilder::PinCategoryInt32;
 								}
 								break;
 
+								case ELiteralType::IntegerArray:
+								{
+									if (RegistryInfo.bIsArrayType)
+									{
+										PinCategory = FGraphBuilder::PinCategoryInt32;
+									}
+									else
+									{
+										PinCategory = FGraphBuilder::PinCategoryObject;
+									}
+								}
+								break;
+
 								case ELiteralType::String:
-								case ELiteralType::StringArray:
 								{
 									PinCategory = FGraphBuilder::PinCategoryString;
 								}
 								break;
 
-								case ELiteralType::UObjectProxy:
-								case ELiteralType::UObjectProxyArray:
+								case ELiteralType::StringArray:
 								{
-									PinCategory = FGraphBuilder::PinCategoryObject;
+									if (RegistryInfo.bIsArrayType)
+									{
+										PinCategory = FGraphBuilder::PinCategoryString;
+									}
+									else
+									{
+										PinCategory = FGraphBuilder::PinCategoryObject;
+									}
 								}
 								break;
 
+								case ELiteralType::UObjectProxy:
+								case ELiteralType::UObjectProxyArray:
 								case ELiteralType::None:
+								case ELiteralType::NoneArray:
 								case ELiteralType::Invalid:
 								default:
 								{
-									// Audio types are ubiquitous, so added as subcategory
-									// to be able to stylize connections (i.e. wire color & wire animation)
-									if (DataTypeName == GetMetasoundDataTypeName<FAudioBuffer>())
-									{
-										PinCategory = FGraphBuilder::PinCategoryAudio;
-									}
+									PinCategory = FGraphBuilder::PinCategoryObject;
 									static_assert(static_cast<int32>(ELiteralType::Invalid) == 12, "Possible missing binding of pin category to primitive type");
 								}
 								break;
