Add missing WaveTable dependencies to MS editor module
#rnx

diff --git a/Engine/Plugins/Runtime/Metasound/Source/MetasoundEditor/MetasoundEditor.Build.cs b/Engine/Plugins/Runtime/Metasound/Source/MetasoundEditor/MetasoundEditor.Build.cs
--- a/Engine/Plugins/Runtime/Metasound/Source/MetasoundEditor/MetasoundEditor.Build.cs
+++ b/Engine/Plugins/Runtime/Metasound/Source/MetasoundEditor/MetasoundEditor.Build.cs
@@ -27,7 +27,9 @@ namespace UnrealBuildTool.Rules
 					"MetasoundGraphCore",
 					"SignalProcessing",
 					"ToolMenus",
-					"ToolWidgets"
+					"ToolWidgets",
+					"WaveTable",
+					"WaveTableEditor"
 				}
 			);
 
