Minor tooltip fix-up for the WaveTableEnvelopeNode
#rnx

diff --git a/Engine/Plugins/Runtime/Metasound/Source/MetasoundEngine/Private/MetasoundWaveTableEnvelopeNode.cpp b/Engine/Plugins/Runtime/Metasound/Source/MetasoundEngine/Private/MetasoundWaveTableEnvelopeNode.cpp
--- a/Engine/Plugins/Runtime/Metasound/Source/MetasoundEngine/Private/MetasoundWaveTableEnvelopeNode.cpp
+++ b/Engine/Plugins/Runtime/Metasound/Source/MetasoundEngine/Private/MetasoundWaveTableEnvelopeNode.cpp
@@ -85,7 +85,7 @@ namespace Metasound
 					1, // Major Version
 					0, // Minor Version
 					LOCTEXT("MetasoundWaveTableEnvelopeNode_Name", "WaveTable Envelope"),
-					LOCTEXT("MetasoundWaveTableEnvelopeNode_Description", "Reads through the given WaveTable at the provided frequency."),
+					LOCTEXT("MetasoundWaveTableEnvelopeNode_Description", "Reads through the given WaveTable over the given duration."),
 					PluginAuthor,
 					PluginNodeMissingPrompt,
 					GetDefaultInterface(),
