Support more audio channels in wave player node.

diff --git a/Engine/Plugins/Runtime/Metasound/Source/MetasoundEngine/Private/MetasoundEngineNodesNames.cpp b/Engine/Plugins/Runtime/Metasound/Source/MetasoundEngine/Private/MetasoundEngineNodesNames.cpp
--- a/Engine/Plugins/Runtime/Metasound/Source/MetasoundEngine/Private/MetasoundEngineNodesNames.cpp
+++ b/Engine/Plugins/Runtime/Metasound/Source/MetasoundEngine/Private/MetasoundEngineNodesNames.cpp
@@ -6,9 +6,13 @@ namespace Metasound
 {
 	namespace EngineNodes
 	{
-		const FName Namespace = TEXT("UE");
+		const FName Namespace = "UE";
 
-		const FName MonoVariant = TEXT("Mono");
-		const FName StereoVariant = TEXT("Stereo");
+		const FName MonoVariant = "Mono";
+		const FName StereoVariant = "Stereo";
+		const FName QuadVariant = "Quad";
+		const FName FiveDotZeroVariant = "5dot0";
+		const FName FiveDotOneVariant = "5dot1";
+		const FName SevenDotOneVariant = "7dot1";
 	}
 }
diff --git a/Engine/Plugins/Runtime/Metasound/Source/MetasoundEngine/Private/MetasoundWavePlayerNode.cpp b/Engine/Plugins/Runtime/Metasound/Source/MetasoundEngine/Private/MetasoundWavePlayerNode.cpp
--- a/Engine/Plugins/Runtime/Metasound/Source/MetasoundEngine/Private/MetasoundWavePlayerNode.cpp
+++ b/Engine/Plugins/Runtime/Metasound/Source/MetasoundEngine/Private/MetasoundWavePlayerNode.cpp
@@ -44,35 +44,19 @@ namespace Metasound
 		METASOUND_PARAM(OutputCuePointLabel, "Cue Point Label", "The cue point label that was triggered (if there was a label parsed in the imported .wav file).")
 		METASOUND_PARAM(OutputLoopRatio, "Loop Percent", "Returns the current playback location as a ratio of the loop (0-1) if looping is enabled.")
 		METASOUND_PARAM(OutputPlaybackLocation, "Playback Location", "Returns the absolute position of the wave playback as a ratio of wave duration (0-1).")
-		METASOUND_PARAM(OutputAudioLeft, "Out Left", "The left channel audio output. Mono wave assets will be upmixed to dual stereo.")
-		METASOUND_PARAM(OutputAudioRight, "Out Right", "The right channel audio output. Mono wave assets will be upmixed to dual stereo.")
+		METASOUND_PARAM(OutputAudioMono, "Out Mono", "The mono channel audio output.")
+		METASOUND_PARAM(OutputAudioLeft, "Out Left", "The left channel audio output.")
+		METASOUND_PARAM(OutputAudioRight, "Out Right", "The right channel audio output.")
+		METASOUND_PARAM(OutputAudioFrontRight, "Out Front Right", "The front right channel audio output.")
+		METASOUND_PARAM(OutputAudioFrontLeft, "Out Front Left", "The front right channel audio output.")
+		METASOUND_PARAM(OutputAudioFrontCenter, "Out Front Center", "The front center channel audio output.")
+		METASOUND_PARAM(OutputAudioLowFrequency, "Out Low Frequency", "The low frequency channel audio output.")
+		METASOUND_PARAM(OutputAudioSideRight, "Out Side Right", "The side right channel audio output.")
+		METASOUND_PARAM(OutputAudioSideLeft, "Out Side Left", "The side right channel audio output.")
+		METASOUND_PARAM(OutputAudioBackRight, "Out Back Right", "The back right channel audio output.")
+		METASOUND_PARAM(OutputAudioBackLeft, "Out Back Left", "The back right channel audio output.")
 	}
 
-	class FWavePlayerNode : public FNode
-	{
-		class FOperatorFactory : public IOperatorFactory
-		{
-			virtual TUniquePtr<IOperator> CreateOperator(const FBuildOperatorParams& InParams, FBuildResults& OutResults) override;
-		};
-
-	public:
-		static FVertexInterface DeclareVertexInterface();
-		static const FNodeClassMetadata& GetNodeInfo();
-
-		FWavePlayerNode(const FVertexName& InName, const FGuid& InInstanceID);
-		FWavePlayerNode(const FNodeInitData& InInitData);
-		virtual ~FWavePlayerNode() = default;
-
-		virtual FOperatorFactorySharedRef GetDefaultOperatorFactory() const override;
-		virtual const FVertexInterface& GetVertexInterface() const override;
-		virtual bool SetVertexInterface(const FVertexInterface& InInterface) override;
-		virtual bool IsVertexInterfaceSupported(const FVertexInterface& InInterface) const override;
-
-	private:
-		FOperatorFactorySharedRef Factory;
-		FVertexInterface Interface;
-	};
-
 	namespace WavePlayerNodePrivate
 	{
 		int32 GetCuePointFrame(const FSoundWaveCuePoint& InPoint)
@@ -313,6 +297,7 @@ namespace Metasound
 	struct FWavePlayerOpArgs
 	{
 		FOperatorSettings Settings;
+		TArray<FOutputDataVertex> OutputAudioVertices;
 		FTriggerReadRef PlayTrigger;
 		FTriggerReadRef StopTrigger;
 		FWaveAssetReadRef WaveAsset;
@@ -332,8 +317,6 @@ namespace Metasound
 		static constexpr float MaxAbsPitchShiftInOctaves = 6.0f;
 		// Maximum decode size in frames. 
 		static constexpr int32 MaxDecodeSizeInFrames = 8192;
-		// Number of output audio channels. 
-		static constexpr int32 NumOutputChannels = 2;
 		// Block size for deinterleaving audio. 
 		static constexpr int32 DeinterleaveBlockSizeInFrames = 512;
 
@@ -355,15 +338,20 @@ namespace Metasound
 			, CuePointLabel(FStringWriteRef::CreateNew(TEXT("")))
 			, LoopPercent(FFloatWriteRef::CreateNew(0.0f))
 			, PlaybackLocation(FFloatWriteRef::CreateNew(0.0f))
-			, AudioBufferL(FAudioBufferWriteRef::CreateNew(InArgs.Settings))
-			, AudioBufferR(FAudioBufferWriteRef::CreateNew(InArgs.Settings))
 		{
-			// Hold on to a view of the output audio. Audio buffers are only writable
-			// by this object and will not be reallocated. 
-			OutputAudioView.Emplace(AudioBufferL->GetData(), AudioBufferL->Num());
-			OutputAudioView.Emplace(AudioBufferR->GetData(), AudioBufferR->Num());
+			NumOutputChannels = InArgs.OutputAudioVertices.Num();
+
+			for (const FOutputDataVertex& OutputAudioVertex : InArgs.OutputAudioVertices)
+			{
+				OutputAudioBufferVertexNames.Add(OutputAudioVertex.VertexName);
 
-			check(OutputAudioView.Num() == NumOutputChannels);
+				FAudioBufferWriteRef AudioBuffer = FAudioBufferWriteRef::CreateNew(InArgs.Settings);
+				OutputAudioBuffers.Add(AudioBuffer);
+
+				// Hold on to a view of the output audio. Audio buffers are only writable
+				// by this object and will not be reallocated. 
+				OutputAudioView.Emplace(AudioBuffer->GetData(), AudioBuffer->Num());
+			}
 		}
 
 		virtual FDataReferenceCollection GetInputs() const override
@@ -396,15 +384,21 @@ namespace Metasound
 			OutputDataReferences.AddDataReadReference(METASOUND_GET_PARAM_NAME(OutputCuePointLabel), CuePointLabel);
 			OutputDataReferences.AddDataReadReference(METASOUND_GET_PARAM_NAME(OutputLoopRatio), LoopPercent);
 			OutputDataReferences.AddDataReadReference(METASOUND_GET_PARAM_NAME(OutputPlaybackLocation), PlaybackLocation);
-			OutputDataReferences.AddDataReadReference(METASOUND_GET_PARAM_NAME(OutputAudioLeft), AudioBufferL);
-			OutputDataReferences.AddDataReadReference(METASOUND_GET_PARAM_NAME(OutputAudioRight), AudioBufferR);
+
+			check(OutputAudioBuffers.Num() == OutputAudioBufferVertexNames.Num());
+
+			for (int32 i = 0; i < OutputAudioBuffers.Num(); i++)
+			{
+				OutputDataReferences.AddDataReadReference(OutputAudioBufferVertexNames[i], OutputAudioBuffers[i]);
+			}
+			
 			return OutputDataReferences;
 		}
 
 
 		void Execute()
 		{
-			METASOUND_TRACE_CPUPROFILER_EVENT_SCOPE(Metasound::FWavePlayerNode::Execute);
+			METASOUND_TRACE_CPUPROFILER_EVENT_SCOPE(Metasound::FWavePlayerOperator::Execute);
 
 			// Advance all triggers owned by this operator. 
 			TriggerOnDone->AdvanceBlock();
@@ -428,8 +422,10 @@ namespace Metasound
 			}
 
 			// zero output buffers
-			FMemory::Memzero(AudioBufferL->GetData(), OperatorSettings.GetNumFramesPerBlock() * sizeof(float));
-			FMemory::Memzero(AudioBufferR->GetData(), OperatorSettings.GetNumFramesPerBlock() * sizeof(float));
+			for (const FAudioBufferWriteRef& OutputBuffer : OutputAudioBuffers)
+			{
+				FMemory::Memzero(OutputBuffer->GetData(), OperatorSettings.GetNumFramesPerBlock() * sizeof(float));
+			}
 
 			// Performs execution per sub block based on triggers.
 			ExecuteSubblocks();
@@ -518,7 +514,7 @@ namespace Metasound
 		{
 			using namespace Audio;
 
-			METASOUND_TRACE_CPUPROFILER_EVENT_SCOPE(Metasound::FWavePlayerNode::RenderFrameRange);
+			METASOUND_TRACE_CPUPROFILER_EVENT_SCOPE(Metasound::FWavePlayerOperator::RenderFrameRange);
 
 			// Assume this is set to true and checked by outside callers
 			check(bIsPlaying);
@@ -610,7 +606,7 @@ namespace Metasound
 		void StartPlaying()
 		{
 			using namespace WavePlayerNodePrivate;
-			METASOUND_TRACE_CPUPROFILER_EVENT_SCOPE(Metasound::FWavePlayerNode::StartPlaying);
+			METASOUND_TRACE_CPUPROFILER_EVENT_SCOPE(Metasound::FWavePlayerOperator::StartPlaying);
 
 			// MetasoundWavePlayerNode DSP Stack
 			//
@@ -961,8 +957,8 @@ namespace Metasound
 		FStringWriteRef CuePointLabel;
 		FFloatWriteRef LoopPercent;
 		FFloatWriteRef PlaybackLocation;
-		FAudioBufferWriteRef AudioBufferL;
-		FAudioBufferWriteRef AudioBufferR;
+		TArray<FAudioBufferWriteRef> OutputAudioBuffers;
+		TArray<FName> OutputAudioBufferVertexNames;
 
 		TUniquePtr<FWaveProxyReader> WaveProxyReader;
 		TUniquePtr<Audio::IConvertDeinterleave> ConvertDeinterleave;
@@ -977,130 +973,274 @@ namespace Metasound
 
 		WavePlayerNodePrivate::FSourceBufferState SourceState;
 		float SampleRateFrameRatio = 1.f;
+		int32 NumOutputChannels;
 		int32 NumDeinterleaveChannels;
 		bool bOnNearlyDoneTriggeredForWave = false;
 		bool bIsPlaying = false;
 		
 	};
 
-	TUniquePtr<IOperator> FWavePlayerNode::FOperatorFactory::CreateOperator(const FBuildOperatorParams& InParams, FBuildResults& OutResults)
+	class FWavePlayerOperatorFactory : public IOperatorFactory
 	{
-		using namespace Audio;
-		using namespace WavePlayerVertexNames;
-
-		const FWavePlayerNode& WaveNode = static_cast<const FWavePlayerNode&>(InParams.Node);
-
-		const FInputVertexInterfaceData& Inputs = InParams.InputData;
+	public:
+		FWavePlayerOperatorFactory(const TArray<FOutputDataVertex>& InOutputAudioVertices)
+		: OutputAudioVertices(InOutputAudioVertices)
+		{
+		}
 
-		FWavePlayerOpArgs Args =
+		virtual TUniquePtr<IOperator> CreateOperator(const FBuildOperatorParams& InParams, FBuildResults& OutResults) override
 		{
-			InParams.OperatorSettings,
-			Inputs.GetOrConstructDataReadReference<FTrigger>(METASOUND_GET_PARAM_NAME(InputTriggerPlay), InParams.OperatorSettings),
-			Inputs.GetOrConstructDataReadReference<FTrigger>(METASOUND_GET_PARAM_NAME(InputTriggerStop), InParams.OperatorSettings),
-			Inputs.GetOrConstructDataReadReference<FWaveAsset>(METASOUND_GET_PARAM_NAME(InputWaveAsset)),
-			Inputs.GetOrCreateDefaultDataReadReference<FTime>(METASOUND_GET_PARAM_NAME(InputStartTime), InParams.OperatorSettings),
-			Inputs.GetOrCreateDefaultDataReadReference<float>(METASOUND_GET_PARAM_NAME(InputPitchShift), InParams.OperatorSettings),
-			Inputs.GetOrCreateDefaultDataReadReference<bool>(METASOUND_GET_PARAM_NAME(InputLoop), InParams.OperatorSettings),
-			Inputs.GetOrCreateDefaultDataReadReference<FTime>(METASOUND_GET_PARAM_NAME(InputLoopStart), InParams.OperatorSettings),
-			Inputs.GetOrCreateDefaultDataReadReference<FTime>(METASOUND_GET_PARAM_NAME(InputLoopDuration), InParams.OperatorSettings)
-		};
+			using namespace Audio;
+			using namespace WavePlayerVertexNames;
 
-		return MakeUnique<FWavePlayerOperator>(Args);
-	}
+			const FInputVertexInterfaceData& Inputs = InParams.InputData;
+
+			FWavePlayerOpArgs Args =
+			{
+				InParams.OperatorSettings,
+				OutputAudioVertices,
+				Inputs.GetOrConstructDataReadReference<FTrigger>(METASOUND_GET_PARAM_NAME(InputTriggerPlay), InParams.OperatorSettings),
+				Inputs.GetOrConstructDataReadReference<FTrigger>(METASOUND_GET_PARAM_NAME(InputTriggerStop), InParams.OperatorSettings),
+				Inputs.GetOrConstructDataReadReference<FWaveAsset>(METASOUND_GET_PARAM_NAME(InputWaveAsset)),
+				Inputs.GetOrCreateDefaultDataReadReference<FTime>(METASOUND_GET_PARAM_NAME(InputStartTime), InParams.OperatorSettings),
+				Inputs.GetOrCreateDefaultDataReadReference<float>(METASOUND_GET_PARAM_NAME(InputPitchShift), InParams.OperatorSettings),
+				Inputs.GetOrCreateDefaultDataReadReference<bool>(METASOUND_GET_PARAM_NAME(InputLoop), InParams.OperatorSettings),
+				Inputs.GetOrCreateDefaultDataReadReference<FTime>(METASOUND_GET_PARAM_NAME(InputLoopStart), InParams.OperatorSettings),
+				Inputs.GetOrCreateDefaultDataReadReference<FTime>(METASOUND_GET_PARAM_NAME(InputLoopDuration), InParams.OperatorSettings)
+			};
+
+			return MakeUnique<FWavePlayerOperator>(Args);
+		}
+	private:
+		TArray<FOutputDataVertex> OutputAudioVertices;
+	};
 
-	FVertexInterface FWavePlayerNode::DeclareVertexInterface()
+	
+	template<typename AudioChannelConfigurationInfoType>
+	class TWavePlayerNode : public FNode
 	{
-		using namespace WavePlayerVertexNames;
 
-		// Workaround to override display name of OutputLoopRatio
-		static const FDataVertexMetadata OutputLoopRatioMetadata
-		{ 
-			METASOUND_GET_PARAM_TT(OutputLoopRatio), // description 
-			METASOUND_LOCTEXT("OutputLoopRatioNotPercentDisplayName", "Loop Ratio") // display name  
-		};
+	public:
+		static FVertexInterface DeclareVertexInterface()
+		{
+			using namespace WavePlayerVertexNames;
 
-		return FVertexInterface(
-			FInputVertexInterface(
-				TInputDataVertex<FTrigger>(METASOUND_GET_PARAM_NAME_AND_METADATA(InputTriggerPlay)),
-				TInputDataVertex<FTrigger>(METASOUND_GET_PARAM_NAME_AND_METADATA(InputTriggerStop)),
-				TInputDataVertex<FWaveAsset>(METASOUND_GET_PARAM_NAME_AND_METADATA(InputWaveAsset)),
-				TInputDataVertex<FTime>(METASOUND_GET_PARAM_NAME_AND_METADATA(InputStartTime), 0.0f),
-				TInputDataVertex<float>(METASOUND_GET_PARAM_NAME_AND_METADATA(InputPitchShift), 0.0f),
-				TInputDataVertex<bool>(METASOUND_GET_PARAM_NAME_AND_METADATA(InputLoop), false),
-				TInputDataVertex<FTime>(METASOUND_GET_PARAM_NAME_AND_METADATA(InputLoopStart), 0.0f),
-				TInputDataVertex<FTime>(METASOUND_GET_PARAM_NAME_AND_METADATA(InputLoopDuration), -1.0f)
-				),
-			FOutputVertexInterface(
-				TOutputDataVertex<FTrigger>(METASOUND_GET_PARAM_NAME_AND_METADATA(OutputTriggerOnPlay)),
-				TOutputDataVertex<FTrigger>(METASOUND_GET_PARAM_NAME_AND_METADATA(OutputTriggerOnDone)),
-				TOutputDataVertex<FTrigger>(METASOUND_GET_PARAM_NAME_AND_METADATA(OutputTriggerOnNearlyDone)),
-				TOutputDataVertex<FTrigger>(METASOUND_GET_PARAM_NAME_AND_METADATA(OutputTriggerOnLooped)),
-				TOutputDataVertex<FTrigger>(METASOUND_GET_PARAM_NAME_AND_METADATA(OutputTriggerOnCuePoint)),
-				TOutputDataVertex<int32>(METASOUND_GET_PARAM_NAME_AND_METADATA(OutputCuePointID)),
-				TOutputDataVertex<FString>(METASOUND_GET_PARAM_NAME_AND_METADATA(OutputCuePointLabel)),
-				TOutputDataVertex<float>(METASOUND_GET_PARAM_NAME(OutputLoopRatio), OutputLoopRatioMetadata),
-				TOutputDataVertex<float>(METASOUND_GET_PARAM_NAME_AND_METADATA(OutputPlaybackLocation)),
-				TOutputDataVertex<FAudioBuffer>(METASOUND_GET_PARAM_NAME_AND_METADATA(OutputAudioLeft)),
-				TOutputDataVertex<FAudioBuffer>(METASOUND_GET_PARAM_NAME_AND_METADATA(OutputAudioRight))
-			)
-		);
-	}
+			// Workaround to override display name of OutputLoopRatio
+			static const FDataVertexMetadata OutputLoopRatioMetadata
+			{ 
+				METASOUND_GET_PARAM_TT(OutputLoopRatio), // description 
+				METASOUND_LOCTEXT("OutputLoopRatioNotPercentDisplayName", "Loop Ratio") // display name  
+			};
+
+			FVertexInterface VertexInterface(
+				FInputVertexInterface(
+					TInputDataVertex<FTrigger>(METASOUND_GET_PARAM_NAME_AND_METADATA(InputTriggerPlay)),
+					TInputDataVertex<FTrigger>(METASOUND_GET_PARAM_NAME_AND_METADATA(InputTriggerStop)),
+					TInputDataVertex<FWaveAsset>(METASOUND_GET_PARAM_NAME_AND_METADATA(InputWaveAsset)),
+					TInputDataVertex<FTime>(METASOUND_GET_PARAM_NAME_AND_METADATA(InputStartTime), 0.0f),
+					TInputDataVertex<float>(METASOUND_GET_PARAM_NAME_AND_METADATA(InputPitchShift), 0.0f),
+					TInputDataVertex<bool>(METASOUND_GET_PARAM_NAME_AND_METADATA(InputLoop), false),
+					TInputDataVertex<FTime>(METASOUND_GET_PARAM_NAME_AND_METADATA(InputLoopStart), 0.0f),
+					TInputDataVertex<FTime>(METASOUND_GET_PARAM_NAME_AND_METADATA(InputLoopDuration), -1.0f)
+					),
+				FOutputVertexInterface(
+					TOutputDataVertex<FTrigger>(METASOUND_GET_PARAM_NAME_AND_METADATA(OutputTriggerOnPlay)),
+					TOutputDataVertex<FTrigger>(METASOUND_GET_PARAM_NAME_AND_METADATA(OutputTriggerOnDone)),
+					TOutputDataVertex<FTrigger>(METASOUND_GET_PARAM_NAME_AND_METADATA(OutputTriggerOnNearlyDone)),
+					TOutputDataVertex<FTrigger>(METASOUND_GET_PARAM_NAME_AND_METADATA(OutputTriggerOnLooped)),
+					TOutputDataVertex<FTrigger>(METASOUND_GET_PARAM_NAME_AND_METADATA(OutputTriggerOnCuePoint)),
+					TOutputDataVertex<int32>(METASOUND_GET_PARAM_NAME_AND_METADATA(OutputCuePointID)),
+					TOutputDataVertex<FString>(METASOUND_GET_PARAM_NAME_AND_METADATA(OutputCuePointLabel)),
+					TOutputDataVertex<float>(METASOUND_GET_PARAM_NAME(OutputLoopRatio), OutputLoopRatioMetadata),
+					TOutputDataVertex<float>(METASOUND_GET_PARAM_NAME_AND_METADATA(OutputPlaybackLocation))
+				)
+			);
+
+			// Add audio outputs dependent upon source info.
+			for (const FOutputDataVertex& OutputDataVertex : AudioChannelConfigurationInfoType::GetAudioOutputs())
+			{
+				VertexInterface.GetOutputInterface().Add(OutputDataVertex);
+			}
 
-	const FNodeClassMetadata& FWavePlayerNode::GetNodeInfo()
-	{
-		auto InitNodeInfo = []() -> FNodeClassMetadata
+			return VertexInterface;
+		}
+
+		static const FNodeClassMetadata& GetNodeInfo()
 		{
-			FNodeClassMetadata Info;
-			Info.ClassName = { Metasound::EngineNodes::Namespace, TEXT("Wave Player"), Metasound::EngineNodes::StereoVariant };
-			Info.MajorVersion = 1;
-			Info.MinorVersion = 0;
-			Info.DisplayName = METASOUND_LOCTEXT("Metasound_WavePlayerNodeDisplayName", "Wave Player");
-			Info.Description = METASOUND_LOCTEXT("Metasound_WavePlayerNodeDescription", "Plays a wave asset.");
-			Info.Author = PluginAuthor;
-			Info.PromptIfMissing = PluginNodeMissingPrompt;
-			Info.DefaultInterface = DeclareVertexInterface();
+			auto InitNodeInfo = []() -> FNodeClassMetadata
+			{
+				FNodeClassMetadata Info;
+				Info.ClassName = { Metasound::EngineNodes::Namespace, TEXT("Wave Player"), AudioChannelConfigurationInfoType::GetVariantName() };
+				Info.MajorVersion = 1;
+				Info.MinorVersion = 0;
+				Info.DisplayName = AudioChannelConfigurationInfoType::GetNodeDisplayName();
+				Info.Description = METASOUND_LOCTEXT("Metasound_WavePlayerNodeDescription", "Plays a wave asset. The wave's channel configurations will be up or down mixed to match the wave players audio channel format.");
+				Info.Author = PluginAuthor;
+				Info.PromptIfMissing = PluginNodeMissingPrompt;
+				Info.DefaultInterface = DeclareVertexInterface();
+
+				return Info;
+			};
+
+			static const FNodeClassMetadata Info = InitNodeInfo();
 
 			return Info;
-		};
+		}
 
-		static const FNodeClassMetadata Info = InitNodeInfo();
+		TWavePlayerNode(const FVertexName& InName, const FGuid& InInstanceID)
+			:	FNode(InName, InInstanceID, GetNodeInfo())
+			,	Factory(MakeOperatorFactoryRef<FWavePlayerOperatorFactory>(AudioChannelConfigurationInfoType::GetAudioOutputs()))
+			,	Interface(DeclareVertexInterface())
+		{
+		}
 
-		return Info;
-	}
+		TWavePlayerNode(const FNodeInitData& InInitData)
+			: TWavePlayerNode(InInitData.InstanceName, InInitData.InstanceID)
+		{
+		}
+
+		virtual ~TWavePlayerNode() = default;
+
+		virtual FOperatorFactorySharedRef GetDefaultOperatorFactory() const override
+		{
+			return Factory;
+		}
+
+		virtual const FVertexInterface& GetVertexInterface() const override
+		{
+			return Interface;
+		}
+
+		virtual bool SetVertexInterface(const FVertexInterface& InInterface) override
+		{
+			return InInterface == Interface;
+		}
+
+		virtual bool IsVertexInterfaceSupported(const FVertexInterface& InInterface) const override
+		{
+			return InInterface == Interface;
+		}
 
-	FWavePlayerNode::FWavePlayerNode(const FVertexName& InName, const FGuid& InInstanceID)
-		:	FNode(InName, InInstanceID, GetNodeInfo())
-		,	Factory(MakeOperatorFactoryRef<FWavePlayerNode::FOperatorFactory>())
-		,	Interface(DeclareVertexInterface())
+	private:
+		FOperatorFactorySharedRef Factory;
+		FVertexInterface Interface;
+	};
+
+	struct FMonoAudioChannelConfigurationInfo
 	{
-	}
+		static FText GetNodeDisplayName() { return METASOUND_LOCTEXT("Metasound_WavePlayerMonoNodeDisplayName", "Wave Player (Mono)"); }
+		static FName GetVariantName() { return Metasound::EngineNodes::MonoVariant; }
+
+		static TArray<FOutputDataVertex> GetAudioOutputs()
+		{
+			using namespace WavePlayerVertexNames;
+			return {
+				TOutputDataVertex<FAudioBuffer>(METASOUND_GET_PARAM_NAME_AND_METADATA(OutputAudioMono))
+			};
+		}
+	};
+	using FMonoWavePlayerNode = TWavePlayerNode<FMonoAudioChannelConfigurationInfo>;
+	METASOUND_REGISTER_NODE(FMonoWavePlayerNode);
 
-	FWavePlayerNode::FWavePlayerNode(const FNodeInitData& InInitData)
-		: FWavePlayerNode(InInitData.InstanceName, InInitData.InstanceID)
+	struct FStereoAudioChannelConfigurationInfo
 	{
-	}
+		static FText GetNodeDisplayName() { return METASOUND_LOCTEXT("Metasound_WavePlayerStereoNodeDisplayName", "Wave Player (Stereo)"); }
+		static FName GetVariantName() { return Metasound::EngineNodes::StereoVariant; }
 
-	FOperatorFactorySharedRef FWavePlayerNode::GetDefaultOperatorFactory() const 
+		static TArray<FOutputDataVertex> GetAudioOutputs()
+		{
+			using namespace WavePlayerVertexNames;
+			return {
+				TOutputDataVertex<FAudioBuffer>(METASOUND_GET_PARAM_NAME_AND_METADATA(OutputAudioLeft)),
+				TOutputDataVertex<FAudioBuffer>(METASOUND_GET_PARAM_NAME_AND_METADATA(OutputAudioRight))
+			};
+		}
+	};
+	using FStereoWavePlayerNode = TWavePlayerNode<FStereoAudioChannelConfigurationInfo>;
+	METASOUND_REGISTER_NODE(FStereoWavePlayerNode);
+
+	struct FQuadAudioChannelConfigurationInfo
 	{
-		return Factory;
-	}
+		static FText GetNodeDisplayName() { return METASOUND_LOCTEXT("Metasound_WavePlayerQuadNodeDisplayName", "Wave Player (Quad)"); }
+		static FName GetVariantName() { return Metasound::EngineNodes::QuadVariant; }
+
+		static TArray<FOutputDataVertex> GetAudioOutputs()
+		{
+			using namespace WavePlayerVertexNames;
+			return {
+				TOutputDataVertex<FAudioBuffer>(METASOUND_GET_PARAM_NAME_AND_METADATA(OutputAudioFrontLeft)),
+				TOutputDataVertex<FAudioBuffer>(METASOUND_GET_PARAM_NAME_AND_METADATA(OutputAudioFrontRight)),
+				TOutputDataVertex<FAudioBuffer>(METASOUND_GET_PARAM_NAME_AND_METADATA(OutputAudioSideLeft)),
+				TOutputDataVertex<FAudioBuffer>(METASOUND_GET_PARAM_NAME_AND_METADATA(OutputAudioSideRight))
+			};
+		}
+	};
+	using FQuadWavePlayerNode = TWavePlayerNode<FQuadAudioChannelConfigurationInfo>;
+	METASOUND_REGISTER_NODE(FQuadWavePlayerNode);
 
-	const FVertexInterface& FWavePlayerNode::GetVertexInterface() const
+	struct FFiveDotZeroAudioChannelConfigurationInfo
 	{
-		return Interface;
-	}
+		static FText GetNodeDisplayName() { return METASOUND_LOCTEXT("Metasound_WavePlayerFiveDotZeroNodeDisplayName", "Wave Player (5.0)"); }
+		static FName GetVariantName() { return Metasound::EngineNodes::FiveDotZeroVariant; }
+
+		static TArray<FOutputDataVertex> GetAudioOutputs()
+		{
+			using namespace WavePlayerVertexNames;
+			return {
+				TOutputDataVertex<FAudioBuffer>(METASOUND_GET_PARAM_NAME_AND_METADATA(OutputAudioFrontLeft)),
+				TOutputDataVertex<FAudioBuffer>(METASOUND_GET_PARAM_NAME_AND_METADATA(OutputAudioFrontRight)),
+				TOutputDataVertex<FAudioBuffer>(METASOUND_GET_PARAM_NAME_AND_METADATA(OutputAudioFrontCenter)),
+				TOutputDataVertex<FAudioBuffer>(METASOUND_GET_PARAM_NAME_AND_METADATA(OutputAudioSideLeft)),
+				TOutputDataVertex<FAudioBuffer>(METASOUND_GET_PARAM_NAME_AND_METADATA(OutputAudioSideRight))
+			};
+		}
+	};
+	using FFiveDotZeroWavePlayerNode = TWavePlayerNode<FFiveDotZeroAudioChannelConfigurationInfo>;
+	METASOUND_REGISTER_NODE(FFiveDotZeroWavePlayerNode);
 
-	bool FWavePlayerNode::SetVertexInterface(const FVertexInterface& InInterface)
+	struct FFiveDotOneAudioChannelConfigurationInfo
 	{
-		return InInterface == Interface;
-	}
+		static FText GetNodeDisplayName() { return METASOUND_LOCTEXT("Metasound_WavePlayerFiveDotOneNodeDisplayName", "Wave Player (5.1)"); }
+		static FName GetVariantName() { return Metasound::EngineNodes::FiveDotOneVariant; }
+
+		static TArray<FOutputDataVertex> GetAudioOutputs()
+		{
+			using namespace WavePlayerVertexNames;
+			return {
+				TOutputDataVertex<FAudioBuffer>(METASOUND_GET_PARAM_NAME_AND_METADATA(OutputAudioFrontLeft)),
+				TOutputDataVertex<FAudioBuffer>(METASOUND_GET_PARAM_NAME_AND_METADATA(OutputAudioFrontRight)),
+				TOutputDataVertex<FAudioBuffer>(METASOUND_GET_PARAM_NAME_AND_METADATA(OutputAudioFrontCenter)),
+				TOutputDataVertex<FAudioBuffer>(METASOUND_GET_PARAM_NAME_AND_METADATA(OutputAudioLowFrequency)),
+				TOutputDataVertex<FAudioBuffer>(METASOUND_GET_PARAM_NAME_AND_METADATA(OutputAudioSideLeft)),
+				TOutputDataVertex<FAudioBuffer>(METASOUND_GET_PARAM_NAME_AND_METADATA(OutputAudioSideRight))
+			};
+		}
+	};
+	using FFiveDotOneWavePlayerNode = TWavePlayerNode<FFiveDotOneAudioChannelConfigurationInfo>;
+	METASOUND_REGISTER_NODE(FFiveDotOneWavePlayerNode);
 
-	bool FWavePlayerNode::IsVertexInterfaceSupported(const FVertexInterface& InInterface) const
+	struct FSevenDotOneAudioChannelConfigurationInfo
 	{
-		return InInterface == Interface;
-	}
+		static FText GetNodeDisplayName() { return METASOUND_LOCTEXT("Metasound_WavePlayerSevenDotOneNodeDisplayName", "Wave Player (7.1)"); }
+		static FName GetVariantName() { return Metasound::EngineNodes::SevenDotOneVariant; }
 
-	METASOUND_REGISTER_NODE(FWavePlayerNode)
+		static TArray<FOutputDataVertex> GetAudioOutputs()
+		{
+			using namespace WavePlayerVertexNames;
+			return {
+				TOutputDataVertex<FAudioBuffer>(METASOUND_GET_PARAM_NAME_AND_METADATA(OutputAudioFrontLeft)),
+				TOutputDataVertex<FAudioBuffer>(METASOUND_GET_PARAM_NAME_AND_METADATA(OutputAudioFrontRight)),
+				TOutputDataVertex<FAudioBuffer>(METASOUND_GET_PARAM_NAME_AND_METADATA(OutputAudioFrontCenter)),
+				TOutputDataVertex<FAudioBuffer>(METASOUND_GET_PARAM_NAME_AND_METADATA(OutputAudioLowFrequency)),
+				TOutputDataVertex<FAudioBuffer>(METASOUND_GET_PARAM_NAME_AND_METADATA(OutputAudioSideLeft)),
+				TOutputDataVertex<FAudioBuffer>(METASOUND_GET_PARAM_NAME_AND_METADATA(OutputAudioSideRight)),
+				TOutputDataVertex<FAudioBuffer>(METASOUND_GET_PARAM_NAME_AND_METADATA(OutputAudioBackLeft)),
+				TOutputDataVertex<FAudioBuffer>(METASOUND_GET_PARAM_NAME_AND_METADATA(OutputAudioBackRight))
+			};
+		}
+	};
+	using FSevenDotOneWavePlayerNode = TWavePlayerNode<FSevenDotOneAudioChannelConfigurationInfo>;
+	METASOUND_REGISTER_NODE(FSevenDotOneWavePlayerNode);
 } // namespace Metasound
 
 #undef LOCTEXT_NAMESPACE // MetasoundWaveNode
diff --git a/Engine/Plugins/Runtime/Metasound/Source/MetasoundEngine/Public/MetasoundEngineNodesNames.h b/Engine/Plugins/Runtime/Metasound/Source/MetasoundEngine/Public/MetasoundEngineNodesNames.h
--- a/Engine/Plugins/Runtime/Metasound/Source/MetasoundEngine/Public/MetasoundEngineNodesNames.h
+++ b/Engine/Plugins/Runtime/Metasound/Source/MetasoundEngine/Public/MetasoundEngineNodesNames.h
@@ -12,5 +12,9 @@ namespace Metasound
 
 		extern const FName MonoVariant;
 		extern const FName StereoVariant;
+		extern const FName QuadVariant;
+		extern const FName FiveDotZeroVariant;
+		extern const FName FiveDotOneVariant;
+		extern const FName SevenDotOneVariant;
 	}
 }
