Fix channel offset not updating/offset being incorrect when authoring wavetable entries
#rnx

diff --git a/Engine/Plugins/Runtime/Metasound/Source/MetasoundEngine/Private/MetasoundWaveTableOscillatorNode.cpp b/Engine/Plugins/Runtime/Metasound/Source/MetasoundEngine/Private/MetasoundWaveTableOscillatorNode.cpp
--- a/Engine/Plugins/Runtime/Metasound/Source/MetasoundEngine/Private/MetasoundWaveTableOscillatorNode.cpp
+++ b/Engine/Plugins/Runtime/Metasound/Source/MetasoundEngine/Private/MetasoundWaveTableOscillatorNode.cpp
@@ -35,7 +35,7 @@ namespace Metasound
 					TInputDataVertex<float>("Freq", FDataVertexMetadata { LOCTEXT("MetasoundWaveTableOscillatorNode_Freq", "Freq") }, 440.0f),
 					TInputDataVertex<FAudioBuffer>("PhaseMod", FDataVertexMetadata
 					{
-						LOCTEXT("MetasoundWaveTableOscillatorNode_PhaseModDescription", "Modulation audio source for modulating oscillation phase of provided table."),
+						LOCTEXT("MetasoundWaveTableOscillatorNode_PhaseModDescription", "Modulation audio source for modulating oscillation phase of provided table. A value of 0 is no phase modulation and 1 a full buffer length (360 degrees) of phase shift."),
 						LOCTEXT("MetasoundWaveTableOscillatorNode_PhaseMod", "Phase Modulator"),
 						true /* bIsAdvancedDisplay */
 					})
diff --git a/Engine/Plugins/Runtime/WaveTable/Source/WaveTableEditor/Private/WaveTableFileUtilities.cpp b/Engine/Plugins/Runtime/WaveTable/Source/WaveTableEditor/Private/WaveTableFileUtilities.cpp
--- a/Engine/Plugins/Runtime/WaveTable/Source/WaveTableEditor/Private/WaveTableFileUtilities.cpp
+++ b/Engine/Plugins/Runtime/WaveTable/Source/WaveTableEditor/Private/WaveTableFileUtilities.cpp
@@ -44,7 +44,7 @@ namespace WaveTable
 								OutPCMData.Empty();
 								OutPCMData.AddZeroed(NumFrames);
 								const int16* RawDataPtr = (const int16*)(RawPCMData.GetData());
-								const int32 ChannelOffset = (InChannelIndex % NumChannels) * sizeof(int16);
+								const int32 ChannelOffset = (InChannelIndex % NumChannels);
 								constexpr float Max16BitAsFloat = static_cast<float>(TNumericLimits<int16>::Max());
 								for (int32 i = 0; i < OutPCMData.Num(); ++i)
 								{
diff --git a/Engine/Plugins/Runtime/WaveTable/Source/WaveTableEditor/Private/WaveTableTransformLayout.cpp b/Engine/Plugins/Runtime/WaveTable/Source/WaveTableEditor/Private/WaveTableTransformLayout.cpp
--- a/Engine/Plugins/Runtime/WaveTable/Source/WaveTableEditor/Private/WaveTableTransformLayout.cpp
+++ b/Engine/Plugins/Runtime/WaveTable/Source/WaveTableEditor/Private/WaveTableTransformLayout.cpp
@@ -79,9 +79,11 @@ namespace WaveTable
 			if (WaveTableOptionsHandle->GetNumChildren(NumWaveTableOptions))
 			{
 				ChannelIndexHandle = WaveTableOptionsHandle->GetChildHandle(GET_MEMBER_NAME_CHECKED(FWaveTableSettings, ChannelIndex));
+				ChannelIndexHandle->SetOnPropertyValueChanged(FSimpleDelegate::CreateLambda([this] { CachePCMFromFile(); }));
 
 				FilePathHandle = WaveTableOptionsHandle->GetChildHandle(GET_MEMBER_NAME_CHECKED(FWaveTableSettings, FilePath));
 				FilePathHandle->SetOnChildPropertyValueChanged(FSimpleDelegate::CreateLambda([this] { CachePCMFromFile(); }));
+
 				FilePathHandle = FilePathHandle->GetChildHandle(GET_MEMBER_NAME_CHECKED(FFilePath, FilePath));
 			}
 
