Node Template Preprocessor Updates
- Conform template build step nominclature to "Document Preprocessing" (as it may be used either directly before building or at time post serialization)
- Use new preprocessor to cache edge data for more performant reroute template processing
#rnx

diff --git a/Engine/Plugins/Runtime/Metasound/Source/MetasoundFrontend/Private/MetasoundAssetBase.cpp b/Engine/Plugins/Runtime/Metasound/Source/MetasoundFrontend/Private/MetasoundAssetBase.cpp
--- a/Engine/Plugins/Runtime/Metasound/Source/MetasoundFrontend/Private/MetasoundAssetBase.cpp
+++ b/Engine/Plugins/Runtime/Metasound/Source/MetasoundFrontend/Private/MetasoundAssetBase.cpp
@@ -27,6 +27,7 @@
 #include "MetasoundParameterTransmitter.h"
 #include "MetasoundTrace.h"
 #include "MetasoundVertex.h"
+#include "NodeTemplates/MetasoundFrontendDocumentTemplatePreprocessor.h"
 #include "StructSerializer.h"
 #include "Templates/SharedPointer.h"
 #include "UObject/MetaData.h"
@@ -59,90 +60,6 @@ namespace Metasound
 					}
 				}
 			}
-
-			// Takes in document from asset and builds all template nodes and removes template classes from document dependencies list.
-			void PreBuildDocument(FMetasoundFrontendDocument& InOutDoc)
-			{
-				using namespace Frontend;
-
-				METASOUND_TRACE_CPUPROFILER_EVENT_SCOPE(Metasound::Frontend::AssetBase::PreBuildDocument);
-
-				FMetasoundFrontendGraph& Graph = InOutDoc.RootGraph.Graph;
-				TArray<FMetasoundFrontendClass>& Dependencies = InOutDoc.Dependencies;
-				const TArray<FMetasoundFrontendEdge>& GraphEdges = Graph.Edges;
-
-				TMap<FGuid, const FMetasoundFrontendClass*> TemplateClassMap;
-				Algo::TransformIf(Dependencies, TemplateClassMap,
-					[](const FMetasoundFrontendClass& Class) { return Class.Metadata.GetType() == EMetasoundFrontendClassType::Template; },
-					[](const FMetasoundFrontendClass& Class) { return TPair<FGuid, const FMetasoundFrontendClass*> { Class.ID, & Class }; }
-				);
-
-				TSet<const FMetasoundFrontendNode*> NodesToRemove;
-				TSet<TPair<FGuid, FGuid>> VerticesRemoved;
-
-				INodeTemplateRegistry& TemplateReg = INodeTemplateRegistry::Get();
-				for (FMetasoundFrontendNode& Node : Graph.Nodes)
-				{
-					const FMetasoundFrontendClass* NodeClass = TemplateClassMap.FindRef(Node.ClassID);
-					if (NodeClass)
-					{
-						if (NodeClass->Metadata.GetType() == EMetasoundFrontendClassType::Template)
-						{
-							const FNodeRegistryKey Key = NodeRegistryKey::CreateKey(NodeClass->Metadata);
-							const INodeTemplate* Template = TemplateReg.FindTemplate(Key);
-							if (ensure(Template))
-							{
-								const bool bTemplateBuilt = Template->BuildTemplate(InOutDoc, Graph, Node);
-								if (bTemplateBuilt)
-								{
-									NodesToRemove.Add(&Node);
-
-									auto GetNodeVertexGuidPair = [NodeID = Node.GetID()](const FMetasoundFrontendVertex& Vertex) { return TPair<FGuid, FGuid> { NodeID, Vertex.VertexID }; };
-									Algo::Transform(Node.Interface.Inputs, VerticesRemoved, GetNodeVertexGuidPair);
-									Algo::Transform(Node.Interface.Outputs, VerticesRemoved, GetNodeVertexGuidPair);
-								}
-							}
-						}
-					}
-				}
-
-				constexpr bool bAllowShrinking = false;
-				for (int32 i = Dependencies.Num() - 1; i >= 0; --i)
-				{
-					if (TemplateClassMap.Contains(Dependencies[i].ID))
-					{
-						Dependencies.RemoveAtSwap(i, 1, bAllowShrinking);
-					}
-				}
-				Dependencies.Shrink();
-
-				for (int32 i = Graph.Nodes.Num() - 1; i >= 0; --i)
-				{
-					if (NodesToRemove.Contains(&Graph.Nodes[i]))
-					{
-						Graph.Nodes.RemoveAtSwap(i, 1, bAllowShrinking);
-					}
-				}
-				Graph.Nodes.Shrink();
-
-				for (int32 i = Graph.Edges.Num() - 1; i >= 0; --i)
-				{
-					const TPair<FGuid, FGuid> FromNodeVertexPair{ Graph.Edges[i].FromNodeID, Graph.Edges[i].FromVertexID };
-					if (VerticesRemoved.Contains(FromNodeVertexPair))
-					{
-						Graph.Edges.RemoveAtSwap(i, 1, bAllowShrinking);
-					}
-					else
-					{
-						const TPair<FGuid, FGuid> ToNodeVertexPair{ Graph.Edges[i].ToNodeID, Graph.Edges[i].ToVertexID };
-						if (VerticesRemoved.Contains(ToNodeVertexPair))
-						{
-							Graph.Edges.RemoveAtSwap(i, 1, bAllowShrinking);
-						}
-					}
-				}
-				Graph.Edges.Shrink();
-			}
 		} // namespace AssetBasePrivate
 
 		FAutoConsoleVariableRef CVarMetaSoundBlockRate(
@@ -240,15 +157,15 @@ void FMetasoundAssetBase::RegisterGraphWithFrontend(Metasound::Frontend::FMetaSo
 	class FNodeRegistryEntry : public INodeRegistryEntry
 	{
 	public:
-		FNodeRegistryEntry(const FString& InName, TSharedPtr<FMetasoundFrontendDocument> InPreBuiltDocument, FName InAssetPath)
+		FNodeRegistryEntry(const FString& InName, TSharedPtr<FMetasoundFrontendDocument> InPreprocessedDoc, FName InAssetPath)
 		: Name(InName)
-		, PreBuiltDocument(InPreBuiltDocument)
+		, PreprocessedDoc(InPreprocessedDoc)
 		{
 			// Copy FrontendClass to preserve original document.
-			FrontendClass = PreBuiltDocument->RootGraph;
+			FrontendClass = PreprocessedDoc->RootGraph;
 			FrontendClass.Metadata.SetType(EMetasoundFrontendClassType::External);
 
-			ClassInfo = FNodeClassInfo(PreBuiltDocument->RootGraph, InAssetPath);
+			ClassInfo = FNodeClassInfo(PreprocessedDoc->RootGraph, InAssetPath);
 		}
 
 		virtual ~FNodeRegistryEntry() = default;
@@ -260,7 +177,7 @@ void FMetasoundAssetBase::RegisterGraphWithFrontend(Metasound::Frontend::FMetaSo
 
 		virtual TUniquePtr<INode> CreateNode(const FNodeInitData&) const override
 		{
-			return FFrontendGraphBuilder().CreateGraph(*PreBuiltDocument, Name);
+			return FFrontendGraphBuilder().CreateGraph(*PreprocessedDoc, Name);
 		}
 
 		virtual TUniquePtr<INode> CreateNode(FDefaultLiteralNodeConstructorParams&&) const override
@@ -285,7 +202,7 @@ void FMetasoundAssetBase::RegisterGraphWithFrontend(Metasound::Frontend::FMetaSo
 
 		virtual TUniquePtr<INodeRegistryEntry> Clone() const override
 		{
-			return MakeUnique<FNodeRegistryEntry>(Name, PreBuiltDocument, ClassInfo.AssetPath);
+			return MakeUnique<FNodeRegistryEntry>(Name, PreprocessedDoc, ClassInfo.AssetPath);
 		}
 
 		virtual bool IsNative() const override
@@ -295,7 +212,7 @@ void FMetasoundAssetBase::RegisterGraphWithFrontend(Metasound::Frontend::FMetaSo
 
 	private:
 		FString Name;
-		TSharedPtr<FMetasoundFrontendDocument> PreBuiltDocument;
+		TSharedPtr<FMetasoundFrontendDocument> PreprocessedDoc;
 		FMetasoundFrontendClass FrontendClass;
 		FNodeClassInfo ClassInfo;
 	};
@@ -312,10 +229,11 @@ void FMetasoundAssetBase::RegisterGraphWithFrontend(Metasound::Frontend::FMetaSo
 	}
 
 	FNodeClassInfo AssetClassInfo = GetAssetClassInfo();
-	TSharedPtr<FMetasoundFrontendDocument> PreBuiltDoc = MakeShared<FMetasoundFrontendDocument>(GetDocumentChecked());
-	AssetBasePrivate::PreBuildDocument(*PreBuiltDoc);
+	TSharedPtr<FMetasoundFrontendDocument> PreprocessedDocument = MakeShared<FMetasoundFrontendDocument>(GetDocumentChecked());
+	FDocumentTemplatePreprocessTransform TemplatePreprocessor;
+	TemplatePreprocessor.Transform(*PreprocessedDocument);
 
-	TUniquePtr<INodeRegistryEntry> RegistryEntry = MakeUnique<FNodeRegistryEntry>(AssetName, PreBuiltDoc, AssetClassInfo.AssetPath);
+	TUniquePtr<INodeRegistryEntry> RegistryEntry = MakeUnique<FNodeRegistryEntry>(AssetName, PreprocessedDocument, AssetClassInfo.AssetPath);
 	RegistryKey = FMetasoundFrontendRegistryContainer::Get()->RegisterNode(MoveTemp(RegistryEntry));
 
 	if (NodeRegistryKey::IsValid(RegistryKey) && FMetasoundFrontendRegistryContainer::Get()->IsNodeRegistered(RegistryKey))
@@ -354,7 +272,7 @@ void FMetasoundAssetBase::RegisterGraphWithFrontend(Metasound::Frontend::FMetaSo
 
 	// Triggers the existing runtime data to be out-of-date.
 	CurrentCachedRuntimeDataChangeID = FGuid::NewGuid();
-	CacheRuntimeData(*PreBuiltDoc);
+	CacheRuntimeData(*PreprocessedDocument);
 }
 
 void FMetasoundAssetBase::UnregisterGraphWithFrontend()
@@ -642,7 +560,7 @@ void FMetasoundAssetBase::ResetSynchronizationState()
 }
 #endif // WITH_EDITOR
 
-TSharedPtr<Metasound::IGraph, ESPMode::ThreadSafe> FMetasoundAssetBase::BuildMetasoundDocument(const FMetasoundFrontendDocument& InPreBuiltDoc, const TArray<FMetasoundFrontendClassInput>& InTransmittableInputs) const
+TSharedPtr<Metasound::IGraph, ESPMode::ThreadSafe> FMetasoundAssetBase::BuildMetasoundDocument(const FMetasoundFrontendDocument& InPreprocessedDoc, const TArray<FMetasoundFrontendClassInput>& InTransmittableInputs) const
 {
 	using namespace Metasound;
 	using namespace Metasound::Frontend;
@@ -650,7 +568,7 @@ TSharedPtr<Metasound::IGraph, ESPMode::ThreadSafe> FMetasoundAssetBase::BuildMet
 	METASOUND_TRACE_CPUPROFILER_EVENT_SCOPE(MetaSoundAssetBase::BuildMetasoundDocument);
 
 	// Create graph which can spawn instances. 
-	TUniquePtr<FFrontendGraph> FrontendGraph = FFrontendGraphBuilder::CreateGraph(InPreBuiltDoc, GetOwningAssetName());
+	TUniquePtr<FFrontendGraph> FrontendGraph = FFrontendGraphBuilder::CreateGraph(InPreprocessedDoc, GetOwningAssetName());
 	if (FrontendGraph.IsValid())
 	{
 		TSet<FVertexName> TransmittableInputNames;
@@ -976,7 +894,7 @@ void FMetasoundAssetBase::RebuildReferencedAssetClassKeys()
 	SetReferencedAssetClassKeys(MoveTemp(ReferencedKeys));
 }
 
-const FMetasoundAssetBase::FRuntimeData& FMetasoundAssetBase::CacheRuntimeData(const FMetasoundFrontendDocument& InPreBuiltDoc)
+const FMetasoundAssetBase::FRuntimeData& FMetasoundAssetBase::CacheRuntimeData(const FMetasoundFrontendDocument& InPreprocessedDoc)
 {
 	METASOUND_TRACE_CPUPROFILER_EVENT_SCOPE(MetaSoundAssetBase::CacheRuntimeData);
 
@@ -991,7 +909,7 @@ const FMetasoundAssetBase::FRuntimeData& FMetasoundAssetBase::CacheRuntimeData(c
 	{
 		// Update CachedRuntimeData.
 		TArray<FMetasoundFrontendClassInput> ClassInputs = GetTransmittableClassInputs();
-		TSharedPtr<Metasound::IGraph, ESPMode::ThreadSafe> Graph = BuildMetasoundDocument(InPreBuiltDoc, ClassInputs);
+		TSharedPtr<Metasound::IGraph, ESPMode::ThreadSafe> Graph = BuildMetasoundDocument(InPreprocessedDoc, ClassInputs);
 
 		CachedRuntimeData =
 		{
diff --git a/Engine/Plugins/Runtime/Metasound/Source/MetasoundFrontend/Private/MetasoundFrontendDocumentTemplatePreprocessor.cpp b/Engine/Plugins/Runtime/Metasound/Source/MetasoundFrontend/Private/MetasoundFrontendDocumentTemplatePreprocessor.cpp
new file mode 100644
--- /dev/null
+++ b/Engine/Plugins/Runtime/Metasound/Source/MetasoundFrontend/Private/MetasoundFrontendDocumentTemplatePreprocessor.cpp
@@ -0,0 +1,135 @@
+// Copyright Epic Games, Inc. All Rights Reserved.
+#pragma once
+
+#include "NodeTemplates/MetasoundFrontendDocumentTemplatePreprocessor.h"
+
+#include "Containers/UnrealString.h"
+#include "MetasoundFrontendDocument.h"
+#include "MetasoundFrontendNodeTemplateRegistry.h"
+#include "MetasoundFrontendRegistries.h"
+#include "MetasoundTrace.h"
+
+
+namespace Metasound
+{
+	namespace Frontend
+	{
+		bool FDocumentTemplatePreprocessTransform::Transform(FDocumentHandle InDocument) const
+		{
+			unimplemented();
+			return false;
+		}
+
+		bool FDocumentTemplatePreprocessTransform::Transform(FMetasoundFrontendDocument& InOutDoc) const
+		{
+			METASOUND_TRACE_CPUPROFILER_EVENT_SCOPE(Metasound::Frontend::AssetBase::Preprocessor);
+
+			bool bModified = false;
+
+			FMetasoundFrontendGraph& Graph = InOutDoc.RootGraph.Graph;
+			TArray<FMetasoundFrontendClass>& Dependencies = InOutDoc.Dependencies;
+			const TArray<FMetasoundFrontendEdge>& GraphEdges = Graph.Edges;
+
+			struct FTemplateParams
+			{
+				FGuid ClassID;
+				const INodeTemplate* Template = nullptr;
+			};
+
+			// 1. Find template dependencies to build
+			TArray<FTemplateParams> TemplateParams;
+			Algo::TransformIf(Dependencies, TemplateParams,
+				[](const FMetasoundFrontendClass& Class) { return Class.Metadata.GetType() == EMetasoundFrontendClassType::Template; },
+				[](const FMetasoundFrontendClass& Class)
+				{
+					const FNodeRegistryKey Key = NodeRegistryKey::CreateKey(Class.Metadata);
+					FGuid ClassID = Class.ID;
+					const INodeTemplate* Template = INodeTemplateRegistry::Get().FindTemplate(Key);
+					ensureMsgf(Template, TEXT("Template not found for template class reference '%s'"), *Class.Metadata.GetClassName().ToString());
+					return FTemplateParams { ClassID, Template };
+				}
+			);
+
+			for (FTemplateParams& Params : TemplateParams)
+			{
+				if (!Params.Template)
+				{
+					continue;
+				}
+
+				TUniquePtr<INodeTransform> NodeTransform = Params.Template->GenerateNodeTransform(InOutDoc);
+				check(NodeTransform.IsValid());
+
+				TSet<const FMetasoundFrontendNode*> NodesToRemove;
+				TSet<TPair<FGuid, FGuid>> VerticesRemoved;
+
+				// 1a. Preprocess template nodes
+				for (FMetasoundFrontendNode& Node : Graph.Nodes)
+				{
+					if (Params.ClassID == Node.ClassID)
+					{
+						NodeTransform->Transform(Node);
+						NodesToRemove.Add(&Node);
+
+						auto GetNodeVertexGuidPair = [NodeID = Node.GetID()](const FMetasoundFrontendVertex& Vertex) { return TPair<FGuid, FGuid> { NodeID, Vertex.VertexID }; };
+						Algo::Transform(Node.Interface.Inputs, VerticesRemoved, GetNodeVertexGuidPair);
+						Algo::Transform(Node.Interface.Outputs, VerticesRemoved, GetNodeVertexGuidPair);
+
+						bModified = true;
+					}
+				}
+
+				// 1b. Remove template node from graph
+				constexpr bool bAllowShrinking = false;
+				for (int32 i = Graph.Nodes.Num() - 1; i >= 0; --i)
+				{
+					if (NodesToRemove.Contains(&Graph.Nodes[i]))
+					{
+						Graph.Nodes.RemoveAtSwap(i, 1, bAllowShrinking);
+					}
+				}
+				Graph.Nodes.Shrink();
+
+				// 1c. Remove edges connecting template node from graph
+				for (int32 i = Graph.Edges.Num() - 1; i >= 0; --i)
+				{
+					const TPair<FGuid, FGuid> FromNodeVertexPair { Graph.Edges[i].FromNodeID, Graph.Edges[i].FromVertexID };
+					if (VerticesRemoved.Contains(FromNodeVertexPair))
+					{
+						Graph.Edges.RemoveAtSwap(i, 1, bAllowShrinking);
+					}
+					else
+					{
+						const TPair<FGuid, FGuid> ToNodeVertexPair { Graph.Edges[i].ToNodeID, Graph.Edges[i].ToVertexID };
+						if (VerticesRemoved.Contains(ToNodeVertexPair))
+						{
+							Graph.Edges.RemoveAtSwap(i, 1, bAllowShrinking);
+						}
+					}
+				}
+				Graph.Edges.Shrink();
+			}
+
+			// 2. Finally, Remove template classes from dependency list
+			{
+				TSet<FString> TemplateKeys;
+				Algo::Transform(TemplateParams, TemplateKeys, [](const FTemplateParams& Params)
+				{
+					if (Params.Template)
+					{
+						return NodeRegistryKey::CreateKey(Params.Template->GetFrontendClass().Metadata);
+					}
+
+					return FString();
+				});
+				constexpr bool bAllowShrinking = false;
+				Dependencies.RemoveAllSwap([&TemplateKeys](const FMetasoundFrontendClass& Class)
+				{
+					return TemplateKeys.Contains(NodeRegistryKey::CreateKey(Class.Metadata));
+				}, bAllowShrinking);
+			}
+
+			return bModified;
+		}
+	} // namespace Frontend
+} // namespace Metasound
diff --git a/Engine/Plugins/Runtime/Metasound/Source/MetasoundFrontend/Private/MetasoundFrontendNodeTemplateRegistry.cpp b/Engine/Plugins/Runtime/Metasound/Source/MetasoundFrontend/Private/MetasoundFrontendNodeTemplateRegistry.cpp
--- a/Engine/Plugins/Runtime/Metasound/Source/MetasoundFrontend/Private/MetasoundFrontendNodeTemplateRegistry.cpp
+++ b/Engine/Plugins/Runtime/Metasound/Source/MetasoundFrontend/Private/MetasoundFrontendNodeTemplateRegistry.cpp
@@ -17,8 +17,6 @@ namespace Metasound
 		class FNodeTemplateRegistry : public INodeTemplateRegistry
 		{
 		public:
-			static FNodeTemplateRegistry& Get();
-
 			FNodeTemplateRegistry() = default;
 			virtual ~FNodeTemplateRegistry() = default;
 
diff --git a/Engine/Plugins/Runtime/Metasound/Source/MetasoundFrontend/Private/MetasoundFrontendTransform.cpp b/Engine/Plugins/Runtime/Metasound/Source/MetasoundFrontend/Private/MetasoundFrontendTransform.cpp
--- a/Engine/Plugins/Runtime/Metasound/Source/MetasoundFrontend/Private/MetasoundFrontendTransform.cpp
+++ b/Engine/Plugins/Runtime/Metasound/Source/MetasoundFrontend/Private/MetasoundFrontendTransform.cpp
@@ -2,16 +2,19 @@
 #include "MetasoundFrontendTransform.h"
 
 #include "Algo/Transform.h"
+#include "MetasoundAccessPtr.h"
 #include "MetasoundArchetype.h"
 #include "MetasoundAssetBase.h"
 #include "MetasoundFrontendArchetypeRegistry.h"
 #include "MetasoundFrontendDocument.h"
+#include "MetasoundFrontendDocumentController.h"
 #include "MetasoundFrontendRegistries.h"
 #include "MetasoundFrontendSearchEngine.h"
 #include "MetasoundLog.h"
 #include "MetasoundTrace.h"
 #include "Misc/App.h"
 
+
 namespace Metasound
 {
 	namespace Frontend
@@ -33,6 +36,11 @@ namespace Metasound
 #endif // WITH_EDITOR
 		} // namespace DocumentTransform
 
+		bool IDocumentTransform::Transform(FMetasoundFrontendDocument& InOutDocument) const
+		{
+			FDocumentAccessPtr DocAccessPtr = MakeAccessPtr<FDocumentAccessPtr>(InOutDocument.AccessPoint, InOutDocument);
+			return Transform(FDocumentController::CreateDocumentHandle(DocAccessPtr));
+		}
 
 		FModifyRootGraphInterfaces::FModifyRootGraphInterfaces(const TArray<FMetasoundFrontendInterface>& InInterfacesToRemove, const TArray<FMetasoundFrontendInterface>& InInterfacesToAdd)
 			: InterfacesToRemove(InInterfacesToRemove)
diff --git a/Engine/Plugins/Runtime/Metasound/Source/MetasoundFrontend/Private/NodeTemplates/MetasoundFrontendDocumentTemplatePreprocessor.h b/Engine/Plugins/Runtime/Metasound/Source/MetasoundFrontend/Private/NodeTemplates/MetasoundFrontendDocumentTemplatePreprocessor.h
new file mode 100644
--- /dev/null
+++ b/Engine/Plugins/Runtime/Metasound/Source/MetasoundFrontend/Private/NodeTemplates/MetasoundFrontendDocumentTemplatePreprocessor.h
@@ -0,0 +1,44 @@
+// Copyright Epic Games, Inc. All Rights Reserved.
+#pragma once
+
+#include "MetasoundFrontendDocument.h"
+#include "MetasoundFrontendTransform.h"
+
+
+namespace Metasound
+{
+	namespace Frontend
+	{
+		// Preprocesses all template node classes and respective node instances within
+		// a given document and removes template classes from document dependencies list.
+		class FDocumentTemplatePreprocessTransform : public IDocumentTransform
+		{
+		public:
+			UE_DEPRECATED(5.1, "Transforms acting on handles/controllers are deprecated.")
+			virtual bool Transform(FDocumentHandle InDocument) const override;
+
+			virtual bool Transform(FMetasoundFrontendDocument& InOutDocument) const override;
+		};
+
+		// Base implementation for preprocessing a given template node
+		class FNodeTemplatePreprocessTransformBase : public INodeTransform
+		{
+		private:
+			FMetasoundFrontendDocument& Document;
+			FMetasoundFrontendGraph& Graph;
+
+		public:
+			FNodeTemplatePreprocessTransformBase(FMetasoundFrontendDocument& InDocument)
+				: Document(InDocument)
+				, Graph(InDocument.RootGraph.Graph)
+			{
+			}
+
+			virtual ~FNodeTemplatePreprocessTransformBase() = default;
+
+			virtual FMetasoundFrontendDocument& GetOwningDocument() const override { return Document; };
+			virtual FMetasoundFrontendGraph& GetOwningGraph() const override { return Graph; };
+		};
+
+	} // namespace Frontend
+} // namespace Metasound
diff --git a/Engine/Plugins/Runtime/Metasound/Source/MetasoundFrontend/Private/NodeTemplates/MetasoundFrontendNodeTemplateReroute.cpp b/Engine/Plugins/Runtime/Metasound/Source/MetasoundFrontend/Private/NodeTemplates/MetasoundFrontendNodeTemplateReroute.cpp
--- a/Engine/Plugins/Runtime/Metasound/Source/MetasoundFrontend/Private/NodeTemplates/MetasoundFrontendNodeTemplateReroute.cpp
+++ b/Engine/Plugins/Runtime/Metasound/Source/MetasoundFrontend/Private/NodeTemplates/MetasoundFrontendNodeTemplateReroute.cpp
@@ -2,10 +2,10 @@
 
 #include "NodeTemplates/MetasoundFrontendNodeTemplateReroute.h"
 
-#include "Algo/AllOf.h"
+#include "Algo/AnyOf.h"
 #include "MetasoundFrontendDataTypeRegistry.h"
 #include "MetasoundFrontendRegistries.h"
-#include "Algo/AnyOf.h"
+#include "NodeTemplates/MetasoundFrontendDocumentTemplatePreprocessor.h"
 
 
 namespace Metasound
@@ -70,44 +70,93 @@ namespace Metasound
 					InOutInputHandles.Add(InHandleToCheck);
 				}
 			}
-		}
 
-		const FMetasoundFrontendClassName FRerouteNodeTemplate::ClassName { "UE", "Reroute", "" };
+			class FRerouteNodeTemplatePreprocessTransform : public FNodeTemplatePreprocessTransformBase
+			{
+			public:
+				FRerouteNodeTemplatePreprocessTransform(FMetasoundFrontendDocument& InDocument)
+					: FNodeTemplatePreprocessTransformBase(InDocument)
+				{
+					TArray<FMetasoundFrontendEdge>& Edges = GetOwningGraph().Edges;
+					for (int32 i = 0; i < Edges.Num(); ++i)
+					{
+						FMetasoundFrontendEdge& Edge = Edges[i];
+						InputEdgeMap.Add({ Edge.ToNodeID, Edge.ToVertexID }, &Edge);
+						OutputEdgeMap.FindOrAdd({ Edge.FromNodeID, Edge.FromVertexID }).Add(&Edge);
+					}
+				}
 
-		const FMetasoundFrontendVersion FRerouteNodeTemplate::Version { ClassName.GetFullName(), { 1, 0 } };
+				virtual ~FRerouteNodeTemplatePreprocessTransform() = default;
 
-		bool FRerouteNodeTemplate::BuildTemplate(FMetasoundFrontendDocument& InOutDocument, FMetasoundFrontendGraph& InOutGraph, FMetasoundFrontendNode& InOutNode) const
-		{
-			if (!ensure(InOutNode.Interface.Inputs.Num() == 1))
-			{
-				return false;
-			}
-			const FMetasoundFrontendVertex& InputVertex = InOutNode.Interface.Inputs.Last();
+				virtual bool Transform(FMetasoundFrontendNode& InOutNode) const override;
 
-			if (!ensure(InOutNode.Interface.Outputs.Num() == 1))
-			{
-				return false;
-			}
-			const FMetasoundFrontendVertex& OutputVertex = InOutNode.Interface.Outputs.Last();
+			private:
+				using FNodeVertexGuidPair = TPair<FGuid, FGuid>;
 
-			int32 SourceIndex = InOutGraph.Edges.IndexOfByPredicate([&](const FMetasoundFrontendEdge& Edge)
-			{
-				return Edge.ToNodeID == InOutNode.GetID() && Edge.ToVertexID == InputVertex.VertexID;
-			});
+				mutable TMap<FNodeVertexGuidPair, FMetasoundFrontendEdge*> InputEdgeMap;
+				mutable TMap<FNodeVertexGuidPair, TArray<FMetasoundFrontendEdge*>> OutputEdgeMap;
+			};
 
-			if (SourceIndex != INDEX_NONE)
+			bool FRerouteNodeTemplatePreprocessTransform::Transform(FMetasoundFrontendNode& InOutNode) const
 			{
-				for (FMetasoundFrontendEdge& Edge : InOutGraph.Edges)
+				using namespace ReroutePrivate;
+
+				const FMetasoundFrontendEdge* InputEdge = nullptr;
+				{
+					if (!ensure(InOutNode.Interface.Inputs.Num() == 1))
+					{
+						return false;
+					}
+
+					const FMetasoundFrontendVertex& InputVertex = InOutNode.Interface.Inputs.Last();
+					FNodeVertexGuidPair InputNodeVertexPair{ InOutNode.GetID(), InputVertex.VertexID };
+					InputEdge = InputEdgeMap.FindRef(InputNodeVertexPair);
+
+					// This can happen if the reroute node isn't provided an input, so its perfectly
+					// acceptable to just ignore this node as it ultimately provides no sourced input.
+					if (!InputEdge)
+					{
+						return false;
+					}
+				}
+
+				TArray<FMetasoundFrontendEdge*>* OutputEdges = nullptr;
 				{
-					if (Edge.FromNodeID == InOutNode.GetID() && Edge.FromVertexID == OutputVertex.VertexID)
+					if (!ensure(InOutNode.Interface.Outputs.Num() == 1))
+					{
+						return false;
+					}
+
+					const FMetasoundFrontendVertex& OutputVertex = InOutNode.Interface.Outputs.Last();
+					FNodeVertexGuidPair OutputNodeVertexPair{ InOutNode.GetID(), OutputVertex.VertexID };
+					OutputEdges = &OutputEdgeMap.FindChecked(OutputNodeVertexPair);
+
+					// This can happen if the reroute node isn't provided any outputs to connect to, so its
+					// perfectly acceptable to just ignore this node as it ultimately provides no sourced input.
+					if (!OutputEdges)
 					{
-						Edge.FromNodeID = InOutGraph.Edges[SourceIndex].FromNodeID;
-						Edge.FromVertexID = InOutGraph.Edges[SourceIndex].FromVertexID;
+						return false;
 					}
 				}
+
+				for (FMetasoundFrontendEdge* OutputEdge : *OutputEdges)
+				{
+					OutputEdge->FromNodeID = InputEdge->FromNodeID;
+					OutputEdge->FromVertexID = InputEdge->FromVertexID;
+				}
+
+				return true;
 			}
+		}
 
-			return true;
+		const FMetasoundFrontendClassName FRerouteNodeTemplate::ClassName { "UE", "Reroute", "" };
+
+		const FMetasoundFrontendVersion FRerouteNodeTemplate::Version { ClassName.GetFullName(), { 1, 0 } };
+
+		TUniquePtr<INodeTransform> FRerouteNodeTemplate::GenerateNodeTransform(FMetasoundFrontendDocument& InPreprocessedDocument) const
+		{
+			using namespace ReroutePrivate;
+			return TUniquePtr<INodeTransform>(new FRerouteNodeTemplatePreprocessTransform(InPreprocessedDocument));
 		}
 
 		const FMetasoundFrontendClass& FRerouteNodeTemplate::GetFrontendClass() const
diff --git a/Engine/Plugins/Runtime/Metasound/Source/MetasoundFrontend/Public/MetasoundAssetBase.h b/Engine/Plugins/Runtime/Metasound/Source/MetasoundFrontend/Public/MetasoundAssetBase.h
--- a/Engine/Plugins/Runtime/Metasound/Source/MetasoundFrontend/Public/MetasoundAssetBase.h
+++ b/Engine/Plugins/Runtime/Metasound/Source/MetasoundFrontend/Public/MetasoundAssetBase.h
@@ -203,7 +203,7 @@ protected:
 	};
 
 	// Returns the cached runtime data. Call updates cached data if out-of-date.
-	const FRuntimeData& CacheRuntimeData(const FMetasoundFrontendDocument& InPreBuiltDoc);
+	const FRuntimeData& CacheRuntimeData(const FMetasoundFrontendDocument& InPreprocessedDoc);
 
 	// Returns the cached runtime data.
 	const FRuntimeData& GetRuntimeData() const;
@@ -219,7 +219,7 @@ private:
 	FGuid CurrentCachedRuntimeDataChangeID;
 	FRuntimeData CachedRuntimeData;
 
-	TSharedPtr<Metasound::IGraph, ESPMode::ThreadSafe> BuildMetasoundDocument(const FMetasoundFrontendDocument& InPreBuildDoc, const TArray<FMetasoundFrontendClassInput>& InTransmittableInputs) const;
+	TSharedPtr<Metasound::IGraph, ESPMode::ThreadSafe> BuildMetasoundDocument(const FMetasoundFrontendDocument& InPreprocessDoc, const TArray<FMetasoundFrontendClassInput>& InTransmittableInputs) const;
 	Metasound::FSendAddress CreateSendAddress(uint64 InInstanceID, const Metasound::FVertexName& InVertexName, const FName& InDataTypeName) const;
 	Metasound::Frontend::FNodeHandle AddInputPinForSendAddress(const Metasound::FMetaSoundParameterTransmitter::FSendInfo& InSendInfo, Metasound::Frontend::FGraphHandle InGraph) const;
 };
diff --git a/Engine/Plugins/Runtime/Metasound/Source/MetasoundFrontend/Public/MetasoundFrontendNodeTemplateRegistry.h b/Engine/Plugins/Runtime/Metasound/Source/MetasoundFrontend/Public/MetasoundFrontendNodeTemplateRegistry.h
--- a/Engine/Plugins/Runtime/Metasound/Source/MetasoundFrontend/Public/MetasoundFrontendNodeTemplateRegistry.h
+++ b/Engine/Plugins/Runtime/Metasound/Source/MetasoundFrontend/Public/MetasoundFrontendNodeTemplateRegistry.h
@@ -4,24 +4,42 @@
 
 #include "MetasoundFrontendController.h"
 #include "MetasoundFrontendDocument.h"
+#include "MetasoundFrontendTransform.h"
 #include "Templates/SharedPointer.h"
+#include "Templates/UniquePtr.h"
 #include "UObject/NoExportTypes.h"
 
+
 namespace Metasound
 {
 	namespace Frontend
 	{
+		/**
+		  * Base interface for a node template, which acts in place of frontend node class and respective instance(s).
+		  * Instances are preprocessed, allowing for custom graph manipulation prior to generating a respective runtime
+		  * graph operator representation.
+		  */
 		class METASOUNDFRONTEND_API INodeTemplate
 		{
 		public:
 			virtual ~INodeTemplate() = default;
 
-			virtual bool BuildTemplate(FMetasoundFrontendDocument& InOutDocument, FMetasoundFrontendGraph& InOutGraph, FMetasoundFrontendNode& InOutNode) const = 0;
+			// Generates node transform that can cache document state for use when individually preprocessing nodes.
+			virtual TUniquePtr<INodeTransform> GenerateNodeTransform(FMetasoundFrontendDocument& InDocument) const = 0;
+
+			// Returns the class definition for the given node class template.
 			virtual const FMetasoundFrontendClass& GetFrontendClass() const = 0;
+
+			// Returns the version of the given node class template.
 			virtual const FMetasoundFrontendVersion& GetVersion() const = 0;
+
+			// Given the provided node interface, returns whether or not it conforms to an expected format
+			// that can be successfully manipulated by a generated node template preprocessor.
 			virtual bool IsValidNodeInterface(const FMetasoundFrontendNodeInterface& InNodeInterface) const = 0;
 
 #if WITH_EDITOR
+			// Returns whether or not the given node template has the necessary
+			// required connections to be preprocessed (editor only).
 			virtual bool HasRequiredConnections(FConstNodeHandle InNodeHandle) const = 0;
 #endif // WITH_EDITOR
 		};
@@ -29,6 +47,7 @@ namespace Metasound
 		class METASOUNDFRONTEND_API INodeTemplateRegistry
 		{
 		public:
+			// Returns singleton template registry.
 			static INodeTemplateRegistry& Get();
 
 			virtual ~INodeTemplateRegistry() = default;
@@ -38,7 +57,7 @@ namespace Metasound
 		};
 
 		// Register & Unregister are not publicly accessible implementation as the API
-		// is in beta and currently, only to be used by reroute nodes.
+		// is in beta and, currently, only to be used by internal implementation (ex. reroute nodes).
 		void RegisterNodeTemplate(TUniquePtr<INodeTemplate>&& InTemplate);
 		void UnregisterNodeTemplate(const FMetasoundFrontendVersion& InNodeTemplateVersion);
 	} // namespace Frontend
diff --git a/Engine/Plugins/Runtime/Metasound/Source/MetasoundFrontend/Public/MetasoundFrontendTransform.h b/Engine/Plugins/Runtime/Metasound/Source/MetasoundFrontend/Public/MetasoundFrontendTransform.h
--- a/Engine/Plugins/Runtime/Metasound/Source/MetasoundFrontend/Public/MetasoundFrontendTransform.h
+++ b/Engine/Plugins/Runtime/Metasound/Source/MetasoundFrontend/Public/MetasoundFrontendTransform.h
@@ -26,21 +26,50 @@ namespace Metasound
 		}
 
 		/** Interface for transforms applied to documents. */
-		class IDocumentTransform
+		class METASOUNDFRONTEND_API IDocumentTransform
 		{
 		public:
 			virtual ~IDocumentTransform() = default;
+
 			/** Return true if InDocument was modified, false otherwise. */
 			virtual bool Transform(FDocumentHandle InDocument) const = 0;
+
+			/** Return true if InDocument was modified, false otherwise.
+			  * This function is soft deprecated.  It is not pure virtual
+			  * to grandfather in old transform implementation. Old transforms
+			  * should be deprecated and rewritten to use the Controller-less
+			  * API in the interest of better performance and simplicity.
+			  */
+			virtual bool Transform(FMetasoundFrontendDocument& InOutDocument) const;
 		};
 
 		/** Interface for transforms applied to a graph. */
-		class IGraphTransform
+		class METASOUNDFRONTEND_API IGraphTransform
 		{
 		public:
 			virtual ~IGraphTransform() = default;
-			/** Return true if InGraph was modified, false otherwise. */
-			virtual bool Transform(FGraphHandle InGraph) const = 0;
+
+			// Returns reference to the node's owning document.
+			virtual FMetasoundFrontendDocument& GetOwningDocument() const = 0;
+
+			/** Return true if the graph was modified, false otherwise. */
+			virtual bool Transform(FMetasoundFrontendGraph& InOutGraph) const = 0;
+		};
+
+		/** Interface for transforming a node. */
+		class METASOUNDFRONTEND_API INodeTransform
+		{
+		public:
+			virtual ~INodeTransform() = default;
+
+			// Returns reference to the node's owning document.
+			virtual FMetasoundFrontendDocument& GetOwningDocument() const = 0;
+
+			// Returns reference to the node's owning graph.
+			virtual FMetasoundFrontendGraph& GetOwningGraph() const = 0;
+
+			/** Return true if the node was modified, false otherwise. */
+			virtual bool Transform(FMetasoundFrontendNode& InOutNode) const = 0;
 		};
 
 		/** Adds or swaps document members (inputs, outputs) and removing any document members where necessary and adding those missing. */
diff --git a/Engine/Plugins/Runtime/Metasound/Source/MetasoundFrontend/Public/NodeTemplates/MetasoundFrontendNodeTemplateReroute.h b/Engine/Plugins/Runtime/Metasound/Source/MetasoundFrontend/Public/NodeTemplates/MetasoundFrontendNodeTemplateReroute.h
--- a/Engine/Plugins/Runtime/Metasound/Source/MetasoundFrontend/Public/NodeTemplates/MetasoundFrontendNodeTemplateReroute.h
+++ b/Engine/Plugins/Runtime/Metasound/Source/MetasoundFrontend/Public/NodeTemplates/MetasoundFrontendNodeTemplateReroute.h
@@ -4,9 +4,11 @@
 #include "MetasoundFrontendController.h"
 #include "MetasoundFrontendDocument.h"
 #include "MetasoundFrontendNodeTemplateRegistry.h"
+#include "MetasoundFrontendTransform.h"
 #include "Templates/SharedPointer.h"
 #include "UObject/NoExportTypes.h"
 
+
 namespace Metasound
 {
 	namespace Frontend
@@ -23,7 +25,7 @@ namespace Metasound
 
 			virtual ~FRerouteNodeTemplate() = default;
 
-			virtual bool BuildTemplate(FMetasoundFrontendDocument& InOutDocument, FMetasoundFrontendGraph& InOutGraph, FMetasoundFrontendNode& InOutNode) const override;
+			virtual TUniquePtr<INodeTransform> GenerateNodeTransform(FMetasoundFrontendDocument& InPreprocessedDocument) const override;
 			virtual const FMetasoundFrontendClass& GetFrontendClass() const override;
 			virtual const FMetasoundFrontendVersion& GetVersion() const override;
 			virtual bool IsValidNodeInterface(const FMetasoundFrontendNodeInterface& InNodeInterface) const override;
