- Top align MS node inputs/outputs in editor unless titles hidden & only one member (ex. math nodes)
- Clean-up of more direct calls to MetaSoundStyle to use safe brush call
- Clean-up a few nodehandle spots that should be const
#rnx

diff --git a/Engine/Plugins/Runtime/Metasound/Source/MetasoundEditor/Private/SMetasoundGraphNode.cpp b/Engine/Plugins/Runtime/Metasound/Source/MetasoundEditor/Private/SMetasoundGraphNode.cpp
--- a/Engine/Plugins/Runtime/Metasound/Source/MetasoundEditor/Private/SMetasoundGraphNode.cpp
+++ b/Engine/Plugins/Runtime/Metasound/Source/MetasoundEditor/Private/SMetasoundGraphNode.cpp
@@ -331,12 +331,9 @@ namespace Metasound
 				{
 					PinWidget = SNew(SMetasoundGraphPin, InPin);
 
-					if (const ISlateStyle* MetasoundStyle = FSlateStyleRegistry::FindSlateStyle("MetaSoundStyle"))
-					{
-						const FSlateBrush* PinConnectedBrush = MetasoundStyle->GetBrush(TEXT("MetasoundEditor.Graph.TriggerPin.Connected"));
-						const FSlateBrush* PinDisconnectedBrush = MetasoundStyle->GetBrush(TEXT("MetasoundEditor.Graph.TriggerPin.Disconnected"));
-						PinWidget->SetCustomPinIcon(PinConnectedBrush, PinDisconnectedBrush);
-					}
+					const FSlateBrush& PinConnectedBrush = Editor::Style::GetSlateBrushSafe("MetasoundEditor.Graph.TriggerPin.Connected");
+					const FSlateBrush& PinDisconnectedBrush = Editor::Style::GetSlateBrushSafe("MetasoundEditor.Graph.TriggerPin.Disconnected");
+					PinWidget->SetCustomPinIcon(&PinConnectedBrush, &PinDisconnectedBrush);
 				}
 			}
 
@@ -356,7 +353,7 @@ namespace Metasound
 				TSharedPtr<SGraphPin> NewPin = CreatePinWidget(InPin);
 				check(NewPin.IsValid());
 
-				Metasound::Frontend::FNodeHandle NodeHandle = GetMetaSoundNode().GetNodeHandle();
+				Frontend::FConstNodeHandle NodeHandle = GetMetaSoundNode().GetConstNodeHandle();
 				if (InPin->Direction == EGPD_Input)
 				{
 					if (!NodeHandle->GetClassStyle().Display.bShowInputNames)
@@ -378,7 +375,7 @@ namespace Metasound
 
 		TSharedRef<SWidget> SMetaSoundGraphNode::CreateTitleWidget(TSharedPtr<SNodeTitle> NodeTitle)
 		{
-			Metasound::Frontend::FNodeHandle NodeHandle = GetMetaSoundNode().GetNodeHandle();
+			Frontend::FConstNodeHandle NodeHandle = GetMetaSoundNode().GetConstNodeHandle();
 			if (!NodeHandle->GetClassStyle().Display.bShowName)
 			{
 				return SNullWidget::NullWidget;
@@ -644,7 +641,7 @@ namespace Metasound
 		{
 			using namespace Frontend;
 
-			FNodeHandle NodeHandle = GetMetaSoundNode().GetNodeHandle();
+			FConstNodeHandle NodeHandle = GetMetaSoundNode().GetConstNodeHandle();
 			const FMetasoundFrontendClassStyleDisplay& StyleDisplay = NodeHandle->GetClassStyle().Display;
 			TSharedPtr<SHorizontalBox> ContentBox = SNew(SHorizontalBox);
 			TSharedPtr<SWidget> OuterContentBox; // currently only used for input float nodes to accommodate the input widget
@@ -858,46 +855,40 @@ namespace Metasound
 				}
 			}
 	
-			static const float GrabPadding = 28.0f;
-
 			// Gives more space for user to grab a bit easier as variables do not have any title area nor icon
-			const float LeftNodeGrabPadding = IsVariableMutator() ? GrabPadding : 0.0f;
+			const float GrabPadding = IsVariableMutator() ? 28.0f : 0.0f;
+
+			const EVerticalAlignment PinNodeAlignInput = (!StyleDisplay.bShowInputNames && NodeHandle->GetNumInputs() == 1) ? VAlign_Center : VAlign_Top;
 			ContentBox->AddSlot()
 			.HAlign(HAlign_Left)
-			.VAlign(VAlign_Top)
+			.VAlign(PinNodeAlignInput)
 			.FillWidth(1.0f)
-			.Padding(0.0f, 0.0f, LeftNodeGrabPadding, 0.0f)
+			.Padding(0.0f, 0.0f, GrabPadding, 0.0f)
 			[
 				SAssignNew(LeftNodeBox, SVerticalBox)
 			];
 
 			if (!StyleDisplay.ImageName.IsNone())
 			{
-				if (const ISlateStyle* MetasoundStyle = FSlateStyleRegistry::FindSlateStyle("MetaSoundStyle"))
-				{
-					if (const FSlateBrush* ImageBrush = MetasoundStyle->GetBrush(StyleDisplay.ImageName))
-					{
-						ContentBox->AddSlot()
-						.AutoWidth()
-						.HAlign(HAlign_Center)
-						.VAlign(VAlign_Center)
-						[
-							SNew(SImage)
-							.Image(ImageBrush)
-							.ColorAndOpacity(FSlateColor::UseForeground())
-							.DesiredSizeOverride(FVector2D(20, 20))
-						];
-					}
-				}
+				const FSlateBrush& ImageBrush = Metasound::Editor::Style::GetSlateBrushSafe(StyleDisplay.ImageName);
+				ContentBox->AddSlot()
+				.AutoWidth()
+				.HAlign(HAlign_Center)
+				.VAlign(VAlign_Center)
+				[
+					SNew(SImage)
+					.Image(&ImageBrush)
+					.ColorAndOpacity(FSlateColor::UseForeground())
+					.DesiredSizeOverride(FVector2D(20, 20))
+				];
 			}
 
-			// Gives more space for user to grab a bit easier as variables do not have any title area nor icon
-			const float RightNodeGrabPadding = IsVariableAccessor() ? GrabPadding : 0.0f;
+			const EVerticalAlignment PinNodeAlignOutput = (!StyleDisplay.bShowInputNames && NodeHandle->GetNumOutputs() == 1) ? VAlign_Center : VAlign_Top;
 			ContentBox->AddSlot()
 				.AutoWidth()
 				.HAlign(HAlign_Right)
-				.VAlign(VAlign_Center)
-				.Padding(RightNodeGrabPadding, 0.0f, 0.0f, 0.0f)
+				.VAlign(PinNodeAlignOutput)
+				.Padding(GrabPadding, 0.0f, 0.0f, 0.0f)
 				[
 					SAssignNew(RightNodeBox, SVerticalBox)
 				];
@@ -982,10 +973,7 @@ namespace Metasound
 			{
 				if (!HasRequiredConnections())
 				{
-					if (const ISlateStyle* MetasoundStyle = FSlateStyleRegistry::FindSlateStyle("MetaSoundStyle"))
-					{
-						return MetasoundStyle->GetBrush(TEXT("MetasoundEditor.Graph.InvalidReroute"));
-					}
+					return &Editor::Style::GetSlateBrushSafe("MetasoundEditor.Graph.InvalidReroute");
 				}
 
 				Pin = FGraphBuilder::FindReroutedOutputPin(Pin);
@@ -1007,19 +995,16 @@ namespace Metasound
 
 			if (bIsConstructorPin)
 			{
-				if (const ISlateStyle* MetasoundStyle = FSlateStyleRegistry::FindSlateStyle("MetaSoundStyle"))
+				const bool bIsConnected = IsConnected();
+				if (IsArray())
 				{
-					const bool bIsConnected = IsConnected();
-					if (IsArray())
-					{
-						return bIsConnected ? MetasoundStyle->GetBrush(TEXT("MetasoundEditor.Graph.ConstructorPinArray")) :
-							MetasoundStyle->GetBrush(TEXT("MetasoundEditor.Graph.ConstructorPinArrayDisconnected"));
-					}
-					else
-					{
-						return bIsConnected ? MetasoundStyle->GetBrush(TEXT("MetasoundEditor.Graph.ConstructorPin")) :
-							MetasoundStyle->GetBrush(TEXT("MetasoundEditor.Graph.ConstructorPinDisconnected"));
-					}
+					const FName BrushName = bIsConnected ? "MetasoundEditor.Graph.ConstructorPinArray" : "MetasoundEditor.Graph.ConstructorPinArrayDisconnected";
+					return &Editor::Style::GetSlateBrushSafe(BrushName);
+				}
+				else
+				{
+					const FName BrushName = bIsConnected ? "MetasoundEditor.Graph.ConstructorPin" : "MetasoundEditor.Graph.ConstructorPinDisconnected";
+					return &Editor::Style::GetSlateBrushSafe(BrushName);
 				}
 			}
 
