Tooltip fix
#rnx

diff --git a/Engine/Plugins/Runtime/Metasound/Source/MetasoundEngine/Private/MetasoundWaveTableOscillatorNode.cpp b/Engine/Plugins/Runtime/Metasound/Source/MetasoundEngine/Private/MetasoundWaveTableOscillatorNode.cpp
--- a/Engine/Plugins/Runtime/Metasound/Source/MetasoundEngine/Private/MetasoundWaveTableOscillatorNode.cpp
+++ b/Engine/Plugins/Runtime/Metasound/Source/MetasoundEngine/Private/MetasoundWaveTableOscillatorNode.cpp
@@ -35,7 +35,7 @@ namespace Metasound
 					TInputDataVertex<float>("Freq", FDataVertexMetadata { LOCTEXT("MetasoundWaveTableOscillatorNode_Freq", "Freq") }, 440.0f),
 					TInputDataVertex<FAudioBuffer>("PhaseMod", FDataVertexMetadata
 					{
-						LOCTEXT("MetasoundWaveTableOscillatorNode_PhaseModDescription", "Modulation audio source for modulating oscillation phase of provided table. A value of 0 is no phase modulation and 1 a full buffer length (360 degrees) of phase shift."),
+						LOCTEXT("MetasoundWaveTableOscillatorNode_PhaseModDescription", "Modulation audio source for modulating oscillation phase of provided table. A value of 0 is no phase modulation and 1 a full table length (360 degrees) of phase shift."),
 						LOCTEXT("MetasoundWaveTableOscillatorNode_PhaseMod", "Phase Modulator"),
 						true /* bIsAdvancedDisplay */
 					})
