Virtual call in ctor & rogue #pragma in impl fixes
#rnx

diff --git a/Engine/Plugins/Runtime/Metasound/Source/MetasoundFrontend/Private/MetasoundFrontendDocumentTemplatePreprocessor.cpp b/Engine/Plugins/Runtime/Metasound/Source/MetasoundFrontend/Private/MetasoundFrontendDocumentTemplatePreprocessor.cpp
--- a/Engine/Plugins/Runtime/Metasound/Source/MetasoundFrontend/Private/MetasoundFrontendDocumentTemplatePreprocessor.cpp
+++ b/Engine/Plugins/Runtime/Metasound/Source/MetasoundFrontend/Private/MetasoundFrontendDocumentTemplatePreprocessor.cpp
@@ -1,5 +1,4 @@
 // Copyright Epic Games, Inc. All Rights Reserved.
-#pragma once
 
 #include "NodeTemplates/MetasoundFrontendDocumentTemplatePreprocessor.h"
 
diff --git a/Engine/Plugins/Runtime/Metasound/Source/MetasoundFrontend/Private/NodeTemplates/MetasoundFrontendDocumentTemplatePreprocessor.h b/Engine/Plugins/Runtime/Metasound/Source/MetasoundFrontend/Private/NodeTemplates/MetasoundFrontendDocumentTemplatePreprocessor.h
--- a/Engine/Plugins/Runtime/Metasound/Source/MetasoundFrontend/Private/NodeTemplates/MetasoundFrontendDocumentTemplatePreprocessor.h
+++ b/Engine/Plugins/Runtime/Metasound/Source/MetasoundFrontend/Private/NodeTemplates/MetasoundFrontendDocumentTemplatePreprocessor.h
@@ -23,7 +23,7 @@ namespace Metasound
 		// Base implementation for preprocessing a given template node
 		class FNodeTemplatePreprocessTransformBase : public INodeTransform
 		{
-		private:
+		protected:
 			FMetasoundFrontendDocument& Document;
 			FMetasoundFrontendGraph& Graph;
 
diff --git a/Engine/Plugins/Runtime/Metasound/Source/MetasoundFrontend/Private/NodeTemplates/MetasoundFrontendNodeTemplateReroute.cpp b/Engine/Plugins/Runtime/Metasound/Source/MetasoundFrontend/Private/NodeTemplates/MetasoundFrontendNodeTemplateReroute.cpp
--- a/Engine/Plugins/Runtime/Metasound/Source/MetasoundFrontend/Private/NodeTemplates/MetasoundFrontendNodeTemplateReroute.cpp
+++ b/Engine/Plugins/Runtime/Metasound/Source/MetasoundFrontend/Private/NodeTemplates/MetasoundFrontendNodeTemplateReroute.cpp
@@ -77,7 +77,7 @@ namespace Metasound
 				FRerouteNodeTemplatePreprocessTransform(FMetasoundFrontendDocument& InDocument)
 					: FNodeTemplatePreprocessTransformBase(InDocument)
 				{
-					TArray<FMetasoundFrontendEdge>& Edges = GetOwningGraph().Edges;
+					TArray<FMetasoundFrontendEdge>& Edges = Graph.Edges;
 					for (int32 i = 0; i < Edges.Num(); ++i)
 					{
 						FMetasoundFrontendEdge& Edge = Edges[i];
