Make literal nodes produce a value reference instead of a read reference.

diff --git a/Engine/Plugins/Runtime/Metasound/Source/MetasoundGraphCore/Private/MetasoundFacade.cpp b/Engine/Plugins/Runtime/Metasound/Source/MetasoundGraphCore/Private/MetasoundFacade.cpp
--- a/Engine/Plugins/Runtime/Metasound/Source/MetasoundGraphCore/Private/MetasoundFacade.cpp
+++ b/Engine/Plugins/Runtime/Metasound/Source/MetasoundGraphCore/Private/MetasoundFacade.cpp
@@ -9,13 +9,25 @@ namespace Metasound
 	{
 	}
 
+	FNodeFacade::FFactory::FFactory(FOriginalCreateOperatorFunction InCreateFunc)
+	:	CreateFunc(WrapOriginalCreateOperatorFunction(InCreateFunc))
+	{
+	}
+
+
 	TUniquePtr<IOperator> FNodeFacade::FFactory::CreateOperator(const FBuildOperatorParams& InParams, FBuildResults& OutResults)
 	{
-		// TODO: update facade to deprecate old CreateOperator functionality.
-		FDataReferenceCollection InputDataReferences = InParams.InputData.ToDataReferenceCollection();
-		FCreateOperatorParams DeprecatedParams{ InParams.Node, InParams.OperatorSettings, InputDataReferences, InParams.Environment, InParams.Builder };
+		return CreateFunc(InParams, OutResults);
+	}
 
-		return CreateFunc(DeprecatedParams, OutResults.Errors);
+	FNodeFacade::FFactory::FCreateOperatorFunction FNodeFacade::FFactory::WrapOriginalCreateOperatorFunction(FOriginalCreateOperatorFunction InCreateFunc)
+	{
+		return [InCreateFunc](const FBuildOperatorParams& InParams, FBuildResults& OutResults)
+		{
+			FDataReferenceCollection InputDataReferences = InParams.InputData.ToDataReferenceCollection();
+			FCreateOperatorParams DeprecatedParams{ InParams.Node, InParams.OperatorSettings, InputDataReferences, InParams.Environment, InParams.Builder };
+			return InCreateFunc(DeprecatedParams, OutResults.Errors);
+		};
 	}
 
 	const FVertexInterface& FNodeFacade::GetVertexInterface() const
diff --git a/Engine/Plugins/Runtime/Metasound/Source/MetasoundGraphCore/Public/MetasoundFacade.h b/Engine/Plugins/Runtime/Metasound/Source/MetasoundGraphCore/Public/MetasoundFacade.h
--- a/Engine/Plugins/Runtime/Metasound/Source/MetasoundGraphCore/Public/MetasoundFacade.h
+++ b/Engine/Plugins/Runtime/Metasound/Source/MetasoundGraphCore/Public/MetasoundFacade.h
@@ -23,6 +23,10 @@ namespace Metasound
 			template<typename T>
 			static uint8 Check(Helper<TUniquePtr<IOperator>(*)(const FCreateOperatorParams& InParams, FBuildErrorArray& OutErrors), &T::CreateOperator>*);
 
+			// Check for "static TUniquePtr<IOperator> U::CreateOperator(const FBuildOperatorParams& Inparams, FBuildErrorArray& OutErrors)"
+			template<typename T>
+			static uint8 Check(Helper<TUniquePtr<IOperator>(*)(const FBuildOperatorParams& InParams, FBuildResults& OutResults), &T::CreateOperator>*);
+
 			template<typename T> static uint16 Check(...);
 
 		public:
@@ -94,14 +98,19 @@ namespace Metasound
 			// the signature of the CreateOperator function.
 			class METASOUNDGRAPHCORE_API FFactory : public IOperatorFactory
 			{
-				using FCreateOperatorFunction = TFunction<TUniquePtr<IOperator>(const FCreateOperatorParams& InParams, FBuildErrorArray& OutErrors)>;
+				using FCreateOperatorFunction = TFunction<TUniquePtr<IOperator>(const FBuildOperatorParams& InParams, FBuildResults& OutResults)>;
+				using FOriginalCreateOperatorFunction = TFunction<TUniquePtr<IOperator>(const FCreateOperatorParams& InParams, FBuildErrorArray& OutErrors)>;
 
 				public:
 					FFactory(FCreateOperatorFunction InCreateFunc);
+					FFactory(FOriginalCreateOperatorFunction InCreateFunc);
 
 					virtual TUniquePtr<IOperator> CreateOperator(const FBuildOperatorParams& InParams, FBuildResults& OutResults) override;
 
 				private:
+					// Wrap the prior version of CreateOperator with the new version of CreateOperator
+					static FCreateOperatorFunction WrapOriginalCreateOperatorFunction(FOriginalCreateOperatorFunction InCreateFunc);
+
 					FCreateOperatorFunction CreateFunc;
 			};
 
diff --git a/Engine/Plugins/Runtime/Metasound/Source/MetasoundGraphCore/Public/MetasoundLiteralNode.h b/Engine/Plugins/Runtime/Metasound/Source/MetasoundGraphCore/Public/MetasoundLiteralNode.h
--- a/Engine/Plugins/Runtime/Metasound/Source/MetasoundGraphCore/Public/MetasoundLiteralNode.h
+++ b/Engine/Plugins/Runtime/Metasound/Source/MetasoundGraphCore/Public/MetasoundLiteralNode.h
@@ -24,18 +24,73 @@ namespace Metasound
 	template<typename DataType>
 	class TLiteralOperator : public IOperator
 	{
+		static_assert(!TExecutableDataType<DataType>::bIsExecutable, "TLiteralOperator is only suitable for non-executable data types");
+
+	public:
+		TLiteralOperator(TDataValueReference<DataType> InValue)
+		: Value (InValue)
+		{
+		}
+
+		virtual ~TLiteralOperator() = default;
+
+		virtual void Bind(FVertexInterfaceData& InVertexData) const
+		{
+			using namespace LiteralNodeNames;
+			InVertexData.GetOutputs().BindValueVertex(METASOUND_GET_PARAM_NAME(OutputValue), Value);
+		}
+
+		virtual FDataReferenceCollection GetInputs() const override
+		{
+			FDataReferenceCollection Inputs;
+			return Inputs;
+		}
+
+		virtual FDataReferenceCollection GetOutputs() const override
+		{
+			using namespace LiteralNodeNames;
+			FDataReferenceCollection Outputs;
+			Outputs.AddDataReference(METASOUND_GET_PARAM_NAME(OutputValue), FAnyDataReference{Value});
+			return Outputs;
+		}
+
+		virtual FExecuteFunction GetExecuteFunction() override
+		{
+			return nullptr;
+		}
+
+	private:
+		TDataValueReference<DataType> Value;
+	};
+
+	/** TExecutableLiteralOperator is used for executable types. The majority of literals
+	 * are constant values, but executable data types cannot be constant values.
+	 * 
+	 * Note: this is supporting for FTrigger literals which are deprecated.
+	 */
+	template<typename DataType>
+	class TExecutableLiteralOperator : public IOperator
+	{
+		static_assert(TExecutableDataType<DataType>::bIsExecutable, "TExecutableLiteralOperator is only suitable for executable data types");
+
 	public:
 		using FDataWriteReference = TDataWriteReference<DataType>;
 
-		TLiteralOperator(FDataWriteReference InDataReference)
+		TExecutableLiteralOperator(FDataWriteReference InDataReference)
 			// Executable DataTypes require a copy of the output to operate on whereas non-executable
 			// types do not. Avoid copy by assigning to reference for non-executable types.
 			: InputValue(InDataReference)
-			, OutputValue(TExecutableDataType<DataType>::bIsExecutable ? FDataWriteReference::CreateNew(*InDataReference) : InDataReference)
+			, OutputValue(FDataWriteReference::CreateNew(*InDataReference))
 		{
 		}
 
-		virtual ~TLiteralOperator() = default;
+		virtual ~TExecutableLiteralOperator() = default;
+
+		virtual void Bind(FVertexInterfaceData& InVertexData) const
+		{
+			using namespace LiteralNodeNames;
+			InVertexData.GetOutputs().BindReadVertex(METASOUND_GET_PARAM_NAME(OutputValue), OutputValue);
+		}
 
 		virtual FDataReferenceCollection GetInputs() const override
 		{
@@ -46,7 +101,6 @@ namespace Metasound
 		virtual FDataReferenceCollection GetOutputs() const override
 		{
 			using namespace LiteralNodeNames;
-
 			FDataReferenceCollection Outputs;
 			Outputs.AddDataReadReference<DataType>(METASOUND_GET_PARAM_NAME(OutputValue), OutputValue);
 			return Outputs;
@@ -59,16 +113,12 @@ namespace Metasound
 
 		static void ExecuteFunction(IOperator* InOperator)
 		{
-			static_cast<TLiteralOperator<DataType>*>(InOperator)->Execute();
+			static_cast<TExecutableLiteralOperator<DataType>*>(InOperator)->Execute();
 		}
 
 		virtual FExecuteFunction GetExecuteFunction() override
 		{
-			if (TExecutableDataType<DataType>::bIsExecutable)
-			{
-				return &TLiteralOperator<DataType>::ExecuteFunction;
-			}
-			return nullptr;
+			return &TExecutableLiteralOperator<DataType>::ExecuteFunction;
 		}
 
 	private:
@@ -100,6 +150,10 @@ namespace Metasound
 	template<typename DataType>
 	class TLiteralNode : public FNode
 	{
+		// Executable data types handle Triggers which need to be advanced every
+		// block.
+		static constexpr bool bIsExecutableDataType = TExecutableDataType<DataType>::bIsExecutable;
+
 	public:
 
 		static FVertexInterface DeclareVertexInterface()
@@ -110,12 +164,27 @@ namespace Metasound
 				  FText::GetEmpty() // description
 				, METASOUND_GET_PARAM_DISPLAYNAME(OutputValue) // display name
 			};
-			return FVertexInterface(
-				FInputVertexInterface(),
-				FOutputVertexInterface(
-					TOutputDataVertex<DataType>(METASOUND_GET_PARAM_NAME(OutputValue), OutputMetadata)
-				)
-			);
+
+			if constexpr (bIsExecutableDataType)
+			{
+				return FVertexInterface(
+					FInputVertexInterface(),
+					FOutputVertexInterface(
+						TOutputDataVertex<DataType>(METASOUND_GET_PARAM_NAME(OutputValue), OutputMetadata)
+					)
+				);
+			}
+			else
+			{
+				// If the data type is not executable, we treat it as a constructor
+				// vertex to produce a constant value.
+				return FVertexInterface(
+					FInputVertexInterface(),
+					FOutputVertexInterface(
+						TOutputConstructorVertex<DataType>(METASOUND_GET_PARAM_NAME(OutputValue), OutputMetadata)
+					)
+				);
+			}
 		}
 
 		static FNodeClassMetadata GetNodeInfo()
@@ -175,10 +244,20 @@ namespace Metasound
 	{
 		using FLiteralNodeType = TLiteralNode<DataType>;
 
-		// Create write reference by calling compatible constructor with literal.
-		FDataWriteReference DataRef = TDataWriteReferenceLiteralFactory<DataType>::CreateExplicitArgs(InParams.OperatorSettings, InitParam);
+		static constexpr bool bIsExecutableDataType = TExecutableDataType<DataType>::bIsExecutable;
 
-		return MakeUnique<TLiteralOperator<DataType>>(DataRef);
+		if constexpr (bIsExecutableDataType)
+		{
+			// Create write reference by calling compatible constructor with literal.
+			TDataWriteReference<DataType> DataRef = TDataWriteReferenceLiteralFactory<DataType>::CreateExplicitArgs(InParams.OperatorSettings, InitParam);
+			return MakeUnique<TExecutableLiteralOperator<DataType>>(DataRef);
+		}
+		else
+		{
+			// Create value reference by calling compatible constructor with literal.
+			TDataValueReference DataRef = TDataValueReferenceLiteralFactory<DataType>::CreateExplicitArgs(InParams.OperatorSettings, InitParam);
+			return MakeUnique<TLiteralOperator<DataType>>(DataRef);
+		}
 	}
 
 } // namespace Metasound
diff --git a/Engine/Plugins/Runtime/Metasound/Source/MetasoundGraphCore/Public/MetasoundVertexData.h b/Engine/Plugins/Runtime/Metasound/Source/MetasoundGraphCore/Public/MetasoundVertexData.h
--- a/Engine/Plugins/Runtime/Metasound/Source/MetasoundGraphCore/Public/MetasoundVertexData.h
+++ b/Engine/Plugins/Runtime/Metasound/Source/MetasoundGraphCore/Public/MetasoundVertexData.h
@@ -180,6 +180,22 @@ namespace Metasound
 				return nullptr;
 			}
 
+			// Gets the value of the bound data reference if it exists. Otherwise
+			// create and return a value by constructing one using the Vertex's 
+			// default literal.
+			template<typename DataType>
+			DataType GetOrCreateDefaultValue(const FOperatorSettings& InSettings) const
+			{
+				if (Data.IsSet())
+				{
+					if (const DataType* Value = Data->GetValue<DataType>())
+					{
+						return *Value;
+					}
+				}
+				return TDataTypeLiteralFactory<DataType>::CreateExplicitArgs(InSettings, Vertex.GetDefaultLiteral());
+			}
+
 			// Set the value with a constant value reference.
 			template<typename DataType>
 			void SetValue(const DataType& InValue)
@@ -307,6 +323,16 @@ namespace Metasound
 			return Binding->GetValue<DataType>();
 		}
 
+		/**  Gets the value of the bound data reference if it exists. Otherwise
+		 * create and return a value by constructing one using the Vertex's 
+		 * default literal. */
+		template<typename DataType>
+		DataType GetOrCreateDefaultValue(const FVertexName& InVertexName, const FOperatorSettings& InSettings) const
+		{
+			const FBinding* Binding = FindChecked(InVertexName);
+			return Binding->GetOrCreateDefaultValue<DataType>(InSettings);
+		}
+
 		/** Set the value of a vertex. */
 		template<typename DataType>
 		void SetValue(const FVertexName& InVertexName, const DataType& InValue)
diff --git a/Engine/Plugins/Runtime/Metasound/Source/MetasoundStandardNodes/Private/MetasoundDelayNode.cpp b/Engine/Plugins/Runtime/Metasound/Source/MetasoundStandardNodes/Private/MetasoundDelayNode.cpp
--- a/Engine/Plugins/Runtime/Metasound/Source/MetasoundStandardNodes/Private/MetasoundDelayNode.cpp
+++ b/Engine/Plugins/Runtime/Metasound/Source/MetasoundStandardNodes/Private/MetasoundDelayNode.cpp
@@ -20,14 +20,16 @@ namespace Metasound
 	namespace Delay
 	{
 		METASOUND_PARAM(InParamAudioInput, "In", "Audio input.")
-		METASOUND_PARAM(InParamDelayTime, "Delay Time", "The amount of time to delay the audio, in seconds.")
+		METASOUND_PARAM(InParamDelayTime, "Delay Time", "The amount of time to delay the audio.")
 		METASOUND_PARAM(InParamDryLevel, "Dry Level", "The dry level of the delay.")
 		METASOUND_PARAM(InParamWetLevel, "Wet Level", "The wet level of the delay.")
 		METASOUND_PARAM(InParamFeedbackAmount, "Feedback", "Feedback amount.")
+		METASOUND_PARAM(InParamMaxDelayTime, "Max Delay Time", "The maximum amount of time to delay the audio.")
 		METASOUND_PARAM(OutParamAudio, "Out", "Audio output.")
 
-		// TODO: make this a static vertex
-		static float MaxDelaySeconds = 5.0f;
+		static constexpr float MinMaxDelaySeconds = 0.001f;
+		static constexpr float MaxMaxDelaySeconds = 1000.f;
+		static constexpr float DefaultMaxDelaySeconds = 5.0f;
 	}
 
 	class FDelayOperator : public TExecutableOperator<FDelayOperator>
@@ -36,15 +38,17 @@ namespace Metasound
 
 		static const FNodeClassMetadata& GetNodeInfo();
 		static const FVertexInterface& GetVertexInterface();
-		static TUniquePtr<IOperator> CreateOperator(const FCreateOperatorParams& InParams, FBuildErrorArray& OutErrors);
+		static TUniquePtr<IOperator> CreateOperator(const FBuildOperatorParams& InParams, FBuildResults& OutResults);
 
 		FDelayOperator(const FOperatorSettings& InSettings, 
 			const FAudioBufferReadRef& InAudioInput, 
 			const FTimeReadRef& InDelayTime, 
 			const FFloatReadRef& InDryLevel, 
 			const FFloatReadRef& InWetLevel, 
-			const FFloatReadRef& InFeedback);
+			const FFloatReadRef& InFeedback,
+			float InMaxDelayTimeSeconds);
 
+		virtual void Bind(FVertexInterfaceData& InVertexData) const override;
 		virtual FDataReferenceCollection GetInputs() const override;
 		virtual FDataReferenceCollection GetOutputs() const override;
 		void Execute();
@@ -74,10 +78,13 @@ namespace Metasound
 		Audio::FDelay DelayBuffer;
 
 		// The previous delay time
-		float PrevDelayTimeMsec;
+		float PrevDelayTimeMsec = 0.f;
 
 		// Feedback sample
-		float FeedbackSample;
+		float FeedbackSample = 0.f;
+
+		// Maximum delay time
+		float MaxDelayTimeSeconds = Delay::DefaultMaxDelaySeconds;
 	};
 
 	FDelayOperator::FDelayOperator(const FOperatorSettings& InSettings,
@@ -85,7 +92,8 @@ namespace Metasound
 		const FTimeReadRef& InDelayTime,
 		const FFloatReadRef& InDryLevel,
 		const FFloatReadRef& InWetLevel,
-		const FFloatReadRef& InFeedback)
+		const FFloatReadRef& InFeedback,
+		float InMaxDelayTimeSeconds)
 
 		: AudioInput(InAudioInput)
 		, DelayTime(InDelayTime)
@@ -93,40 +101,57 @@ namespace Metasound
 		, WetLevel(InWetLevel)
 		, Feedback(InFeedback)
 		, AudioOutput(FAudioBufferWriteRef::CreateNew(InSettings))
-		, PrevDelayTimeMsec(GetInputDelayTimeMsec())
-		, FeedbackSample(0.0f)
 	{
-		DelayBuffer.Init(InSettings.GetSampleRate(), Delay::MaxDelaySeconds);
+		MaxDelayTimeSeconds = FMath::Clamp(InMaxDelayTimeSeconds, Delay::MinMaxDelaySeconds, Delay::MaxMaxDelaySeconds);
+		PrevDelayTimeMsec = GetInputDelayTimeMsec();
+
+		DelayBuffer.Init(InSettings.GetSampleRate(), MaxDelayTimeSeconds);
 		DelayBuffer.SetDelayMsec(PrevDelayTimeMsec);
 	}
 
-	FDataReferenceCollection FDelayOperator::GetInputs() const
+	void FDelayOperator::Bind(FVertexInterfaceData& InVertexData) const
 	{
 		using namespace Delay;
 
-		FDataReferenceCollection InputDataReferences;
-		InputDataReferences.AddDataReadReference(METASOUND_GET_PARAM_NAME(InParamAudioInput), FAudioBufferReadRef(AudioInput));
-		InputDataReferences.AddDataReadReference(METASOUND_GET_PARAM_NAME(InParamDelayTime), FTimeReadRef(DelayTime));
-		InputDataReferences.AddDataReadReference(METASOUND_GET_PARAM_NAME(InParamDryLevel), FFloatReadRef(DryLevel));
-		InputDataReferences.AddDataReadReference(METASOUND_GET_PARAM_NAME(InParamWetLevel), FFloatReadRef(WetLevel));
-		InputDataReferences.AddDataReadReference(METASOUND_GET_PARAM_NAME(InParamFeedbackAmount), FFloatReadRef(Feedback));
+		FInputVertexInterfaceData& Inputs = InVertexData.GetInputs();
+
+		Inputs.BindReadVertex(METASOUND_GET_PARAM_NAME(InParamAudioInput), FAudioBufferReadRef(AudioInput));
+		Inputs.BindReadVertex(METASOUND_GET_PARAM_NAME(InParamDelayTime), FTimeReadRef(DelayTime));
+		Inputs.BindReadVertex(METASOUND_GET_PARAM_NAME(InParamDryLevel), FFloatReadRef(DryLevel));
+		Inputs.BindReadVertex(METASOUND_GET_PARAM_NAME(InParamWetLevel), FFloatReadRef(WetLevel));
+		Inputs.BindReadVertex(METASOUND_GET_PARAM_NAME(InParamFeedbackAmount), FFloatReadRef(Feedback));
+		Inputs.SetValue(METASOUND_GET_PARAM_NAME(InParamMaxDelayTime), FTime::FromSeconds(MaxDelayTimeSeconds));
+
+		FOutputVertexInterfaceData& Outputs = InVertexData.GetOutputs();
+		Outputs.BindReadVertex(METASOUND_GET_PARAM_NAME(OutParamAudio), AudioOutput);
+	}
 
+	FDataReferenceCollection FDelayOperator::GetInputs() const
+	{
+		// This should never be called. Bind(...) is called instead. This method
+		// exists as a stop-gap until the API can be deprecated and removed. 
+		checkNoEntry();
+
+		FDataReferenceCollection InputDataReferences;
 		return InputDataReferences;
 	}
 
 	FDataReferenceCollection FDelayOperator::GetOutputs() const
 	{
+		// This should never be called. Bind(...) is called instead. This method
+		// exists as a stop-gap until the API can be deprecated and removed. 
+		checkNoEntry();
+
 		using namespace Delay;
 
 		FDataReferenceCollection OutputDataReferences;
-		OutputDataReferences.AddDataReadReference(METASOUND_GET_PARAM_NAME(OutParamAudio), FAudioBufferReadRef(AudioOutput));
 		return OutputDataReferences;
 	}
 
 	float FDelayOperator::GetInputDelayTimeMsec() const
 	{
 		// Clamp the delay time to the max delay allowed
-		return 1000.0f * FMath::Clamp((float)DelayTime->GetSeconds(), 0.0f, Delay::MaxDelaySeconds);
+		return 1000.0f * FMath::Clamp((float)DelayTime->GetSeconds(), 0.0f, MaxDelayTimeSeconds);
 	}
 
 	void FDelayOperator::Execute()
@@ -175,13 +200,17 @@ namespace Metasound
 	{
 		using namespace Delay;
 
+		FDataVertexMetadata MaxDelayTimeMetadata = METASOUND_GET_PARAM_METADATA(InParamMaxDelayTime);
+		MaxDelayTimeMetadata.bIsAdvancedDisplay = true;
+
 		static const FVertexInterface Interface(
 			FInputVertexInterface(
 				TInputDataVertex<FAudioBuffer>(METASOUND_GET_PARAM_NAME_AND_METADATA(InParamAudioInput)),
 				TInputDataVertex<FTime>(METASOUND_GET_PARAM_NAME_AND_METADATA(InParamDelayTime), 1.0f),
 				TInputDataVertex<float>(METASOUND_GET_PARAM_NAME_AND_METADATA(InParamDryLevel), 0.0f),
 				TInputDataVertex<float>(METASOUND_GET_PARAM_NAME_AND_METADATA(InParamWetLevel), 1.0f),
-				TInputDataVertex<float>(METASOUND_GET_PARAM_NAME_AND_METADATA(InParamFeedbackAmount), 0.0f)
+				TInputDataVertex<float>(METASOUND_GET_PARAM_NAME_AND_METADATA(InParamFeedbackAmount), 0.0f),
+				TInputConstructorVertex<FTime>(METASOUND_GET_PARAM_NAME(InParamMaxDelayTime), MaxDelayTimeMetadata, DefaultMaxDelaySeconds)
 			),
 			FOutputVertexInterface(
 				TOutputDataVertex<FAudioBuffer>(METASOUND_GET_PARAM_NAME_AND_METADATA(OutParamAudio))
@@ -213,20 +242,20 @@ namespace Metasound
 		return Info;
 	}
 
-	TUniquePtr<IOperator> FDelayOperator::CreateOperator(const FCreateOperatorParams& InParams, FBuildErrorArray& OutErrors)
+	TUniquePtr<IOperator> FDelayOperator::CreateOperator(const FBuildOperatorParams& InParams, FBuildResults& OutResults)
 	{
 		using namespace Delay; 
 
-		const FDataReferenceCollection& InputCollection = InParams.InputDataReferences;
-		const FInputVertexInterface& InputInterface = GetVertexInterface().GetInputInterface();
+		const FInputVertexInterfaceData& InputData = InParams.InputData;
 
-		FAudioBufferReadRef AudioIn = InputCollection.GetDataReadReferenceOrConstruct<FAudioBuffer>(METASOUND_GET_PARAM_NAME(InParamAudioInput), InParams.OperatorSettings);
-		FTimeReadRef DelayTime = InputCollection.GetDataReadReferenceOrConstructWithVertexDefault<FTime>(InputInterface, METASOUND_GET_PARAM_NAME(InParamDelayTime), InParams.OperatorSettings);
-		FFloatReadRef DryLevel = InputCollection.GetDataReadReferenceOrConstructWithVertexDefault<float>(InputInterface, METASOUND_GET_PARAM_NAME(InParamDryLevel), InParams.OperatorSettings);
-		FFloatReadRef WetLevel = InputCollection.GetDataReadReferenceOrConstructWithVertexDefault<float>(InputInterface, METASOUND_GET_PARAM_NAME(InParamWetLevel), InParams.OperatorSettings);
-		FFloatReadRef Feedback = InputCollection.GetDataReadReferenceOrConstructWithVertexDefault<float>(InputInterface, METASOUND_GET_PARAM_NAME(InParamFeedbackAmount), InParams.OperatorSettings);
+		FAudioBufferReadRef AudioIn = InputData.GetOrConstructDataReadReference<FAudioBuffer>(METASOUND_GET_PARAM_NAME(InParamAudioInput), InParams.OperatorSettings);
+		FTimeReadRef DelayTime = InputData.GetOrCreateDefaultDataReadReference<FTime>(METASOUND_GET_PARAM_NAME(InParamDelayTime), InParams.OperatorSettings);
+		FFloatReadRef DryLevel = InputData.GetOrCreateDefaultDataReadReference<float>(METASOUND_GET_PARAM_NAME(InParamDryLevel), InParams.OperatorSettings);
+		FFloatReadRef WetLevel = InputData.GetOrCreateDefaultDataReadReference<float>(METASOUND_GET_PARAM_NAME(InParamWetLevel), InParams.OperatorSettings);
+		FFloatReadRef Feedback = InputData.GetOrCreateDefaultDataReadReference<float>(METASOUND_GET_PARAM_NAME(InParamFeedbackAmount), InParams.OperatorSettings);
+		FTime MaxDelayTime = InputData.GetOrCreateDefaultValue<FTime>(METASOUND_GET_PARAM_NAME(InParamMaxDelayTime), InParams.OperatorSettings);
 
-		return MakeUnique<FDelayOperator>(InParams.OperatorSettings, AudioIn, DelayTime, DryLevel, WetLevel, Feedback);
+		return MakeUnique<FDelayOperator>(InParams.OperatorSettings, AudioIn, DelayTime, DryLevel, WetLevel, Feedback, MaxDelayTime.GetSeconds());
 	}
 
 	class FDelayNode : public FNodeFacade
