More MetaSound pin category clean-up & object pin category for WT fix-up
#rnx

diff --git a/Engine/Plugins/Runtime/Metasound/Source/MetasoundEditor/Private/MetasoundEditorGraphBuilder.cpp b/Engine/Plugins/Runtime/Metasound/Source/MetasoundEditor/Private/MetasoundEditorGraphBuilder.cpp
--- a/Engine/Plugins/Runtime/Metasound/Source/MetasoundEditor/Private/MetasoundEditorGraphBuilder.cpp
+++ b/Engine/Plugins/Runtime/Metasound/Source/MetasoundEditor/Private/MetasoundEditorGraphBuilder.cpp
@@ -13,6 +13,7 @@
 #include "Kismet/KismetSystemLibrary.h"
 #include "Metasound.h"
 #include "MetasoundAssetBase.h"
+#include "MetasoundAudioBuffer.h"
 #include "MetasoundEditor.h"
 #include "MetasoundEditorGraph.h"
 #include "MetasoundEditorGraphBuilder.h"
@@ -31,13 +32,16 @@
 #include "MetasoundFrontendSearchEngine.h"
 #include "MetasoundFrontendTransform.h"
 #include "MetasoundLiteral.h"
+#include "MetasoundTime.h"
 #include "MetasoundUObjectRegistry.h"
 #include "MetasoundVariableNodes.h"
 #include "MetasoundVertex.h"
+#include "MetasoundWaveTable.h"
 #include "Modules/ModuleManager.h"
 #include "NodeTemplates/MetasoundFrontendNodeTemplateReroute.h"
 #include "Templates/Tuple.h"
 #include "Toolkits/ToolkitManager.h"
+#include "WaveTable.h"
 #include "Widgets/Notifications/SNotificationList.h"
 
 #define LOCTEXT_NAMESPACE "MetaSoundEditor"
@@ -47,18 +51,6 @@ namespace Metasound
 {
 	namespace Editor
 	{
-		const FName FGraphBuilder::PinCategoryAudio = "audio";
-		const FName FGraphBuilder::PinCategoryBoolean = "bool";
-		//const FName FGraphBuilder::PinCategoryDouble = "double";
-		const FName FGraphBuilder::PinCategoryFloat = "float";
-		const FName FGraphBuilder::PinCategoryInt32 = "int";
-		//const FName FGraphBuilder::PinCategoryInt64 = "int64";
-		const FName FGraphBuilder::PinCategoryObject = "object";
-		const FName FGraphBuilder::PinCategoryString = "string";
-		const FName FGraphBuilder::PinCategoryTrigger = "trigger";
-
-		const FName FGraphBuilder::PinSubCategoryTime = "time";
-
 		namespace GraphBuilderPrivate
 		{
 			void DeleteNode(UObject& InMetaSound, Frontend::FNodeHandle InNodeHandle)
@@ -88,6 +80,29 @@ namespace Metasound
 			}
 		} // namespace GraphBuilderPrivate
 
+		// Categories corresponding with POD DataTypes
+		const FName FGraphBuilder::PinCategoryObject = "object"; // Basket for all UObject proxy types (corresponds to multiple DataTypes)
+		const FName FGraphBuilder::PinCategoryBoolean = GetMetasoundDataTypeName<bool>();
+		const FName FGraphBuilder::PinCategoryFloat = GetMetasoundDataTypeName<float>();
+		const FName FGraphBuilder::PinCategoryInt32 = GetMetasoundDataTypeName<int32>();
+		const FName FGraphBuilder::PinCategoryString = GetMetasoundDataTypeName<FString>();
+
+		// Categories corresponding with MetaSound DataTypes with custom visualization
+		const FName FGraphBuilder::PinCategoryAudio = GetMetasoundDataTypeName<FAudioBuffer>();
+		const FName FGraphBuilder::PinCategoryTime = GetMetasoundDataTypeName<FTime>();
+		const FName FGraphBuilder::PinCategoryTimeArray = GetMetasoundDataTypeName<TArray<FTime>>();
+		const FName FGraphBuilder::PinCategoryTrigger = GetMetasoundDataTypeName<FTrigger>();
+		const FName FGraphBuilder::PinCategoryWaveTable = GetMetasoundDataTypeName<WaveTable::FWaveTable>();
+
+		bool FGraphBuilder::IsPinCategoryMetaSoundCustomDataType(FName InPinCategoryName)
+		{
+			return InPinCategoryName == PinCategoryAudio
+				|| InPinCategoryName == PinCategoryTime
+				|| InPinCategoryName == PinCategoryTimeArray
+				|| InPinCategoryName == PinCategoryTrigger
+				|| InPinCategoryName == PinCategoryWaveTable;
+		}
+
 		bool FGraphBuilder::CanInspectPin(const UEdGraphPin* InPin)
 		{
 			// Can't inspect the value on an invalid pin object.
@@ -618,11 +633,6 @@ namespace Metasound
 				return Settings->AudioPinTypeColor;
 			}
 
-			if (PinType.PinCategory == PinCategoryTrigger)
-			{
-				return Settings->TriggerPinTypeColor;
-			}
-
 			if (PinType.PinCategory == PinCategoryBoolean)
 			{
 				return Settings->BooleanPinTypeColor;
@@ -630,10 +640,6 @@ namespace Metasound
 
 			if (PinType.PinCategory == PinCategoryFloat)
 			{
-				if (PinType.PinSubCategory == PinSubCategoryTime)
-				{
-					return Settings->TimePinTypeColor;
-				}
 				return Settings->FloatPinTypeColor;
 			}
 
@@ -642,24 +648,29 @@ namespace Metasound
 				return Settings->IntPinTypeColor;
 			}
 
-			//if (PinType.PinCategory == PinCategoryInt64)
-			//{
-			//	return Settings->Int64PinTypeColor;
-			//}
+			if (PinType.PinCategory == PinCategoryObject)
+			{
+				return Settings->ObjectPinTypeColor;
+			}
 
 			if (PinType.PinCategory == PinCategoryString)
 			{
 				return Settings->StringPinTypeColor;
 			}
 
-			//if (PinType.PinCategory == PinCategoryDouble)
-			//{
-			//	return Settings->DoublePinTypeColor;
-			//}
+			if (PinType.PinCategory == PinCategoryTime || PinType.PinCategory == PinCategoryTimeArray)
+			{
+				return Settings->TimePinTypeColor;
+			}
 
-			if (PinType.PinCategory == PinCategoryObject)
+			if (PinType.PinCategory == PinCategoryTrigger)
 			{
-				return Settings->ObjectPinTypeColor;
+				return Settings->TriggerPinTypeColor;
+			}
+
+			if (PinType.PinCategory == PinCategoryWaveTable)
+			{
+				return Settings->WaveTablePinTypeColor;
 			}
 
 			return Settings->DefaultPinTypeColor;
diff --git a/Engine/Plugins/Runtime/Metasound/Source/MetasoundEditor/Private/MetasoundEditorGraphBuilder.h b/Engine/Plugins/Runtime/Metasound/Source/MetasoundEditor/Private/MetasoundEditorGraphBuilder.h
--- a/Engine/Plugins/Runtime/Metasound/Source/MetasoundEditor/Private/MetasoundEditorGraphBuilder.h
+++ b/Engine/Plugins/Runtime/Metasound/Source/MetasoundEditor/Private/MetasoundEditorGraphBuilder.h
@@ -47,16 +47,15 @@ namespace Metasound
 		public:
 			static const FName PinCategoryAudio;
 			static const FName PinCategoryBoolean;
-			//static const FName PinCategoryDouble;
 			static const FName PinCategoryFloat;
 			static const FName PinCategoryInt32;
-			//static const FName PinCategoryInt64;
 			static const FName PinCategoryObject;
 			static const FName PinCategoryString;
+			static const FName PinCategoryTime;
+			static const FName PinCategoryTimeArray;
 			static const FName PinCategoryTrigger;
+			static const FName PinCategoryWaveTable;
 
-			// Custom pin-related styles for non-literal types (ex. wire color, pin heads, etc.)
-			static const FName PinSubCategoryTime; // Time type
 
 			static const FText FunctionMenuName;
 			static const FText GraphMenuName;
@@ -82,6 +81,9 @@ namespace Metasound
 			// Wraps RegisterGraphWithFrontend logic in Frontend with any additional logic required to refresh editor & respective editor object state.
 			static void UnregisterGraphWithFrontend(UObject& InMetaSound);
 
+			// Returns whether pin category is a custom MetaSound DataType
+			static bool IsPinCategoryMetaSoundCustomDataType(FName InPinCategoryName);
+
 			// Determines if pin supports inspection/probe view.
 			static bool CanInspectPin(const UEdGraphPin* InPin);
 
diff --git a/Engine/Plugins/Runtime/Metasound/Source/MetasoundEditor/Private/MetasoundEditorModule.cpp b/Engine/Plugins/Runtime/Metasound/Source/MetasoundEditor/Private/MetasoundEditorModule.cpp
--- a/Engine/Plugins/Runtime/Metasound/Source/MetasoundEditor/Private/MetasoundEditorModule.cpp
+++ b/Engine/Plugins/Runtime/Metasound/Source/MetasoundEditor/Private/MetasoundEditorModule.cpp
@@ -469,19 +469,11 @@ namespace Metasound
 						FName PinCategory = DataTypeName;
 						FName PinSubCategory;
 
-						// Execution path triggers are specialized
-						if (DataTypeName == GetMetasoundDataTypeName<FTrigger>())
-						{
-							PinCategory = FGraphBuilder::PinCategoryTrigger;
-						}
-						// Audio type is ubiquitous, so connections are stylized (i.e. wire color & animation)
-						else if (DataTypeName == GetMetasoundDataTypeName<FAudioBuffer>())
-						{
-							PinCategory = FGraphBuilder::PinCategoryAudio;
-						}
-						// Primitives
-						else
+						// Types like triggers & AudioBuffer are specialized, so ignore their preferred
+						// literal types to classify the category.
+						if (!FGraphBuilder::IsPinCategoryMetaSoundCustomDataType(PinCategory))
 						{
+							// Primitives
 							switch (RegistryInfo.PreferredLiteralType)
 							{
 								case ELiteralType::Boolean:
@@ -494,12 +486,6 @@ namespace Metasound
 								case ELiteralType::Float:
 								{
 									PinCategory = FGraphBuilder::PinCategoryFloat;
-
-									if (DataTypeName == GetMetasoundDataTypeName<FTime>() ||
-										DataTypeName == CreateArrayTypeNameFromElementTypeName(GetMetasoundDataTypeName<FTime>()))
-									{
-										PinSubCategory = FGraphBuilder::PinSubCategoryTime;
-									}
 								}
 								break;
 
@@ -509,10 +495,6 @@ namespace Metasound
 									{
 										PinCategory = FGraphBuilder::PinCategoryFloat;
 									}
-									else
-									{
-										PinCategory = FGraphBuilder::PinCategoryObject;
-									}
 								}
 								break;
 
@@ -528,10 +510,6 @@ namespace Metasound
 									{
 										PinCategory = FGraphBuilder::PinCategoryInt32;
 									}
-									else
-									{
-										PinCategory = FGraphBuilder::PinCategoryObject;
-									}
 								}
 								break;
 
@@ -547,21 +525,21 @@ namespace Metasound
 									{
 										PinCategory = FGraphBuilder::PinCategoryString;
 									}
-									else
-									{
-										PinCategory = FGraphBuilder::PinCategoryObject;
-									}
 								}
 								break;
 
 								case ELiteralType::UObjectProxy:
 								case ELiteralType::UObjectProxyArray:
+								{
+									PinCategory = FGraphBuilder::PinCategoryObject;
+								}
+								break;
+
 								case ELiteralType::None:
 								case ELiteralType::NoneArray:
 								case ELiteralType::Invalid:
 								default:
 								{
-									PinCategory = FGraphBuilder::PinCategoryObject;
 									static_assert(static_cast<int32>(ELiteralType::Invalid) == 12, "Possible missing binding of pin category to primitive type");
 								}
 								break;
diff --git a/Engine/Plugins/Runtime/Metasound/Source/MetasoundEditor/Private/MetasoundEditorSettings.cpp b/Engine/Plugins/Runtime/Metasound/Source/MetasoundEditor/Private/MetasoundEditorSettings.cpp
--- a/Engine/Plugins/Runtime/Metasound/Source/MetasoundEditor/Private/MetasoundEditorSettings.cpp
+++ b/Engine/Plugins/Runtime/Metasound/Source/MetasoundEditor/Private/MetasoundEditorSettings.cpp
@@ -13,14 +13,13 @@ UMetasoundEditorSettings::UMetasoundEditorSettings(const FObjectInitializer& Obj
 
 	AudioPinTypeColor = FLinearColor(1.0f, 0.3f, 1.0f, 1.0f);					// magenta
 	BooleanPinTypeColor = FLinearColor(0.300000f, 0.0f, 0.0f, 1.0f);			// maroon
-	//DoublePinTypeColor = FLinearColor(0.039216f, 0.666667f, 0.0f, 1.0f);		// darker green
 	FloatPinTypeColor = FLinearColor(0.357667f, 1.0f, 0.060000f, 1.0f);			// bright green
 	IntPinTypeColor = FLinearColor(0.013575f, 0.770000f, 0.429609f, 1.0f);		// green-blue
-	//Int64PinTypeColor = FLinearColor(0.413575f, 0.770000f, 0.429609f, 1.0f);	// green
 	ObjectPinTypeColor = FLinearColor(0.0f, 0.4f, 0.910000f, 1.0f);				// sharp blue
 	StringPinTypeColor = FLinearColor(1.0f, 0.0f, 0.660537f, 1.0f);				// bright pink
 	TimePinTypeColor = FLinearColor(0.3f, 1.0f, 1.0f, 1.0f);					// cyan
 	TriggerPinTypeColor = FLinearColor(1.0f, 1.0f, 1.0f, 1.0f);					// white
+	WaveTablePinTypeColor = FLinearColor(0.580392f, 0.0f, 0.827450f, 1.0f);		// purple
 
 	NativeNodeTitleColor = FLinearColor(0.4f, 0.85f, 0.35f, 1.0f);				// pale green
 	AssetReferenceNodeTitleColor = FLinearColor(0.047f, 0.686f, 0.988f);		// sky blue
diff --git a/Engine/Plugins/Runtime/Metasound/Source/MetasoundEditor/Private/MetasoundEditorSettings.h b/Engine/Plugins/Runtime/Metasound/Source/MetasoundEditor/Private/MetasoundEditorSettings.h
--- a/Engine/Plugins/Runtime/Metasound/Source/MetasoundEditor/Private/MetasoundEditorSettings.h
+++ b/Engine/Plugins/Runtime/Metasound/Source/MetasoundEditor/Private/MetasoundEditorSettings.h
@@ -55,10 +55,6 @@ public:
 	UPROPERTY(EditAnywhere, config, Category=PinColors)
 	FLinearColor BooleanPinTypeColor;
 
-	/** Double pin type color */
-	//UPROPERTY(EditAnywhere, config, Category = PinColors)
-	//FLinearColor DoublePinTypeColor;
-
 	/** Floating-point pin type color */
 	UPROPERTY(EditAnywhere, config, Category=PinColors)
 	FLinearColor FloatPinTypeColor;
@@ -67,10 +63,6 @@ public:
 	UPROPERTY(EditAnywhere, config, Category=PinColors)
 	FLinearColor IntPinTypeColor;
 
-	/** Integer64 pin type color */
-	//UPROPERTY(EditAnywhere, config, Category = PinColors)
-	//FLinearColor Int64PinTypeColor;
-
 	/** Object pin type color */
 	UPROPERTY(EditAnywhere, config, Category=PinColors)
 	FLinearColor ObjectPinTypeColor;
@@ -87,6 +79,10 @@ public:
 	UPROPERTY(EditAnywhere, config, Category=PinColors)
 	FLinearColor TriggerPinTypeColor;
 
+	/** WaveTable pin type color */
+	UPROPERTY(EditAnywhere, config, Category = PinColors)
+	FLinearColor WaveTablePinTypeColor;
+
 	/** Native node class title color */
 	UPROPERTY(EditAnywhere, config, Category = NodeTitleColors)
 	FLinearColor NativeNodeTitleColor;
