Adding tags to MetaSound nodes

diff --git a/Engine/Plugins/Runtime/Metasound/Source/MetasoundEngine/Private/MetasoundWavePlayerNode.cpp b/Engine/Plugins/Runtime/Metasound/Source/MetasoundEngine/Private/MetasoundWavePlayerNode.cpp
--- a/Engine/Plugins/Runtime/Metasound/Source/MetasoundEngine/Private/MetasoundWavePlayerNode.cpp
+++ b/Engine/Plugins/Runtime/Metasound/Source/MetasoundEngine/Private/MetasoundWavePlayerNode.cpp
@@ -1078,6 +1078,9 @@ namespace Metasound
 				Info.Author = PluginAuthor;
 				Info.PromptIfMissing = PluginNodeMissingPrompt;
 				Info.DefaultInterface = DeclareVertexInterface();
+				Info.Keywords = { METASOUND_LOCTEXT("WavePlayerSoundKeyword", "Sound"),
+				                  METASOUND_LOCTEXT("WavePlayerCueKeyword", "Cue")
+				};
 
 				return Info;
 			};
diff --git a/Engine/Plugins/Runtime/Metasound/Source/MetasoundEngine/Private/MetasoundWaveTableEnvelopeNode.cpp b/Engine/Plugins/Runtime/Metasound/Source/MetasoundEngine/Private/MetasoundWaveTableEnvelopeNode.cpp
--- a/Engine/Plugins/Runtime/Metasound/Source/MetasoundEngine/Private/MetasoundWaveTableEnvelopeNode.cpp
+++ b/Engine/Plugins/Runtime/Metasound/Source/MetasoundEngine/Private/MetasoundWaveTableEnvelopeNode.cpp
@@ -9,6 +9,7 @@
 #include "MetasoundPrimitives.h"
 #include "MetasoundSourceInterface.h"
 #include "MetasoundStandardNodesNames.h"
+#include "MetasoundStandardNodesCategories.h"
 #include "MetasoundTime.h"
 #include "MetasoundTrace.h"
 #include "MetasoundTrigger.h"
@@ -89,8 +90,8 @@ namespace Metasound
 					PluginAuthor,
 					PluginNodeMissingPrompt,
 					GetDefaultInterface(),
-					{ },
-					{ },
+					{ NodeCategories::Envelopes },
+					{ METASOUND_LOCTEXT("WaveTableEnvelopeSynthesisKeyword", "Synthesis")},
 					{ }
 				};
 
diff --git a/Engine/Plugins/Runtime/Metasound/Source/MetasoundEngine/Private/MetasoundWaveTableOscillatorNode.cpp b/Engine/Plugins/Runtime/Metasound/Source/MetasoundEngine/Private/MetasoundWaveTableOscillatorNode.cpp
--- a/Engine/Plugins/Runtime/Metasound/Source/MetasoundEngine/Private/MetasoundWaveTableOscillatorNode.cpp
+++ b/Engine/Plugins/Runtime/Metasound/Source/MetasoundEngine/Private/MetasoundWaveTableOscillatorNode.cpp
@@ -62,8 +62,8 @@ namespace Metasound
 					PluginAuthor,
 					PluginNodeMissingPrompt,
 					GetDefaultInterface(),
-					{ NodeCategories::Envelopes },
-					{ },
+					{ NodeCategories::Generators },
+					{ METASOUND_LOCTEXT("WaveTableOscillatorSynthesisKeyword", "Synthesis")},
 					{ }
 				};
 
diff --git a/Engine/Plugins/Runtime/Metasound/Source/MetasoundStandardNodes/Private/MetasoundAdditiveSynthNode.cpp b/Engine/Plugins/Runtime/Metasound/Source/MetasoundStandardNodes/Private/MetasoundAdditiveSynthNode.cpp
--- a/Engine/Plugins/Runtime/Metasound/Source/MetasoundStandardNodes/Private/MetasoundAdditiveSynthNode.cpp
+++ b/Engine/Plugins/Runtime/Metasound/Source/MetasoundStandardNodes/Private/MetasoundAdditiveSynthNode.cpp
@@ -50,6 +50,8 @@ namespace Metasound
 				Info.Author = PluginAuthor;
 				Info.PromptIfMissing = PluginNodeMissingPrompt;
 				Info.DefaultInterface = GetVertexInterface();
+				Info.CategoryHierarchy.Emplace(NodeCategories::Generators);
+				Info.Keywords = { METASOUND_LOCTEXT("AdditiveSynthesisKeyword", "Synthesis") };
 
 				return Info;
 			};
diff --git a/Engine/Plugins/Runtime/Metasound/Source/MetasoundStandardNodes/Private/MetasoundBPMToSecondsNode.cpp b/Engine/Plugins/Runtime/Metasound/Source/MetasoundStandardNodes/Private/MetasoundBPMToSecondsNode.cpp
--- a/Engine/Plugins/Runtime/Metasound/Source/MetasoundStandardNodes/Private/MetasoundBPMToSecondsNode.cpp
+++ b/Engine/Plugins/Runtime/Metasound/Source/MetasoundStandardNodes/Private/MetasoundBPMToSecondsNode.cpp
@@ -141,6 +141,7 @@ namespace Metasound
 			Info.Author = PluginAuthor;
 			Info.PromptIfMissing = PluginNodeMissingPrompt;
 			Info.DefaultInterface = GetVertexInterface();
+			Info.Keywords = { METASOUND_LOCTEXT("BPMNodeRhythmKeyword", "Rhythm"), METASOUND_LOCTEXT("BPMNodeTempoKeyword", "Tempo") };
 
 			return Info;
 		};
diff --git a/Engine/Plugins/Runtime/Metasound/Source/MetasoundStandardNodes/Private/MetasoundBasicFilters.cpp b/Engine/Plugins/Runtime/Metasound/Source/MetasoundStandardNodes/Private/MetasoundBasicFilters.cpp
--- a/Engine/Plugins/Runtime/Metasound/Source/MetasoundStandardNodes/Private/MetasoundBasicFilters.cpp
+++ b/Engine/Plugins/Runtime/Metasound/Source/MetasoundStandardNodes/Private/MetasoundBasicFilters.cpp
@@ -34,7 +34,7 @@ namespace Metasound
 	namespace BasicFilterParameterNames
 	{
 		// inputs
-		METASOUND_PARAM(ParamAudioInput, "In", "Audio to be processed by the filter.");	
+		METASOUND_PARAM(ParamAudioInput, "In", "Audio to be processed by the filter.");
 		METASOUND_PARAM(ParamCutoffFrequency, "Cutoff Frequency", "Controls cutoff frequency.");
 		METASOUND_PARAM(ParamResonance, "Resonance", "Controls filter resonance.");
 		METASOUND_PARAM(ParamBandwidth, "Bandwidth", "Controls bandwidth when applicable to the current filter type.");
@@ -54,6 +54,15 @@ namespace Metasound
 	using namespace BasicFilterParameterNames;
 #pragma endregion
 
+	static const TArray<FText> LowPassKeywords = { METASOUND_LOCTEXT("LowPassKeyword", "Lowpass"), METASOUND_LOCTEXT("LowPassSpaceKeyword", "Low Pass"), METASOUND_LOCTEXT("LowPassDashKeyword", "Low-pass"), METASOUND_LOCTEXT("LPFKeyword", "lpf")};
+	static const TArray<FText> HighPassKeywords = { METASOUND_LOCTEXT("HighPassKeyword", "Highpass"), METASOUND_LOCTEXT("HighPassSpaceKeyword", "High Pass"), METASOUND_LOCTEXT("HighPassDashKeyword", "High-pass"), METASOUND_LOCTEXT("HPFKeyword", "hpf") };
+	static const TArray<FText> BandPassKeywords = { METASOUND_LOCTEXT("BandPassKeyword", "Bandpass"), METASOUND_LOCTEXT("BandPassSpaceKeyword", "Band Pass"), METASOUND_LOCTEXT("BandPassDashKeyword", "Band-pass"), METASOUND_LOCTEXT("BPFKeyword", "bpf") };
+	static const TArray<FText> NotchKeywords = { METASOUND_LOCTEXT("NotchKeyword", "Notch"), METASOUND_LOCTEXT("NFKeyword", "nf"), METASOUND_LOCTEXT("BandStopKeyword", "BandStop"), METASOUND_LOCTEXT("BandStopSpaceKeyword", "Band Stop"), METASOUND_LOCTEXT("BandStopDashKeyword", "Band-stop"), METASOUND_LOCTEXT("BSFKeyword", "bsf") };
+	static const TArray<FText> ParametricEQKeywords = { METASOUND_LOCTEXT("EQKeyword", "EQ"), METASOUND_LOCTEXT("ParametricKeyword", "Parametric") };
+	static const TArray<FText> LowShelfKeywords = { METASOUND_LOCTEXT("LowShelfSpaceKeyword", "Low Shelf"), METASOUND_LOCTEXT("LowShelfDashKeyword", "Low-shelf"), METASOUND_LOCTEXT("LSFKeyword", "lsf") };
+	static const TArray<FText> HighShelfKeywords = { METASOUND_LOCTEXT("HighShelfSpaceKeyword", "High Shelf"), METASOUND_LOCTEXT("HighShelfDashKeyword", "High-shelf"), METASOUND_LOCTEXT("HSFKeyword", "hsf") };
+	static const TArray<FText> AllPassKeywords = { METASOUND_LOCTEXT("AllPassKeyword", "Allpass"), METASOUND_LOCTEXT("AllPassSpaceKeyword", "All Pass"), METASOUND_LOCTEXT("AllPassDashKeyword", "All-pass"), METASOUND_LOCTEXT("APFKeyword", "apf") };
+	static const TArray<FText> ButterworthKeywords = { METASOUND_LOCTEXT("ButterworthKeyword", "Butterworth") };
 
 #pragma region Biquad Filter Enum
 	DECLARE_METASOUND_ENUM(Audio::EBiquadFilter::Type, Audio::EBiquadFilter::Lowpass,
@@ -169,6 +178,7 @@ namespace Metasound
 			Info.PromptIfMissing = PluginNodeMissingPrompt;
 			Info.DefaultInterface = DeclareVertexInterface();
 			Info.CategoryHierarchy.Emplace(NodeCategories::Filters);
+			Info.Keywords = LowPassKeywords;
 
 			return Info;
 		};
@@ -351,6 +361,10 @@ namespace Metasound
 			Info.PromptIfMissing = PluginNodeMissingPrompt;
 			Info.DefaultInterface = DeclareVertexInterface();
 			Info.CategoryHierarchy.Emplace(NodeCategories::Filters);
+			Info.Keywords.Append(LowPassKeywords);
+			Info.Keywords.Append(HighPassKeywords); 
+			Info.Keywords.Append(BandPassKeywords);
+			Info.Keywords.Append(NotchKeywords);
 
 			return Info;
 		};
@@ -523,6 +537,7 @@ namespace Metasound
 			Info.PromptIfMissing = PluginNodeMissingPrompt;
 			Info.DefaultInterface = DeclareVertexInterface();
 			Info.CategoryHierarchy.Emplace(NodeCategories::Filters);
+			Info.Keywords = LowPassKeywords;
 
 			return Info;
 		};
@@ -660,6 +675,7 @@ namespace Metasound
 			Info.PromptIfMissing = PluginNodeMissingPrompt;
 			Info.DefaultInterface = DeclareVertexInterface();
 			Info.CategoryHierarchy.Emplace(NodeCategories::Filters);
+			Info.Keywords = HighPassKeywords;
 
 			return Info;
 		};
@@ -718,6 +734,7 @@ namespace Metasound
 	{
 	private:
 		static constexpr float InvalidValue = -1.f;
+		
 	public:
 		static const FNodeClassMetadata& GetNodeInfo();
 
@@ -816,6 +833,15 @@ namespace Metasound
 			Info.PromptIfMissing = PluginNodeMissingPrompt;
 			Info.DefaultInterface = DeclareVertexInterface();
 			Info.CategoryHierarchy.Emplace(NodeCategories::Filters);
+			Info.Keywords.Append(LowPassKeywords);
+			Info.Keywords.Append(HighPassKeywords);
+			Info.Keywords.Append(BandPassKeywords);
+			Info.Keywords.Append(NotchKeywords);
+			Info.Keywords.Append(ParametricEQKeywords);
+			Info.Keywords.Append(AllPassKeywords);
+			Info.Keywords.Append(LowShelfKeywords);
+			Info.Keywords.Append(HighShelfKeywords);
+			Info.Keywords.Append(ButterworthKeywords);
 
 			return Info;
 		};
diff --git a/Engine/Plugins/Runtime/Metasound/Source/MetasoundStandardNodes/Private/MetasoundCompressorNode.cpp b/Engine/Plugins/Runtime/Metasound/Source/MetasoundStandardNodes/Private/MetasoundCompressorNode.cpp
--- a/Engine/Plugins/Runtime/Metasound/Source/MetasoundStandardNodes/Private/MetasoundCompressorNode.cpp
+++ b/Engine/Plugins/Runtime/Metasound/Source/MetasoundStandardNodes/Private/MetasoundCompressorNode.cpp
@@ -132,7 +132,7 @@ namespace Metasound
 					PluginNodeMissingPrompt,
 					NodeInterface,
 					{ NodeCategories::Dynamics },
-					{ },
+					{ METASOUND_LOCTEXT("SidechainKeyword", "Sidechain"), METASOUND_LOCTEXT("UpwardsKeyword", "Upwards Compressor")},
 					FNodeDisplayStyle()
 				};
 
diff --git a/Engine/Plugins/Runtime/Metasound/Source/MetasoundStandardNodes/Private/MetasoundFreqToMIDINode.cpp b/Engine/Plugins/Runtime/Metasound/Source/MetasoundStandardNodes/Private/MetasoundFreqToMIDINode.cpp
--- a/Engine/Plugins/Runtime/Metasound/Source/MetasoundStandardNodes/Private/MetasoundFreqToMIDINode.cpp
+++ b/Engine/Plugins/Runtime/Metasound/Source/MetasoundStandardNodes/Private/MetasoundFreqToMIDINode.cpp
@@ -121,6 +121,7 @@ namespace Metasound
 			Info.PromptIfMissing = PluginNodeMissingPrompt;
 			Info.DefaultInterface = GetVertexInterface();
 			Info.CategoryHierarchy.Emplace(NodeCategories::Music);
+			Info.Keywords.Add(METASOUND_LOCTEXT("FreqToMIDIPitchKeyword", "Pitch"));
 
 			return Info;
 		};
diff --git a/Engine/Plugins/Runtime/Metasound/Source/MetasoundStandardNodes/Private/MetasoundITDPannerNode.cpp b/Engine/Plugins/Runtime/Metasound/Source/MetasoundStandardNodes/Private/MetasoundITDPannerNode.cpp
--- a/Engine/Plugins/Runtime/Metasound/Source/MetasoundStandardNodes/Private/MetasoundITDPannerNode.cpp
+++ b/Engine/Plugins/Runtime/Metasound/Source/MetasoundStandardNodes/Private/MetasoundITDPannerNode.cpp
@@ -259,6 +259,7 @@ namespace Metasound
 			Info.PromptIfMissing = PluginNodeMissingPrompt;
 			Info.DefaultInterface = GetVertexInterface();
 			Info.CategoryHierarchy.Emplace(NodeCategories::Spatialization);
+			Info.Keywords = { METASOUND_LOCTEXT("ITDBinauralKeyword", "Binaural"), METASOUND_LOCTEXT("ITDInterauralKeyword", "Interaural Time Delay")};
 			return Info;
 		};
 
diff --git a/Engine/Plugins/Runtime/Metasound/Source/MetasoundStandardNodes/Private/MetasoundInterpToNode.cpp b/Engine/Plugins/Runtime/Metasound/Source/MetasoundStandardNodes/Private/MetasoundInterpToNode.cpp
--- a/Engine/Plugins/Runtime/Metasound/Source/MetasoundStandardNodes/Private/MetasoundInterpToNode.cpp
+++ b/Engine/Plugins/Runtime/Metasound/Source/MetasoundStandardNodes/Private/MetasoundInterpToNode.cpp
@@ -157,6 +157,7 @@ namespace Metasound
 			Info.Author = PluginAuthor;
 			Info.PromptIfMissing = PluginNodeMissingPrompt;
 			Info.DefaultInterface = GetVertexInterface();
+			Info.Keywords.Add(METASOUND_LOCTEXT("LerpKeyword", "Lerp"));
 
 			return Info;
 		};
diff --git a/Engine/Plugins/Runtime/Metasound/Source/MetasoundStandardNodes/Private/MetasoundLinearToLogFrequencyNode.cpp b/Engine/Plugins/Runtime/Metasound/Source/MetasoundStandardNodes/Private/MetasoundLinearToLogFrequencyNode.cpp
--- a/Engine/Plugins/Runtime/Metasound/Source/MetasoundStandardNodes/Private/MetasoundLinearToLogFrequencyNode.cpp
+++ b/Engine/Plugins/Runtime/Metasound/Source/MetasoundStandardNodes/Private/MetasoundLinearToLogFrequencyNode.cpp
@@ -46,6 +46,7 @@ namespace Metasound
 				Info.PromptIfMissing = PluginNodeMissingPrompt;
 				Info.DefaultInterface = GetVertexInterface();
 				Info.CategoryHierarchy.Emplace(NodeCategories::Math);
+				Info.Keywords.Add(METASOUND_LOCTEXT("LinearToLogPitchKeyword", "Pitch"));
 
 				return Info;
 			};
diff --git a/Engine/Plugins/Runtime/Metasound/Source/MetasoundStandardNodes/Private/MetasoundMIDIToFreqNode.cpp b/Engine/Plugins/Runtime/Metasound/Source/MetasoundStandardNodes/Private/MetasoundMIDIToFreqNode.cpp
--- a/Engine/Plugins/Runtime/Metasound/Source/MetasoundStandardNodes/Private/MetasoundMIDIToFreqNode.cpp
+++ b/Engine/Plugins/Runtime/Metasound/Source/MetasoundStandardNodes/Private/MetasoundMIDIToFreqNode.cpp
@@ -163,6 +163,7 @@ namespace Metasound
 			Info.PromptIfMissing = PluginNodeMissingPrompt;
 			Info.DefaultInterface = GetVertexInterface();
 			Info.CategoryHierarchy.Emplace(NodeCategories::Music);
+			Info.Keywords.Add(METASOUND_LOCTEXT("MIDIToFreqPitchKeyword", "Pitch"));
 
 			return Info;
 		};
diff --git a/Engine/Plugins/Runtime/Metasound/Source/MetasoundStandardNodes/Private/MetasoundMidiNoteQuantizerNode.cpp b/Engine/Plugins/Runtime/Metasound/Source/MetasoundStandardNodes/Private/MetasoundMidiNoteQuantizerNode.cpp
--- a/Engine/Plugins/Runtime/Metasound/Source/MetasoundStandardNodes/Private/MetasoundMidiNoteQuantizerNode.cpp
+++ b/Engine/Plugins/Runtime/Metasound/Source/MetasoundStandardNodes/Private/MetasoundMidiNoteQuantizerNode.cpp
@@ -119,6 +119,7 @@ namespace Metasound
 			Info.PromptIfMissing = PluginNodeMissingPrompt;
 			Info.DefaultInterface = DeclareVertexInterface();
 			Info.CategoryHierarchy.Emplace(NodeCategories::Music);
+			Info.Keywords.Add(METASOUND_LOCTEXT("QuantizerPitchKeyword", "Pitch"));
 
 			return Info;
 		};
diff --git a/Engine/Plugins/Runtime/Metasound/Source/MetasoundStandardNodes/Private/MetasoundMidiScaleToArrayNode.cpp b/Engine/Plugins/Runtime/Metasound/Source/MetasoundStandardNodes/Private/MetasoundMidiScaleToArrayNode.cpp
--- a/Engine/Plugins/Runtime/Metasound/Source/MetasoundStandardNodes/Private/MetasoundMidiScaleToArrayNode.cpp
+++ b/Engine/Plugins/Runtime/Metasound/Source/MetasoundStandardNodes/Private/MetasoundMidiScaleToArrayNode.cpp
@@ -113,6 +113,8 @@ namespace Metasound
 			Info.PromptIfMissing = PluginNodeMissingPrompt;
 			Info.DefaultInterface = DeclareVertexInterface();
 			Info.CategoryHierarchy.Emplace(NodeCategories::Music);
+			Info.Keywords.Add(METASOUND_LOCTEXT("ScaleToArrayPitchKeyword", "Pitch"));
+			Info.Keywords.Add(METASOUND_LOCTEXT("ScaleToArrayMIDIKeyword", "MIDI"));
 
 			return Info;
 		};
diff --git a/Engine/Plugins/Runtime/Metasound/Source/MetasoundStandardNodes/Private/MetasoundOscillatorNodes.cpp b/Engine/Plugins/Runtime/Metasound/Source/MetasoundStandardNodes/Private/MetasoundOscillatorNodes.cpp
--- a/Engine/Plugins/Runtime/Metasound/Source/MetasoundStandardNodes/Private/MetasoundOscillatorNodes.cpp
+++ b/Engine/Plugins/Runtime/Metasound/Source/MetasoundStandardNodes/Private/MetasoundOscillatorNodes.cpp
@@ -795,6 +795,7 @@ namespace Metasound
 				Info.PromptIfMissing = PluginNodeMissingPrompt;
 				Info.DefaultInterface = GetVertexInterface();
 				Info.CategoryHierarchy.Emplace(NodeCategories::Generators);
+				Info.Keywords = { METASOUND_LOCTEXT("OscSineKeyword", "Osc"), METASOUND_LOCTEXT("FMSineKeyword", "FM"), METASOUND_LOCTEXT("SineSynthesisKeyword", "Synthesis") };
 				return Info;
 			};
 			static const FNodeClassMetadata Info = InitNodeInfo();
@@ -925,6 +926,7 @@ namespace Metasound
 				Info.PromptIfMissing = PluginNodeMissingPrompt;
 				Info.DefaultInterface = GetVertexInterface();
 				Info.CategoryHierarchy.Emplace(NodeCategories::Generators);
+				Info.Keywords = { METASOUND_LOCTEXT("OscSawKeyword", "Osc"), METASOUND_LOCTEXT("FMSawKeyword", "FM"), METASOUND_LOCTEXT("SawSynthesisKeyword", "Synthesis") };
 				return Info;
 			};
 			static const FNodeClassMetadata Info = InitNodeInfo();
@@ -1165,6 +1167,7 @@ namespace Metasound
 				Info.PromptIfMissing = PluginNodeMissingPrompt;
 				Info.DefaultInterface = GetVertexInterface();
 				Info.CategoryHierarchy.Emplace(NodeCategories::Generators);
+				Info.Keywords = { METASOUND_LOCTEXT("OscSquareKeyword", "Osc"), METASOUND_LOCTEXT("FMSquareKeyword", "FM"), METASOUND_LOCTEXT("SquareSynthesisKeyword", "Synthesis") };
 				return Info;
 			};
 			static const FNodeClassMetadata Info = InitNodeInfo();
@@ -1286,6 +1289,7 @@ namespace Metasound
 				Info.PromptIfMissing = PluginNodeMissingPrompt;
 				Info.DefaultInterface = GetVertexInterface();
 				Info.CategoryHierarchy.Emplace(NodeCategories::Generators);
+				Info.Keywords = {METASOUND_LOCTEXT("OscTriKeyword", "Osc"), METASOUND_LOCTEXT("FMTriKeyword", "FM"), METASOUND_LOCTEXT("TriSynthesisKeyword", "Synthesis") };
 				return Info;
 			};
 			static const FNodeClassMetadata Info = InitNodeInfo();
@@ -1356,6 +1360,8 @@ namespace Metasound
 		METASOUND_PARAM(LfoPulseWidthPin, "Pulse Width", "Pulse Width (0..1)")
 	}
 
+	static const TArray<FText> LFOKeywords = { METASOUND_LOCTEXT("SineKeyword", "Sine"), METASOUND_LOCTEXT("SquareKeyword", "Square"), METASOUND_LOCTEXT("TriangleKeyword", "Triangle"), METASOUND_LOCTEXT("SawKeyword", "Saw"), METASOUND_LOCTEXT("OscLFOKeyword", "Oscillator"), METASOUND_LOCTEXT("ModulateKeyword", "Modulate")};
+
 	// Blockrate All-Purpose Oscillator
 	class FLfoOperator : public TExecutableOperator<FLfoOperator>
 	{
@@ -1396,6 +1402,7 @@ namespace Metasound
 				Info.PromptIfMissing = PluginNodeMissingPrompt;
 				Info.DefaultInterface = GetVertexInterface();
 				Info.CategoryHierarchy.Emplace(NodeCategories::Generators);
+				Info.Keywords = LFOKeywords;
 				return Info;
 			};
 			static const FNodeClassMetadata Info = InitNodeInfo();
diff --git a/Engine/Plugins/Runtime/Metasound/Source/MetasoundStandardNodes/Private/MetasoundTriggerCoinNode.cpp b/Engine/Plugins/Runtime/Metasound/Source/MetasoundStandardNodes/Private/MetasoundTriggerCoinNode.cpp
--- a/Engine/Plugins/Runtime/Metasound/Source/MetasoundStandardNodes/Private/MetasoundTriggerCoinNode.cpp
+++ b/Engine/Plugins/Runtime/Metasound/Source/MetasoundStandardNodes/Private/MetasoundTriggerCoinNode.cpp
@@ -205,6 +205,7 @@ namespace Metasound
 			Info.PromptIfMissing = PluginNodeMissingPrompt;
 			Info.DefaultInterface = GetVertexInterface();
 			Info.CategoryHierarchy.Emplace(NodeCategories::Trigger);
+			Info.Keywords = { METASOUND_LOCTEXT("TriggerFilterRandomKeyword", "Random"), METASOUND_LOCTEXT("TriggerFilterProbabilityKeyword", "Probability"), METASOUND_LOCTEXT("TriggerFilterCoinKeyword", "Coin"), METASOUND_LOCTEXT("TriggerFilterBernoulliKeyword", "Bernoulli Gate")};
 
 			return Info;
 		};
diff --git a/Engine/Plugins/Runtime/Metasound/Source/MetasoundStandardNodes/Private/MetasoundTriggerCompareNode.cpp b/Engine/Plugins/Runtime/Metasound/Source/MetasoundStandardNodes/Private/MetasoundTriggerCompareNode.cpp
--- a/Engine/Plugins/Runtime/Metasound/Source/MetasoundStandardNodes/Private/MetasoundTriggerCompareNode.cpp
+++ b/Engine/Plugins/Runtime/Metasound/Source/MetasoundStandardNodes/Private/MetasoundTriggerCompareNode.cpp
@@ -21,7 +21,7 @@ namespace Metasound
 				PluginNodeMissingPrompt,
 				InDefaultInterface,
 				{ NodeCategories::Trigger },
-				{ },
+				{ METASOUND_LOCTEXT("EqualsKeyword", "="), METASOUND_LOCTEXT("NotEqualsKeyword", "!="), METASOUND_LOCTEXT("LessThanKeyword", "<"), METASOUND_LOCTEXT("GreaterThanKeyword", ">"), METASOUND_LOCTEXT("LessEqualsKeyword", "<="), METASOUND_LOCTEXT("GreaterEqualsKeyword", ">=")},
 				FNodeDisplayStyle()
 			};
 
diff --git a/Engine/Plugins/Runtime/Metasound/Source/MetasoundStandardNodes/Private/MetasoundTriggerControlNode.cpp b/Engine/Plugins/Runtime/Metasound/Source/MetasoundStandardNodes/Private/MetasoundTriggerControlNode.cpp
--- a/Engine/Plugins/Runtime/Metasound/Source/MetasoundStandardNodes/Private/MetasoundTriggerControlNode.cpp
+++ b/Engine/Plugins/Runtime/Metasound/Source/MetasoundStandardNodes/Private/MetasoundTriggerControlNode.cpp
@@ -210,6 +210,7 @@ namespace Metasound
 			Info.PromptIfMissing = PluginNodeMissingPrompt;
 			Info.DefaultInterface = GetVertexInterface();
 			Info.CategoryHierarchy.Emplace(NodeCategories::Trigger);
+			Info.Keywords = { METASOUND_LOCTEXT("TriggerControlGateKeyword", "Gate") };
 
 			return Info;
 		};
diff --git a/Engine/Plugins/Runtime/Metasound/Source/MetasoundStandardNodes/Private/MetasoundTriggerRepeatNode.cpp b/Engine/Plugins/Runtime/Metasound/Source/MetasoundStandardNodes/Private/MetasoundTriggerRepeatNode.cpp
--- a/Engine/Plugins/Runtime/Metasound/Source/MetasoundStandardNodes/Private/MetasoundTriggerRepeatNode.cpp
+++ b/Engine/Plugins/Runtime/Metasound/Source/MetasoundStandardNodes/Private/MetasoundTriggerRepeatNode.cpp
@@ -178,6 +178,7 @@ namespace Metasound
 			Info.PromptIfMissing = PluginNodeMissingPrompt;
 			Info.DefaultInterface = GetVertexInterface();
 			Info.CategoryHierarchy.Emplace(NodeCategories::Trigger);
+			Info.Keywords.Add(METASOUND_LOCTEXT("TriggerRepeatLoopKeyword", "Loop"));
 
 			return Info;
 		};
