Make AnalyzerAddress & GraphAnalyzerAddress real TransmissionAddress citizens
#rnx

diff --git a/Engine/Plugins/Runtime/Metasound/Source/MetasoundFrontend/Private/Analysis/MetasoundFrontendAnalyzerAddress.cpp b/Engine/Plugins/Runtime/Metasound/Source/MetasoundFrontend/Private/Analysis/MetasoundFrontendAnalyzerAddress.cpp
--- a/Engine/Plugins/Runtime/Metasound/Source/MetasoundFrontend/Private/Analysis/MetasoundFrontendAnalyzerAddress.cpp
+++ b/Engine/Plugins/Runtime/Metasound/Source/MetasoundFrontend/Private/Analysis/MetasoundFrontendAnalyzerAddress.cpp
@@ -1,5 +1,8 @@
 // Copyright Epic Games, Inc. All Rights Reserved.
+
 #include "Analysis/MetasoundFrontendAnalyzerAddress.h"
+#include "MetasoundDataTypeRegistrationMacro.h"
+#include "Misc/AssertionMacros.h"
 
 
 namespace Metasound
@@ -8,6 +11,65 @@ namespace Metasound
 	{
 		const FString FAnalyzerAddress::PathSeparator = TEXT("/");
 
+		FAnalyzerAddress::FAnalyzerAddress(const FString& InAddressString)
+		{
+			TArray<FString> Tokens;
+			if (ensureAlwaysMsgf(InAddressString.ParseIntoArray(Tokens, *PathSeparator) == 7, TEXT("Invalid Analyzer Address String Format")))
+			{
+				InstanceID = static_cast<uint64>(FCString::Atoi64(*Tokens[0]));
+				NodeID = FGuid(Tokens[1]);
+				OutputName = *Tokens[2];
+				DataType = *Tokens[3];
+				AnalyzerName = *Tokens[4];
+				AnalyzerInstanceID = FGuid(Tokens[5]);
+				AnalyzerMemberName = *Tokens[6];
+			}
+		}
+
+		FName FAnalyzerAddress::GetAddressType() const
+		{
+			return "Analyzer";
+		}
+
+		FName FAnalyzerAddress::GetDataType() const
+		{
+			return DataType;
+		}
+
+		TUniquePtr<FTransmissionAddress> FAnalyzerAddress::Clone() const
+		{
+			return TUniquePtr<FTransmissionAddress>(new FAnalyzerAddress(*this));
+		}
+
+		uint32 FAnalyzerAddress::GetHash() const
+		{
+			uint32 AddressHash = HashCombineFast(AnalyzerInstanceID.A, ::GetTypeHash(AnalyzerMemberName));
+			AddressHash = HashCombineFast(AddressHash, ::GetTypeHash(AnalyzerName));
+			AddressHash = HashCombineFast(AddressHash, ::GetTypeHash(DataType));
+			AddressHash = HashCombineFast(AddressHash, ::GetTypeHash(InstanceID));
+			AddressHash = HashCombineFast(AddressHash, NodeID.A);
+			AddressHash = HashCombineFast(AddressHash, ::GetTypeHash(OutputName));
+
+			return AddressHash;
+		}
+
+		bool FAnalyzerAddress::IsEqual(const FTransmissionAddress& InOther) const
+		{
+			if (InOther.GetAddressType() != GetAddressType())
+			{
+				return false;
+			}
+
+			const FAnalyzerAddress& OtherAddr = static_cast<const FAnalyzerAddress&>(InOther);
+			return OtherAddr.AnalyzerInstanceID == AnalyzerInstanceID
+				&& OtherAddr.AnalyzerMemberName == AnalyzerMemberName
+				&& OtherAddr.AnalyzerName == AnalyzerName
+				&& OtherAddr.DataType == DataType
+				&& OtherAddr.InstanceID == InstanceID
+				&& OtherAddr.NodeID == NodeID
+				&& OtherAddr.OutputName == OutputName;
+		}
+
 		FString FAnalyzerAddress::ToString() const
 		{
 			return FString::Join(TArray<FString>
@@ -21,41 +83,7 @@ namespace Metasound
 				*AnalyzerMemberName.ToString()
 			}, *PathSeparator);
 		}
-
-		FSendAddress FAnalyzerAddress::ToSendAddress() const
-		{
-			const TArray<FString> ChannelTokens
-			{
-				NodeID.ToString(),
-				OutputName.ToString(),
-				AnalyzerName.ToString(),
-				AnalyzerInstanceID.ToString(),
-				AnalyzerMemberName.ToString()
-			};
-
-			// TODO: This is bad as its generating FNames like crazy. One idea was to include guid support in FSendAddresses to avoid FName generation,
-			// which would enable this factory to map analyzer keys to send guid.
-			const FName Channel = FName(*FString::Join(ChannelTokens, *PathSeparator));
-			return FSendAddress { Channel, DataType, InstanceID };
-		}
-
-		bool FAnalyzerAddress::ParseKey(const FString& InAnalyzerKey, FAnalyzerAddress& OutAddress)
-		{
-			TArray<FString> Tokens;
-			if (InAnalyzerKey.ParseIntoArray(Tokens, *PathSeparator) == 7)
-			{
-				OutAddress.InstanceID = static_cast<uint64>(FCString::Atoi64(*Tokens[0]));
-				OutAddress.NodeID = FGuid(Tokens[1]);
-				OutAddress.OutputName = *Tokens[2];
-				OutAddress.DataType = *Tokens[3];
-				OutAddress.AnalyzerName = *Tokens[4];
-				OutAddress.AnalyzerInstanceID = FGuid(Tokens[5]);
-				OutAddress.AnalyzerMemberName = *Tokens[6];
-
-				return true;
-			}
-
-			return false;
-		}
 	} // namespace Frontend
+
+	DEFINE_METASOUND_DATA_TYPE(Frontend::FAnalyzerAddress, "AnalyzerAddress");
 } // namespace Metasound
diff --git a/Engine/Plugins/Runtime/Metasound/Source/MetasoundFrontend/Private/Analysis/MetasoundFrontendAnalyzerView.cpp b/Engine/Plugins/Runtime/Metasound/Source/MetasoundFrontend/Private/Analysis/MetasoundFrontendAnalyzerView.cpp
--- a/Engine/Plugins/Runtime/Metasound/Source/MetasoundFrontend/Private/Analysis/MetasoundFrontendAnalyzerView.cpp
+++ b/Engine/Plugins/Runtime/Metasound/Source/MetasoundFrontend/Private/Analysis/MetasoundFrontendAnalyzerView.cpp
@@ -32,9 +32,8 @@ namespace Metasound
 					FAnalyzerAddress OutputAddress = AnalyzerAddress;
 					OutputAddress.AnalyzerMemberName = Output.Name;
 					OutputAddress.DataType = Output.DataType;
-					const FSendAddress SendAddress = OutputAddress.ToSendAddress();
 					const FReceiverInitParams ReceiverParams { InOperatorSettings };
-					IReceiver* Receiver = FDataTransmissionCenter::Get().RegisterNewReceiver(SendAddress, ReceiverParams).Release();
+					IReceiver* Receiver = FDataTransmissionCenter::Get().RegisterNewReceiver(OutputAddress, ReceiverParams).Release();
 					OutputReceivers.Add({ Output.Name, TSharedPtr<IReceiver>(Receiver) });
 				}
 			}
diff --git a/Engine/Plugins/Runtime/Metasound/Source/MetasoundFrontend/Private/Analysis/MetasoundFrontendGraphAnalyzer.cpp b/Engine/Plugins/Runtime/Metasound/Source/MetasoundFrontend/Private/Analysis/MetasoundFrontendGraphAnalyzer.cpp
--- a/Engine/Plugins/Runtime/Metasound/Source/MetasoundFrontend/Private/Analysis/MetasoundFrontendGraphAnalyzer.cpp
+++ b/Engine/Plugins/Runtime/Metasound/Source/MetasoundFrontend/Private/Analysis/MetasoundFrontendGraphAnalyzer.cpp
@@ -18,68 +18,62 @@ namespace Metasound
 			, InternalDataReferences(MoveTemp(InReferences))
 		{
 			using namespace Frontend;
-			FSendAddress ActiveAnalyzerAddr(GetAnalyzerArraySendChannelName(InInstanceID), GetMetasoundDataTypeName<TArray<FString>>(), InInstanceID);
-			ActiveAnalyzerReceiver = FDataTransmissionCenter::Get().RegisterNewReceiver<TArray<FString>>(ActiveAnalyzerAddr, FReceiverInitParams { InSettings });
+			FGraphAnalyzerAddress GraphAddress(InstanceID);
+			ActiveAnalyzerReceiver = FDataTransmissionCenter::Get().RegisterNewReceiver(GraphAddress, FReceiverInitParams { InSettings });
 		}
 
 		void FGraphAnalyzer::Execute()
 		{
 			// 1. Check if any message has been received that determines what analyzers are active.
-			TSet<FString> ReceiverKeys;
+			TSet<FAnalyzerAddress> ReceiverAddresses;
 			bool bUpdateReceivers = false;
 			if (ActiveAnalyzerReceiver.IsValid())
 			{
-				TReceiver<TArray<FString>>& AnalyzerReceiver = ActiveAnalyzerReceiver->GetAs<TReceiver<TArray<FString>>>();
+				TReceiver<TArray<FAnalyzerAddress>>& AnalyzerReceiver = ActiveAnalyzerReceiver->GetAs<TReceiver<TArray<FAnalyzerAddress>>>();
 				if (AnalyzerReceiver.CanPop())
 				{
 					bUpdateReceivers = true;
-					TArray<FString> AnalyzerAddresses;
+					TArray<FAnalyzerAddress> AnalyzerAddresses;
 					AnalyzerReceiver.Pop(AnalyzerAddresses);
-					ReceiverKeys.Append(AnalyzerAddresses);
+					ReceiverAddresses.Append(AnalyzerAddresses);
 				}
 			}
 
 			// 2. Check if any message has been received that determines what analyzers are active and remove stale analyzers.
-			for (int32 i = Analyzers.Num() - 1; i >= 0; --i)
+			if (bUpdateReceivers)
 			{
-				TUniquePtr<IVertexAnalyzer>& Analyzer = Analyzers[i];
-				check(Analyzer.IsValid());
-				if (bUpdateReceivers)
+				for (int32 i = Analyzers.Num() - 1; i >= 0; --i)
 				{
-					const FString AnalyzerKey = Analyzer->GetAnalyzerAddress().ToString();
-					if (!ReceiverKeys.Contains(AnalyzerKey))
+					TUniquePtr<IVertexAnalyzer>& Analyzer = Analyzers[i];
+					check(Analyzer.IsValid());
+					const FAnalyzerAddress& AnalyzerAddress = Analyzer->GetAnalyzerAddress();
+					if (ReceiverAddresses.Contains(AnalyzerAddress))
 					{
-						constexpr bool bAllowShrinking = false;
-						Analyzers.RemoveAtSwap(i, 1, bAllowShrinking);
-						continue;
+						ReceiverAddresses.Remove(AnalyzerAddress);
 					}
 
-					ReceiverKeys.Remove(AnalyzerKey);
+					constexpr bool bAllowShrinking = false;
+					Analyzers.RemoveAtSwap(i, 1, bAllowShrinking);
 				}
 			}
 
 			// 3. If message received to update analyzers, create missing analyzers
 			if (bUpdateReceivers)
 			{
-				for (const FString& AnalyzerKey : ReceiverKeys)
+				for (const FAnalyzerAddress& AnalyzerAddress : ReceiverAddresses)
 				{
-					FAnalyzerAddress AnalyzerAddress;
-					const bool bAnalyzerKeyValid = FAnalyzerAddress::ParseKey(AnalyzerKey, AnalyzerAddress);
-					if (ensureMsgf(bAnalyzerKeyValid, TEXT("Failed to create MetaSoundAnalyzer: AnalyzerKey '%s' invalid."), *AnalyzerKey))
-					{
-						const FDataReferenceCollection* Collection = InternalDataReferences.Find(AnalyzerAddress.NodeID);
-						// TODO: This currently fails for composed graphs.  Figure out how to differentiate addresses for composed graphs or refactor to only support parent graph analysis.
+					const FDataReferenceCollection* Collection = InternalDataReferences.Find(AnalyzerAddress.NodeID);
+					// TODO: This currently fails for composed graphs.  Figure out how to differentiate addresses for composed graphs or refactor to only support parent graph analysis.
 // 						if (ensureMsgf(Collection != nullptr, TEXT("Failed to create MetaSoundAnalyzer: DataReferenceCollection for node analyzer at address '%s' not found."), *AnalyzerKey))
-						if (Collection)
+					if (Collection)
+					{
+						if (const IVertexAnalyzerFactory* Factory = IVertexAnalyzerRegistry::Get().FindAnalyzerFactory(AnalyzerAddress.AnalyzerName))
 						{
-							if (const IVertexAnalyzerFactory* Factory = IVertexAnalyzerRegistry::Get().FindAnalyzerFactory(AnalyzerAddress.AnalyzerName))
+							if (const FAnyDataReference* DataRef = Collection->FindDataReference(AnalyzerAddress.OutputName))
 							{
-								if (const FAnyDataReference* DataRef = Collection->FindDataReference(AnalyzerAddress.OutputName))
-								{
-									FCreateAnalyzerParams Params{ AnalyzerAddress, OperatorSettings, *DataRef };
-									TUniquePtr<IVertexAnalyzer> NewAnalyzer = Factory->CreateAnalyzer(Params);
-									Analyzers.Add(MoveTemp(NewAnalyzer));
-								}
+								FCreateAnalyzerParams Params{ AnalyzerAddress, OperatorSettings, *DataRef };
+								TUniquePtr<IVertexAnalyzer> NewAnalyzer = Factory->CreateAnalyzer(Params);
+								Analyzers.Add(MoveTemp(NewAnalyzer));
 							}
 						}
 					}
diff --git a/Engine/Plugins/Runtime/Metasound/Source/MetasoundFrontend/Private/Analysis/MetasoundFrontendGraphAnalyzerView.cpp b/Engine/Plugins/Runtime/Metasound/Source/MetasoundFrontend/Private/Analysis/MetasoundFrontendGraphAnalyzerView.cpp
--- a/Engine/Plugins/Runtime/Metasound/Source/MetasoundFrontend/Private/Analysis/MetasoundFrontendGraphAnalyzerView.cpp
+++ b/Engine/Plugins/Runtime/Metasound/Source/MetasoundFrontend/Private/Analysis/MetasoundFrontendGraphAnalyzerView.cpp
@@ -60,10 +60,10 @@ namespace Metasound
 		{
 			// TODO: Mirrored from FMetaSoundParameterTransmitter. Fix here when this is refactored.
 			const float DelayTimeInSeconds = 0.1f;
-			FSendAddress ActiveAnalyzerAddr(FGraphAnalyzer::GetAnalyzerArraySendChannelName(InInstanceID), GetMetasoundDataTypeName<TArray<FString>>(), InstanceID);
-			FSenderInitParams SenderParams{ OperatorSettings, DelayTimeInSeconds };
+			FSenderInitParams SenderParams { OperatorSettings, DelayTimeInSeconds };
 
-			ActiveAnalyzerSender = FDataTransmissionCenter::Get().RegisterNewSender<TArray<FString>>(ActiveAnalyzerAddr, SenderParams);
+			const FGraphAnalyzerAddress AnalyzerAddress(InInstanceID);
+			ActiveAnalyzerSender = FDataTransmissionCenter::Get().RegisterNewSender(AnalyzerAddress, SenderParams);
 		}
 
 		void FMetasoundGraphAnalyzerView::AddAnalyzerForAllSupportedOutputs(FName InAnalyzerName, bool bInRequiresConnection)
@@ -79,8 +79,7 @@ namespace Metasound
 
 			IterateOutputsSupportingAnalyzer(GraphHandle, InstanceID, InAnalyzerName, bInRequiresConnection, [this, &Factory, &InAnalyzerName](FConstOutputHandle OutputHandle, const FAnalyzerAddress& AnalyzerAddress)
 			{
-				FString AnalyzerKey = AnalyzerAddress.ToString();
-				ActiveAnalyzerKeys.Add(MoveTemp(AnalyzerKey));
+				ActiveAnalyzers.Add(AnalyzerAddress);
 
 				const TArray<FAnalyzerOutput>& AnalyzerOutputs = Factory->GetAnalyzerOutputs();
 				for (const FAnalyzerOutput& AnalyzerOutput : AnalyzerOutputs)
@@ -95,7 +94,10 @@ namespace Metasound
 					AnalyzerViews.FindOrAdd(OutputKey).Add(MoveTemp(NewView));
 				}
 			});
-			ActiveAnalyzerSender->PushLiteral(ActiveAnalyzerKeys.Array());
+
+			TArray<FString> ActiveAnalyzerStrings;
+			Algo::Transform(ActiveAnalyzers, ActiveAnalyzerStrings, [](const FAnalyzerAddress& Address) { return Address.ToString(); });
+			ActiveAnalyzerSender->PushLiteral(MoveTemp(ActiveAnalyzerStrings));
 		}
 
 		TArray<FMetasoundAnalyzerView*> FMetasoundGraphAnalyzerView::GetAnalyzerViews(FName InAnalyzerName)
@@ -194,19 +196,19 @@ namespace Metasound
 				AnalyzerViews.Remove(Key);
 			}
 
-			TArray<FString> KeyArray = ActiveAnalyzerKeys.Array();
-			for (FString& AnalyzerKey : KeyArray)
+			TArray<FAnalyzerAddress> AnalyzerAddresses = ActiveAnalyzers.Array();
+			for (const FAnalyzerAddress& AnalyzerAddress : AnalyzerAddresses)
 			{
-				FAnalyzerAddress AnalyzerAddress;
-				if (ensureAlways(FAnalyzerAddress::ParseKey(AnalyzerKey, AnalyzerAddress)))
+				if (AnalyzerAddress.AnalyzerName == InAnalyzerName)
 				{
-					if (AnalyzerAddress.AnalyzerName == InAnalyzerName)
-					{
-						ActiveAnalyzerKeys.Remove(AnalyzerKey);
-					}
+					ActiveAnalyzers.Remove(AnalyzerAddress);
+					break;
 				}
 			}
-			ActiveAnalyzerSender->PushLiteral(ActiveAnalyzerKeys.Array());
+
+			TArray<FString> ActiveAnalyzerStrings;
+			Algo::Transform(ActiveAnalyzers, ActiveAnalyzerStrings, [] (const FAnalyzerAddress& Address) { return Address.ToString(); });
+			ActiveAnalyzerSender->PushLiteral(MoveTemp(ActiveAnalyzerStrings));
 		}
 
 		const FMetasoundAssetBase& FMetasoundGraphAnalyzerView::GetMetaSoundAssetChecked() const
diff --git a/Engine/Plugins/Runtime/Metasound/Source/MetasoundFrontend/Public/Analysis/MetasoundFrontendAnalyzerAddress.h b/Engine/Plugins/Runtime/Metasound/Source/MetasoundFrontend/Public/Analysis/MetasoundFrontendAnalyzerAddress.h
--- a/Engine/Plugins/Runtime/Metasound/Source/MetasoundFrontend/Public/Analysis/MetasoundFrontendAnalyzerAddress.h
+++ b/Engine/Plugins/Runtime/Metasound/Source/MetasoundFrontend/Public/Analysis/MetasoundFrontendAnalyzerAddress.h
@@ -3,9 +3,11 @@
 
 #include "Containers/UnrealString.h"
 #include "HAL/Platform.h"
+#include "MetasoundDataReferenceMacro.h"
 #include "MetasoundRouter.h"
 #include "MetasoundVertex.h"
 #include "Misc/Guid.h"
+#include "UObject/NameTypes.h"
 
 
 namespace Metasound
@@ -14,10 +16,29 @@ namespace Metasound
 	{
 		// String serializable (as key) channel of analyzer or its internal members
 		// that can be written to or read from using the Transmission System.
-		struct METASOUNDFRONTEND_API FAnalyzerAddress
+		class METASOUNDFRONTEND_API FAnalyzerAddress : public FTransmissionAddress
 		{
+		public:
 			static const FString PathSeparator;
 
+			FAnalyzerAddress(const FString& InAddressString);
+			FAnalyzerAddress() = default;
+
+			virtual ~FAnalyzerAddress() = default;
+
+			virtual FName GetAddressType() const override;
+
+			virtual FName GetDataType() const override;
+
+			virtual TUniquePtr<FTransmissionAddress> Clone() const override;
+
+			// Converts AnalyzerAddress to String representation using the PathSeparator
+			virtual FString ToString() const override;
+
+			virtual uint32 GetHash() const override;
+
+			virtual bool IsEqual(const FTransmissionAddress& InOther) const override;
+
 			// Active Instance ID to monitor
 			uint64 InstanceID = TNumericLimits<uint64>::Max();
 
@@ -43,14 +64,8 @@ namespace Metasound
 			// the same name. Useful if the analyzer requires outputting multiple analysis values.
 			// Can potentially be used as an input as well to modify analyzer settings.
 			FName AnalyzerMemberName;
-
-			// Converts AnalyzerAddress to String representation using the PathSeparator
-			FString ToString() const;
-
-			// Converts AnalyzerAddress to SendAddress
-			FSendAddress ToSendAddress() const;
-
-			static bool ParseKey(const FString& InAnalyzerKey, FAnalyzerAddress& OutAnalyzer);
 		};
 	} // namespace Frontend
+
+	DECLARE_METASOUND_DATA_REFERENCE_TYPES(Frontend::FAnalyzerAddress, METASOUNDFRONTEND_API, FAnalyzerAddressTypeInfo, FAnalyzerAddressReadRef, FAnalyzerAddressWriteRef)
 } // namespace Metasound
diff --git a/Engine/Plugins/Runtime/Metasound/Source/MetasoundFrontend/Public/Analysis/MetasoundFrontendGraphAnalyzer.h b/Engine/Plugins/Runtime/Metasound/Source/MetasoundFrontend/Public/Analysis/MetasoundFrontendGraphAnalyzer.h
--- a/Engine/Plugins/Runtime/Metasound/Source/MetasoundFrontend/Public/Analysis/MetasoundFrontendGraphAnalyzer.h
+++ b/Engine/Plugins/Runtime/Metasound/Source/MetasoundFrontend/Public/Analysis/MetasoundFrontendGraphAnalyzer.h
@@ -15,6 +15,55 @@ namespace Metasound
 	{
 		using FNodeVertexDataReferenceMap = TMap<FGuid, FDataReferenceCollection>;
 
+		class METASOUNDFRONTEND_API FGraphAnalyzerAddress : public FTransmissionAddress
+		{
+			uint64 InstanceID = INDEX_NONE;
+
+		public:
+			FGraphAnalyzerAddress(uint64 InInstanceID)
+				: InstanceID(InInstanceID)
+			{
+			}
+
+			virtual ~FGraphAnalyzerAddress() = default;
+
+			virtual FName GetAddressType() const override
+			{
+				return "GraphAnalyzer";
+			}
+
+			virtual FName GetDataType() const override
+			{
+				return GetMetasoundDataTypeName<TArray<FAnalyzerAddress>>();
+			}
+
+			virtual TUniquePtr<FTransmissionAddress> Clone() const override
+			{
+				return TUniquePtr<FTransmissionAddress>(new FGraphAnalyzerAddress(*this));
+			}
+
+			virtual FString ToString() const override
+			{
+				return FString::Printf(TEXT("%s%s%lld"), *GetAddressType().ToString(), * FAnalyzerAddress::PathSeparator, InstanceID);
+			}
+
+			virtual uint32 GetHash() const override
+			{
+				return static_cast<uint32>(InstanceID);
+			}
+
+			virtual bool IsEqual(const FTransmissionAddress& InOther) const override
+			{
+				if (InOther.GetAddressType() != GetAddressType())
+				{
+					return false;
+				}
+
+				const FGraphAnalyzerAddress& OtherAddr = static_cast<const FGraphAnalyzerAddress&>(InOther);
+				return OtherAddr.InstanceID == InstanceID;
+			}
+		};
+
 		// Handles intrinsic analysis operations within a given graph
 		// should the graph's operator be enabled for analysis.
 		class METASOUNDFRONTEND_API FGraphAnalyzer
@@ -30,14 +79,6 @@ namespace Metasound
 			FGraphAnalyzer(const FOperatorSettings& InSettings, uint64 InInstanceID, FNodeVertexDataReferenceMap&& InGraphReferences);
 			~FGraphAnalyzer() = default;
 
-			// Creates a send channel name unique for the given sound instance used to send array of analyzer
-			// addresses to the profiler of what analyzers are expected to be active for a given instance.
-			static const FName GetAnalyzerArraySendChannelName(uint64 InInstanceID)
-			{
-				const FString ChannelName = FString::Printf(TEXT("AnalyzerArray%s%lld"), *FAnalyzerAddress::PathSeparator, InInstanceID);
-				return *ChannelName;
-			}
-
 			// Execute analysis for the current block
 			void Execute();
 		};
diff --git a/Engine/Plugins/Runtime/Metasound/Source/MetasoundFrontend/Public/Analysis/MetasoundFrontendGraphAnalyzerView.h b/Engine/Plugins/Runtime/Metasound/Source/MetasoundFrontend/Public/Analysis/MetasoundFrontendGraphAnalyzerView.h
--- a/Engine/Plugins/Runtime/Metasound/Source/MetasoundFrontend/Public/Analysis/MetasoundFrontendGraphAnalyzerView.h
+++ b/Engine/Plugins/Runtime/Metasound/Source/MetasoundFrontend/Public/Analysis/MetasoundFrontendGraphAnalyzerView.h
@@ -23,8 +23,8 @@ namespace Metasound
 			// Sender in charge of supplying expected vertex analyzers currently being analyzed.
 			TUniquePtr<ISender> ActiveAnalyzerSender;
 
-			// Set of active keys used for what senders should be active on the instance graph
-			TSet<FString> ActiveAnalyzerKeys;
+			// Set of active analyzer addresses used for what analyzers should be active on the associated graph instance
+			TSet<FAnalyzerAddress> ActiveAnalyzers;
 
 			const FMetasoundAssetBase& GetMetaSoundAssetChecked() const;
 
diff --git a/Engine/Plugins/Runtime/Metasound/Source/MetasoundFrontend/Public/Analysis/MetasoundFrontendVertexAnalyzer.h b/Engine/Plugins/Runtime/Metasound/Source/MetasoundFrontend/Public/Analysis/MetasoundFrontendVertexAnalyzer.h
--- a/Engine/Plugins/Runtime/Metasound/Source/MetasoundFrontend/Public/Analysis/MetasoundFrontendVertexAnalyzer.h
+++ b/Engine/Plugins/Runtime/Metasound/Source/MetasoundFrontend/Public/Analysis/MetasoundFrontendVertexAnalyzer.h
@@ -75,7 +75,7 @@ namespace Metasound
 			TBoundAnalyzerOutput(const FAnalyzerAddress& InAnalyzerOutputAddress, const FOperatorSettings& InOperatorSettings, TDataReadReference<DataType>&& InData)
 				: DataRef(MoveTemp(InData))
 			{
-				Sender = FDataTransmissionCenter::Get().RegisterNewSender<DataType>(InAnalyzerOutputAddress.ToSendAddress(), FSenderInitParams { InOperatorSettings, 0 });
+				Sender = FDataTransmissionCenter::Get().RegisterNewSender<DataType>(InAnalyzerOutputAddress, FSenderInitParams { InOperatorSettings, 0 });
 				ensure(Sender);
 			}
 
