Register data type as it is required to generate transmission channel to avoid ensure in editor
#rnx

diff --git a/Engine/Plugins/Runtime/Metasound/Source/MetasoundFrontend/Private/Analysis/MetasoundFrontendAnalyzerAddress.cpp b/Engine/Plugins/Runtime/Metasound/Source/MetasoundFrontend/Private/Analysis/MetasoundFrontendAnalyzerAddress.cpp
--- a/Engine/Plugins/Runtime/Metasound/Source/MetasoundFrontend/Private/Analysis/MetasoundFrontendAnalyzerAddress.cpp
+++ b/Engine/Plugins/Runtime/Metasound/Source/MetasoundFrontend/Private/Analysis/MetasoundFrontendAnalyzerAddress.cpp
@@ -85,5 +85,5 @@ namespace Metasound
 		}
 	} // namespace Frontend
 
-	DEFINE_METASOUND_DATA_TYPE(Frontend::FAnalyzerAddress, "AnalyzerAddress");
+	REGISTER_METASOUND_DATATYPE(Frontend::FAnalyzerAddress, "AnalyzerAddress");
 } // namespace Metasound
