[MetaSound] Allow Constructor pin to be set from BP

diff --git a/Engine/Plugins/Runtime/Metasound/Source/MetasoundEngine/Private/MetasoundSource.cpp b/Engine/Plugins/Runtime/Metasound/Source/MetasoundEngine/Private/MetasoundSource.cpp
--- a/Engine/Plugins/Runtime/Metasound/Source/MetasoundEngine/Private/MetasoundSource.cpp
+++ b/Engine/Plugins/Runtime/Metasound/Source/MetasoundEngine/Private/MetasoundSource.cpp
@@ -381,24 +381,26 @@ void UMetaSoundSource::InitParameters(TArray<FAudioParameter>& InParametersToIni
 	// via InitResources. If it has, this call is fast and returns the already cached RuntimeData.
 	RegisterGraphWithFrontend(GetInitRegistrationOptions());
 	const FRuntimeData& RuntimeData = GetRuntimeData();
-	const TArray<FMetasoundFrontendClassInput>& TransmittableInputs = RuntimeData.TransmittableInputs;
+	const TArray<FMetasoundFrontendClassInput>& PublicInputs = RuntimeData.PublicInputs;
+
+	TMap<FName, FMetasoundFrontendVertex> PublicInputMap;
+	Algo::Transform(PublicInputs, PublicInputMap, [](const FMetasoundFrontendClassInput& Input)
+	{
+		return TPair<FName, FMetasoundFrontendVertex>(Input.Name, Input);
+	});
 
 	// Removes values that are not explicitly defined by the ParamType and returns
 	// whether or not the parameter is a valid input and should be included.
-	auto Sanitize = [&TransmittableInputs](FAudioParameter& Parameter) -> bool
+	auto Sanitize = [&PublicInputMap](FAudioParameter& Parameter) -> bool
 	{
-		auto FindInput = [&Parameter](const FMetasoundFrontendClassInput& Input)
+		const FMetasoundFrontendVertex* Input = PublicInputMap.Find(Parameter.ParamName);
+		if (!Input)
 		{
-			if (Parameter.ParamName == Input.Name)
-			{
-				return Parameter.TypeName.IsNone() || Parameter.TypeName == Input.TypeName;
-			}
-
 			return false;
-		};
+		}
 
-		const FMetasoundFrontendClassInput* Input = TransmittableInputs.FindByPredicate(FindInput);
-		if (!Input)
+		const bool bIsMatchingType = Parameter.TypeName.IsNone() || (Parameter.TypeName == Input->TypeName);
+		if (!bIsMatchingType)
 		{
 			return false;
 		}
@@ -524,11 +526,6 @@ void UMetaSoundSource::InitParameters(TArray<FAudioParameter>& InParametersToIni
 		}
 	};
 
-	TMap<FName, FName> InputNameTypeMap;
-	Algo::Transform(GetDocumentChecked().RootGraph.Interface.Inputs, InputNameTypeMap, [](const FMetasoundFrontendClassInput& Input)
-	{
-		return TPair<FName, FName>(Input.Name, Input.TypeName);
-	});
 
 	for (int32 i = InParametersToInit.Num() - 1; i >= 0; --i)
 	{
@@ -541,7 +538,7 @@ void UMetaSoundSource::InitParameters(TArray<FAudioParameter>& InParametersToIni
 		
 		if (Sanitize(Parameter))
 		{
-			if (IsParameterValid(Parameter, InputNameTypeMap))
+			if (IsParameterValid(Parameter, PublicInputMap))
 			{
 				ConstructProxies(Parameter);
 			}
@@ -608,7 +605,7 @@ bool UMetaSoundSource::ImplementsParameterInterface(Audio::FParameterInterfacePt
 	return GetDocumentChecked().Interfaces.Contains(Version);
 }
 
-ISoundGeneratorPtr UMetaSoundSource::CreateSoundGenerator(const FSoundGeneratorInitParams& InParams)
+ISoundGeneratorPtr UMetaSoundSource::CreateSoundGenerator(const FSoundGeneratorInitParams& InParams, TArray<FAudioParameter>&& InDefaultParameters)
 {
 	using namespace Metasound;
 	using namespace Metasound::Frontend;
@@ -639,7 +636,8 @@ ISoundGeneratorPtr UMetaSoundSource::CreateSoundGenerator(const FSoundGeneratorI
 		MetasoundGraph,
 		Environment,
 		GetName(),
-		GetOutputAudioChannelOrder()
+		GetOutputAudioChannelOrder(),
+		MoveTemp(InDefaultParameters)
 	};
 
 	return ISoundGeneratorPtr(new FMetasoundGenerator(MoveTemp(InitParams)));
@@ -652,8 +650,8 @@ bool UMetaSoundSource::GetAllDefaultParameters(TArray<FAudioParameter>& OutParam
 	using namespace Metasound::Engine;
 
 	// TODO: Make this use the cached runtime data's input copy. This call can become expensive if called repeatedly.
-	TArray<FMetasoundFrontendClassInput> TransmittableInputs = GetTransmittableClassInputs();
-	for(const FMetasoundFrontendClassInput& Input : TransmittableInputs)
+	TArray<FMetasoundFrontendClassInput> PublicInputs = GetPublicClassInputs();
+	for(const FMetasoundFrontendClassInput& Input : PublicInputs)
 	{
 		FAudioParameter Params;
 		Params.ParamName = Input.Name;
@@ -747,15 +745,15 @@ bool UMetaSoundSource::GetAllDefaultParameters(TArray<FAudioParameter>& OutParam
 
 bool UMetaSoundSource::IsParameterValid(const FAudioParameter& InParameter) const
 {
-	TMap<FName, FName> InputNameTypeMap;
+	TMap<FName, FMetasoundFrontendVertex> InputNameTypeMap;
 	Algo::Transform(GetDocumentChecked().RootGraph.Interface.Inputs, InputNameTypeMap, [] (const FMetasoundFrontendClassInput& Input)
 	{
-		return TPair<FName, FName>(Input.Name, Input.TypeName);
+		return TPair<FName, FMetasoundFrontendVertex>(Input.Name, Input);
 	});
 	return IsParameterValid(InParameter, InputNameTypeMap);
 }
 
-bool UMetaSoundSource::IsParameterValid(const FAudioParameter& InParameter, const TMap<FName, FName>& InInputNameTypeMap) const
+bool UMetaSoundSource::IsParameterValid(const FAudioParameter& InParameter, const TMap<FName, FMetasoundFrontendVertex>& InInputNameVertexMap) const
 {
 	using namespace Metasound;
 	using namespace Metasound::Frontend;
@@ -765,11 +763,12 @@ bool UMetaSoundSource::IsParameterValid(const FAudioParameter& InParameter, cons
 		return false;
 	}
 
-	const FName* TypeName = InInputNameTypeMap.Find(InParameter.ParamName);
-	if (!TypeName)
+	const FMetasoundFrontendVertex* Vertex = InInputNameVertexMap.Find(InParameter.ParamName);
+	if (!Vertex)
 	{
 		return false;
 	}
+	const FName& TypeName = Vertex->TypeName;
 
 	bool bIsValid = false;
 	switch (InParameter.ParamType)
@@ -777,7 +776,7 @@ bool UMetaSoundSource::IsParameterValid(const FAudioParameter& InParameter, cons
 		case EAudioParameterType::Boolean:
 		{
 			FDataTypeRegistryInfo DataTypeInfo;
-			bIsValid = IDataTypeRegistry::Get().GetDataTypeInfo(*TypeName, DataTypeInfo);
+			bIsValid = IDataTypeRegistry::Get().GetDataTypeInfo(TypeName, DataTypeInfo);
 			bIsValid &= DataTypeInfo.bIsBoolParsable;
 		}
 		break;
@@ -785,7 +784,7 @@ bool UMetaSoundSource::IsParameterValid(const FAudioParameter& InParameter, cons
 		case EAudioParameterType::BooleanArray:
 		{
 			FDataTypeRegistryInfo DataTypeInfo;
-			bIsValid = IDataTypeRegistry::Get().GetDataTypeInfo(*TypeName, DataTypeInfo);
+			bIsValid = IDataTypeRegistry::Get().GetDataTypeInfo(TypeName, DataTypeInfo);
 			bIsValid &= DataTypeInfo.bIsBoolArrayParsable;
 		}
 		break;
@@ -793,7 +792,7 @@ bool UMetaSoundSource::IsParameterValid(const FAudioParameter& InParameter, cons
 		case EAudioParameterType::Float:
 		{
 			FDataTypeRegistryInfo DataTypeInfo;
-			bIsValid = IDataTypeRegistry::Get().GetDataTypeInfo(*TypeName, DataTypeInfo);
+			bIsValid = IDataTypeRegistry::Get().GetDataTypeInfo(TypeName, DataTypeInfo);
 			bIsValid &= DataTypeInfo.bIsFloatParsable;
 		}
 		break;
@@ -801,7 +800,7 @@ bool UMetaSoundSource::IsParameterValid(const FAudioParameter& InParameter, cons
 		case EAudioParameterType::FloatArray:
 		{
 			FDataTypeRegistryInfo DataTypeInfo;
-			bIsValid = IDataTypeRegistry::Get().GetDataTypeInfo(*TypeName, DataTypeInfo);
+			bIsValid = IDataTypeRegistry::Get().GetDataTypeInfo(TypeName, DataTypeInfo);
 			bIsValid &= DataTypeInfo.bIsFloatArrayParsable;
 		}
 		break;
@@ -809,7 +808,7 @@ bool UMetaSoundSource::IsParameterValid(const FAudioParameter& InParameter, cons
 		case EAudioParameterType::Integer:
 		{
 			FDataTypeRegistryInfo DataTypeInfo;
-			bIsValid = IDataTypeRegistry::Get().GetDataTypeInfo(*TypeName, DataTypeInfo);
+			bIsValid = IDataTypeRegistry::Get().GetDataTypeInfo(TypeName, DataTypeInfo);
 			bIsValid &= DataTypeInfo.bIsIntParsable;
 		}
 		break;
@@ -817,7 +816,7 @@ bool UMetaSoundSource::IsParameterValid(const FAudioParameter& InParameter, cons
 		case EAudioParameterType::IntegerArray:
 		{
 			FDataTypeRegistryInfo DataTypeInfo;
-			bIsValid = IDataTypeRegistry::Get().GetDataTypeInfo(*TypeName, DataTypeInfo);
+			bIsValid = IDataTypeRegistry::Get().GetDataTypeInfo(TypeName, DataTypeInfo);
 			bIsValid &= DataTypeInfo.bIsIntArrayParsable;
 
 		}
@@ -828,7 +827,7 @@ bool UMetaSoundSource::IsParameterValid(const FAudioParameter& InParameter, cons
 			FDataTypeRegistryInfo DataTypeInfo;
 			bIsValid = IDataTypeRegistry::Get().GetDataTypeInfo(InParameter.ObjectParam, DataTypeInfo);
 			bIsValid &= DataTypeInfo.bIsProxyParsable;
-			bIsValid &= DataTypeInfo.DataTypeName == *TypeName;
+			bIsValid &= DataTypeInfo.DataTypeName == TypeName;
 		}
 		break;
 
@@ -836,7 +835,7 @@ bool UMetaSoundSource::IsParameterValid(const FAudioParameter& InParameter, cons
 		{
 			bIsValid = true;
 
-			const FName ElementTypeName = CreateElementTypeNameFromArrayTypeName(*TypeName);
+			const FName ElementTypeName = CreateElementTypeNameFromArrayTypeName(TypeName);
 			for (UObject* Object : InParameter.ArrayObjectParam)
 			{
 				FDataTypeRegistryInfo DataTypeInfo;
@@ -855,7 +854,7 @@ bool UMetaSoundSource::IsParameterValid(const FAudioParameter& InParameter, cons
 		case EAudioParameterType::String:
 		{
 			FDataTypeRegistryInfo DataTypeInfo;
-			bIsValid = IDataTypeRegistry::Get().GetDataTypeInfo(*TypeName, DataTypeInfo);
+			bIsValid = IDataTypeRegistry::Get().GetDataTypeInfo(TypeName, DataTypeInfo);
 			bIsValid &= DataTypeInfo.bIsStringParsable;
 		}
 		break;
@@ -863,7 +862,7 @@ bool UMetaSoundSource::IsParameterValid(const FAudioParameter& InParameter, cons
 		case EAudioParameterType::StringArray:
 		{
 			FDataTypeRegistryInfo DataTypeInfo;
-			bIsValid = IDataTypeRegistry::Get().GetDataTypeInfo(*TypeName, DataTypeInfo);
+			bIsValid = IDataTypeRegistry::Get().GetDataTypeInfo(TypeName, DataTypeInfo);
 			bIsValid &= DataTypeInfo.bIsStringArrayParsable;
 		}
 		break;
@@ -871,14 +870,14 @@ bool UMetaSoundSource::IsParameterValid(const FAudioParameter& InParameter, cons
 		case EAudioParameterType::NoneArray:
 		{
 			FDataTypeRegistryInfo DataTypeInfo;
-			bIsValid = IDataTypeRegistry::Get().GetDataTypeInfo(*TypeName, DataTypeInfo);
+			bIsValid = IDataTypeRegistry::Get().GetDataTypeInfo(TypeName, DataTypeInfo);
 			bIsValid &= DataTypeInfo.bIsDefaultArrayParsable;
 		}
 		case EAudioParameterType::None:
 		default:
 		{
 			FDataTypeRegistryInfo DataTypeInfo;
-			bIsValid = IDataTypeRegistry::Get().GetDataTypeInfo(*TypeName, DataTypeInfo);
+			bIsValid = IDataTypeRegistry::Get().GetDataTypeInfo(TypeName, DataTypeInfo);
 			bIsValid &= DataTypeInfo.bIsDefaultParsable;
 		}
 		break;
@@ -912,7 +911,6 @@ TUniquePtr<Audio::IParameterTransmitter> UMetaSoundSource::CreateParameterTransm
 	}
 
 	TUniquePtr<Audio::IParameterTransmitter> NewTransmitter = MakeUnique<Metasound::FMetaSoundParameterTransmitter>(InitParams);
-	NewTransmitter->SetParameters(MoveTemp(InParams.DefaultParams));
 
 	return NewTransmitter;
 }
diff --git a/Engine/Plugins/Runtime/Metasound/Source/MetasoundEngine/Public/MetasoundSource.h b/Engine/Plugins/Runtime/Metasound/Source/MetasoundEngine/Public/MetasoundSource.h
--- a/Engine/Plugins/Runtime/Metasound/Source/MetasoundEngine/Public/MetasoundSource.h
+++ b/Engine/Plugins/Runtime/Metasound/Source/MetasoundEngine/Public/MetasoundSource.h
@@ -175,19 +175,14 @@ public:
 	virtual bool SupportsSubtitles() const override;
 	virtual float GetDuration() const override;
 	virtual bool ImplementsParameterInterface(Audio::FParameterInterfacePtr InInterface) const override;
-	virtual ISoundGeneratorPtr CreateSoundGenerator(const FSoundGeneratorInitParams& InParams) override;
+	virtual ISoundGeneratorPtr CreateSoundGenerator(const FSoundGeneratorInitParams& InParams, TArray<FAudioParameter>&& InDefaultParameters) override;
 	virtual TUniquePtr<Audio::IParameterTransmitter> CreateParameterTransmitter(Audio::FParameterTransmitterInitParams&& InParams) const override;
 	virtual bool IsParameterValid(const FAudioParameter& InParameter) const override;
 	virtual bool IsLooping() const override;
 	virtual bool IsOneShot() const override;
 	virtual bool EnableSubmixSendsOnPreview() const override { return true; }
 protected:
-	/** Gets all the default parameters for this Asset.  */
-	virtual bool GetAllDefaultParameters(TArray<FAudioParameter>& OutParameters) const override;
 
-	virtual void SetReferencedAssetClassKeys(TSet<Metasound::Frontend::FNodeRegistryKey>&& InKeys) override;
-
-	bool IsParameterValid(const FAudioParameter& InParameter, const TMap<FName, FName>& InInputNameTypePairs) const;
 
 	Metasound::Frontend::FDocumentAccessPtr GetDocument() override
 	{
@@ -205,8 +200,15 @@ protected:
 		return MakeAccessPtr<FConstDocumentAccessPtr>(RootMetasoundDocument.AccessPoint, RootMetasoundDocument);
 	}
 
+	/** Gets all the default parameters for this Asset.  */
+	virtual bool GetAllDefaultParameters(TArray<FAudioParameter>& OutParameters) const override;
+
+	virtual void SetReferencedAssetClassKeys(TSet<Metasound::Frontend::FNodeRegistryKey>&& InKeys) override;
+
 private:
 
+	bool IsParameterValid(const FAudioParameter& InParameter, const TMap<FName, FMetasoundFrontendVertex>& InInputNameVertexMap) const;
+
 	Metasound::FOperatorSettings GetOperatorSettings(Metasound::FSampleRate InSampleRate) const;
 	Metasound::FMetasoundEnvironment CreateEnvironment() const;
 	Metasound::FMetasoundEnvironment CreateEnvironment(const FSoundGeneratorInitParams& InParams) const;
diff --git a/Engine/Plugins/Runtime/Metasound/Source/MetasoundFrontend/Private/MetasoundAssetBase.cpp b/Engine/Plugins/Runtime/Metasound/Source/MetasoundFrontend/Private/MetasoundAssetBase.cpp
--- a/Engine/Plugins/Runtime/Metasound/Source/MetasoundFrontend/Private/MetasoundAssetBase.cpp
+++ b/Engine/Plugins/Runtime/Metasound/Source/MetasoundFrontend/Private/MetasoundAssetBase.cpp
@@ -3,6 +3,7 @@
 #include "MetasoundAssetBase.h"
 
 #include "Algo/AnyOf.h"
+#include "Algo/Copy.h"
 #include "Algo/ForEach.h"
 #include "Algo/Transform.h"
 #include "Containers/Set.h"
@@ -840,14 +841,25 @@ TArray<FMetasoundFrontendClassInput> FMetasoundAssetBase::GetTransmittableClassI
 	METASOUND_TRACE_CPUPROFILER_EVENT_SCOPE(FMetasoundAssetBase::GetTransmittableClassInputs);
 
 	check(IsInGameThread() || IsInAudioThread());
-	TArray<FMetasoundFrontendClassInput> Inputs;
 
-	const FMetasoundFrontendDocument& Doc = GetDocumentChecked();
+	TArray<FMetasoundFrontendClassInput> TransmittableInputs = GetPublicClassInputs();
+	RemoveNonTransmittableClassInputs(TransmittableInputs);
+
+	return TransmittableInputs;
+}
+
+TArray<FMetasoundFrontendClassInput> FMetasoundAssetBase::GetPublicClassInputs() const
+{
+	using namespace Metasound;
+	using namespace Metasound::Frontend;
+
+	METASOUND_TRACE_CPUPROFILER_EVENT_SCOPE(FMetasoundAssetBase::GetPublicClassInputs);
+
 	auto GetInputName = [](const FMetasoundFrontendClassInput& InInput) { return InInput.Name; };
 
-	// Do not transmit vertices defined in interface marked as non-transmittable
+	// Inputs which are controlled by an interface are private. 
 	TArray<const IInterfaceRegistryEntry*> Interfaces;
-	TSet<FVertexName> NonTransmittableInputs;
+	TSet<FVertexName> PrivateInputs;
 	GetDeclaredInterfaces(Interfaces);
 	for (const IInterfaceRegistryEntry* InterfaceEntry : Interfaces)
 	{
@@ -856,34 +868,47 @@ TArray<FMetasoundFrontendClassInput> FMetasoundAssetBase::GetTransmittableClassI
 			if (InterfaceEntry->GetRouterName() != Audio::IParameterTransmitter::RouterName)
 			{
 				const FMetasoundFrontendInterface& Interface = InterfaceEntry->GetInterface();
-				Algo::Transform(Interface.Inputs, NonTransmittableInputs, GetInputName);
+				Algo::Transform(Interface.Inputs, PrivateInputs, GetInputName);
 			}
 		}
 	}
 
+	auto IsPublic = [&PrivateInputs](const FMetasoundFrontendClassVertex& InVertex)
+	{
+		return !PrivateInputs.Contains(InVertex.Name);
+	};
+
+	const FMetasoundFrontendDocument& Doc = GetDocumentChecked();
+	TArray<FMetasoundFrontendClassInput> PublicInputs;
+
+	Algo::CopyIf(Doc.RootGraph.Interface.Inputs, PublicInputs, IsPublic);
+
+	return PublicInputs;
+}
+
+void FMetasoundAssetBase::RemoveNonTransmittableClassInputs(TArray<FMetasoundFrontendClassInput>& InOutClassInputs) const
+{
+	using namespace Metasound::Frontend;
+
 	// Do not transmit vertices which are not transmittable. Async communication
 	// is not supported without transmission.
 	IDataTypeRegistry& Registry = IDataTypeRegistry::Get();
-	auto IsTransmittable = [&Registry, &NonTransmittableInputs](const FMetasoundFrontendClassVertex& InVertex)
+	auto IsNotTransmittable = [&Registry](const FMetasoundFrontendClassVertex& InVertex)
 	{
 		// Don't transmit constructor inputs 
 		if (InVertex.AccessType == EMetasoundFrontendVertexAccessType::Reference)
 		{
-			if (!NonTransmittableInputs.Contains(InVertex.Name))
+			FDataTypeRegistryInfo Info;
+			if (Registry.GetDataTypeInfo(InVertex.TypeName, Info))
 			{
-				FDataTypeRegistryInfo Info;
-				if (Registry.GetDataTypeInfo(InVertex.TypeName, Info))
-				{
-					return Info.bIsTransmittable;
-				}
+				return !Info.bIsTransmittable;
 			}
 		}
 
-		return false;
+		return true;
 	};
-	Algo::TransformIf(Doc.RootGraph.Interface.Inputs, Inputs, IsTransmittable, [] (const FMetasoundFrontendClassInput& Input) { return Input; });
 
-	return Inputs;
+	InOutClassInputs.RemoveAllSwap(IsNotTransmittable);
 }
 
 void FMetasoundAssetBase::RebuildReferencedAssetClassKeys()
@@ -908,13 +933,18 @@ const FMetasoundAssetBase::FRuntimeData& FMetasoundAssetBase::CacheRuntimeData(c
 	if (CachedRuntimeData.ChangeID != CurrentCachedRuntimeDataChangeID)
 	{
 		// Update CachedRuntimeData.
-		TArray<FMetasoundFrontendClassInput> ClassInputs = GetTransmittableClassInputs();
-		TSharedPtr<Metasound::IGraph, ESPMode::ThreadSafe> Graph = BuildMetasoundDocument(InPreprocessedDoc, ClassInputs);
+
+		TArray<FMetasoundFrontendClassInput> PublicInputs = GetPublicClassInputs();
+		TArray<FMetasoundFrontendClassInput> TransmittableInputs = PublicInputs;
+		RemoveNonTransmittableClassInputs(TransmittableInputs);
+
+		TSharedPtr<Metasound::IGraph, ESPMode::ThreadSafe> Graph = BuildMetasoundDocument(InPreprocessedDoc, TransmittableInputs);
 
 		CachedRuntimeData =
 		{
 			CurrentCachedRuntimeDataChangeID,
-			MoveTemp(ClassInputs),
+			MoveTemp(PublicInputs),
+			MoveTemp(TransmittableInputs),
 			MoveTemp(Graph)
 		};
 	}
diff --git a/Engine/Plugins/Runtime/Metasound/Source/MetasoundFrontend/Private/MetasoundFrontendDataTypeRegistry.cpp b/Engine/Plugins/Runtime/Metasound/Source/MetasoundFrontend/Private/MetasoundFrontendDataTypeRegistry.cpp
--- a/Engine/Plugins/Runtime/Metasound/Source/MetasoundFrontend/Private/MetasoundFrontendDataTypeRegistry.cpp
+++ b/Engine/Plugins/Runtime/Metasound/Source/MetasoundFrontend/Private/MetasoundFrontendDataTypeRegistry.cpp
@@ -505,6 +505,7 @@ namespace Metasound
 				virtual FLiteral CreateLiteralFromUObject(const FName& InDataType, UObject* InObject) const override;
 				virtual FLiteral CreateLiteralFromUObjectArray(const FName& InDataType, const TArray<UObject*>& InObjectArray) const override;
 
+				virtual TOptional<FAnyDataReference> CreateDataReference(const FName& InDataType, EDataReferenceAccessType InAccessType, const FLiteral& InLiteral, const FOperatorSettings& InOperatorSettings) const override;
 				virtual TSharedPtr<IDataChannel, ESPMode::ThreadSafe> CreateDataChannel(const FName& InDataType, const FOperatorSettings& InOperatorSettings) const override;
 
 				virtual bool GetFrontendInputClass(const FName& InDataType, FMetasoundFrontendClass& OutClass) const override;
@@ -868,6 +869,15 @@ namespace Metasound
 				return Metasound::FLiteral(MoveTemp(ProxyArray));
 			}
 
+			TOptional<FAnyDataReference> FDataTypeRegistry::CreateDataReference(const FName& InDataType, EDataReferenceAccessType InAccessType, const FLiteral& InLiteral, const FOperatorSettings& InOperatorSettings) const
+			{
+				if (const IDataTypeRegistryEntry* Entry = FindDataTypeEntry(InDataType))
+				{
+					return Entry->CreateDataReference(InAccessType, InLiteral, InOperatorSettings);
+				}
+				return TOptional<FAnyDataReference>();
+			}
+
 			TSharedPtr<IDataChannel, ESPMode::ThreadSafe> FDataTypeRegistry::CreateDataChannel(const FName& InDataType, const FOperatorSettings& InOperatorSettings) const
 			{
 				if (const IDataTypeRegistryEntry* Entry = FindDataTypeEntry(InDataType))
diff --git a/Engine/Plugins/Runtime/Metasound/Source/MetasoundFrontend/Public/MetasoundAssetBase.h b/Engine/Plugins/Runtime/Metasound/Source/MetasoundFrontend/Public/MetasoundAssetBase.h
--- a/Engine/Plugins/Runtime/Metasound/Source/MetasoundFrontend/Public/MetasoundAssetBase.h
+++ b/Engine/Plugins/Runtime/Metasound/Source/MetasoundFrontend/Public/MetasoundAssetBase.h
@@ -195,6 +195,9 @@ protected:
 		// Current ID of graph.
 		FGuid ChangeID;
 
+		// Array of inputs which can be set for construction. 
+		TArray<FMetasoundFrontendClassInput> PublicInputs;
+
 		// Array of inputs which can be transmitted to.
 		TArray<FMetasoundFrontendClassInput> TransmittableInputs;
 
@@ -208,11 +211,18 @@ protected:
 	// Returns the cached runtime data.
 	const FRuntimeData& GetRuntimeData() const;
 
+	// Returns all public class inputs.  This is a potentially expensive.
+	// Prefer accessing public class inputs using CacheRuntimeData.
+	TArray<FMetasoundFrontendClassInput> GetPublicClassInputs() const;
+
 	// Returns all transmissible class inputs.  This is a potentially expensive.
 	// Prefer accessing transmissible class inputs using CacheRuntimeData.
 	TArray<FMetasoundFrontendClassInput> GetTransmittableClassInputs() const;
 
 private:
+	// Remove all inputs from the array which are of non-transmittable data types. 
+	void RemoveNonTransmittableClassInputs(TArray<FMetasoundFrontendClassInput>& InOutClassInputs) const;
+
 	Metasound::Frontend::FNodeRegistryKey RegistryKey;
 
 	// Cache ID is used to determine whether CachedRuntimeData is out-of-date.
diff --git a/Engine/Plugins/Runtime/Metasound/Source/MetasoundFrontend/Public/MetasoundDataTypeRegistrationMacro.h b/Engine/Plugins/Runtime/Metasound/Source/MetasoundFrontend/Public/MetasoundDataTypeRegistrationMacro.h
--- a/Engine/Plugins/Runtime/Metasound/Source/MetasoundFrontend/Public/MetasoundDataTypeRegistrationMacro.h
+++ b/Engine/Plugins/Runtime/Metasound/Source/MetasoundFrontend/Public/MetasoundDataTypeRegistrationMacro.h
@@ -526,6 +526,28 @@ namespace Metasound
 					return Audio::IProxyDataPtr(nullptr);
 				}
 
+				virtual TOptional<FAnyDataReference> CreateDataReference(EDataReferenceAccessType InAccessType, const FLiteral& InLiteral, const FOperatorSettings& InOperatorSettings) const override
+				{
+					if constexpr(bIsParsable)
+					{
+						switch (InAccessType)
+						{
+							case EDataReferenceAccessType::Read:
+								return FAnyDataReference{TDataReadReferenceLiteralFactory<TDataType>::CreateExplicitArgs(InOperatorSettings, InLiteral)};
+
+							case EDataReferenceAccessType::Write:
+								return FAnyDataReference{TDataWriteReferenceLiteralFactory<TDataType>::CreateExplicitArgs(InOperatorSettings, InLiteral)};
+
+							case EDataReferenceAccessType::Value:
+								return FAnyDataReference{TDataValueReferenceLiteralFactory<TDataType>::CreateExplicitArgs(InOperatorSettings, InLiteral)};
+
+							default:
+								break;
+						}
+					}
+					return TOptional<FAnyDataReference>();
+				}
+
 				virtual TSharedPtr<IDataChannel, ESPMode::ThreadSafe> CreateDataChannel(const FOperatorSettings& InOperatorSettings) const override
 				{
 					if constexpr (bIsParsable)
diff --git a/Engine/Plugins/Runtime/Metasound/Source/MetasoundFrontend/Public/MetasoundFrontendDataTypeRegistry.h b/Engine/Plugins/Runtime/Metasound/Source/MetasoundFrontend/Public/MetasoundFrontendDataTypeRegistry.h
--- a/Engine/Plugins/Runtime/Metasound/Source/MetasoundFrontend/Public/MetasoundFrontendDataTypeRegistry.h
+++ b/Engine/Plugins/Runtime/Metasound/Source/MetasoundFrontend/Public/MetasoundFrontendDataTypeRegistry.h
@@ -196,6 +196,9 @@ namespace Metasound
 
 			/** Create a get delayed variable node for this data type. */
 			virtual TUniquePtr<INode> CreateVariableDeferredAccessorNode(const FNodeInitData&) const = 0;
+			
+			/** Create a data reference from a literal. */
+			virtual TOptional<FAnyDataReference> CreateDataReference(EDataReferenceAccessType InAccessType, const FLiteral& InLiteral, const FOperatorSettings& InOperatorSettings) const = 0;
 
 			/** Create a proxy from a UObject. If this data type does not support
 			 * UObject proxies, return a nullptr. */
@@ -207,6 +210,7 @@ namespace Metasound
 
 			/** Clone this registry entry. */
 			virtual TUniquePtr<IDataTypeRegistryEntry> Clone() const = 0;
+
 		};
 
 		class METASOUNDFRONTEND_API IDataTypeRegistry
@@ -255,6 +259,9 @@ namespace Metasound
 			virtual FLiteral CreateLiteralFromUObject(const FName& InDataType, UObject* InObject) const = 0;
 			virtual FLiteral CreateLiteralFromUObjectArray(const FName& InDataType, const TArray<UObject*>& InObjectArray) const = 0;
 
+			/** Create a data reference of the data type given a literal. If not supported, the result will be not be set. */
+			virtual TOptional<FAnyDataReference> CreateDataReference(const FName& InDataType, EDataReferenceAccessType InAccessType, const FLiteral& InLiteral, const FOperatorSettings& InOperatorSettings) const = 0;
+
 			virtual TSharedPtr<IDataChannel, ESPMode::ThreadSafe> CreateDataChannel(const FName& InDataType, const FOperatorSettings& InOperatorSettings) const = 0;
 
 			/** Return an FMetasoundFrontendClass representing an input node of the data type. */
diff --git a/Engine/Plugins/Runtime/Metasound/Source/MetasoundGenerator/Private/MetasoundGenerator.cpp b/Engine/Plugins/Runtime/Metasound/Source/MetasoundGenerator/Private/MetasoundGenerator.cpp
--- a/Engine/Plugins/Runtime/Metasound/Source/MetasoundGenerator/Private/MetasoundGenerator.cpp
+++ b/Engine/Plugins/Runtime/Metasound/Source/MetasoundGenerator/Private/MetasoundGenerator.cpp
@@ -13,6 +13,7 @@
 #include "MetasoundSourceInterface.h"
 #include "MetasoundTrace.h"
 #include "MetasoundTrigger.h"
+#include "MetasoundVertexData.h"
 
 
 namespace Metasound
@@ -55,50 +56,37 @@ namespace Metasound
 
 		METASOUND_TRACE_CPUPROFILER_EVENT_SCOPE_TEXT(*FString::Printf(TEXT("AsyncMetaSoundBuilder::DoWork %s"), *InitParams.MetaSoundName));
 
-		// Create an instance of the new graph
-		FDataReferenceCollection DataReferenceCollection;
-		FInputVertexInterfaceData InputVertexData;
-		FBuildGraphOperatorParams BuildParams { *InitParams.Graph, InitParams.OperatorSettings, InputVertexData, InitParams.Environment };
 		FBuildResults BuildResults;
 
-		TUniquePtr<IOperator> GraphOperator = FOperatorBuilder(InitParams.BuilderSettings).BuildGraphOperator(BuildParams, BuildResults);
+		// Create an instance of the new graph
+		TUniquePtr<IOperator> GraphOperator = BuildGraphOperator(InitParams.DefaultParameters, BuildResults);
+		LogBuildErrors(BuildResults);
 
-		TUniquePtr<FGraphAnalyzer> GraphAnalyzer;
-		if (InitParams.BuilderSettings.bPopulateInternalDataReferences)
+		if (GraphOperator.IsValid())
 		{
-			const uint64 InstanceID = BuildParams.Environment.GetValue<uint64>(SourceInterface::Environment::TransmitterID);
-			GraphAnalyzer = MakeUnique<FGraphAnalyzer>(InitParams.OperatorSettings, InstanceID, MoveTemp(BuildResults.InternalDataReferences));
-		}
+			// Create graph analyzer
+			TUniquePtr<FGraphAnalyzer> GraphAnalyzer = BuildGraphAnalyzer(MoveTemp(BuildResults.InternalDataReferences));
 
-		// Log build errors
-		for (const IOperatorBuilder::FBuildErrorPtr& Error : BuildResults.Errors)
-		{
-			if (Error.IsValid())
-			{
-				UE_LOG(LogMetaSound, Warning, TEXT("MetasoundSource [%s] build error [%s] \"%s\""), *InitParams.MetaSoundName, *(Error->GetErrorType().ToString()), *(Error->GetErrorDescription().ToString()));
-			}
-		}
+			// Gather relevant input and output references
+			FVertexInterfaceData VertexData(InitParams.Graph->GetVertexInterface());
+			GraphOperator->Bind(VertexData);
 
-		if (GraphOperator.IsValid())
-		{
-			TArray<FAudioBufferReadRef> OutputBuffers;
-			FDataReferenceCollection Outputs = GraphOperator->GetOutputs();
+			// Get inputs
+			FTriggerWriteRef PlayTrigger = VertexData.GetInputs().GetOrConstructDataWriteReference<FTrigger>(SourceInterface::Inputs::OnPlay, InitParams.OperatorSettings, false);
+
+			// Get outputs
+			TArray<FAudioBufferReadRef> OutputBuffers = FindOutputAudioBuffers(VertexData);
+			FTriggerReadRef FinishTrigger = TDataReadReferenceFactory<FTrigger>::CreateExplicitArgs(InitParams.OperatorSettings, false);
 
-			// Get output audio buffers.
-			for (const FVertexName& AudioOutputName : InitParams.AudioOutputNames)
+			if (InitParams.Graph->GetVertexInterface().GetOutputInterface().Contains(SourceOneShotInterface::Outputs::OnFinished))
 			{
-				if (!Outputs.ContainsDataReadReference<FAudioBuffer>(AudioOutputName))
-				{
-					UE_LOG(LogMetaSound, Warning, TEXT("MetasoundSource [%s] does not contain audio output [%s] in output"), *InitParams.MetaSoundName, *AudioOutputName.ToString());
-				}
-				OutputBuffers.Add(Outputs.GetDataReadReferenceOrConstruct<FAudioBuffer>(AudioOutputName, InitParams.OperatorSettings));
+				// Only attempt to retrieve the on finished trigger if it exists.
+				// Attempting to retrieve a data reference from a non-existent vertex 
+				// will log an error. 
+				FinishTrigger = VertexData.GetOutputs().GetOrConstructDataReadReference<FTrigger>(SourceOneShotInterface::Outputs::OnFinished, InitParams.OperatorSettings, false);
 			}
 
-			// References must be cached before moving the operator to the InitParams
-			FDataReferenceCollection Inputs = GraphOperator->GetInputs();
-			FTriggerWriteRef PlayTrigger = Inputs.GetDataWriteReferenceOrConstruct<FTrigger>(SourceInterface::Inputs::OnPlay, InitParams.OperatorSettings, false);
-			FTriggerReadRef FinishTrigger = Outputs.GetDataReadReferenceOrConstruct<FTrigger>(SourceOneShotInterface::Outputs::OnFinished, InitParams.OperatorSettings, false);
-
+			// Set data needed for graph
 			FMetasoundGeneratorData GeneratorData
 			{
 				InitParams.OperatorSettings,
@@ -122,6 +110,89 @@ namespace Metasound
 		InitParams.Release();
 	}
 
+	TUniquePtr<IOperator> FAsyncMetaSoundBuilder::BuildGraphOperator(TArray<FAudioParameter>& InParameters, FBuildResults& OutBuildResults) const
+	{
+		using namespace Frontend;
+
+		// Set input data based on the input parameters and the input interface
+		const FInputVertexInterface& InputInterface = InitParams.Graph->GetVertexInterface().GetInputInterface();
+		FInputVertexInterfaceData InputData(InputInterface);
+
+		IDataTypeRegistry& DataRegistry = IDataTypeRegistry::Get();
+		for (FAudioParameter& Parameter : InParameters)
+		{
+			if (const FInputDataVertex* InputVertex = InputInterface.Find(Parameter.ParamName))
+			{
+				FLiteral Literal = Frontend::ConvertParameterToLiteral(MoveTemp(Parameter));
+				TOptional<FAnyDataReference> DataReference = DataRegistry.CreateDataReference(InputVertex->DataTypeName, EDataReferenceAccessType::Value, Literal, InitParams.OperatorSettings);
+
+				if (DataReference)
+				{
+					InputData.BindVertex(Parameter.ParamName, *DataReference);
+				}
+				else
+				{
+					UE_LOG(LogMetaSound, Error, TEXT("Failed to create initial input data reference from parameter %s of type %s on graph in MetaSoundSource [%s]"), *Parameter.ParamName.ToString(), *InputVertex->DataTypeName.ToString(), *InitParams.MetaSoundName);
+				}
+			}
+			else
+			{
+				UE_LOG(LogMetaSound, Error, TEXT("Failed to set initial input parameter %s on graph in MetaSoundSource [%s]"), *Parameter.ParamName.ToString(), *InitParams.MetaSoundName);
+			}
+
+		}
+
+		// Reset as elements in array have been moved.
+		InParameters.Reset();
+
+		// Create an instance of the new graph
+		FBuildGraphOperatorParams BuildParams { *InitParams.Graph, InitParams.OperatorSettings, InputData, InitParams.Environment };
+		return FOperatorBuilder(InitParams.BuilderSettings).BuildGraphOperator(BuildParams, OutBuildResults);
+	}
+
+	TUniquePtr<Frontend::FGraphAnalyzer> FAsyncMetaSoundBuilder::BuildGraphAnalyzer(TMap<FGuid, FDataReferenceCollection>&& InInternalDataReferences) const
+	{
+		using namespace Frontend;
+
+		if (InitParams.BuilderSettings.bPopulateInternalDataReferences)
+		{
+			const uint64 InstanceID = InitParams.Environment.GetValue<uint64>(SourceInterface::Environment::TransmitterID);
+			return MakeUnique<FGraphAnalyzer>(InitParams.OperatorSettings, InstanceID, MoveTemp(InInternalDataReferences));
+		}
+		return nullptr;
+	}
+
+	void FAsyncMetaSoundBuilder::LogBuildErrors(const FBuildResults& InBuildResults) const
+	{
+		// Log build errors
+		for (const IOperatorBuilder::FBuildErrorPtr& Error : InBuildResults.Errors)
+		{
+			if (Error.IsValid())
+			{
+				UE_LOG(LogMetaSound, Warning, TEXT("MetasoundSource [%s] build error [%s] \"%s\""), *InitParams.MetaSoundName, *(Error->GetErrorType().ToString()), *(Error->GetErrorDescription().ToString()));
+			}
+		}
+	}
+
+	TArray<FAudioBufferReadRef> FAsyncMetaSoundBuilder::FindOutputAudioBuffers(const FVertexInterfaceData& InVertexData) const
+	{
+		TArray<FAudioBufferReadRef> OutputBuffers;
+
+		const FOutputVertexInterfaceData& OutputVertexData = InVertexData.GetOutputs();
+
+		// Get output audio buffers.
+		for (const FVertexName& AudioOutputName : InitParams.AudioOutputNames)
+		{
+			if (!OutputVertexData.IsVertexBound(AudioOutputName))
+			{
+				UE_LOG(LogMetaSound, Warning, TEXT("MetasoundSource [%s] does not contain audio output [%s] in output"), *InitParams.MetaSoundName, *AudioOutputName.ToString());
+			}
+			OutputBuffers.Add(OutputVertexData.GetOrConstructDataReadReference<FAudioBuffer>(AudioOutputName, InitParams.OperatorSettings));
+		}
+
+		return OutputBuffers;
+	}
+
 	FMetasoundGenerator::FMetasoundGenerator(FMetasoundGeneratorInitParams&& InParams)
 		: MetasoundName(InParams.MetaSoundName)
 		, bIsFinishTriggered(false)
diff --git a/Engine/Plugins/Runtime/Metasound/Source/MetasoundGenerator/Public/MetasoundGenerator.h b/Engine/Plugins/Runtime/Metasound/Source/MetasoundGenerator/Public/MetasoundGenerator.h
--- a/Engine/Plugins/Runtime/Metasound/Source/MetasoundGenerator/Public/MetasoundGenerator.h
+++ b/Engine/Plugins/Runtime/Metasound/Source/MetasoundGenerator/Public/MetasoundGenerator.h
@@ -14,6 +14,7 @@
 #include "MetasoundRouter.h"
 #include "MetasoundTrigger.h"
 #include "MetasoundVertex.h"
+#include "MetasoundVertexData.h"
 #include "Sound/SoundGenerator.h"
 #include "Tickable.h"
 
@@ -29,6 +30,7 @@ namespace Metasound
 		FMetasoundEnvironment Environment;
 		FString MetaSoundName;
 		TArray<FVertexName> AudioOutputNames;
+		TArray<FAudioParameter> DefaultParameters;
 
 		void Release();
 	};
@@ -60,6 +62,10 @@ namespace Metasound
 		}
 
 	private:
+		TUniquePtr<IOperator> BuildGraphOperator(TArray<FAudioParameter>& InParameters, FBuildResults& OutBuildResults) const;
+		TUniquePtr<Frontend::FGraphAnalyzer> BuildGraphAnalyzer(TMap<FGuid, FDataReferenceCollection>&& InInternalDataReferences) const;
+		void LogBuildErrors(const FBuildResults& InBuildResults) const;
+		TArray<FAudioBufferReadRef> FindOutputAudioBuffers(const FVertexInterfaceData& InVertexData) const;
 
 		FMetasoundGenerator* Generator;
 		FMetasoundGeneratorInitParams InitParams;
diff --git a/Engine/Source/Runtime/AudioMixer/Private/AudioMixerSource.cpp b/Engine/Source/Runtime/AudioMixer/Private/AudioMixerSource.cpp
--- a/Engine/Source/Runtime/AudioMixer/Private/AudioMixerSource.cpp
+++ b/Engine/Source/Runtime/AudioMixer/Private/AudioMixerSource.cpp
@@ -873,10 +873,12 @@ namespace Audio
 		// Active sound instance ID is the audio component ID of active sound.
 		uint64 InstanceID = 0;
 		bool bActiveSoundIsPreviewSound = false;
+		TArray<FAudioParameter> DefaultParameters;
 		if (WaveInstance->ActiveSound)
 		{
 			InstanceID = WaveInstance->ActiveSound->GetAudioComponentID();
 			bActiveSoundIsPreviewSound = WaveInstance->ActiveSound->bIsPreviewSound;
+			DefaultParameters = MoveTemp(WaveInstance->ActiveSound->DefaultParameters);
 		}
 
 		FMixerSourceBufferInitArgs BufferInitArgs;
@@ -890,7 +892,7 @@ namespace Audio
 		BufferInitArgs.bIsSeeking = bIsSeeking;
 		BufferInitArgs.bIsPreviewSound = bActiveSoundIsPreviewSound;
 
-		MixerSourceBuffer = FMixerSourceBuffer::Create(BufferInitArgs);
+		MixerSourceBuffer = FMixerSourceBuffer::Create(BufferInitArgs, MoveTemp(DefaultParameters));
 		
 		if (!MixerSourceBuffer.IsValid())
 		{
diff --git a/Engine/Source/Runtime/AudioMixer/Private/AudioMixerSourceBuffer.cpp b/Engine/Source/Runtime/AudioMixer/Private/AudioMixerSourceBuffer.cpp
--- a/Engine/Source/Runtime/AudioMixer/Private/AudioMixerSourceBuffer.cpp
+++ b/Engine/Source/Runtime/AudioMixer/Private/AudioMixerSourceBuffer.cpp
@@ -56,16 +56,16 @@ namespace Audio
 		return CurrentSample >= NumSamples;
 	}
 
-	TSharedPtr<FMixerSourceBuffer, ESPMode::ThreadSafe> FMixerSourceBuffer::Create(FMixerSourceBufferInitArgs& InArgs)
+	TSharedPtr<FMixerSourceBuffer, ESPMode::ThreadSafe> FMixerSourceBuffer::Create(FMixerSourceBufferInitArgs& InArgs, TArray<FAudioParameter>&& InDefaultParams)
 	{
 		LLM_SCOPE(ELLMTag::AudioMixer);
 
-		TSharedPtr<FMixerSourceBuffer, ESPMode::ThreadSafe> NewSourceBuffer = MakeShareable(new FMixerSourceBuffer(InArgs));
+		TSharedPtr<FMixerSourceBuffer, ESPMode::ThreadSafe> NewSourceBuffer = MakeShareable(new FMixerSourceBuffer(InArgs, MoveTemp(InDefaultParams)));
 
 		return NewSourceBuffer;
 	}
 
-	FMixerSourceBuffer::FMixerSourceBuffer(FMixerSourceBufferInitArgs& InArgs)
+	FMixerSourceBuffer::FMixerSourceBuffer(FMixerSourceBufferInitArgs& InArgs, TArray<FAudioParameter>&& InDefaultParams)
 		: NumBuffersQeueued(0)
 		, CurrentBuffer(0)
 		, SoundWave(InArgs.SoundWave)
@@ -101,7 +101,7 @@ namespace Audio
 			InitParams.InstanceID = InArgs.InstanceID;
 			InitParams.bIsPreviewSound = InArgs.bIsPreviewSound;
 
-			SoundGenerator = InArgs.SoundWave->CreateSoundGenerator(InitParams);
+			SoundGenerator = InArgs.SoundWave->CreateSoundGenerator(InitParams, MoveTemp(InDefaultParams));
 
 			// In the case of procedural audio generation, the mixer source buffer will never "loop" -- i.e. when it's done, it's done
 			LoopingMode = LOOP_Never;
diff --git a/Engine/Source/Runtime/AudioMixer/Private/AudioMixerSourceBuffer.h b/Engine/Source/Runtime/AudioMixer/Private/AudioMixerSourceBuffer.h
--- a/Engine/Source/Runtime/AudioMixer/Private/AudioMixerSourceBuffer.h
+++ b/Engine/Source/Runtime/AudioMixer/Private/AudioMixerSourceBuffer.h
@@ -64,7 +64,7 @@ namespace Audio
 	class FMixerSourceBuffer : public ISoundWaveClient
 	{
 	public:
-		static FMixerSourceBufferPtr Create(FMixerSourceBufferInitArgs& InArgs);
+		static FMixerSourceBufferPtr Create(FMixerSourceBufferInitArgs& InArgs, TArray<FAudioParameter>&& InDefaultParams=TArray<FAudioParameter>());
 
 		~FMixerSourceBuffer();
 
@@ -115,7 +115,7 @@ namespace Audio
 		bool IsGeneratorFinished() const;
 
 	private:
-		FMixerSourceBuffer(FMixerSourceBufferInitArgs& InArgs);
+		FMixerSourceBuffer(FMixerSourceBufferInitArgs& InArgs, TArray<FAudioParameter>&& InDefaultParams);
 
 		void SubmitInitialPCMBuffers();
 		void SubmitInitialRealtimeBuffers();
diff --git a/Engine/Source/Runtime/Engine/Classes/Sound/SoundBase.h b/Engine/Source/Runtime/Engine/Classes/Sound/SoundBase.h
--- a/Engine/Source/Runtime/Engine/Classes/Sound/SoundBase.h
+++ b/Engine/Source/Runtime/Engine/Classes/Sound/SoundBase.h
@@ -310,6 +310,9 @@ public:
 
 	/** Creates a sound generator instance from this sound base. Return true if this is being implemented by a subclass. Sound generators procedurally generate audio in the audio render thread. */
 	virtual ISoundGeneratorPtr CreateSoundGenerator(const FSoundGeneratorInitParams& InParams) { return nullptr; }
+	
+	/** Creates a sound generator instance from this sound base. Return true if this is being implemented by a subclass. Sound generators procedurally generate audio in the audio render thread. */
+	virtual ISoundGeneratorPtr CreateSoundGenerator(const FSoundGeneratorInitParams& InParams, TArray<FAudioParameter>&& InDefaultParameters) { return CreateSoundGenerator(InParams); }
 
 	/** Creates a parameter transmitter for communicating with active sound instances. */
 	virtual TUniquePtr<Audio::IParameterTransmitter> CreateParameterTransmitter(Audio::FParameterTransmitterInitParams&& InParams) const;
diff --git a/Engine/Source/Runtime/Engine/Private/AudioDevice.cpp b/Engine/Source/Runtime/Engine/Private/AudioDevice.cpp
--- a/Engine/Source/Runtime/Engine/Private/AudioDevice.cpp
+++ b/Engine/Source/Runtime/Engine/Private/AudioDevice.cpp
@@ -5045,22 +5045,24 @@ void FAudioDevice::AddNewActiveSoundInternal(const FActiveSound& InNewActiveSoun
 		VirtualActiveSound.ClearAudioComponent();
 	}
 
-	// Retriggering a virtualized ActiveSound results in an
-	// already created, valid transmitter, and thus forwarded
-	// initial params should be ignored to avoid stomping existing
-	// parameter state.
+	// Retriggering a virtualized ActiveSound which already have a transmitter
+	// should not be given a new transmitter.
 	if (!ActiveSound->InstanceTransmitter.IsValid())
 	{
 		Audio::FParameterTransmitterInitParams TransmitterInitParams
 		{
 			ActiveSound->GetAudioComponentID(),
 			GetSampleRate(),
-			MoveTemp(InDefaultParams)
+			InDefaultParams
 		};
 
 		ActiveSound->InstanceTransmitter = Sound->CreateParameterTransmitter(MoveTemp(TransmitterInitParams));
 	}
 
+	// Set the default parameter used during creation of the active sound.
+	ActiveSound->DefaultParameters = MoveTemp(InDefaultParams);
+
+
 	ActiveSounds.Add(ActiveSound);
 
 	if (ActiveSound->GetAudioComponentID() > 0)
diff --git a/Engine/Source/Runtime/Engine/Public/ActiveSound.h b/Engine/Source/Runtime/Engine/Public/ActiveSound.h
--- a/Engine/Source/Runtime/Engine/Public/ActiveSound.h
+++ b/Engine/Source/Runtime/Engine/Public/ActiveSound.h
@@ -465,6 +465,9 @@ public:
 		Concurrency
 	};
 
+	/** The default initial parameters use when starting a new active sound. */
+	TArray<FAudioParameter> DefaultParameters;
+
 	/** Whether or not the sound has checked if it was occluded already. Used to initialize a sound as occluded and bypassing occlusion interpolation. */
 	uint8 bHasCheckedOcclusion:1;
 
