Remove cvar and unhide constructor pins from MetaSound editor

#rnx

diff --git a/Engine/Plugins/Runtime/Metasound/Source/MetasoundEditor/Private/MetasoundNodeDetailCustomization.cpp b/Engine/Plugins/Runtime/Metasound/Source/MetasoundEditor/Private/MetasoundNodeDetailCustomization.cpp
--- a/Engine/Plugins/Runtime/Metasound/Source/MetasoundEditor/Private/MetasoundNodeDetailCustomization.cpp
+++ b/Engine/Plugins/Runtime/Metasound/Source/MetasoundEditor/Private/MetasoundNodeDetailCustomization.cpp
@@ -46,15 +46,6 @@ namespace Metasound
 	{
 		namespace MemberCustomizationPrivate
 		{
-			static int32 EnableMetasoundConstructorPins = 0;
-
-			FAutoConsoleVariableRef CVarMetaSoundEditorEnableConstructorPins(
-				TEXT("au.MetaSounds.Editor.ConstructorPinsEnabled"),
-				EnableMetasoundConstructorPins,
-				TEXT("Whether option to make inputs and outputs constructor pins is available in editor.\n")
-				TEXT("0 (default): Disabled, !0: Enabled"),
-				ECVF_Default);
-
 			/** Set of input types which are valid registered types, but should
 			 * not show up as an input type option in the MetaSound editor. */
 			static const TSet<FName> HiddenInputTypeNames =
@@ -1189,14 +1180,11 @@ namespace Metasound
 			}
 
 			// Constructor pin 
-			if (MemberCustomizationPrivate::EnableMetasoundConstructorPins)
+			Frontend::FDataTypeRegistryInfo DataTypeInfo;
+			Frontend::IDataTypeRegistry::Get().GetDataTypeInfo(Vertex->GetDataType(), DataTypeInfo);
+			if (DataTypeInfo.bIsConstructorType)
 			{
-				Frontend::FDataTypeRegistryInfo DataTypeInfo;
-				Frontend::IDataTypeRegistry::Get().GetDataTypeInfo(Vertex->GetDataType(), DataTypeInfo);
-				if (DataTypeInfo.bIsConstructorType)
-				{
-					AddConstructorPinRow(InDetailLayout);
-				}
+				AddConstructorPinRow(InDetailLayout);
 			}
 
 			// Sort order
