Summer Cleaning
- Rename MetasoundArchetype Source file to MetasoundFrontendInterface
- Move MetaSound interface files to subdirectory.
- Break out document versioning transforms to own impl/header
#rnx

diff --git a/Engine/Plugins/Runtime/AudioModulation/Source/AudioModulation/Private/MetaSoundNodes/GetModulatorValueNode.cpp b/Engine/Plugins/Runtime/AudioModulation/Source/AudioModulation/Private/MetaSoundNodes/GetModulatorValueNode.cpp
--- a/Engine/Plugins/Runtime/AudioModulation/Source/AudioModulation/Private/MetaSoundNodes/GetModulatorValueNode.cpp
+++ b/Engine/Plugins/Runtime/AudioModulation/Source/AudioModulation/Private/MetaSoundNodes/GetModulatorValueNode.cpp
@@ -3,12 +3,12 @@
 #if WITH_AUDIOMODULATION_METASOUND_SUPPORT
 #include "AudioDefines.h"
 #include "AudioModulation.h"
+#include "Interfaces/MetasoundFrontendSourceInterface.h"
 #include "MetasoundDataFactory.h"
 #include "MetasoundExecutableOperator.h"
 #include "MetasoundFacade.h"
 #include "MetasoundNodeRegistrationMacro.h"
 #include "MetasoundPrimitives.h"
-#include "MetasoundSourceInterface.h"
 #include "SoundModulatorAsset.h"
 
 #define LOCTEXT_NAMESPACE "AudioModulationNodes"
diff --git a/Engine/Plugins/Runtime/AudioModulation/Source/AudioModulation/Private/MetaSoundNodes/MixModulatorsNode.cpp b/Engine/Plugins/Runtime/AudioModulation/Source/AudioModulation/Private/MetaSoundNodes/MixModulatorsNode.cpp
--- a/Engine/Plugins/Runtime/AudioModulation/Source/AudioModulation/Private/MetaSoundNodes/MixModulatorsNode.cpp
+++ b/Engine/Plugins/Runtime/AudioModulation/Source/AudioModulation/Private/MetaSoundNodes/MixModulatorsNode.cpp
@@ -4,12 +4,12 @@
 
 #include "AudioModulation.h"
 #include "IAudioModulation.h"
+#include "Interfaces/MetasoundFrontendSourceInterface.h"
 #include "MetasoundDataFactory.h"
 #include "MetasoundExecutableOperator.h"
 #include "MetasoundFacade.h"
 #include "MetasoundNodeRegistrationMacro.h"
 #include "MetasoundPrimitives.h"
-#include "MetasoundSourceInterface.h"
 #include "SoundModulatorAsset.h"
 
 #define LOCTEXT_NAMESPACE "AudioModulationNodes"
diff --git a/Engine/Plugins/Runtime/Metasound/Source/MetasoundEditor/Private/MetasoundEditorGraphBuilder.cpp b/Engine/Plugins/Runtime/Metasound/Source/MetasoundEditor/Private/MetasoundEditorGraphBuilder.cpp
--- a/Engine/Plugins/Runtime/Metasound/Source/MetasoundEditor/Private/MetasoundEditorGraphBuilder.cpp
+++ b/Engine/Plugins/Runtime/Metasound/Source/MetasoundEditor/Private/MetasoundEditorGraphBuilder.cpp
@@ -43,6 +43,7 @@
 #include "Toolkits/ToolkitManager.h"
 #include "WaveTable.h"
 #include "Widgets/Notifications/SNotificationList.h"
+#include "../../MetasoundFrontend/Private/MetasoundFrontendDocumentVersioning.h"
 
 #define LOCTEXT_NAMESPACE "MetaSoundEditor"
 
@@ -1476,7 +1477,7 @@ namespace Metasound
 			// 2. Set default doc version Metadata
 			FDocumentHandle DocumentHandle = MetaSoundAsset->GetDocumentHandle();
 			FMetasoundFrontendDocumentMetadata DocMetadata = DocumentHandle->GetMetadata();
-			DocMetadata.Version.Number = FVersionDocument::GetMaxVersion();
+			DocMetadata.Version.Number = FMetasoundFrontendDocument::GetMaxVersion();
 			DocumentHandle->SetMetadata(DocMetadata);
 
 			MetaSoundAsset->AddDefaultInterfaces();
diff --git a/Engine/Plugins/Runtime/Metasound/Source/MetasoundEngine/Private/MetasoundEngineArchetypes.cpp b/Engine/Plugins/Runtime/Metasound/Source/MetasoundEngine/Private/MetasoundEngineArchetypes.cpp
--- a/Engine/Plugins/Runtime/Metasound/Source/MetasoundEngine/Private/MetasoundEngineArchetypes.cpp
+++ b/Engine/Plugins/Runtime/Metasound/Source/MetasoundEngine/Private/MetasoundEngineArchetypes.cpp
@@ -3,12 +3,12 @@
 #include "MetasoundEngineArchetypes.h"
 
 #include "IAudioParameterInterfaceRegistry.h"
+#include "Interfaces/MetasoundFrontendOutputFormatInterfaces.h"
 #include "MetasoundAudioFormats.h"
 #include "MetasoundDataReference.h"
 #include "MetasoundFrontendArchetypeRegistry.h"
 #include "MetasoundFrontendTransform.h"
 #include "MetasoundInterface.h"
-#include "MetasoundOutputFormatInterfaces.h"
 #include "MetasoundRouter.h"
 #include "MetasoundSource.h"
 #include "MetasoundUObjectRegistry.h"
diff --git a/Engine/Plugins/Runtime/Metasound/Source/MetasoundEngine/Private/MetasoundInterface.cpp b/Engine/Plugins/Runtime/Metasound/Source/MetasoundEngine/Private/MetasoundInterface.cpp
--- a/Engine/Plugins/Runtime/Metasound/Source/MetasoundEngine/Private/MetasoundInterface.cpp
+++ b/Engine/Plugins/Runtime/Metasound/Source/MetasoundEngine/Private/MetasoundInterface.cpp
@@ -5,15 +5,15 @@
 #include "AudioParameterControllerInterface.h"
 #include "IAudioParameterInterfaceRegistry.h"
 #include "IAudioParameterTransmitter.h"
+#include "Interfaces/MetasoundFrontendSourceInterface.h"
 #include "MetasoundEngineArchetypes.h"
 #include "MetasoundFrontendDataTypeRegistry.h"
 #include "MetasoundFrontendDocument.h"
 #include "MetasoundFrontendTransform.h"
 #include "MetasoundLog.h"
-#include "MetasoundOutputFormatInterfaces.h"
+#include "Interfaces/MetasoundFrontendOutputFormatInterfaces.h"
 #include "MetasoundParameterTransmitter.h"
 #include "MetasoundSource.h"
-#include "MetasoundSourceInterface.h"
 #include "MetasoundUObjectRegistry.h"
 #include "Templates/SharedPointer.h"
 #include "Templates/UniquePtr.h"
diff --git a/Engine/Plugins/Runtime/Metasound/Source/MetasoundEngine/Private/MetasoundSource.cpp b/Engine/Plugins/Runtime/Metasound/Source/MetasoundEngine/Private/MetasoundSource.cpp
--- a/Engine/Plugins/Runtime/Metasound/Source/MetasoundEngine/Private/MetasoundSource.cpp
+++ b/Engine/Plugins/Runtime/Metasound/Source/MetasoundEngine/Private/MetasoundSource.cpp
@@ -5,6 +5,8 @@
 #include "AssetRegistry/AssetRegistryModule.h"
 #include "AudioDeviceManager.h"
 #include "IAudioParameterInterfaceRegistry.h"
+#include "Interfaces/MetasoundFrontendOutputFormatInterfaces.h"
+#include "Interfaces/MetasoundFrontendSourceInterface.h"
 #include "Internationalization/Text.h"
 #include "MetasoundAssetBase.h"
 #include "MetasoundAssetManager.h"
@@ -22,12 +24,10 @@
 #include "MetasoundLog.h"
 #include "MetasoundOperatorBuilderSettings.h"
 #include "MetasoundOperatorSettings.h"
-#include "MetasoundOutputFormatInterfaces.h"
 #include "MetasoundParameterTransmitter.h"
 #include "MetasoundPrimitives.h"
 #include "MetasoundReceiveNode.h"
 #include "MetasoundSettings.h"
-#include "MetasoundSourceInterface.h"
 #include "MetasoundTrace.h"
 #include "MetasoundTrigger.h"
 #include "MetasoundUObjectRegistry.h"
diff --git a/Engine/Plugins/Runtime/Metasound/Source/MetasoundEngine/Private/MetasoundWaveTableEnvelopeNode.cpp b/Engine/Plugins/Runtime/Metasound/Source/MetasoundEngine/Private/MetasoundWaveTableEnvelopeNode.cpp
--- a/Engine/Plugins/Runtime/Metasound/Source/MetasoundEngine/Private/MetasoundWaveTableEnvelopeNode.cpp
+++ b/Engine/Plugins/Runtime/Metasound/Source/MetasoundEngine/Private/MetasoundWaveTableEnvelopeNode.cpp
@@ -1,5 +1,6 @@
 // Copyright Epic Games, Inc. All Rights Reserved.
 
+#include "Interfaces/MetasoundFrontendSourceInterface.h"
 #include "MetasoundDataFactory.h"
 #include "MetasoundEngineNodesNames.h"
 #include "MetasoundEnumRegistrationMacro.h"
@@ -7,7 +8,6 @@
 #include "MetasoundFacade.h"
 #include "MetasoundNodeRegistrationMacro.h"
 #include "MetasoundPrimitives.h"
-#include "MetasoundSourceInterface.h"
 #include "MetasoundStandardNodesNames.h"
 #include "MetasoundStandardNodesCategories.h"
 #include "MetasoundTime.h"
diff --git a/Engine/Plugins/Runtime/Metasound/Source/MetasoundEngine/Private/MetasoundWaveTableGetNode.cpp b/Engine/Plugins/Runtime/Metasound/Source/MetasoundEngine/Private/MetasoundWaveTableGetNode.cpp
--- a/Engine/Plugins/Runtime/Metasound/Source/MetasoundEngine/Private/MetasoundWaveTableGetNode.cpp
+++ b/Engine/Plugins/Runtime/Metasound/Source/MetasoundEngine/Private/MetasoundWaveTableGetNode.cpp
@@ -1,13 +1,13 @@
 // Copyright Epic Games, Inc. All Rights Reserved.
 
 #include "DSP/FloatArrayMath.h"
+#include "Interfaces/MetasoundFrontendSourceInterface.h"
 #include "MetasoundDataFactory.h"
 #include "MetasoundEngineNodesNames.h"
 #include "MetasoundExecutableOperator.h"
 #include "MetasoundFacade.h"
 #include "MetasoundNodeRegistrationMacro.h"
 #include "MetasoundPrimitives.h"
-#include "MetasoundSourceInterface.h"
 #include "MetasoundTrace.h"
 #include "MetasoundWaveTable.h"
 #include "WaveTable.h"
diff --git a/Engine/Plugins/Runtime/Metasound/Source/MetasoundEngine/Private/MetasoundWaveTableOscillatorNode.cpp b/Engine/Plugins/Runtime/Metasound/Source/MetasoundEngine/Private/MetasoundWaveTableOscillatorNode.cpp
--- a/Engine/Plugins/Runtime/Metasound/Source/MetasoundEngine/Private/MetasoundWaveTableOscillatorNode.cpp
+++ b/Engine/Plugins/Runtime/Metasound/Source/MetasoundEngine/Private/MetasoundWaveTableOscillatorNode.cpp
@@ -1,18 +1,18 @@
 // Copyright Epic Games, Inc. All Rights Reserved.
 
+#include "Interfaces/MetasoundFrontendSourceInterface.h"
 #include "MetasoundDataFactory.h"
 #include "MetasoundEngineNodesNames.h"
 #include "MetasoundExecutableOperator.h"
 #include "MetasoundFacade.h"
 #include "MetasoundNodeRegistrationMacro.h"
 #include "MetasoundPrimitives.h"
-#include "WaveTableSampler.h"
-#include "MetasoundSourceInterface.h"
 #include "MetasoundStandardNodesCategories.h"
 #include "MetasoundStandardNodesNames.h"
 #include "MetasoundTrace.h"
 #include "MetasoundTrigger.h"
 #include "MetasoundWaveTable.h"
+#include "WaveTableSampler.h"
 
 #define LOCTEXT_NAMESPACE "MetasoundStandardNodes"
 
diff --git a/Engine/Plugins/Runtime/Metasound/Source/MetasoundEngineTest/Private/EngineTestMetasoundSource.cpp b/Engine/Plugins/Runtime/Metasound/Source/MetasoundEngineTest/Private/EngineTestMetasoundSource.cpp
--- a/Engine/Plugins/Runtime/Metasound/Source/MetasoundEngineTest/Private/EngineTestMetasoundSource.cpp
+++ b/Engine/Plugins/Runtime/Metasound/Source/MetasoundEngineTest/Private/EngineTestMetasoundSource.cpp
@@ -4,11 +4,11 @@
 #include "Components/AudioComponent.h"
 #include "IAudioParameterInterfaceRegistry.h"
 #include "Interfaces/IPluginManager.h"
+#include "Interfaces/MetasoundFrontendOutputFormatInterfaces.h"
+#include "Interfaces/MetasoundFrontendSourceInterface.h"
 #include "MetasoundFrontendController.h"
 #include "MetasoundFrontendSearchEngine.h"
-#include "MetasoundOutputFormatInterfaces.h"
 #include "MetasoundSource.h"
-#include "MetasoundSourceInterface.h"
 #include "Misc/AutomationTest.h"
 #include "Tests/AutomationCommon.h"
 
diff --git a/Engine/Plugins/Runtime/Metasound/Source/MetasoundFrontend/Private/Interfaces/MetasoundFrontendInterface.cpp b/Engine/Plugins/Runtime/Metasound/Source/MetasoundFrontend/Private/Interfaces/MetasoundFrontendInterface.cpp
new file mode 100644
--- /dev/null
+++ b/Engine/Plugins/Runtime/Metasound/Source/MetasoundFrontend/Private/Interfaces/MetasoundFrontendInterface.cpp
@@ -0,0 +1,281 @@
+// Copyright Epic Games, Inc. All Rights Reserved.
+#include "Interfaces/MetasoundFrontendInterface.h"
+
+#include "Algo/ForEach.h"
+#include "MetasoundFrontendDocument.h"
+
+
+namespace Metasound
+{
+	namespace Frontend
+	{
+		namespace InterfacePrivate
+		{
+			bool IsLessThanVertex(const FMetasoundFrontendVertex& VertexA, const FMetasoundFrontendVertex& VertexB)
+			{
+				// For interfaces we can ignore metadata when comparing
+				if (VertexA.Name == VertexB.Name)
+				{
+					return VertexA.TypeName.FastLess(VertexB.TypeName);
+				}
+				return VertexA.Name.FastLess(VertexB.Name);
+			}
+
+			bool IsLessThanEnvironmentVariable(const FMetasoundFrontendClassEnvironmentVariable& EnvironmentA, const FMetasoundFrontendClassEnvironmentVariable& EnvironmentB)
+			{
+				// For interfaces we can ignore metadata when comparing
+				if (EnvironmentA.Name == EnvironmentB.Name)
+				{
+					return EnvironmentA.TypeName.FastLess(EnvironmentB.TypeName);
+				}
+				return EnvironmentA.Name.FastLess(EnvironmentB.Name);
+			}
+
+			bool IsEquivalentEnvironmentVariable(const FMetasoundFrontendClassEnvironmentVariable& EnvironmentA, const FMetasoundFrontendClassEnvironmentVariable& EnvironmentB)
+			{
+				// For interfaces we can ignore display name and tooltip
+				return (EnvironmentA.Name == EnvironmentB.Name) && (EnvironmentA.TypeName == EnvironmentB.TypeName);
+			}
+
+			template<typename Type>
+			TArray<const Type*> MakePointerToArrayElements(const TArray<Type>& InArray)
+			{
+				TArray<const Type*> PointerArray;
+				for (const Type& Element : InArray)
+				{
+					PointerArray.Add(&Element);
+				}
+				return PointerArray;
+			}
+
+			template<typename FrontendTypeA, typename FrontendTypeB, typename ComparisonType>
+			void SetDifference(const TArray<FrontendTypeA>& InSetA, const TArray<FrontendTypeB>& InSetB, TArray<const FrontendTypeA*>& OutUniqueA, TArray<const FrontendTypeB*>& OutUniqueB, ComparisonType Compare)
+			{
+				TArray<const FrontendTypeA*> SortedA = MakePointerToArrayElements(InSetA);
+				SortedA.Sort(Compare);
+
+				TArray<const FrontendTypeB*> SortedB = MakePointerToArrayElements(InSetB);
+				SortedB.Sort(Compare);
+
+				int32 IndexA = 0;
+				int32 IndexB = 0;
+
+				while ((IndexA < SortedA.Num()) && (IndexB < SortedB.Num()))
+				{
+					if (Compare(*SortedA[IndexA], *SortedB[IndexB]))
+					{
+						OutUniqueA.Add(SortedA[IndexA]);
+						IndexA++;
+					}
+					else if (Compare(*SortedB[IndexB], *SortedA[IndexA]))
+					{
+						OutUniqueB.Add(SortedB[IndexB]);
+						IndexB++;
+					}
+					else
+					{
+						// Equal. Increment both indices.
+						IndexA++;
+						IndexB++;
+					}
+				}
+
+				if (IndexA < SortedA.Num())
+				{
+					OutUniqueA.Append(&SortedA[IndexA], SortedA.Num() - IndexA);
+				}
+
+				if (IndexB < SortedB.Num())
+				{
+					OutUniqueB.Append(&SortedB[IndexB], SortedB.Num() - IndexB);
+				}
+			}
+
+			template<typename FrontendTypeA, typename FrontendTypeB, typename ComparisonType>
+			bool IsSetEquivalent(const TArray<FrontendTypeA>& InSetA, const TArray<FrontendTypeB>& InSetB, ComparisonType Compare)
+			{
+				TArray<const FrontendTypeA*> UniqueA;
+				TArray<const FrontendTypeB*> UniqueB;
+
+				SetDifference(InSetA, InSetB, UniqueA, UniqueB, Compare);
+
+				return (UniqueA.Num() == 0) && (UniqueB.Num() == 0);
+			}
+			
+			template<typename FrontendTypeA, typename FrontendTypeB, typename ComparisonType>
+			bool IsSetIncluded(const TArray<FrontendTypeA>& InSubset, const TArray<FrontendTypeB>& InSuperset, ComparisonType Compare)
+			{
+				TArray<const FrontendTypeA*> UniqueSubset;
+				TArray<const FrontendTypeB*> UniqueSuperset;
+
+				SetDifference(InSubset, InSuperset, UniqueSubset, UniqueSuperset, Compare);
+
+				return UniqueSubset.Num() == 0;
+			}
+
+			template<typename FrontendTypeA, typename FrontendTypeB, typename ComparisonType>
+			int32 SetDifferenceCount(const TArray<FrontendTypeA>& InSetA, const TArray<FrontendTypeB>& InSetB, ComparisonType Compare)
+			{
+				TArray<const FrontendTypeA*> UniqueA;
+				TArray<const FrontendTypeB*> UniqueB;
+
+				SetDifference(InSetA, InSetB, UniqueA, UniqueB, Compare);
+
+				return UniqueA.Num() + UniqueB.Num();
+			}
+		} // namespace InterfacePrivate
+
+		bool IsSubsetOfInterface(const FMetasoundFrontendInterface& InSubsetInterface, const FMetasoundFrontendInterface& InSupersetInterface)
+		{
+			using namespace InterfacePrivate;
+
+			const bool bIsInputSetIncluded = IsSetIncluded(InSubsetInterface.Inputs, InSupersetInterface.Inputs, &IsLessThanVertex);
+
+			const bool bIsOutputSetIncluded = IsSetIncluded(InSubsetInterface.Outputs, InSupersetInterface.Outputs, &IsLessThanVertex);
+
+			const bool bIsEnvironmentVariableSetIncluded = IsSetIncluded(InSubsetInterface.Environment, InSupersetInterface.Environment, &IsLessThanEnvironmentVariable);
+
+			return bIsInputSetIncluded && bIsOutputSetIncluded && bIsEnvironmentVariableSetIncluded;
+		}
+
+		bool IsSubsetOfClass(const FMetasoundFrontendInterface& InSubsetInterface, const FMetasoundFrontendClass& InSupersetClass)
+		{
+			// TODO: Environment variables are ignored as they are poorly supported and classes describe which environment variables are required,
+			// not which are supported 
+			using namespace InterfacePrivate;
+
+			const bool bIsInputSetIncluded = IsSetIncluded(InSubsetInterface.Inputs, InSupersetClass.Interface.Inputs, &IsLessThanVertex);
+			const bool bIsOutputSetIncluded = IsSetIncluded(InSubsetInterface.Outputs, InSupersetClass.Interface.Outputs, &IsLessThanVertex);
+
+			return bIsInputSetIncluded && bIsOutputSetIncluded;
+		}
+
+		bool IsEquivalentInterface(const FMetasoundFrontendInterface& InInputInterface, const FMetasoundFrontendInterface& InTargetInterface)
+		{
+			using namespace InterfacePrivate;
+
+			const bool bIsInputSetEquivalent = IsSetEquivalent(InTargetInterface.Inputs, InInputInterface.Inputs, &IsLessThanVertex);
+			const bool bIsOutputSetEquivalent = IsSetEquivalent(InTargetInterface.Outputs, InInputInterface.Outputs, &IsLessThanVertex);
+			const bool bIsEnvironmentVariableSetEquivalent = IsSetEquivalent(InTargetInterface.Environment, InInputInterface.Environment, &IsLessThanEnvironmentVariable);
+
+			return bIsInputSetEquivalent && bIsOutputSetEquivalent && bIsEnvironmentVariableSetEquivalent;
+		}
+
+		bool IsEqualInterface(const FMetasoundFrontendInterface& InInputInterface, const FMetasoundFrontendInterface& InTargetInterface)
+		{
+			const bool bIsMetadataEqual = InInputInterface.Version == InTargetInterface.Version;
+			if (bIsMetadataEqual)
+			{
+				// If metadata is equal, then inputs/outputs/environment variables should also be equal. 
+				// Consider updating interface version info if this ensure is hit.
+				return ensure(IsEquivalentInterface(InInputInterface, InTargetInterface));
+			}
+
+			return false;
+		}
+
+		bool DeclaredInterfaceVersionsMatch(const FMetasoundFrontendDocument& InDocumentA, const FMetasoundFrontendDocument& InDocumentB)
+		{
+			if (InDocumentA.Interfaces.Num() != InDocumentB.Interfaces.Num())
+			{
+				return false;
+			}
+
+			return InDocumentA.Interfaces.Includes(InDocumentB.Interfaces);
+		}
+
+		int32 InputOutputDifferenceCount(const FMetasoundFrontendClass& InClass, const FMetasoundFrontendInterface& InInterface)
+		{
+			using namespace InterfacePrivate;
+
+			int32 DiffCount = SetDifferenceCount(InClass.Interface.Inputs, InInterface.Inputs, &IsLessThanVertex);
+			DiffCount += SetDifferenceCount(InClass.Interface.Outputs, InInterface.Outputs, &IsLessThanVertex);
+
+			return DiffCount;
+		}
+
+		int32 InputOutputDifferenceCount(const FMetasoundFrontendInterface& InInterfaceA, const FMetasoundFrontendInterface& InInterfaceB)
+		{
+			using namespace InterfacePrivate;
+
+			int32 DiffCount = SetDifferenceCount(InInterfaceA.Inputs, InInterfaceB.Inputs, &IsLessThanVertex);
+			DiffCount += SetDifferenceCount(InInterfaceA.Outputs, InInterfaceB.Outputs, &IsLessThanVertex);
+
+			return DiffCount;
+		}
+
+		void GatherRequiredEnvironmentVariables(const FMetasoundFrontendGraphClass& InRootGraph, const TArray<FMetasoundFrontendClass>& InDependencies, const TArray<FMetasoundFrontendGraphClass>& InSubgraphs, TArray<FMetasoundFrontendClassEnvironmentVariable>& OutEnvironmentVariables)
+		{
+			auto GatherRequiredEnvironmentVariablesFromClass = [&](const FMetasoundFrontendClass& InClass)
+			{
+				using namespace InterfacePrivate;
+
+				for (const FMetasoundFrontendClassEnvironmentVariable& EnvVar : InClass.Interface.Environment)
+				{
+					if (EnvVar.bIsRequired)
+					{
+						auto IsEquivalentEnvVar = [&](const FMetasoundFrontendClassEnvironmentVariable& OtherEnvVar)
+						{
+							return IsEquivalentEnvironmentVariable(EnvVar, OtherEnvVar);
+						};
+
+						// Basically same as TArray::AddUnique except uses the `IsEquivalentEnvironmentVariable` instead of `operator==` to test for uniqueness.
+						if (nullptr == OutEnvironmentVariables.FindByPredicate(IsEquivalentEnvVar))
+						{
+							OutEnvironmentVariables.Add(EnvVar);
+						}
+					}
+				}
+			};
+
+			GatherRequiredEnvironmentVariablesFromClass(InRootGraph);
+			Algo::ForEach(InDependencies, GatherRequiredEnvironmentVariablesFromClass);
+			Algo::ForEach(InSubgraphs, GatherRequiredEnvironmentVariablesFromClass);
+		}
+
+		static FMetasoundFrontendVersionNumber GetMaxVersion()
+		{
+			return FMetasoundFrontendVersionNumber{ 1, 10 };
+		}
+
+		const FMetasoundFrontendInterface* FindMostSimilarInterfaceSupportingEnvironment(const FMetasoundFrontendGraphClass& InRootGraph, const TArray<FMetasoundFrontendClass>& InDependencies, const TArray<FMetasoundFrontendGraphClass>& InSubgraphs, const TArray<FMetasoundFrontendInterface>& InCandidateInterfaces)
+		{
+			using namespace InterfacePrivate;
+
+			const FMetasoundFrontendInterface* Result = nullptr;
+
+			TArray<const FMetasoundFrontendInterface*> CandidateInterfaces = MakePointerToArrayElements(InCandidateInterfaces);
+
+			TArray<FMetasoundFrontendClassEnvironmentVariable> AllEnvironmentVariables;
+			GatherRequiredEnvironmentVariables(InRootGraph, InDependencies, InSubgraphs, AllEnvironmentVariables);
+
+			// Remove all interfaces which do not provide all environment variables.
+			auto DoesNotProvideAllEnvironmentVariables = [&](const FMetasoundFrontendInterface* CandidateInterface)
+			{
+				return !IsSetIncluded(AllEnvironmentVariables, CandidateInterface->Environment, &IsLessThanEnvironmentVariable);
+			};
+
+			CandidateInterfaces.RemoveAll(DoesNotProvideAllEnvironmentVariables);
+
+			// Return the interface with the least amount of differences.
+			auto DifferencesFromRootGraph = [&](const FMetasoundFrontendInterface* CandidateInterface) -> int32
+			{ 
+				return InputOutputDifferenceCount(InRootGraph, *CandidateInterface);
+			};
+
+			Algo::SortBy(CandidateInterfaces, DifferencesFromRootGraph);
+
+			if (0 == CandidateInterfaces.Num())
+			{
+				return nullptr;
+			}
+
+			return CandidateInterfaces[0];
+		}
+		
+		const FMetasoundFrontendInterface* FindMostSimilarInterfaceSupportingEnvironment(const FMetasoundFrontendDocument& InDocument, const TArray<FMetasoundFrontendInterface>& InCandidateInterfaces)
+		{
+			return FindMostSimilarInterfaceSupportingEnvironment(InDocument.RootGraph, InDocument.Dependencies, InDocument.Subgraphs, InCandidateInterfaces);
+		}
+	} // namespace Frontend
+} // namespace Metasound
diff --git a/Engine/Plugins/Runtime/Metasound/Source/MetasoundFrontend/Private/Interfaces/MetasoundFrontendOutputFormatInterfaces.cpp b/Engine/Plugins/Runtime/Metasound/Source/MetasoundFrontend/Private/Interfaces/MetasoundFrontendOutputFormatInterfaces.cpp
new file mode 100644
--- /dev/null
+++ b/Engine/Plugins/Runtime/Metasound/Source/MetasoundFrontend/Private/Interfaces/MetasoundFrontendOutputFormatInterfaces.cpp
@@ -0,0 +1,391 @@
+// Copyright Epic Games, Inc. All Rights Reserved.
+#include "Interfaces/MetasoundFrontendOutputFormatInterfaces.h"
+
+#include "IAudioParameterInterfaceRegistry.h"
+#include "MetasoundAudioBuffer.h"
+#include "MetasoundTrigger.h"
+#include "Templates/SharedPointer.h"
+#include "UObject/Class.h"
+
+
+#define LOCTEXT_NAMESPACE "MetasoundFrontend"
+namespace Metasound
+{
+	namespace Frontend
+	{
+		namespace MetasoundOutputFormatInterfacesPrivate
+		{
+			Audio::FParameterInterface::FOutput GetFrontLeftOutput(const FName& InVertexName)  
+			{
+				return {
+					LOCTEXT("OutputFormatSurroundInterface_GeneratedFrontLeftDisplayName", "Out Front Left"),
+					LOCTEXT("OutputFormatSurroundInterface_GeneratedFrontLeftDescription", "The resulting front left channel output audio from this source."),
+					GetMetasoundDataTypeName<FAudioBuffer>(),
+					InVertexName,
+					FText::GetEmpty(), // RequiredText
+					EAudioParameterType::None,
+					100 // SortOrder
+				};
+			}
+
+			Audio::FParameterInterface::FOutput GetFrontRightOutput(const FName& InVertexName)  
+			{
+				return {
+					LOCTEXT("OutputFormatSurroundInterface_GeneratedFrontRightDisplayName", "Out Front Right"),
+					LOCTEXT("OutputFormatSurroundInterface_GeneratedFrontRightDescription", "The resulting front right channel output audio from this source."),
+					GetMetasoundDataTypeName<FAudioBuffer>(),
+					InVertexName,
+					FText::GetEmpty(), // RequiredText
+					EAudioParameterType::None,
+					101
+				};
+			}
+
+			Audio::FParameterInterface::FOutput GetFrontCenterOutput(const FName& InVertexName)  
+			{
+				return {
+					LOCTEXT("OutputFormatSurroundInterface_GeneratedCenterDisplayName", "Out Front Center"),
+					LOCTEXT("OutputFormatSurroundInterface_GeneratedCenterDescription", "The resulting front center channel output audio from this source."),
+					GetMetasoundDataTypeName<FAudioBuffer>(),
+					InVertexName,
+					FText::GetEmpty(), // RequiredText
+					EAudioParameterType::None,
+					102
+				};
+			}
+
+			Audio::FParameterInterface::FOutput GetLowFrequencyOutput(const FName& InVertexName)  
+			{
+				return {
+					LOCTEXT("OutputFormatSurroundInterface_GeneratedLowFrequencyDisplayName", "Out Low Frequency"),
+					LOCTEXT("OutputFormatSurroundInterface_GeneratedLowFrequencyDescription", "The resulting low frequency channel output audio from this source."),
+					GetMetasoundDataTypeName<FAudioBuffer>(),
+					InVertexName,
+					FText::GetEmpty(), // RequiredText
+					EAudioParameterType::None,
+					103
+				};
+			}
+
+			Audio::FParameterInterface::FOutput GetSideLeftOutput(const FName& InVertexName)  
+			{
+				return {
+					LOCTEXT("OutputFormatSurroundInterface_GeneratedSideLeftDisplayName", "Out Side Left"),
+					LOCTEXT("OutputFormatSurroundInterface_GeneratedSideLeftDescription", "The resulting side left channel output audio from this source."),
+					GetMetasoundDataTypeName<FAudioBuffer>(),
+					InVertexName,
+					FText::GetEmpty(), // RequiredText
+					EAudioParameterType::None,
+					104
+				};
+			}
+
+			Audio::FParameterInterface::FOutput GetSideRightOutput(const FName& InVertexName)  
+			{
+				return {
+					LOCTEXT("OutputFormatSurroundInterface_GeneratedSideRightDisplayName", "Out Side Right"),
+					LOCTEXT("OutputFormatSurroundInterface_GeneratedSideRightDescription", "The resulting side right channel output audio from this source."),
+					GetMetasoundDataTypeName<FAudioBuffer>(),
+					InVertexName,
+					FText::GetEmpty(), // RequiredText
+					EAudioParameterType::None,
+					105
+				};
+			}
+
+			Audio::FParameterInterface::FOutput GetBackLeftOutput(const FName& InVertexName)  
+			{
+				return {
+					LOCTEXT("OutputFormatSurroundInterface_GeneratedBackLeftDisplayName", "Out Back Left"),
+					LOCTEXT("OutputFormatSurroundInterface_GeneratedBackLeftDescription", "The resulting back left channel output audio from this source."),
+					GetMetasoundDataTypeName<FAudioBuffer>(),
+					InVertexName,
+					FText::GetEmpty(), // RequiredText
+					EAudioParameterType::None,
+					106
+				};
+			}
+
+			Audio::FParameterInterface::FOutput GetBackRightOutput(const FName& InVertexName)  
+			{
+				return {
+					LOCTEXT("OutputFormatSurroundInterface_GeneratedBackRightDisplayName", "Out Back Right"),
+					LOCTEXT("OutputFormatSurroundInterface_GeneratedBackRightDescription", "The resulting back right channel output audio from this source."),
+					GetMetasoundDataTypeName<FAudioBuffer>(),
+					InVertexName,
+					FText::GetEmpty(), // RequiredText
+					EAudioParameterType::None,
+					107
+				};
+			}
+		} // namespace MetasoundOutputFormatInterfacesPrivate
+
+#define AUDIO_PARAMETER_INTERFACE_NAMESPACE "UE.OutputFormat.Mono"
+		namespace OutputFormatMonoInterface
+		{
+			const FMetasoundFrontendVersion& GetVersion()
+			{
+				static const FMetasoundFrontendVersion Version = { AUDIO_PARAMETER_INTERFACE_NAMESPACE, { 1, 0 } };
+				return Version;
+			}
+
+			namespace Outputs
+			{
+				const FName MonoOut = AUDIO_PARAMETER_INTERFACE_MEMBER_DEFINE("Audio:0");
+			}
+
+			Audio::FParameterInterfacePtr CreateInterface(const UClass& InUClass)
+			{
+				struct FInterface : public Audio::FParameterInterface
+				{
+					FInterface(const UClass& InAssetClass)
+						: FParameterInterface(OutputFormatMonoInterface::GetVersion().Name, OutputFormatMonoInterface::GetVersion().Number.ToInterfaceVersion(), InAssetClass)
+					{
+						Outputs =
+						{
+							{
+								LOCTEXT("GeneratedAudioDisplayName", "Out Mono"),
+								LOCTEXT("GeneratedAudioDescription", "The resulting mono output from this source."),
+								GetMetasoundDataTypeName<FAudioBuffer>(),
+								Outputs::MonoOut,
+								FText::GetEmpty(), // RequiredText
+								EAudioParameterType::None,
+								100
+							}
+						};
+					}
+				};
+
+				return MakeShared<FInterface>(InUClass);
+			}
+		} // namespace OutputFormatMonoInterface
+#undef AUDIO_PARAMETER_INTERFACE_NAMESPACE
+
+#define AUDIO_PARAMETER_INTERFACE_NAMESPACE "UE.OutputFormat.Stereo"
+		namespace OutputFormatStereoInterface
+		{
+			const FMetasoundFrontendVersion& GetVersion()
+			{
+				static const FMetasoundFrontendVersion Version = { AUDIO_PARAMETER_INTERFACE_NAMESPACE, { 1, 0 } };
+				return Version;
+			}
+
+			namespace Outputs
+			{
+				const FName LeftOut = AUDIO_PARAMETER_INTERFACE_MEMBER_DEFINE("Audio:0");
+				const FName RightOut = AUDIO_PARAMETER_INTERFACE_MEMBER_DEFINE("Audio:1");
+			}
+
+			Audio::FParameterInterfacePtr CreateInterface(const UClass& InUClass)
+			{
+				struct FInterface : public Audio::FParameterInterface
+				{
+					FInterface(const UClass& InAssetClass)
+						: FParameterInterface(OutputFormatStereoInterface::GetVersion().Name, OutputFormatStereoInterface::GetVersion().Number.ToInterfaceVersion(), InAssetClass)
+					{
+						Outputs =
+						{
+							{
+								LOCTEXT("OutputFormatStereoInterface_GeneratedLeftDisplayName", "Out Left"),
+								LOCTEXT("OutputFormatStereoInterface_GeneratedLeftDescription", "The resulting left channel output audio from this source."),
+								GetMetasoundDataTypeName<FAudioBuffer>(),
+								Outputs::LeftOut,
+								FText::GetEmpty(), // RequiredText
+								EAudioParameterType::None,
+								100 // Sort order
+							},
+							{
+								LOCTEXT("OutputFormatStereoInterface_GeneratedRightDisplayName", "Out Right"),
+								LOCTEXT("OutputFormatStereoInterface_GeneratedRightDescription", "The resulting right channel output audio from this source."),
+								GetMetasoundDataTypeName<FAudioBuffer>(),
+								Outputs::RightOut,
+								FText::GetEmpty(), // RequiredText
+								EAudioParameterType::None,
+								101 // Sort order
+							}
+						};
+					}
+				};
+
+				return MakeShared<FInterface>(InUClass);
+			}
+		} // namespace OutputFormatStereoInterface
+#undef AUDIO_PARAMETER_INTERFACE_NAMESPACE
+
+#define AUDIO_PARAMETER_INTERFACE_NAMESPACE "UE.OutputFormat.Quad"
+		namespace OutputFormatQuadInterface
+		{
+			const FMetasoundFrontendVersion& GetVersion()
+			{
+				static const FMetasoundFrontendVersion Version = { AUDIO_PARAMETER_INTERFACE_NAMESPACE, { 1, 0 } };
+				return Version;
+			}
+
+			namespace Outputs
+			{
+				const FName FrontLeftOut = AUDIO_PARAMETER_INTERFACE_MEMBER_DEFINE("Audio:0");
+				const FName FrontRightOut = AUDIO_PARAMETER_INTERFACE_MEMBER_DEFINE("Audio:1");
+				const FName SideLeftOut = AUDIO_PARAMETER_INTERFACE_MEMBER_DEFINE("Audio:4");
+				const FName SideRightOut = AUDIO_PARAMETER_INTERFACE_MEMBER_DEFINE("Audio:5");
+			}
+
+			Audio::FParameterInterfacePtr CreateInterface(const UClass& InUClass)
+			{
+				struct FInterface : public Audio::FParameterInterface
+				{
+					FInterface(const UClass& InAssetClass)
+						: FParameterInterface(OutputFormatQuadInterface::GetVersion().Name, OutputFormatQuadInterface::GetVersion().Number.ToInterfaceVersion(), InAssetClass)
+					{
+						using namespace MetasoundOutputFormatInterfacesPrivate;
+
+						Outputs =
+						{
+							GetFrontLeftOutput(Outputs::FrontLeftOut),
+							GetFrontRightOutput(Outputs::FrontRightOut),
+							GetSideLeftOutput(Outputs::SideLeftOut),
+							GetSideRightOutput(Outputs::SideRightOut)
+						};
+					}
+				};
+
+				return MakeShared<FInterface>(InUClass);
+			}
+		} // namespace OutputFormatQuadInterface
+#undef AUDIO_PARAMETER_INTERFACE_NAMESPACE
+
+#define AUDIO_PARAMETER_INTERFACE_NAMESPACE "UE.OutputFormat.5dot0"
+		namespace OutputFormatFiveDotZeroInterface
+		{
+			const FMetasoundFrontendVersion& GetVersion()
+			{
+				static const FMetasoundFrontendVersion Version = { AUDIO_PARAMETER_INTERFACE_NAMESPACE, { 1, 0 } };
+				return Version;
+			}
+
+			namespace Outputs
+			{
+				const FName FrontLeftOut = AUDIO_PARAMETER_INTERFACE_MEMBER_DEFINE("Audio:0");
+				const FName FrontRightOut = AUDIO_PARAMETER_INTERFACE_MEMBER_DEFINE("Audio:1");
+				const FName FrontCenterOut = AUDIO_PARAMETER_INTERFACE_MEMBER_DEFINE("Audio:2");
+				const FName SideLeftOut = AUDIO_PARAMETER_INTERFACE_MEMBER_DEFINE("Audio:4");
+				const FName SideRightOut = AUDIO_PARAMETER_INTERFACE_MEMBER_DEFINE("Audio:5");
+			}
+
+			Audio::FParameterInterfacePtr CreateInterface(const UClass& InUClass)
+			{
+				struct FInterface : public Audio::FParameterInterface
+				{
+					FInterface(const UClass& InAssetClass)
+						: FParameterInterface(OutputFormatFiveDotZeroInterface::GetVersion().Name, OutputFormatFiveDotZeroInterface::GetVersion().Number.ToInterfaceVersion(), InAssetClass)
+					{
+						using namespace MetasoundOutputFormatInterfacesPrivate;
+						Outputs =
+						{
+							GetFrontLeftOutput(Outputs::FrontLeftOut),
+							GetFrontRightOutput(Outputs::FrontRightOut),
+							GetFrontCenterOutput(Outputs::FrontCenterOut),
+							GetSideLeftOutput(Outputs::SideLeftOut),
+							GetSideRightOutput(Outputs::SideRightOut)
+						};
+					}
+				};
+
+				return MakeShared<FInterface>(InUClass);
+			}
+		} // namespace OutputFormatFiveDotZeroInterface
+#undef AUDIO_PARAMETER_INTERFACE_NAMESPACE
+
+#define AUDIO_PARAMETER_INTERFACE_NAMESPACE "UE.OutputFormat.5dot1"
+		namespace OutputFormatFiveDotOneInterface
+		{
+			const FMetasoundFrontendVersion& GetVersion()
+			{
+				static const FMetasoundFrontendVersion Version = { AUDIO_PARAMETER_INTERFACE_NAMESPACE, { 1, 0 } };
+				return Version;
+			}
+
+			namespace Outputs
+			{
+				const FName FrontLeftOut = AUDIO_PARAMETER_INTERFACE_MEMBER_DEFINE("Audio:0");
+				const FName FrontRightOut = AUDIO_PARAMETER_INTERFACE_MEMBER_DEFINE("Audio:1");
+				const FName FrontCenterOut = AUDIO_PARAMETER_INTERFACE_MEMBER_DEFINE("Audio:2");
+				const FName LowFrequencyOut = AUDIO_PARAMETER_INTERFACE_MEMBER_DEFINE("Audio:3");
+				const FName SideLeftOut = AUDIO_PARAMETER_INTERFACE_MEMBER_DEFINE("Audio:4");
+				const FName SideRightOut = AUDIO_PARAMETER_INTERFACE_MEMBER_DEFINE("Audio:5");
+			}
+
+			Audio::FParameterInterfacePtr CreateInterface(const UClass& InUClass)
+			{
+				struct FInterface : public Audio::FParameterInterface
+				{
+					FInterface(const UClass& InAssetClass)
+						: FParameterInterface(OutputFormatFiveDotOneInterface::GetVersion().Name, OutputFormatFiveDotOneInterface::GetVersion().Number.ToInterfaceVersion(), InAssetClass)
+					{
+						using namespace MetasoundOutputFormatInterfacesPrivate;
+						Outputs =
+						{
+							GetFrontLeftOutput(Outputs::FrontLeftOut),
+							GetFrontRightOutput(Outputs::FrontRightOut),
+							GetFrontCenterOutput(Outputs::FrontCenterOut),
+							GetLowFrequencyOutput(Outputs::LowFrequencyOut),
+							GetSideLeftOutput(Outputs::SideLeftOut),
+							GetSideRightOutput(Outputs::SideRightOut)
+						};
+					}
+				};
+
+				return MakeShared<FInterface>(InUClass);
+			}
+		} // namespace OutputFormatFiveDotOneInterface
+#undef AUDIO_PARAMETER_INTERFACE_NAMESPACE
+
+#define AUDIO_PARAMETER_INTERFACE_NAMESPACE "UE.OutputFormat.7dot1"
+		namespace OutputFormatSevenDotOneInterface
+		{
+			const FMetasoundFrontendVersion& GetVersion()
+			{
+				static const FMetasoundFrontendVersion Version = { AUDIO_PARAMETER_INTERFACE_NAMESPACE, { 1, 0 } };
+				return Version;
+			}
+
+			namespace Outputs
+			{
+				const FName FrontLeftOut = AUDIO_PARAMETER_INTERFACE_MEMBER_DEFINE("Audio:0");
+				const FName FrontRightOut = AUDIO_PARAMETER_INTERFACE_MEMBER_DEFINE("Audio:1");
+				const FName FrontCenterOut = AUDIO_PARAMETER_INTERFACE_MEMBER_DEFINE("Audio:2");
+				const FName LowFrequencyOut = AUDIO_PARAMETER_INTERFACE_MEMBER_DEFINE("Audio:3");
+				const FName SideLeftOut = AUDIO_PARAMETER_INTERFACE_MEMBER_DEFINE("Audio:4");
+				const FName SideRightOut = AUDIO_PARAMETER_INTERFACE_MEMBER_DEFINE("Audio:5");
+				const FName BackLeftOut = AUDIO_PARAMETER_INTERFACE_MEMBER_DEFINE("Audio:6");
+				const FName BackRightOut = AUDIO_PARAMETER_INTERFACE_MEMBER_DEFINE("Audio:7");
+			}
+
+			Audio::FParameterInterfacePtr CreateInterface(const UClass& InUClass)
+			{
+				struct FInterface : public Audio::FParameterInterface
+				{
+					FInterface(const UClass& InAssetClass)
+						: FParameterInterface(OutputFormatSevenDotOneInterface::GetVersion().Name, OutputFormatSevenDotOneInterface::GetVersion().Number.ToInterfaceVersion(), InAssetClass)
+					{
+						using namespace MetasoundOutputFormatInterfacesPrivate;
+						Outputs =
+						{
+							GetFrontLeftOutput(Outputs::FrontLeftOut),
+							GetFrontRightOutput(Outputs::FrontRightOut),
+							GetFrontCenterOutput(Outputs::FrontCenterOut),
+							GetLowFrequencyOutput(Outputs::LowFrequencyOut),
+							GetSideLeftOutput(Outputs::SideLeftOut),
+							GetSideRightOutput(Outputs::SideRightOut),
+							GetBackLeftOutput(Outputs::BackLeftOut),
+							GetBackRightOutput(Outputs::BackRightOut)
+						};
+					}
+				};
+
+				return MakeShared<FInterface>(InUClass);
+			}
+		} // namespace OutputFormatFiveDotOneInterface
+#undef AUDIO_PARAMETER_INTERFACE_NAMESPACE
+	}
+}
+#undef LOCTEXT_NAMESPACE
diff --git a/Engine/Plugins/Runtime/Metasound/Source/MetasoundFrontend/Private/Interfaces/MetasoundFrontendSourceInterface.cpp b/Engine/Plugins/Runtime/Metasound/Source/MetasoundFrontend/Private/Interfaces/MetasoundFrontendSourceInterface.cpp
new file mode 100644
--- /dev/null
+++ b/Engine/Plugins/Runtime/Metasound/Source/MetasoundFrontend/Private/Interfaces/MetasoundFrontendSourceInterface.cpp
@@ -0,0 +1,285 @@
+// Copyright Epic Games, Inc. All Rights Reserved.
+#include "Interfaces/MetasoundFrontendSourceInterface.h"
+
+#include "AudioParameter.h"
+#include "MetasoundDataReference.h"
+#include "MetasoundPrimitives.h"
+#include "MetasoundTrigger.h"
+#include "Templates/SharedPointer.h"
+#include "UObject/Class.h"
+
+#define LOCTEXT_NAMESPACE "Metasound"
+
+namespace Metasound
+{
+	namespace Frontend
+	{
+	
+#define AUDIO_PARAMETER_INTERFACE_NAMESPACE "UE.Source.OneShot"
+
+		namespace SourceOneShotInterface
+		{
+			const FMetasoundFrontendVersion& GetVersion()
+			{
+				static const FMetasoundFrontendVersion Version = { AUDIO_PARAMETER_INTERFACE_NAMESPACE, { 1, 0 } };
+				return Version;
+			}
+
+			namespace Outputs
+			{
+				const FName OnFinished = AUDIO_PARAMETER_INTERFACE_MEMBER_DEFINE("OnFinished");
+			}
+
+			Audio::FParameterInterfacePtr CreateInterface(const UClass& InUClass)
+			{
+				struct FInterface : public Audio::FParameterInterface
+				{
+					FInterface(const UClass& InAssetClass)
+						: FParameterInterface(SourceOneShotInterface::GetVersion().Name, SourceOneShotInterface::GetVersion().Number.ToInterfaceVersion(), InAssetClass)
+					{
+						Outputs =
+						{
+							{
+								LOCTEXT("OnFinished", "On Finished"),
+								LOCTEXT("OnFinishedDescription", "Trigger executed to initiate stopping the source."),
+								GetMetasoundDataTypeName<FTrigger>(),
+								Outputs::OnFinished,								
+								LOCTEXT("OnFinishedWarning", "\"On Finished\" should be connected for OneShot MetaSound sources. For sources with undefined duration (e.g. looping), remove the OneShot interface and use an audio component to avoid leaking the source."),
+							}
+						};
+					}
+				};
+
+				return MakeShared<FInterface>(InUClass);
+			}
+		} // namespace SourceOneShotInterface
+
+#undef AUDIO_PARAMETER_INTERFACE_NAMESPACE
+
+
+#define AUDIO_PARAMETER_INTERFACE_NAMESPACE "UE.Source"
+
+		namespace SourceInterfaceV1_0
+		{
+			const FMetasoundFrontendVersion& GetVersion()
+			{
+				static const FMetasoundFrontendVersion Version = { AUDIO_PARAMETER_INTERFACE_NAMESPACE, { 1, 0 }};
+				return Version;
+			}
+
+			namespace Inputs
+			{
+				const FName OnPlay = AUDIO_PARAMETER_INTERFACE_MEMBER_DEFINE("OnPlay");
+			}
+
+			namespace Outputs
+			{
+				const FName OnFinished = AUDIO_PARAMETER_INTERFACE_MEMBER_DEFINE("OnFinished");
+			}
+
+			namespace Environment
+			{
+				const FName DeviceID = AUDIO_PARAMETER_INTERFACE_MEMBER_DEFINE("AudioDeviceID");
+				const FName GraphName = AUDIO_PARAMETER_INTERFACE_MEMBER_DEFINE("GraphName");
+				const FName IsPreview = AUDIO_PARAMETER_INTERFACE_MEMBER_DEFINE("IsPreviewSound");
+				const FName SoundUniqueID = AUDIO_PARAMETER_INTERFACE_MEMBER_DEFINE("SoundUniqueID");
+				const FName TransmitterID = AUDIO_PARAMETER_INTERFACE_MEMBER_DEFINE("TransmitterID");
+			}
+
+			Audio::FParameterInterfacePtr CreateInterface(const UClass& InUClass)
+			{
+				struct FInterface : public Audio::FParameterInterface
+				{
+					FInterface(const UClass& InAssetClass)
+						: FParameterInterface(SourceInterfaceV1_0::GetVersion().Name, SourceInterfaceV1_0::GetVersion().Number.ToInterfaceVersion(), InAssetClass)
+					{
+						Inputs =
+						{
+							{
+								LOCTEXT("OnPlay", "On Play"),
+								LOCTEXT("OnPlayDescription", "Trigger executed when the source is played."),
+								GetMetasoundDataTypeName<FTrigger>(),
+								{ Inputs::OnPlay, false }
+							}
+						};
+
+						Outputs =
+						{
+							{
+								LOCTEXT("OnFinished", "On Finished"),
+								LOCTEXT("OnFinishedDescription", "Trigger executed to initiate stopping the source."),
+								GetMetasoundDataTypeName<FTrigger>(),
+								Outputs::OnFinished
+							}
+						};
+
+						Environment =
+						{
+							{
+								LOCTEXT("AudioDeviceIDDisplayName", "Audio Device ID"),
+								LOCTEXT("AudioDeviceIDDescription", "ID of AudioDevice source is played from."),
+								FName(), // TODO: Align environment data types with environment (ex. this is actually set/get as a uint32)
+								Environment::DeviceID
+							},
+							{
+								LOCTEXT("GraphNameDisplayName", "Graph Name"),
+								LOCTEXT("GraphNameDescription", "Name of source graph (for debugging/logging)."),
+								GetMetasoundDataTypeName<FString>(),
+								Environment::GraphName
+							},
+							{
+								LOCTEXT("IsPreviewSoundDisplayName", "Is Preview Sound"),
+								LOCTEXT("IsPreviewSoundDescription", "Whether source is being played as a previewed sound."),
+								GetMetasoundDataTypeName<bool>(),
+								Environment::IsPreview
+							},
+							{
+								LOCTEXT("TransmitterIDDisplayName", "Transmitter ID"),
+								LOCTEXT("TransmitterIDDescription", "ID used by Transmission System to generate a unique send address for each source instance."),
+								FName(), // TODO: Align environment data types with environment (ex. this is actually set/get as a uint64)
+								Environment::TransmitterID
+							},
+							{
+								LOCTEXT("SoundUniqueIdDisplayName", "Sound Unique ID"),
+								LOCTEXT("SoundUniqueIdDescription", "ID of unique source instance."),
+								FName(), // TODO: Align environment data types with environment (ex. this is actually set/get as a uint32)
+								Environment::SoundUniqueID
+							}
+						};
+					}
+				};
+
+				return MakeShared<FInterface>(InUClass);
+			}
+		} // namespace SourceInterfaceV1_0
+
+
+		namespace SourceInterface
+		{
+			const FMetasoundFrontendVersion& GetVersion()
+			{
+				static const FMetasoundFrontendVersion Version = { AUDIO_PARAMETER_INTERFACE_NAMESPACE, { 1, 1 } };
+				return Version;
+			}
+
+			namespace Inputs
+			{
+				const FName OnPlay = AUDIO_PARAMETER_INTERFACE_MEMBER_DEFINE("OnPlay");
+			}
+
+			namespace Environment
+			{
+				const FName DeviceID = AUDIO_PARAMETER_INTERFACE_MEMBER_DEFINE("AudioDeviceID");
+				const FName GraphName = AUDIO_PARAMETER_INTERFACE_MEMBER_DEFINE("GraphName");
+				const FName IsPreview = AUDIO_PARAMETER_INTERFACE_MEMBER_DEFINE("IsPreviewSound");
+				const FName SoundUniqueID = AUDIO_PARAMETER_INTERFACE_MEMBER_DEFINE("SoundUniqueID");
+				const FName TransmitterID = AUDIO_PARAMETER_INTERFACE_MEMBER_DEFINE("TransmitterID");
+				const FName AudioMixerNumOutputFrames = AUDIO_PARAMETER_INTERFACE_MEMBER_DEFINE("AudioMixerNumOutputFrames");
+			}
+
+			Audio::FParameterInterfacePtr CreateInterface(const UClass& InUClass)
+			{
+				struct FInterface : public Audio::FParameterInterface
+				{
+					FInterface(const UClass& InAssetClass)
+						: FParameterInterface(SourceInterface::GetVersion().Name, SourceInterface::GetVersion().Number.ToInterfaceVersion(), InAssetClass)
+					{
+						Inputs =
+						{
+							{
+								LOCTEXT("OnPlay", "On Play"),
+								LOCTEXT("OnPlayDescription", "Trigger executed when the source is played."),
+								GetMetasoundDataTypeName<FTrigger>(),
+								{ Inputs::OnPlay, false }
+							}
+						};
+
+						Environment =
+						{
+							{
+								LOCTEXT("AudioDeviceIDDisplayName1", "Audio Device ID"),
+								LOCTEXT("AudioDeviceIDDescription2", "ID of AudioDevice source is played from."),
+								FName(), // TODO: Align environment data types with environment (ex. this is actually set/get as a uint32)
+								Environment::DeviceID
+							},
+							{
+								LOCTEXT("GraphNameDisplayName", "Graph Name"),
+								LOCTEXT("AudioDeviceIDDescription3", "Name of source graph (for debugging/logging)."),
+								GetMetasoundDataTypeName<FString>(),
+								Environment::GraphName
+							},
+							{
+								LOCTEXT("IsPreviewSoundDisplayName", "Is Preview Sound"),
+								LOCTEXT("IsPreviewSoundDescription4", "Whether source is being played as a previewed sound."),
+								GetMetasoundDataTypeName<bool>(),
+								Environment::IsPreview
+							},
+							{
+								LOCTEXT("TransmitterIDDisplayName", "Transmitter ID"),
+								LOCTEXT("TransmitterIDDescription", "ID used by Transmission System to generate a unique send address for each source instance."),
+								FName(), // TODO: Align environment data types with environment (ex. this is actually set/get as a uint64)
+								Environment::TransmitterID
+							},
+							{
+								LOCTEXT("SoundUniqueDisplayName", "Sound Unique ID"),
+								LOCTEXT("SoundUniqueDescription", "ID of unique source instance."),
+								FName(), // TODO: Align environment data types with environment (ex. this is actually set/get as a uint32)
+								Environment::SoundUniqueID
+							},
+							{
+								LOCTEXT("AudioMixerOutputFramesDisplayName", "Audio Mixer Output Frames"),
+								LOCTEXT("AudioMixerOutputFramesDescription", "The number of output frames rendered per buffer in the audio mixer."),
+								FName(), 
+								Environment::AudioMixerNumOutputFrames
+							}
+						};
+					}
+				};
+
+				return MakeShared<FInterface>(InUClass);
+			}
+
+			bool FUpdateInterface::Transform(Frontend::FDocumentHandle InDocument) const
+			{
+				using namespace Frontend;
+
+				// When upgrading, we only want to add the one-shot interface if the MetaSound actually has the OnFinished trigger connected.
+				bool bIsOnFinishedConnected = false;
+				InDocument->GetRootGraph()->IterateConstNodes([&](FConstNodeHandle NodeHandle)
+				{
+					NodeHandle->IterateConstInputs([&](FConstInputHandle InputHandle)
+					{
+						if (InputHandle->GetName() == SourceInterfaceV1_0::Outputs::OnFinished)
+						{
+							bIsOnFinishedConnected = InputHandle->IsConnected();
+						}
+					});
+				}, EMetasoundFrontendClassType::Output);
+
+				const TArray<FMetasoundFrontendVersion> InterfacesToRemove
+				{
+					SourceInterfaceV1_0::GetVersion()
+				};
+
+				TArray<FMetasoundFrontendVersion> InterfacesToAdd
+				{
+					SourceInterface::GetVersion()
+				};
+
+				if (bIsOnFinishedConnected)
+				{
+					InterfacesToAdd.Add(SourceOneShotInterface::GetVersion());
+				}
+
+				FModifyRootGraphInterfaces InterfaceTransform(InterfacesToRemove, InterfacesToAdd);
+				return InterfaceTransform.Transform(InDocument);
+			}
+
+} // namespace SourceInterface
+
+#undef AUDIO_PARAMETER_INTERFACE_NAMESPACE
+
+	} // namespace Frontend
+} // namespace Metasound
+
+#undef LOCTEXT_NAMESPACE
diff --git a/Engine/Plugins/Runtime/Metasound/Source/MetasoundFrontend/Private/MetasoundArchetype.cpp b/Engine/Plugins/Runtime/Metasound/Source/MetasoundFrontend/Private/MetasoundArchetype.cpp
deleted file mode 100644
diff --git a/Engine/Plugins/Runtime/Metasound/Source/MetasoundFrontend/Private/MetasoundAssetBase.cpp b/Engine/Plugins/Runtime/Metasound/Source/MetasoundFrontend/Private/MetasoundAssetBase.cpp
--- a/Engine/Plugins/Runtime/Metasound/Source/MetasoundFrontend/Private/MetasoundAssetBase.cpp
+++ b/Engine/Plugins/Runtime/Metasound/Source/MetasoundFrontend/Private/MetasoundAssetBase.cpp
@@ -9,14 +9,15 @@
 #include "Containers/Set.h"
 #include "HAL/FileManager.h"
 #include "IAudioParameterTransmitter.h"
+#include "Interfaces/MetasoundFrontendInterface.h"
 #include "Internationalization/Text.h"
 #include "IStructSerializerBackend.h"
 #include "Logging/LogMacros.h"
-#include "MetasoundArchetype.h"
 #include "MetasoundAssetManager.h"
 #include "MetasoundFrontendArchetypeRegistry.h"
 #include "MetasoundFrontendController.h"
 #include "MetasoundFrontendDocument.h"
+#include "MetasoundFrontendDocumentVersioning.h"
 #include "MetasoundFrontendGraph.h"
 #include "MetasoundFrontendInjectReceiveNodes.h"
 #include "MetasoundFrontendNodeTemplateRegistry.h"
diff --git a/Engine/Plugins/Runtime/Metasound/Source/MetasoundFrontend/Private/MetasoundFrontendDocument.cpp b/Engine/Plugins/Runtime/Metasound/Source/MetasoundFrontend/Private/MetasoundFrontendDocument.cpp
--- a/Engine/Plugins/Runtime/Metasound/Source/MetasoundFrontend/Private/MetasoundFrontendDocument.cpp
+++ b/Engine/Plugins/Runtime/Metasound/Source/MetasoundFrontend/Private/MetasoundFrontendDocument.cpp
@@ -7,6 +7,7 @@
 #include "Algo/Transform.h"
 #include "IAudioParameterInterfaceRegistry.h"
 #include "Logging/LogMacros.h"
+#include "MetasoundFrontendDocumentVersioning.h"
 #include "MetasoundFrontend.h"
 #include "MetasoundFrontendRegistries.h"
 #include "MetasoundLog.h"
@@ -573,6 +574,11 @@ FMetasoundFrontendGraphClass::FMetasoundFrontendGraphClass()
 	Metadata.SetType(EMetasoundFrontendClassType::Graph);
 }
 
+FMetasoundFrontendVersionNumber FMetasoundFrontendDocument::GetMaxVersion()
+{
+	return Metasound::Frontend::FVersionDocument::GetMaxVersion();
+}
+
 FMetasoundFrontendDocument::FMetasoundFrontendDocument()
 {
 	RootGraph.ID = FGuid::NewGuid();
diff --git a/Engine/Plugins/Runtime/Metasound/Source/MetasoundFrontend/Private/MetasoundFrontendDocumentVersioning.cpp b/Engine/Plugins/Runtime/Metasound/Source/MetasoundFrontend/Private/MetasoundFrontendDocumentVersioning.cpp
new file mode 100644
--- /dev/null
+++ b/Engine/Plugins/Runtime/Metasound/Source/MetasoundFrontend/Private/MetasoundFrontendDocumentVersioning.cpp
@@ -0,0 +1,525 @@
+// Copyright Epic Games, Inc. All Rights Reserved.
+#include "MetasoundFrontendDocumentVersioning.h"
+
+#include "Algo/Transform.h"
+#include "Interfaces/MetasoundFrontendInterface.h"
+#include "MetasoundAccessPtr.h"
+#include "MetasoundFrontendArchetypeRegistry.h"
+#include "MetasoundFrontendDocumentController.h"
+#include "MetasoundFrontendRegistries.h"
+#include "MetasoundFrontendSearchEngine.h"
+#include "MetasoundLog.h"
+#include "MetasoundTrace.h"
+#include "Misc/App.h"
+
+
+namespace Metasound
+{
+	namespace Frontend
+	{
+		class FVersionDocumentTransform : public IDocumentTransform
+		{
+			protected:
+				virtual FMetasoundFrontendVersionNumber GetTargetVersion() const = 0;
+				virtual void TransformInternal(FDocumentHandle InDocument) const = 0;
+
+			public:
+				bool Transform(FDocumentHandle InDocument) const override
+				{
+					const FMetasoundFrontendDocumentMetadata& Metadata = InDocument->GetMetadata();
+					if (Metadata.Version.Number >= GetTargetVersion())
+					{
+						return false;
+					}
+
+					TransformInternal(InDocument);
+
+					FMetasoundFrontendDocumentMetadata NewMetadata = Metadata;
+					NewMetadata.Version.Number = GetTargetVersion();
+					InDocument->SetMetadata(NewMetadata);
+
+					return true;
+				}
+		};
+
+		/** Versions document from 1.0 to 1.1. */
+		class FVersionDocument_1_1 : public FVersionDocumentTransform
+		{
+			FName Name;
+			const FString& Path;
+
+		public:
+			FVersionDocument_1_1(FName InName, const FString& InPath)
+			: Name(InName)
+			, Path(InPath)
+			{
+			}
+
+			FMetasoundFrontendVersionNumber GetTargetVersion() const override
+			{
+				return { 1, 1 };
+			}
+
+			void TransformInternal(FDocumentHandle InDocument) const override
+			{
+#if WITH_EDITOR
+				FGraphHandle GraphHandle = InDocument->GetRootGraph();
+				TArray<FNodeHandle> FrontendNodes = GraphHandle->GetNodes();
+
+				// Before literals could be stored on node inputs directly, they were stored
+				// by creating hidden input nodes. Update the doc by finding all hidden input
+				// nodes, placing the literal value of the input node directly on the
+				// downstream node's input. Then delete the hidden input node.
+				for (FNodeHandle& NodeHandle : FrontendNodes)
+				{
+					const bool bIsHiddenNode = NodeHandle->GetNodeStyle().Display.Visibility == EMetasoundFrontendNodeStyleDisplayVisibility::Hidden;
+					const bool bIsInputNode = EMetasoundFrontendClassType::Input == NodeHandle->GetClassMetadata().GetType();
+					const bool bIsHiddenInputNode = bIsHiddenNode && bIsInputNode;
+
+					if (bIsHiddenInputNode)
+					{
+						// Get literal value from input node.
+						const FGuid VertexID = GraphHandle->GetVertexIDForInputVertex(NodeHandle->GetNodeName());
+						const FMetasoundFrontendLiteral DefaultLiteral = GraphHandle->GetDefaultInput(VertexID);
+
+						// Apply literal value to downstream node's inputs.
+						TArray<FOutputHandle> OutputHandles = NodeHandle->GetOutputs();
+						if (ensure(OutputHandles.Num() == 1))
+						{
+							FOutputHandle OutputHandle = OutputHandles[0];
+							TArray<FInputHandle> Inputs = OutputHandle->GetConnectedInputs();
+							OutputHandle->Disconnect();
+
+							for (FInputHandle& Input : Inputs)
+							{
+								if (const FMetasoundFrontendLiteral* Literal = Input->GetClassDefaultLiteral())
+								{
+									if (!Literal->IsEqual(DefaultLiteral))
+									{
+										Input->SetLiteral(DefaultLiteral);
+									}
+								}
+								else
+								{
+									Input->SetLiteral(DefaultLiteral);
+								}
+							}
+						}
+						GraphHandle->RemoveNode(*NodeHandle);
+					}
+				}
+#else
+				UE_LOG(LogMetaSound, Error, TEXT("Asset '%s' at '%s' must be saved with editor enabled in order to version document to target version '%s'."), *Name.ToString(), *Path, *GetTargetVersion().ToString());
+#endif // !WITH_EDITOR
+			}
+		};
+
+		/** Versions document from 1.1 to 1.2. */
+		class FVersionDocument_1_2 : public FVersionDocumentTransform
+		{
+		private:
+			const FName Name;
+			const FString& Path;
+
+		public:
+			FVersionDocument_1_2(const FName InName, const FString& InPath)
+				: Name(InName)
+				, Path(InPath)
+			{
+			}
+
+			FMetasoundFrontendVersionNumber GetTargetVersion() const override
+			{
+				return { 1, 2 };
+			}
+
+			void TransformInternal(FDocumentHandle InDocument) const override
+			{
+#if WITH_EDITOR
+				const FMetasoundFrontendGraphClass& GraphClass = InDocument->GetRootGraphClass();
+				FMetasoundFrontendClassMetadata Metadata = GraphClass.Metadata;
+
+				Metadata.SetClassName({ "GraphAsset", Name, *Path });
+				Metadata.SetDisplayName(FText::FromString(Name.ToString()));
+				InDocument->GetRootGraph()->SetGraphMetadata(Metadata);
+#else
+				UE_LOG(LogMetaSound, Error, TEXT("Asset '%s' at '%s' must be saved with editor enabled in order to version document to target version '%s'."), *Name.ToString(), *Path, *GetTargetVersion().ToString());
+#endif // !WITH_EDITOR
+			}
+		};
+
+		/** Versions document from 1.2 to 1.3. */
+		class FVersionDocument_1_3 : public FVersionDocumentTransform
+		{
+		public:
+			FVersionDocument_1_3()
+			{
+			}
+
+			FMetasoundFrontendVersionNumber GetTargetVersion() const override
+			{
+				return {1, 3};
+			}
+
+			void TransformInternal(FDocumentHandle InDocument) const override
+			{
+				const FMetasoundFrontendGraphClass& GraphClass = InDocument->GetRootGraphClass();
+				FMetasoundFrontendClassMetadata Metadata = GraphClass.Metadata;
+
+				Metadata.SetClassName(FMetasoundFrontendClassName { FName(), *FGuid::NewGuid().ToString(), FName() });
+				InDocument->GetRootGraph()->SetGraphMetadata(Metadata);
+			}
+		};
+
+		/** Versions document from 1.3 to 1.4. */
+		class FVersionDocument_1_4 : public FVersionDocumentTransform
+		{
+		public:
+			FVersionDocument_1_4()
+			{
+			}
+
+			FMetasoundFrontendVersionNumber GetTargetVersion() const override
+			{
+				return {1, 4};
+			}
+
+			void TransformInternal(FDocumentHandle InDocument) const override
+			{
+				check(InDocument->GetMetadata().Version.Number.Major == 1);
+				check(InDocument->GetMetadata().Version.Number.Minor == 3);
+
+				const TSet<FMetasoundFrontendVersion>& Interfaces = InDocument->GetInterfaceVersions();
+
+				// Version 1.3 did not have an "InterfaceVersion" property on the
+				// document, so any document that is being updated should start off
+				// with an "Invalid" interface version.
+				if (ensure(Interfaces.IsEmpty()))
+				{
+					// At the time when version 1.4 of the document was introduced, 
+					// these were the only available interfaces. 
+					static const FMetasoundFrontendVersion PreexistingInterfaceVersions[] = {
+						FMetasoundFrontendVersion{"MetaSound", {1, 0}},
+						FMetasoundFrontendVersion{"MonoSource", {1, 0}},
+						FMetasoundFrontendVersion{"StereoSource", {1, 0}},
+						FMetasoundFrontendVersion{"MonoSource", {1, 1}},
+						FMetasoundFrontendVersion{"StereoSource", {1, 1}}
+					};
+					static const int32 NumPreexistingInterfaceVersions = sizeof(PreexistingInterfaceVersions) / sizeof(PreexistingInterfaceVersions[0]);
+
+					TArray<FMetasoundFrontendInterface> CandidateInterfaces;
+					IInterfaceRegistry& InterfaceRegistry = IInterfaceRegistry::Get();
+					for (int32 i = 0; i < NumPreexistingInterfaceVersions; i++)
+					{
+						FMetasoundFrontendInterface Interface;
+						if (InterfaceRegistry.FindInterface(GetInterfaceRegistryKey(PreexistingInterfaceVersions[i]), Interface))
+						{
+							CandidateInterfaces.Add(Interface);
+						}
+					}
+
+					const FMetasoundFrontendGraphClass& RootGraph = InDocument->GetRootGraphClass();
+					const TArray<FMetasoundFrontendClass>& Dependencies = InDocument->GetDependencies();
+					const TArray<FMetasoundFrontendGraphClass>& Subgraphs = InDocument->GetSubgraphs();
+
+					if (const FMetasoundFrontendInterface* Interface = FindMostSimilarInterfaceSupportingEnvironment(RootGraph, Dependencies, Subgraphs, CandidateInterfaces))
+					{
+						UE_LOG(LogMetaSound, Display, TEXT("Assigned interface [InterfaceVersion:%s] to document [RootGraphClassName:%s]"),
+							*Interface->Version.ToString(), *RootGraph.Metadata.GetClassName().ToString());
+
+						InDocument->AddInterfaceVersion(Interface->Version);
+					}
+					else
+					{
+						UE_LOG(LogMetaSound, Warning, TEXT("Failed to find interface for document [RootGraphClassName:%s]"),
+							*RootGraph.Metadata.GetClassName().ToString());
+					}
+				}
+			}
+		};
+
+		/** Versions document from 1.4 to 1.5. */
+		class FVersionDocument_1_5 : public FVersionDocumentTransform
+		{
+			FName Name;
+			const FString& Path;
+
+		public:
+			FVersionDocument_1_5(FName InName, const FString& InPath)
+				: Name(InName)
+				, Path(InPath)
+			{
+			}
+
+			FMetasoundFrontendVersionNumber GetTargetVersion() const override
+			{
+				return { 1, 5 };
+			}
+
+			void TransformInternal(FDocumentHandle InDocument) const override
+			{
+#if WITH_EDITOR
+				const FMetasoundFrontendClassMetadata& Metadata = InDocument->GetRootGraphClass().Metadata;
+				const FText NewAssetName = FText::FromString(Name.ToString());
+				if (Metadata.GetDisplayName().CompareTo(NewAssetName) != 0)
+				{
+					FMetasoundFrontendClassMetadata NewMetadata = Metadata;
+					NewMetadata.SetDisplayName(NewAssetName);
+					InDocument->GetRootGraph()->SetGraphMetadata(NewMetadata);
+				}
+#else
+				UE_LOG(LogMetaSound, Error, TEXT("Asset '%s' at '%s' must be saved with editor enabled in order to version document to target version '%s'."), *Name.ToString(), *Path, *GetTargetVersion().ToString());
+#endif // !WITH_EDITOR
+			}
+		};
+
+		/** Versions document from 1.5 to 1.6. */
+		class FVersionDocument_1_6 : public FVersionDocumentTransform
+		{
+		public:
+			FVersionDocument_1_6() = default;
+
+			FMetasoundFrontendVersionNumber GetTargetVersion() const override
+			{
+				return { 1, 6 };
+			}
+
+			void TransformInternal(FDocumentHandle InDocument) const override
+			{
+				const FGuid NewAssetClassID = FGuid::NewGuid();
+				FRenameRootGraphClass::Generate(InDocument, NewAssetClassID);
+			}
+		};
+
+		/** Versions document from 1.6 to 1.7. */
+		class FVersionDocument_1_7 : public FVersionDocumentTransform
+		{
+			FName Name;
+			const FString& Path;
+
+		public:
+			FVersionDocument_1_7(FName InName, const FString& InPath)
+				: Name(InName)
+				, Path(InPath)
+			{
+			}
+
+
+			FMetasoundFrontendVersionNumber GetTargetVersion() const override
+			{
+				return { 1, 7 };
+			}
+
+			void TransformInternal(FDocumentHandle InDocument) const override
+			{
+#if WITH_EDITOR
+				auto RenameTransform = [](FNodeHandle NodeHandle)
+				{
+					// Required nodes are all (at the point of this transform) providing
+					// unique names and customized display names (ex. 'Audio' for both mono &
+					// L/R output, On Play, & 'On Finished'), so do not replace them by nulling
+					// out the guid as a name and using the converted FName of the FText DisplayName.
+					if (!NodeHandle->IsInterfaceMember())
+					{
+						const FName NewNodeName = *NodeHandle->GetDisplayName().ToString();
+						NodeHandle->IterateInputs([&](FInputHandle InputHandle)
+						{
+							InputHandle->SetName(NewNodeName);
+						});
+
+						NodeHandle->IterateOutputs([&](FOutputHandle OutputHandle)
+						{
+							OutputHandle->SetName(NewNodeName);
+						});
+
+						NodeHandle->SetDisplayName(FText());
+						NodeHandle->SetNodeName(NewNodeName);
+					}
+				};
+
+				InDocument->GetRootGraph()->IterateNodes(RenameTransform, EMetasoundFrontendClassType::Input);
+				InDocument->GetRootGraph()->IterateNodes(RenameTransform, EMetasoundFrontendClassType::Output);
+#else
+				UE_LOG(LogMetaSound, Error, TEXT("Asset '%s' at '%s' must be saved with editor enabled in order to version document to target version '%s'."), *Name.ToString(), *Path, *GetTargetVersion().ToString());
+#endif // !WITH_EDITOR
+			}
+		};
+
+		/** Versions document from 1.7 to 1.8. */
+		class FVersionDocument_1_8 : public FVersionDocumentTransform
+		{
+			FName Name;
+			const FString& Path;
+
+		public:
+			FVersionDocument_1_8(FName InName, const FString& InPath)
+				: Name(InName)
+				, Path(InPath)
+			{
+			}
+
+			FMetasoundFrontendVersionNumber GetTargetVersion() const override
+			{
+				return { 1, 8 };
+			}
+
+			void TransformInternal(FDocumentHandle InDocument) const override
+			{
+#if WITH_EDITOR
+				// Do not serialize MetaData text for dependencies as
+				// CacheRegistryData dynamically provides this.
+				InDocument->IterateDependencies([](FMetasoundFrontendClass& Dependency)
+				{
+					constexpr bool bSerializeText = false;
+					Dependency.Metadata.SetSerializeText(bSerializeText);
+
+					for (FMetasoundFrontendClassInput& Input : Dependency.Interface.Inputs)
+					{
+						Input.Metadata.SetSerializeText(false);
+					}
+
+					for (FMetasoundFrontendClassOutput& Output : Dependency.Interface.Outputs)
+					{
+						Output.Metadata.SetSerializeText(false);
+					}
+				});
+
+				const TSet<FMetasoundFrontendVersion>& InterfaceVersions = InDocument->GetInterfaceVersions();
+
+				using FNameDataTypePair = TPair<FName, FName>;
+				TSet<FNameDataTypePair> InterfaceInputs;
+				TSet<FNameDataTypePair> InterfaceOutputs;
+
+				for (const FMetasoundFrontendVersion& Version : InterfaceVersions)
+				{
+					FInterfaceRegistryKey RegistryKey = GetInterfaceRegistryKey(Version);
+					const IInterfaceRegistryEntry* Entry = IInterfaceRegistry::Get().FindInterfaceRegistryEntry(RegistryKey);
+					if (ensure(Entry))
+					{
+						const FMetasoundFrontendInterface& Interface = Entry->GetInterface();
+						Algo::Transform(Interface.Inputs, InterfaceInputs, [](const FMetasoundFrontendClassInput& Input)
+						{
+							return FNameDataTypePair(Input.Name, Input.TypeName);
+						});
+
+						Algo::Transform(Interface.Outputs, InterfaceOutputs, [](const FMetasoundFrontendClassOutput& Output)
+						{
+							return FNameDataTypePair(Output.Name, Output.TypeName);
+						});
+					}
+				}
+
+				// Only serialize MetaData text for inputs owned by the graph (not by interfaces)
+				FMetasoundFrontendGraphClass RootGraphClass = InDocument->GetRootGraphClass();
+				for (FMetasoundFrontendClassInput& Input : RootGraphClass.Interface.Inputs)
+				{
+					const bool bSerializeText = !InterfaceInputs.Contains(FNameDataTypePair(Input.Name, Input.TypeName));
+					Input.Metadata.SetSerializeText(bSerializeText);
+				}
+
+				// Only serialize MetaData text for outputs owned by the graph (not by interfaces)
+				for (FMetasoundFrontendClassOutput& Output : RootGraphClass.Interface.Outputs)
+				{
+					const bool bSerializeText = !InterfaceOutputs.Contains(FNameDataTypePair(Output.Name, Output.TypeName));
+					Output.Metadata.SetSerializeText(bSerializeText);
+				}
+
+				InDocument->SetRootGraphClass(MoveTemp(RootGraphClass));
+#else
+			UE_LOG(LogMetaSound, Error, TEXT("Asset '%s' at '%s' must be saved with editor enabled in order to version document to target version '%s'."), *Name.ToString(), *Path, *GetTargetVersion().ToString());
+#endif // !WITH_EDITOR
+			}
+		};
+
+		/** Versions document from 1.8 to 1.9. */
+		class FVersionDocument_1_9 : public FVersionDocumentTransform
+		{
+			FName Name;
+			const FString& Path;
+
+		public:
+			FVersionDocument_1_9(FName InName, const FString& InPath)
+				: Name(InName)
+				, Path(InPath)
+			{
+			}
+
+			FMetasoundFrontendVersionNumber GetTargetVersion() const override
+			{
+				return { 1, 9 };
+			}
+
+			void TransformInternal(FDocumentHandle InDocument) const override
+			{
+#if WITH_EDITOR
+				// Display name text is no longer copied at this versioning point for assets
+				// from the asset's FName to avoid FText warnings regarding generation from
+				// an FString.  It also avoids desync if asset gets moved.
+				FMetasoundFrontendGraphClass RootGraphClass = InDocument->GetRootGraphClass();
+				RootGraphClass.Metadata.SetDisplayName(FText());
+				InDocument->SetRootGraphClass(MoveTemp(RootGraphClass));
+#else
+				UE_LOG(LogMetaSound, Error, TEXT("Asset '%s' at '%s' must be saved with editor enabled in order to version document to target version '%s'."), *Name.ToString(), *Path, *GetTargetVersion().ToString());
+#endif // !WITH_EDITOR
+			}
+		};
+
+		/** Versions document from 1.9 to 1.10. */
+		class FVersionDocument_1_10 : public FVersionDocumentTransform
+		{
+		public:
+			FVersionDocument_1_10() = default;
+
+			FMetasoundFrontendVersionNumber GetTargetVersion() const override
+			{
+				return { 1, 10 };
+			}
+
+			void TransformInternal(FDocumentHandle InDocument) const override
+			{
+				FMetasoundFrontendGraphClass Class = InDocument->GetRootGraphClass();
+				FMetasoundFrontendGraphClassPresetOptions PresetOptions = Class.PresetOptions;
+				Class.PresetOptions.bIsPreset = Class.Metadata.GetAndClearAutoUpdateManagesInterface_Deprecated();
+				InDocument->SetRootGraphClass(MoveTemp(Class));
+			}
+		};
+
+		FVersionDocument::FVersionDocument(FName InName, const FString& InPath)
+			: Name(InName)
+			, Path(InPath)
+		{
+		}
+
+		bool FVersionDocument::Transform(FDocumentHandle InDocument) const
+		{
+			if (!ensure(InDocument->IsValid()))
+			{
+				return false;
+			}
+
+			bool bWasUpdated = false;
+
+			const FMetasoundFrontendVersionNumber InitVersionNumber = InDocument->GetMetadata().Version.Number;
+
+			// Add additional transforms here after defining them above, example below.
+			bWasUpdated |= FVersionDocument_1_1(Name, Path).Transform(InDocument);
+			bWasUpdated |= FVersionDocument_1_2(Name, Path).Transform(InDocument);
+			bWasUpdated |= FVersionDocument_1_3().Transform(InDocument);
+			bWasUpdated |= FVersionDocument_1_4().Transform(InDocument);
+			bWasUpdated |= FVersionDocument_1_5(Name, Path).Transform(InDocument);
+			bWasUpdated |= FVersionDocument_1_6().Transform(InDocument);
+			bWasUpdated |= FVersionDocument_1_7(Name, Path).Transform(InDocument);
+			bWasUpdated |= FVersionDocument_1_8(Name, Path).Transform(InDocument);
+			bWasUpdated |= FVersionDocument_1_9(Name, Path).Transform(InDocument);
+			bWasUpdated |= FVersionDocument_1_10().Transform(InDocument);
+
+			if (bWasUpdated)
+			{
+				const FMetasoundFrontendVersionNumber NewVersionNumber = InDocument->GetMetadata().Version.Number;
+				UE_LOG(LogMetaSound, Display, TEXT("MetaSound at '%s' Document Versioned: '%s' --> '%s'"), *Path, *InitVersionNumber.ToString(), *NewVersionNumber.ToString());
+			}
+
+			return bWasUpdated;
+		}
+	} // namespace Frontend
+} // namespace Metasound
diff --git a/Engine/Plugins/Runtime/Metasound/Source/MetasoundFrontend/Private/MetasoundFrontendDocumentVersioning.h b/Engine/Plugins/Runtime/Metasound/Source/MetasoundFrontend/Private/MetasoundFrontendDocumentVersioning.h
new file mode 100644
--- /dev/null
+++ b/Engine/Plugins/Runtime/Metasound/Source/MetasoundFrontend/Private/MetasoundFrontendDocumentVersioning.h
@@ -0,0 +1,33 @@
+// Copyright Epic Games, Inc. All Rights Reserved.
+#pragma once
+
+#include "Containers/UnrealString.h"
+#include "MetasoundFrontendController.h"
+#include "MetasoundFrontendDocumentVersioning.h"
+#include "MetasoundFrontendTransform.h"
+#include "MetasoundFrontendDocument.h"
+#include "UObject/NameTypes.h"
+
+
+namespace Metasound
+{
+	namespace Frontend
+	{
+		/** Base class for versioning a document. */
+		class METASOUNDFRONTEND_API FVersionDocument : public IDocumentTransform
+		{
+			const FName Name;
+			const FString& Path;
+
+		public:
+			static FMetasoundFrontendVersionNumber GetMaxVersion()
+			{
+				return FMetasoundFrontendVersionNumber { 1, 10 };
+			}
+
+			FVersionDocument(FName InName, const FString& InPath);
+
+			bool Transform(FDocumentHandle InDocument) const override;
+		};
+	} // namespace Frontend
+} // namespace Metasound
diff --git a/Engine/Plugins/Runtime/Metasound/Source/MetasoundFrontend/Private/MetasoundFrontendTransform.cpp b/Engine/Plugins/Runtime/Metasound/Source/MetasoundFrontend/Private/MetasoundFrontendTransform.cpp
--- a/Engine/Plugins/Runtime/Metasound/Source/MetasoundFrontend/Private/MetasoundFrontendTransform.cpp
+++ b/Engine/Plugins/Runtime/Metasound/Source/MetasoundFrontend/Private/MetasoundFrontendTransform.cpp
@@ -2,8 +2,8 @@
 #include "MetasoundFrontendTransform.h"
 
 #include "Algo/Transform.h"
+#include "Interfaces/MetasoundFrontendInterface.h"
 #include "MetasoundAccessPtr.h"
-#include "MetasoundArchetype.h"
 #include "MetasoundAssetBase.h"
 #include "MetasoundFrontendArchetypeRegistry.h"
 #include "MetasoundFrontendDocument.h"
@@ -947,511 +947,5 @@ namespace Metasound
 
 			return true;
 		}
-
-		// Versioning Transforms
-		class FVersionDocumentTransform : public IDocumentTransform
-		{
-			protected:
-				virtual FMetasoundFrontendVersionNumber GetTargetVersion() const = 0;
-				virtual void TransformInternal(FDocumentHandle InDocument) const = 0;
-
-			public:
-				bool Transform(FDocumentHandle InDocument) const override
-				{
-					const FMetasoundFrontendDocumentMetadata& Metadata = InDocument->GetMetadata();
-					if (Metadata.Version.Number >= GetTargetVersion())
-					{
-						return false;
-					}
-
-					TransformInternal(InDocument);
-
-					FMetasoundFrontendDocumentMetadata NewMetadata = Metadata;
-					NewMetadata.Version.Number = GetTargetVersion();
-					InDocument->SetMetadata(NewMetadata);
-
-					return true;
-				}
-		};
-
-		/** Versions document from 1.0 to 1.1. */
-		class FVersionDocument_1_1 : public FVersionDocumentTransform
-		{
-			FName Name;
-			const FString& Path;
-
-		public:
-			FVersionDocument_1_1(FName InName, const FString& InPath)
-			: Name(InName)
-			, Path(InPath)
-			{
-			}
-
-			FMetasoundFrontendVersionNumber GetTargetVersion() const override
-			{
-				return { 1, 1 };
-			}
-
-			void TransformInternal(FDocumentHandle InDocument) const override
-			{
-#if WITH_EDITOR
-				FGraphHandle GraphHandle = InDocument->GetRootGraph();
-				TArray<FNodeHandle> FrontendNodes = GraphHandle->GetNodes();
-
-				// Before literals could be stored on node inputs directly, they were stored
-				// by creating hidden input nodes. Update the doc by finding all hidden input
-				// nodes, placing the literal value of the input node directly on the
-				// downstream node's input. Then delete the hidden input node.
-				for (FNodeHandle& NodeHandle : FrontendNodes)
-				{
-					const bool bIsHiddenNode = NodeHandle->GetNodeStyle().Display.Visibility == EMetasoundFrontendNodeStyleDisplayVisibility::Hidden;
-					const bool bIsInputNode = EMetasoundFrontendClassType::Input == NodeHandle->GetClassMetadata().GetType();
-					const bool bIsHiddenInputNode = bIsHiddenNode && bIsInputNode;
-
-					if (bIsHiddenInputNode)
-					{
-						// Get literal value from input node.
-						const FGuid VertexID = GraphHandle->GetVertexIDForInputVertex(NodeHandle->GetNodeName());
-						const FMetasoundFrontendLiteral DefaultLiteral = GraphHandle->GetDefaultInput(VertexID);
-
-						// Apply literal value to downstream node's inputs.
-						TArray<FOutputHandle> OutputHandles = NodeHandle->GetOutputs();
-						if (ensure(OutputHandles.Num() == 1))
-						{
-							FOutputHandle OutputHandle = OutputHandles[0];
-							TArray<FInputHandle> Inputs = OutputHandle->GetConnectedInputs();
-							OutputHandle->Disconnect();
-
-							for (FInputHandle& Input : Inputs)
-							{
-								if (const FMetasoundFrontendLiteral* Literal = Input->GetClassDefaultLiteral())
-								{
-									if (!Literal->IsEqual(DefaultLiteral))
-									{
-										Input->SetLiteral(DefaultLiteral);
-									}
-								}
-								else
-								{
-									Input->SetLiteral(DefaultLiteral);
-								}
-							}
-						}
-						GraphHandle->RemoveNode(*NodeHandle);
-					}
-				}
-#else
-				UE_LOG(LogMetaSound, Error, TEXT("Asset '%s' at '%s' must be saved with editor enabled in order to version document to target version '%s'."), *Name.ToString(), *Path, *GetTargetVersion().ToString());
-#endif // !WITH_EDITOR
-			}
-		};
-
-		/** Versions document from 1.1 to 1.2. */
-		class FVersionDocument_1_2 : public FVersionDocumentTransform
-		{
-		private:
-			const FName Name;
-			const FString& Path;
-
-		public:
-			FVersionDocument_1_2(const FName InName, const FString& InPath)
-				: Name(InName)
-				, Path(InPath)
-			{
-			}
-
-			FMetasoundFrontendVersionNumber GetTargetVersion() const override
-			{
-				return { 1, 2 };
-			}
-
-			void TransformInternal(FDocumentHandle InDocument) const override
-			{
-#if WITH_EDITOR
-				const FMetasoundFrontendGraphClass& GraphClass = InDocument->GetRootGraphClass();
-				FMetasoundFrontendClassMetadata Metadata = GraphClass.Metadata;
-
-				Metadata.SetClassName({ "GraphAsset", Name, *Path });
-				Metadata.SetDisplayName(FText::FromString(Name.ToString()));
-				InDocument->GetRootGraph()->SetGraphMetadata(Metadata);
-#else
-				UE_LOG(LogMetaSound, Error, TEXT("Asset '%s' at '%s' must be saved with editor enabled in order to version document to target version '%s'."), *Name.ToString(), *Path, *GetTargetVersion().ToString());
-#endif // !WITH_EDITOR
-			}
-		};
-
-		/** Versions document from 1.2 to 1.3. */
-		class FVersionDocument_1_3 : public FVersionDocumentTransform
-		{
-		public:
-			FVersionDocument_1_3()
-			{
-			}
-
-			FMetasoundFrontendVersionNumber GetTargetVersion() const override
-			{
-				return {1, 3};
-			}
-
-			void TransformInternal(FDocumentHandle InDocument) const override
-			{
-				const FMetasoundFrontendGraphClass& GraphClass = InDocument->GetRootGraphClass();
-				FMetasoundFrontendClassMetadata Metadata = GraphClass.Metadata;
-
-				Metadata.SetClassName(FMetasoundFrontendClassName { FName(), *FGuid::NewGuid().ToString(), FName() });
-				InDocument->GetRootGraph()->SetGraphMetadata(Metadata);
-			}
-		};
-
-		/** Versions document from 1.3 to 1.4. */
-		class FVersionDocument_1_4 : public FVersionDocumentTransform
-		{
-		public:
-			FVersionDocument_1_4()
-			{
-			}
-
-			FMetasoundFrontendVersionNumber GetTargetVersion() const override
-			{
-				return {1, 4};
-			}
-
-			void TransformInternal(FDocumentHandle InDocument) const override
-			{
-				check(InDocument->GetMetadata().Version.Number.Major == 1);
-				check(InDocument->GetMetadata().Version.Number.Minor == 3);
-
-				const TSet<FMetasoundFrontendVersion>& Interfaces = InDocument->GetInterfaceVersions();
-
-				// Version 1.3 did not have an "InterfaceVersion" property on the
-				// document, so any document that is being updated should start off
-				// with an "Invalid" interface version.
-				if (ensure(Interfaces.IsEmpty()))
-				{
-					// At the time when version 1.4 of the document was introduced, 
-					// these were the only available interfaces. 
-					static const FMetasoundFrontendVersion PreexistingInterfaceVersions[] = {
-						FMetasoundFrontendVersion{"MetaSound", {1, 0}},
-						FMetasoundFrontendVersion{"MonoSource", {1, 0}},
-						FMetasoundFrontendVersion{"StereoSource", {1, 0}},
-						FMetasoundFrontendVersion{"MonoSource", {1, 1}},
-						FMetasoundFrontendVersion{"StereoSource", {1, 1}}
-					};
-					static const int32 NumPreexistingInterfaceVersions = sizeof(PreexistingInterfaceVersions) / sizeof(PreexistingInterfaceVersions[0]);
-
-					TArray<FMetasoundFrontendInterface> CandidateInterfaces;
-					IInterfaceRegistry& InterfaceRegistry = IInterfaceRegistry::Get();
-					for (int32 i = 0; i < NumPreexistingInterfaceVersions; i++)
-					{
-						FMetasoundFrontendInterface Interface;
-						if (InterfaceRegistry.FindInterface(GetInterfaceRegistryKey(PreexistingInterfaceVersions[i]), Interface))
-						{
-							CandidateInterfaces.Add(Interface);
-						}
-					}
-
-					const FMetasoundFrontendGraphClass& RootGraph = InDocument->GetRootGraphClass();
-					const TArray<FMetasoundFrontendClass>& Dependencies = InDocument->GetDependencies();
-					const TArray<FMetasoundFrontendGraphClass>& Subgraphs = InDocument->GetSubgraphs();
-
-					if (const FMetasoundFrontendInterface* Interface = FindMostSimilarInterfaceSupportingEnvironment(RootGraph, Dependencies, Subgraphs, CandidateInterfaces))
-					{
-						UE_LOG(LogMetaSound, Display, TEXT("Assigned interface [InterfaceVersion:%s] to document [RootGraphClassName:%s]"),
-							*Interface->Version.ToString(), *RootGraph.Metadata.GetClassName().ToString());
-
-						InDocument->AddInterfaceVersion(Interface->Version);
-					}
-					else
-					{
-						UE_LOG(LogMetaSound, Warning, TEXT("Failed to find interface for document [RootGraphClassName:%s]"),
-							*RootGraph.Metadata.GetClassName().ToString());
-					}
-				}
-			}
-		};
-
-		/** Versions document from 1.4 to 1.5. */
-		class FVersionDocument_1_5 : public FVersionDocumentTransform
-		{
-			FName Name;
-			const FString& Path;
-
-		public:
-			FVersionDocument_1_5(FName InName, const FString& InPath)
-				: Name(InName)
-				, Path(InPath)
-			{
-			}
-
-			FMetasoundFrontendVersionNumber GetTargetVersion() const override
-			{
-				return { 1, 5 };
-			}
-
-			void TransformInternal(FDocumentHandle InDocument) const override
-			{
-#if WITH_EDITOR
-				const FMetasoundFrontendClassMetadata& Metadata = InDocument->GetRootGraphClass().Metadata;
-				const FText NewAssetName = FText::FromString(Name.ToString());
-				if (Metadata.GetDisplayName().CompareTo(NewAssetName) != 0)
-				{
-					FMetasoundFrontendClassMetadata NewMetadata = Metadata;
-					NewMetadata.SetDisplayName(NewAssetName);
-					InDocument->GetRootGraph()->SetGraphMetadata(NewMetadata);
-				}
-#else
-				UE_LOG(LogMetaSound, Error, TEXT("Asset '%s' at '%s' must be saved with editor enabled in order to version document to target version '%s'."), *Name.ToString(), *Path, *GetTargetVersion().ToString());
-#endif // !WITH_EDITOR
-			}
-		};
-
-		/** Versions document from 1.5 to 1.6. */
-		class FVersionDocument_1_6 : public FVersionDocumentTransform
-		{
-		public:
-			FVersionDocument_1_6() = default;
-
-			FMetasoundFrontendVersionNumber GetTargetVersion() const override
-			{
-				return { 1, 6 };
-			}
-
-			void TransformInternal(FDocumentHandle InDocument) const override
-			{
-				const FGuid NewAssetClassID = FGuid::NewGuid();
-				FRenameRootGraphClass::Generate(InDocument, NewAssetClassID);
-			}
-		};
-
-		/** Versions document from 1.6 to 1.7. */
-		class FVersionDocument_1_7 : public FVersionDocumentTransform
-		{
-			FName Name;
-			const FString& Path;
-
-		public:
-			FVersionDocument_1_7(FName InName, const FString& InPath)
-				: Name(InName)
-				, Path(InPath)
-			{
-			}
-
-
-			FMetasoundFrontendVersionNumber GetTargetVersion() const override
-			{
-				return { 1, 7 };
-			}
-
-			void TransformInternal(FDocumentHandle InDocument) const override
-			{
-#if WITH_EDITOR
-				auto RenameTransform = [](FNodeHandle NodeHandle)
-				{
-					// Required nodes are all (at the point of this transform) providing
-					// unique names and customized display names (ex. 'Audio' for both mono &
-					// L/R output, On Play, & 'On Finished'), so do not replace them by nulling
-					// out the guid as a name and using the converted FName of the FText DisplayName.
-					if (!NodeHandle->IsInterfaceMember())
-					{
-						const FName NewNodeName = *NodeHandle->GetDisplayName().ToString();
-						NodeHandle->IterateInputs([&](FInputHandle InputHandle)
-						{
-							InputHandle->SetName(NewNodeName);
-						});
-
-						NodeHandle->IterateOutputs([&](FOutputHandle OutputHandle)
-						{
-							OutputHandle->SetName(NewNodeName);
-						});
-
-						NodeHandle->SetDisplayName(FText());
-						NodeHandle->SetNodeName(NewNodeName);
-					}
-				};
-
-				InDocument->GetRootGraph()->IterateNodes(RenameTransform, EMetasoundFrontendClassType::Input);
-				InDocument->GetRootGraph()->IterateNodes(RenameTransform, EMetasoundFrontendClassType::Output);
-#else
-				UE_LOG(LogMetaSound, Error, TEXT("Asset '%s' at '%s' must be saved with editor enabled in order to version document to target version '%s'."), *Name.ToString(), *Path, *GetTargetVersion().ToString());
-#endif // !WITH_EDITOR
-			}
-		};
-
-		/** Versions document from 1.7 to 1.8. */
-		class FVersionDocument_1_8 : public FVersionDocumentTransform
-		{
-			FName Name;
-			const FString& Path;
-
-		public:
-			FVersionDocument_1_8(FName InName, const FString& InPath)
-				: Name(InName)
-				, Path(InPath)
-			{
-			}
-
-			FMetasoundFrontendVersionNumber GetTargetVersion() const override
-			{
-				return { 1, 8 };
-			}
-
-			void TransformInternal(FDocumentHandle InDocument) const override
-			{
-#if WITH_EDITOR
-				// Do not serialize MetaData text for dependencies as
-				// CacheRegistryData dynamically provides this.
-				InDocument->IterateDependencies([](FMetasoundFrontendClass& Dependency)
-				{
-					constexpr bool bSerializeText = false;
-					Dependency.Metadata.SetSerializeText(bSerializeText);
-
-					for (FMetasoundFrontendClassInput& Input : Dependency.Interface.Inputs)
-					{
-						Input.Metadata.SetSerializeText(false);
-					}
-
-					for (FMetasoundFrontendClassOutput& Output : Dependency.Interface.Outputs)
-					{
-						Output.Metadata.SetSerializeText(false);
-					}
-				});
-
-				const TSet<FMetasoundFrontendVersion>& InterfaceVersions = InDocument->GetInterfaceVersions();
-
-				using FNameDataTypePair = TPair<FName, FName>;
-				TSet<FNameDataTypePair> InterfaceInputs;
-				TSet<FNameDataTypePair> InterfaceOutputs;
-
-				for (const FMetasoundFrontendVersion& Version : InterfaceVersions)
-				{
-					FInterfaceRegistryKey RegistryKey = GetInterfaceRegistryKey(Version);
-					const IInterfaceRegistryEntry* Entry = IInterfaceRegistry::Get().FindInterfaceRegistryEntry(RegistryKey);
-					if (ensure(Entry))
-					{
-						const FMetasoundFrontendInterface& Interface = Entry->GetInterface();
-						Algo::Transform(Interface.Inputs, InterfaceInputs, [](const FMetasoundFrontendClassInput& Input)
-						{
-							return FNameDataTypePair(Input.Name, Input.TypeName);
-						});
-
-						Algo::Transform(Interface.Outputs, InterfaceOutputs, [](const FMetasoundFrontendClassOutput& Output)
-						{
-							return FNameDataTypePair(Output.Name, Output.TypeName);
-						});
-					}
-				}
-
-				// Only serialize MetaData text for inputs owned by the graph (not by interfaces)
-				FMetasoundFrontendGraphClass RootGraphClass = InDocument->GetRootGraphClass();
-				for (FMetasoundFrontendClassInput& Input : RootGraphClass.Interface.Inputs)
-				{
-					const bool bSerializeText = !InterfaceInputs.Contains(FNameDataTypePair(Input.Name, Input.TypeName));
-					Input.Metadata.SetSerializeText(bSerializeText);
-				}
-
-				// Only serialize MetaData text for outputs owned by the graph (not by interfaces)
-				for (FMetasoundFrontendClassOutput& Output : RootGraphClass.Interface.Outputs)
-				{
-					const bool bSerializeText = !InterfaceOutputs.Contains(FNameDataTypePair(Output.Name, Output.TypeName));
-					Output.Metadata.SetSerializeText(bSerializeText);
-				}
-
-				InDocument->SetRootGraphClass(MoveTemp(RootGraphClass));
-#else
-			UE_LOG(LogMetaSound, Error, TEXT("Asset '%s' at '%s' must be saved with editor enabled in order to version document to target version '%s'."), *Name.ToString(), *Path, *GetTargetVersion().ToString());
-#endif // !WITH_EDITOR
-			}
-		};
-
-		/** Versions document from 1.8 to 1.9. */
-		class FVersionDocument_1_9 : public FVersionDocumentTransform
-		{
-			FName Name;
-			const FString& Path;
-
-		public:
-			FVersionDocument_1_9(FName InName, const FString& InPath)
-				: Name(InName)
-				, Path(InPath)
-			{
-			}
-
-			FMetasoundFrontendVersionNumber GetTargetVersion() const override
-			{
-				return { 1, 9 };
-			}
-
-			void TransformInternal(FDocumentHandle InDocument) const override
-			{
-#if WITH_EDITOR
-				// Display name text is no longer copied at this versioning point for assets
-				// from the asset's FName to avoid FText warnings regarding generation from
-				// an FString.  It also avoids desync if asset gets moved.
-				FMetasoundFrontendGraphClass RootGraphClass = InDocument->GetRootGraphClass();
-				RootGraphClass.Metadata.SetDisplayName(FText());
-				InDocument->SetRootGraphClass(MoveTemp(RootGraphClass));
-#else
-				UE_LOG(LogMetaSound, Error, TEXT("Asset '%s' at '%s' must be saved with editor enabled in order to version document to target version '%s'."), *Name.ToString(), *Path, *GetTargetVersion().ToString());
-#endif // !WITH_EDITOR
-			}
-		};
-
-		/** Versions document from 1.9 to 1.10. */
-		class FVersionDocument_1_10 : public FVersionDocumentTransform
-		{
-		public:
-			FVersionDocument_1_10() = default;
-
-			FMetasoundFrontendVersionNumber GetTargetVersion() const override
-			{
-				return { 1, 10 };
-			}
-
-			void TransformInternal(FDocumentHandle InDocument) const override
-			{
-				FMetasoundFrontendGraphClass Class = InDocument->GetRootGraphClass();
-				FMetasoundFrontendGraphClassPresetOptions PresetOptions = Class.PresetOptions;
-				Class.PresetOptions.bIsPreset = Class.Metadata.GetAndClearAutoUpdateManagesInterface_Deprecated();
-				InDocument->SetRootGraphClass(MoveTemp(Class));
-			}
-		};
-
-		FVersionDocument::FVersionDocument(FName InName, const FString& InPath)
-			: Name(InName)
-			, Path(InPath)
-		{
-		}
-
-		bool FVersionDocument::Transform(FDocumentHandle InDocument) const
-		{
-			if (!ensure(InDocument->IsValid()))
-			{
-				return false;
-			}
-
-			bool bWasUpdated = false;
-
-			const FMetasoundFrontendVersionNumber InitVersionNumber = InDocument->GetMetadata().Version.Number;
-
-			// Add additional transforms here after defining them above, example below.
-			bWasUpdated |= FVersionDocument_1_1(Name, Path).Transform(InDocument);
-			bWasUpdated |= FVersionDocument_1_2(Name, Path).Transform(InDocument);
-			bWasUpdated |= FVersionDocument_1_3().Transform(InDocument);
-			bWasUpdated |= FVersionDocument_1_4().Transform(InDocument);
-			bWasUpdated |= FVersionDocument_1_5(Name, Path).Transform(InDocument);
-			bWasUpdated |= FVersionDocument_1_6().Transform(InDocument);
-			bWasUpdated |= FVersionDocument_1_7(Name, Path).Transform(InDocument);
-			bWasUpdated |= FVersionDocument_1_8(Name, Path).Transform(InDocument);
-			bWasUpdated |= FVersionDocument_1_9(Name, Path).Transform(InDocument);
-			bWasUpdated |= FVersionDocument_1_10().Transform(InDocument);
-
-			if (bWasUpdated)
-			{
-				const FMetasoundFrontendVersionNumber NewVersionNumber = InDocument->GetMetadata().Version.Number;
-				UE_LOG(LogMetaSound, Display, TEXT("MetaSound at '%s' Document Versioned: '%s' --> '%s'"), *Path, *InitVersionNumber.ToString(), *NewVersionNumber.ToString());
-			}
-
-			return bWasUpdated;
-		}
 	} // namespace Frontend
 } // namespace Metasound
diff --git a/Engine/Plugins/Runtime/Metasound/Source/MetasoundFrontend/Private/MetasoundOutputFormatInterfaces.cpp b/Engine/Plugins/Runtime/Metasound/Source/MetasoundFrontend/Private/MetasoundOutputFormatInterfaces.cpp
deleted file mode 100644
diff --git a/Engine/Plugins/Runtime/Metasound/Source/MetasoundFrontend/Private/MetasoundParameterTransmitter.cpp b/Engine/Plugins/Runtime/Metasound/Source/MetasoundFrontend/Private/MetasoundParameterTransmitter.cpp
--- a/Engine/Plugins/Runtime/Metasound/Source/MetasoundFrontend/Private/MetasoundParameterTransmitter.cpp
+++ b/Engine/Plugins/Runtime/Metasound/Source/MetasoundFrontend/Private/MetasoundParameterTransmitter.cpp
@@ -3,8 +3,8 @@
 #include "MetasoundParameterTransmitter.h"
 
 #include "IAudioParameterInterfaceRegistry.h"
+#include "Interfaces/MetasoundFrontendSourceInterface.h"
 #include "MetasoundLog.h"
-#include "MetasoundSourceInterface.h"
 
 
 namespace Metasound
diff --git a/Engine/Plugins/Runtime/Metasound/Source/MetasoundFrontend/Private/MetasoundSourceInterface.cpp b/Engine/Plugins/Runtime/Metasound/Source/MetasoundFrontend/Private/MetasoundSourceInterface.cpp
deleted file mode 100644
diff --git a/Engine/Plugins/Runtime/Metasound/Source/MetasoundFrontend/Public/Interfaces/MetasoundFrontendInterface.h b/Engine/Plugins/Runtime/Metasound/Source/MetasoundFrontend/Public/Interfaces/MetasoundFrontendInterface.h
new file mode 100644
--- /dev/null
+++ b/Engine/Plugins/Runtime/Metasound/Source/MetasoundFrontend/Public/Interfaces/MetasoundFrontendInterface.h
@@ -0,0 +1,31 @@
+// Copyright Epic Games, Inc. All Rights Reserved.
+#pragma once
+
+#include "Containers/Array.h"
+#include "MetasoundFrontendDocument.h"
+
+namespace Metasound
+{
+	namespace Frontend
+	{
+		METASOUNDFRONTEND_API bool IsSubsetOfInterface(const FMetasoundFrontendInterface& InSubsetInterface, const FMetasoundFrontendInterface& InSupersetInterface);
+
+		METASOUNDFRONTEND_API bool IsSubsetOfClass(const FMetasoundFrontendInterface& InSubsetInterface, const FMetasoundFrontendClass& InSupersetClass);
+
+		METASOUNDFRONTEND_API bool IsEquivalentInterface(const FMetasoundFrontendInterface& InInputInterface, const FMetasoundFrontendInterface& InTargetInterface);
+
+		METASOUNDFRONTEND_API bool IsEqualInterface(const FMetasoundFrontendInterface& InInputInterface, const FMetasoundFrontendInterface& InTargetInterface);
+
+		METASOUNDFRONTEND_API int32 InputOutputDifferenceCount(const FMetasoundFrontendClass& InClass, const FMetasoundFrontendInterface& InInterface);
+
+		METASOUNDFRONTEND_API int32 InputOutputDifferenceCount(const FMetasoundFrontendInterface& InInterfaceA, const FMetasoundFrontendInterface& InInterfaceB);
+
+		METASOUNDFRONTEND_API bool DeclaredInterfaceVersionsMatch(const FMetasoundFrontendDocument& InDocumentA, const FMetasoundFrontendDocument& InDocumentB);
+
+		METASOUNDFRONTEND_API void GatherRequiredEnvironmentVariables(const FMetasoundFrontendClass& InClass, TArray<FMetasoundFrontendClassEnvironmentVariable>& OutEnvironmentVariables);
+
+		METASOUNDFRONTEND_API const FMetasoundFrontendInterface* FindMostSimilarInterfaceSupportingEnvironment(const FMetasoundFrontendGraphClass& InRootGraph, const TArray<FMetasoundFrontendClass>& InDependencies, const TArray<FMetasoundFrontendGraphClass>& InSubgraphs, const TArray<FMetasoundFrontendInterface>& InCandidateInterfaces);
+
+		METASOUNDFRONTEND_API const FMetasoundFrontendInterface* FindMostSimilarInterfaceSupportingEnvironment(const FMetasoundFrontendDocument& InDocument, const TArray<FMetasoundFrontendInterface>& InCandidateInterfaces);
+	} // namespace Frontend
+} // namespace Metasound
diff --git a/Engine/Plugins/Runtime/Metasound/Source/MetasoundFrontend/Public/Interfaces/MetasoundFrontendOutputFormatInterfaces.h b/Engine/Plugins/Runtime/Metasound/Source/MetasoundFrontend/Public/Interfaces/MetasoundFrontendOutputFormatInterfaces.h
new file mode 100644
--- /dev/null
+++ b/Engine/Plugins/Runtime/Metasound/Source/MetasoundFrontend/Public/Interfaces/MetasoundFrontendOutputFormatInterfaces.h
@@ -0,0 +1,101 @@
+// Copyright Epic Games, Inc. All Rights Reserved.
+#pragma once
+
+#include "IAudioParameterInterfaceRegistry.h"
+#include "MetasoundFrontendDocument.h"
+#include "UObject/NameTypes.h"
+
+// Forward Declarations
+class UClass;
+
+namespace Metasound
+{
+	namespace Frontend
+	{
+		namespace OutputFormatMonoInterface
+		{
+			namespace Outputs
+			{
+				METASOUNDFRONTEND_API const extern FName MonoOut;
+			} // namespace Outputs
+
+			METASOUNDFRONTEND_API const FMetasoundFrontendVersion& GetVersion();
+			METASOUNDFRONTEND_API Audio::FParameterInterfacePtr CreateInterface(const UClass& InClass);
+		} // namespace OutputFormatMonoInterface
+
+		namespace OutputFormatStereoInterface
+		{
+			namespace Outputs
+			{
+				METASOUNDFRONTEND_API const extern FName LeftOut;
+				METASOUNDFRONTEND_API const extern FName RightOut;
+			} // namespace Outputs
+
+			METASOUNDFRONTEND_API const FMetasoundFrontendVersion& GetVersion();
+			METASOUNDFRONTEND_API Audio::FParameterInterfacePtr CreateInterface(const UClass& InClass);
+		} // namespace OutputFormatStereoInterface
+
+		namespace OutputFormatQuadInterface
+		{
+			namespace Outputs
+			{
+				METASOUNDFRONTEND_API const extern FName FrontLeftOut;
+				METASOUNDFRONTEND_API const extern FName FrontRightOut;
+				METASOUNDFRONTEND_API const extern FName SideLeftOut;
+				METASOUNDFRONTEND_API const extern FName SideRightOut;
+			} // namespace Outputs
+
+			METASOUNDFRONTEND_API const FMetasoundFrontendVersion& GetVersion();
+			METASOUNDFRONTEND_API Audio::FParameterInterfacePtr CreateInterface(const UClass& InClass);
+		} // namespace OutputFormatQuadInterface
+
+		namespace OutputFormatFiveDotZeroInterface
+		{
+			namespace Outputs
+			{
+				METASOUNDFRONTEND_API const extern FName FrontLeftOut;
+				METASOUNDFRONTEND_API const extern FName FrontRightOut;
+				METASOUNDFRONTEND_API const extern FName FrontCenterOut;
+				METASOUNDFRONTEND_API const extern FName SideLeftOut;
+				METASOUNDFRONTEND_API const extern FName SideRightOut;
+			} // namespace Outputs
+
+			METASOUNDFRONTEND_API const FMetasoundFrontendVersion& GetVersion();
+			METASOUNDFRONTEND_API Audio::FParameterInterfacePtr CreateInterface(const UClass& InClass);
+		} // namespace OutputFormatFiveDotZeroInterface
+
+		namespace OutputFormatFiveDotOneInterface
+		{
+			namespace Outputs
+			{
+				METASOUNDFRONTEND_API const extern FName FrontLeftOut;
+				METASOUNDFRONTEND_API const extern FName FrontRightOut;
+				METASOUNDFRONTEND_API const extern FName FrontCenterOut;
+				METASOUNDFRONTEND_API const extern FName LowFrequencyOut;
+				METASOUNDFRONTEND_API const extern FName SideLeftOut;
+				METASOUNDFRONTEND_API const extern FName SideRightOut;
+			} // namespace Outputs
+
+			METASOUNDFRONTEND_API const FMetasoundFrontendVersion& GetVersion();
+			METASOUNDFRONTEND_API Audio::FParameterInterfacePtr CreateInterface(const UClass& InClass);
+		} // namespace OutputFormatFiveDotOneInterface
+
+		namespace OutputFormatSevenDotOneInterface
+		{
+			namespace Outputs
+			{
+				METASOUNDFRONTEND_API const extern FName FrontLeftOut;
+				METASOUNDFRONTEND_API const extern FName FrontRightOut;
+				METASOUNDFRONTEND_API const extern FName FrontCenterOut;
+				METASOUNDFRONTEND_API const extern FName LowFrequencyOut;
+				METASOUNDFRONTEND_API const extern FName SideLeftOut;
+				METASOUNDFRONTEND_API const extern FName SideRightOut;
+				METASOUNDFRONTEND_API const extern FName BackLeftOut;
+				METASOUNDFRONTEND_API const extern FName BackRightOut;
+			} // namespace Outputs
+
+			METASOUNDFRONTEND_API const FMetasoundFrontendVersion& GetVersion();
+			METASOUNDFRONTEND_API Audio::FParameterInterfacePtr CreateInterface(const UClass& InClass);
+		} // namespace OutputFormatSevenDotOneInterface
+	} // namespace Frontend
+} // namespace Metasound
diff --git a/Engine/Plugins/Runtime/Metasound/Source/MetasoundFrontend/Public/Interfaces/MetasoundFrontendSourceInterface.h b/Engine/Plugins/Runtime/Metasound/Source/MetasoundFrontend/Public/Interfaces/MetasoundFrontendSourceInterface.h
new file mode 100644
--- /dev/null
+++ b/Engine/Plugins/Runtime/Metasound/Source/MetasoundFrontend/Public/Interfaces/MetasoundFrontendSourceInterface.h
@@ -0,0 +1,83 @@
+// Copyright Epic Games, Inc. All Rights Reserved.
+#pragma once
+
+#include "IAudioParameterInterfaceRegistry.h"
+#include "MetasoundFrontendDocument.h"
+#include "MetasoundFrontendTransform.h"
+#include "UObject/NameTypes.h"
+//#include "MetasoundEngineArchetypes.h"
+#include "MetasoundFrontendController.h"
+// Forward Declarations
+struct FMetasoundFrontendVersion;
+
+namespace Metasound
+{
+	namespace Frontend
+	{
+		namespace SourceOneShotInterface
+		{
+			namespace Outputs
+			{
+				METASOUNDFRONTEND_API const extern FName OnFinished;
+			}
+
+			METASOUNDFRONTEND_API const FMetasoundFrontendVersion& GetVersion();
+			METASOUNDFRONTEND_API Audio::FParameterInterfacePtr CreateInterface(const UClass& InClass);
+		}
+
+		namespace SourceInterfaceV1_0
+		{
+			namespace Inputs
+			{
+				METASOUNDFRONTEND_API const extern FName OnPlay;
+			}
+
+			namespace Outputs
+			{
+				METASOUNDFRONTEND_API const extern FName OnFinished;
+			}
+
+			namespace Environment
+			{
+				METASOUNDFRONTEND_API const extern FName DeviceID;
+				METASOUNDFRONTEND_API const extern FName GraphName;
+				METASOUNDFRONTEND_API const extern FName IsPreview;
+				METASOUNDFRONTEND_API const extern FName SoundUniqueID;
+				METASOUNDFRONTEND_API const extern FName TransmitterID;
+			}
+
+			METASOUNDFRONTEND_API const FMetasoundFrontendVersion& GetVersion();
+			METASOUNDFRONTEND_API Audio::FParameterInterfacePtr CreateInterface(const UClass& InClass);
+		}
+
+		namespace SourceInterface
+		{
+			namespace Inputs
+			{ 
+				METASOUNDFRONTEND_API const extern FName OnPlay;
+			}
+
+			namespace Environment
+			{
+				METASOUNDFRONTEND_API const extern FName DeviceID;
+				METASOUNDFRONTEND_API const extern FName GraphName;
+				METASOUNDFRONTEND_API const extern FName IsPreview;
+				METASOUNDFRONTEND_API const extern FName SoundUniqueID;
+				METASOUNDFRONTEND_API const extern FName TransmitterID;
+				METASOUNDFRONTEND_API const extern FName AudioMixerNumOutputFrames;
+			}
+
+			METASOUNDFRONTEND_API const FMetasoundFrontendVersion& GetVersion();
+			METASOUNDFRONTEND_API Audio::FParameterInterfacePtr CreateInterface(const UClass& InClass);
+
+			class METASOUNDFRONTEND_API FUpdateInterface : public Frontend::IDocumentTransform
+			{
+			public:
+				virtual bool Transform(Frontend::FDocumentHandle InDocument) const override;
+			};
+		}
+
+
+
+	} // namespace Frontend
+} // namespace Metasound
diff --git a/Engine/Plugins/Runtime/Metasound/Source/MetasoundFrontend/Public/MetasoundArchetype.h b/Engine/Plugins/Runtime/Metasound/Source/MetasoundFrontend/Public/MetasoundArchetype.h
deleted file mode 100644
diff --git a/Engine/Plugins/Runtime/Metasound/Source/MetasoundFrontend/Public/MetasoundArrayRandomNode.h b/Engine/Plugins/Runtime/Metasound/Source/MetasoundFrontend/Public/MetasoundArrayRandomNode.h
--- a/Engine/Plugins/Runtime/Metasound/Source/MetasoundFrontend/Public/MetasoundArrayRandomNode.h
+++ b/Engine/Plugins/Runtime/Metasound/Source/MetasoundFrontend/Public/MetasoundArrayRandomNode.h
@@ -2,11 +2,10 @@
 #pragma once
 
 #include "Containers/CircularQueue.h"
-
+#include "Interfaces/MetasoundFrontendSourceInterface.h"
 #include "MetasoundArrayNodes.h"
 #include "MetasoundExecutableOperator.h"
 #include "MetasoundFacade.h"
-#include "MetasoundSourceInterface.h"
 #include "MetasoundNodeInterface.h"
 #include "MetasoundParamHelper.h"
 #include "MetasoundTrigger.h"
diff --git a/Engine/Plugins/Runtime/Metasound/Source/MetasoundFrontend/Public/MetasoundArrayShuffleNode.h b/Engine/Plugins/Runtime/Metasound/Source/MetasoundFrontend/Public/MetasoundArrayShuffleNode.h
--- a/Engine/Plugins/Runtime/Metasound/Source/MetasoundFrontend/Public/MetasoundArrayShuffleNode.h
+++ b/Engine/Plugins/Runtime/Metasound/Source/MetasoundFrontend/Public/MetasoundArrayShuffleNode.h
@@ -2,7 +2,7 @@
 
 #pragma once
 
-#include "CoreMinimal.h"
+#include "Interfaces/MetasoundFrontendSourceInterface.h"
 #include "Internationalization/Text.h"
 #include "MetasoundArrayNodes.h"
 #include "MetasoundBuilderInterface.h"
@@ -14,7 +14,6 @@
 #include "MetasoundOperatorInterface.h"
 #include "MetasoundParamHelper.h"
 #include "MetasoundPrimitives.h"
-#include "MetasoundSourceInterface.h"
 #include "MetasoundTrigger.h"
 #include "MetasoundVertex.h"
 #include "Misc/ScopeLock.h"
diff --git a/Engine/Plugins/Runtime/Metasound/Source/MetasoundFrontend/Public/MetasoundFrontendDocument.h b/Engine/Plugins/Runtime/Metasound/Source/MetasoundFrontend/Public/MetasoundFrontendDocument.h
--- a/Engine/Plugins/Runtime/Metasound/Source/MetasoundFrontend/Public/MetasoundFrontendDocument.h
+++ b/Engine/Plugins/Runtime/Metasound/Source/MetasoundFrontend/Public/MetasoundFrontendDocument.h
@@ -1364,6 +1364,8 @@ struct METASOUNDFRONTEND_API FMetasoundFrontendDocument
 {
 	GENERATED_BODY()
 
+	static FMetasoundFrontendVersionNumber GetMaxVersion();
+
 	Metasound::Frontend::FAccessPoint AccessPoint;
 
 	FMetasoundFrontendDocument();
diff --git a/Engine/Plugins/Runtime/Metasound/Source/MetasoundFrontend/Public/MetasoundFrontendTransform.h b/Engine/Plugins/Runtime/Metasound/Source/MetasoundFrontend/Public/MetasoundFrontendTransform.h
--- a/Engine/Plugins/Runtime/Metasound/Source/MetasoundFrontend/Public/MetasoundFrontendTransform.h
+++ b/Engine/Plugins/Runtime/Metasound/Source/MetasoundFrontend/Public/MetasoundFrontendTransform.h
@@ -225,22 +225,5 @@ namespace Metasound
 
 			bool Transform(FDocumentHandle InDocument) const override;
 		};
-
-		/** Base class for versioning a document. */
-		class METASOUNDFRONTEND_API FVersionDocument : public IDocumentTransform
-		{
-			const FName Name;
-			const FString& Path;
-
-		public:
-			static FMetasoundFrontendVersionNumber GetMaxVersion()
-			{
-				return FMetasoundFrontendVersionNumber { 1, 10 };
-			}
-
-			FVersionDocument(FName InName, const FString& InPath);
-
-			bool Transform(FDocumentHandle InDocument) const override;
-		};
 	}
 }
diff --git a/Engine/Plugins/Runtime/Metasound/Source/MetasoundFrontend/Public/MetasoundOutputFormatInterfaces.h b/Engine/Plugins/Runtime/Metasound/Source/MetasoundFrontend/Public/MetasoundOutputFormatInterfaces.h
deleted file mode 100644
diff --git a/Engine/Plugins/Runtime/Metasound/Source/MetasoundFrontend/Public/MetasoundSourceInterface.h b/Engine/Plugins/Runtime/Metasound/Source/MetasoundFrontend/Public/MetasoundSourceInterface.h
deleted file mode 100644
diff --git a/Engine/Plugins/Runtime/Metasound/Source/MetasoundGenerator/Private/MetasoundGenerator.cpp b/Engine/Plugins/Runtime/Metasound/Source/MetasoundGenerator/Private/MetasoundGenerator.cpp
--- a/Engine/Plugins/Runtime/Metasound/Source/MetasoundGenerator/Private/MetasoundGenerator.cpp
+++ b/Engine/Plugins/Runtime/Metasound/Source/MetasoundGenerator/Private/MetasoundGenerator.cpp
@@ -4,13 +4,13 @@
 #include "Analysis/MetasoundFrontendVertexAnalyzer.h"
 #include "AudioParameter.h"
 #include "DSP/Dsp.h"
+#include "Interfaces/MetasoundFrontendSourceInterface.h"
 #include "MetasoundGraph.h"
 #include "MetasoundInputNode.h"
 #include "MetasoundOperatorBuilder.h"
 #include "MetasoundOperatorInterface.h"
 #include "MetasoundOutputNode.h"
 #include "MetasoundRouter.h"
-#include "MetasoundSourceInterface.h"
 #include "MetasoundTrace.h"
 #include "MetasoundTrigger.h"
 #include "MetasoundVertexData.h"
diff --git a/Engine/Plugins/Runtime/Metasound/Source/MetasoundStandardNodes/Private/MetasoundDebugLogNode.cpp b/Engine/Plugins/Runtime/Metasound/Source/MetasoundStandardNodes/Private/MetasoundDebugLogNode.cpp
--- a/Engine/Plugins/Runtime/Metasound/Source/MetasoundStandardNodes/Private/MetasoundDebugLogNode.cpp
+++ b/Engine/Plugins/Runtime/Metasound/Source/MetasoundStandardNodes/Private/MetasoundDebugLogNode.cpp
@@ -1,11 +1,11 @@
 // Copyright Epic Games, Inc. All Rights Reserved.
 
+#include "Interfaces/MetasoundFrontendSourceInterface.h"
 #include "MetasoundFacade.h"
 #include "MetasoundNodeRegistrationMacro.h"
 #include "MetasoundParamHelper.h"
 #include "MetasoundPrimitives.h"
 #include "MetasoundStandardNodesCategories.h"
-#include "MetasoundSourceInterface.h"
 #include "MetasoundStandardNodesNames.h"
 #include "MetasoundTrigger.h"
 
