Fix for static analysis include cycle

diff --git a/Engine/Plugins/Runtime/Metasound/Source/MetasoundFrontend/Private/MetasoundFrontendDocumentVersioning.h b/Engine/Plugins/Runtime/Metasound/Source/MetasoundFrontend/Private/MetasoundFrontendDocumentVersioning.h
--- a/Engine/Plugins/Runtime/Metasound/Source/MetasoundFrontend/Private/MetasoundFrontendDocumentVersioning.h
+++ b/Engine/Plugins/Runtime/Metasound/Source/MetasoundFrontend/Private/MetasoundFrontendDocumentVersioning.h
@@ -3,7 +3,6 @@
 
 #include "Containers/UnrealString.h"
 #include "MetasoundFrontendController.h"
-#include "MetasoundFrontendDocumentVersioning.h"
 #include "MetasoundFrontendTransform.h"
 #include "MetasoundFrontendDocument.h"
 #include "UObject/NameTypes.h"
