MetaSound constructor pin tooltip clarification

#rnx

diff --git a/Engine/Plugins/Runtime/Metasound/Source/MetasoundEditor/Private/MetasoundEditorGraphSchema.cpp b/Engine/Plugins/Runtime/Metasound/Source/MetasoundEditor/Private/MetasoundEditorGraphSchema.cpp
--- a/Engine/Plugins/Runtime/Metasound/Source/MetasoundEditor/Private/MetasoundEditorGraphSchema.cpp
+++ b/Engine/Plugins/Runtime/Metasound/Source/MetasoundEditor/Private/MetasoundEditorGraphSchema.cpp
@@ -1288,7 +1288,7 @@ const FPinConnectionResponse UMetasoundEditorGraphSchema::CanCreateConnection(co
 			}
 			else if (Frontend::FConnectability::EReason::IncompatibleAccessTypes == Connectability.Reason)
 			{
-				return FPinConnectionResponse(CONNECT_RESPONSE_DISALLOW, LOCTEXT("ConnectionIncompatibleAccessTypes", "Cannot create connection between incompatible access types. Constructor pin outputs can only be connected to constructor pin inputs."));
+				return FPinConnectionResponse(CONNECT_RESPONSE_DISALLOW, LOCTEXT("ConnectionIncompatibleAccessTypes", "Cannot create connection between incompatible access types. Constructor input pins can only be connected to constructor output pins."));
 			}
 			else
 			{
diff --git a/Engine/Plugins/Runtime/Metasound/Source/MetasoundEditor/Private/MetasoundNodeDetailCustomization.cpp b/Engine/Plugins/Runtime/Metasound/Source/MetasoundEditor/Private/MetasoundNodeDetailCustomization.cpp
--- a/Engine/Plugins/Runtime/Metasound/Source/MetasoundEditor/Private/MetasoundNodeDetailCustomization.cpp
+++ b/Engine/Plugins/Runtime/Metasound/Source/MetasoundEditor/Private/MetasoundNodeDetailCustomization.cpp
@@ -60,7 +60,7 @@ namespace Metasound
 
 			static const FText ConstructorPinText = LOCTEXT("ConstructorPinText", "Is Constructor Pin");
 			static const FText ConstructorPinTooltip = LOCTEXT("ConstructorPinTooltip",
-				"Whether this input or output is a constructor pin (value is only read on construction, and is not dynamically updated at runtime).");
+				"Whether this input or output is a constructor pin. Constructor values are only read on construction (on play), and are not dynamically updated at runtime.");
 
 			void GetDataTypeFromElementPropertyHandle(TSharedPtr<IPropertyHandle> ElementPropertyHandle, Frontend::FDataTypeRegistryInfo& OutDataTypeInfo)
 			{
