Float and Int Trigger Threshold Nodes

diff --git a/Engine/Plugins/Runtime/Metasound/Source/MetasoundStandardNodes/Private/MetasoundTriggerOnThresholdNode.cpp b/Engine/Plugins/Runtime/Metasound/Source/MetasoundStandardNodes/Private/MetasoundTriggerOnThresholdNode.cpp
--- a/Engine/Plugins/Runtime/Metasound/Source/MetasoundStandardNodes/Private/MetasoundTriggerOnThresholdNode.cpp
+++ b/Engine/Plugins/Runtime/Metasound/Source/MetasoundStandardNodes/Private/MetasoundTriggerOnThresholdNode.cpp
@@ -14,7 +14,7 @@
 #include "MetasoundVertex.h"
 #include "MetasoundStandardNodesCategories.h"
 
-#define LOCTEXT_NAMESPACE "MetasoundStandardNodes_ThresholdNode"
+#define LOCTEXT_NAMESPACE "MetasoundStandardNodes"
 
 namespace Metasound
 {
@@ -42,51 +42,100 @@ namespace Metasound
 		METASOUND_PARAM(InTriggerType, "Type", "Trigger Threshold Type");
 	}
 
-	class FTriggerOnThresholdOperator : public IOperator
+	// This object is a set of templatized functions containing all the code that varies between int/float and Audio input in the Operator.
+	// The base implementation is int/float, and there is a specialized template class for Audio below this one.
+	// The reason we templatize ValueType and ThresholdType is because for Audio Buffers, they are not the same.
+	// The actual operator class uses this as a more general set of functions that works for all the data types.
+	template <typename ValueType, typename ThresholdType = ValueType>
+	struct TTriggerOnThresholdHelper
 	{
-	public:
-		static constexpr float DefaultThreshold = 0.85f;
+		// Generate is called in Execute, with a different PREDICATE depending on whether it's Rising or Falling Edge.
+		// That PREDICATE is the only difference between the implementation of those two enum values.
+		template <typename PREDICATE>
+		static void Generate(const PREDICATE& ValueTester, const ValueType& InputValue, const ThresholdType& Threshold, ValueType& LastSample, FTriggerWriteRef& Out)
+		{
+			// block-rate only fires on the first frame
+			if (!ValueTester(LastSample, Threshold) && ValueTester(InputValue, Threshold))
+			{
+				Out->TriggerFrame(0);
+			}
+
+			// Remember the last sample for next block.
+			LastSample = InputValue;
+		}
+
+		// Absolute Threshold has different implementation--the other one does not work for gain tracking a bipolar signal like audio (ex. range [-1.0, 1.0]).
+		static void GenerateAbs(const ValueType& InputValue, const ThresholdType& Threshold, bool& bTriggered, FTriggerWriteRef& Out)
+		{
+			const float ThresholdSqr = Threshold * Threshold;
+			const float CurrentSqr = InputValue * InputValue;
+			
+			// Just went above the threshold, so send the trigger
+			if (CurrentSqr > ThresholdSqr && !bTriggered)
+			{
+				bTriggered = true;
+				Out->TriggerFrame(0);
+			}
+			// Below the threshold, so it's safe to send another trigger next time
+			else if (CurrentSqr < ThresholdSqr && bTriggered)
+			{
+				bTriggered = false;
+			}
+		}
 
-		static TUniquePtr<IOperator> CreateOperator(const FCreateOperatorParams& InParams, FBuildErrorArray& OutErrors);
-		static const FNodeClassMetadata& GetNodeInfo();
-		static FVertexInterface DeclareVertexInterface();
+		// In DeclareVertexInterface, the only change is the first TInputDataVertex and how it's constructed.
+		// The Audio version cannot have a default value, but float/int need it.
+		// For this node, the value of the default doesn't matter, as it needs to receive input from a connection elsewhere to function as a TriggerOnThreshold.
+		static const FVertexInterface DeclareVertexInterface(const float& DefaultThreshold)
+		{
+			using namespace TriggerOnThresholdVertexNames;
+
+			static const FVertexInterface Interface(
+				FInputVertexInterface(
+					TInputDataVertex<ValueType>(METASOUND_GET_PARAM_NAME_AND_METADATA(InPin), 0),
+					TInputDataVertex<ThresholdType>(METASOUND_GET_PARAM_NAME_AND_METADATA(InThresholdPin), DefaultThreshold),
+					TInputDataVertex<FEnumBufferTriggerType>(METASOUND_GET_PARAM_NAME_AND_METADATA(InTriggerType))
+				),
+				FOutputVertexInterface(
+					TOutputDataVertex<FTrigger>(METASOUND_GET_PARAM_NAME_AND_METADATA(OutPin))
+				)
+			);
+			return Interface;
+		}
 
-		FTriggerOnThresholdOperator(const FOperatorSettings& InSettings, FAudioBufferReadRef InBuffer, FFloatReadRef InThreshold);
+		// In CreateOperator, the only thing that varies between Audio and non-audio input is whether there is a default value for the Input. 
+		static TDataReadReference<ValueType> CreateInput(const FCreateOperatorParams& InParams)
+		{
+			using namespace TriggerOnThresholdVertexNames;
 
-		virtual FDataReferenceCollection GetInputs() const override;
-		virtual FDataReferenceCollection GetOutputs() const override;
+			const FDataReferenceCollection& InputCol = InParams.InputDataReferences;
+			const FOperatorSettings& Settings = InParams.OperatorSettings;
+			const FInputVertexInterface& InputInterface = InParams.Node.GetVertexInterface().GetInputInterface();
 
-	protected:
-		FAudioBufferReadRef In;
-		FFloatReadRef Threshold;
-		FTriggerWriteRef Out;
-		bool bTriggered = false;
-		float LastSample = 0.f;
-	};
+			return InputCol.GetDataReadReferenceOrConstructWithVertexDefault<ValueType>(InputInterface, METASOUND_GET_PARAM_NAME(InPin), Settings);
+		}
 
-	// Mac Clang requires linkage on constexpr
-	constexpr float FTriggerOnThresholdOperator::DefaultThreshold;
+	};
 
-	struct FTriggerOnThresholdOperator_EdgeCommon : public FTriggerOnThresholdOperator
+	// Audio Buffer specialization.
+	// For details about each function, see the general template class above this one. They look mostly the same,
+	// but most of the differences are difficult or impossible to replicate without specialization.
+	template<>
+	struct TTriggerOnThresholdHelper<FAudioBuffer, float>
 	{
-		using FTriggerOnThresholdOperator::FTriggerOnThresholdOperator;
-		
-		template<typename PREDICATE>
-		void Generate(const PREDICATE& ValueTester)
-		{	
-			Out->AdvanceBlock();
+		template <typename PREDICATE>
+		static void Generate(const PREDICATE& ValueTester, const FAudioBuffer& Input, const float& Threshold, float& LastSample, FTriggerWriteRef& Out)
+		{
+			const float* InputBuffer = Input.GetData();
+			const int32 NumFrames = Input.Num();
 
-			const float* InputBuffer = In->GetData();
-			const int32 NumFrames = In->Num();
-			const float ThreshValue = *Threshold;
-			
 			float Previous = LastSample;
 			for (int32 i = 0; i < NumFrames; ++i)
 			{
 				const float Current = *InputBuffer;
 
 				// If Previous didn't trigger but current value does... fire!
-				if (!ValueTester(Previous, ThreshValue) && ValueTester(Current,ThreshValue))
+				if (!ValueTester(Previous, Threshold) && ValueTester(Current, Threshold))
 				{
 					Out->TriggerFrame(i);
 				}
@@ -95,57 +144,22 @@ namespace Metasound
 				++InputBuffer;
 			}
 
-			// Remember the last sample for next frame.
+			// Remember the last sample for next Audio Block.
 			LastSample = Previous;
 		}
-	};
-
-	struct FTriggerOnThresholdOperator_RisingEdge final : public FTriggerOnThresholdOperator_EdgeCommon
-	{
-		using FTriggerOnThresholdOperator_EdgeCommon::FTriggerOnThresholdOperator_EdgeCommon;
-
-		void Execute() 
-		{
-			Generate(TGreater<float>{});
-		}
-
-		static void ExecuteFunction(IOperator* InOperator) { static_cast<FTriggerOnThresholdOperator_RisingEdge*>(InOperator)->Execute(); }
-		FExecuteFunction GetExecuteFunction() override { return &FTriggerOnThresholdOperator_RisingEdge::ExecuteFunction; }
-	};
-
-	struct FTriggerOnThresholdOperator_FallingEdge final : public FTriggerOnThresholdOperator_EdgeCommon
-	{
-		using FTriggerOnThresholdOperator_EdgeCommon::FTriggerOnThresholdOperator_EdgeCommon;
-
-		static void ExecuteFunction(IOperator* InOperator) { static_cast<FTriggerOnThresholdOperator_FallingEdge*>(InOperator)->Execute(); }
-		FExecuteFunction GetExecuteFunction() override { return &FTriggerOnThresholdOperator_FallingEdge::ExecuteFunction; }
 
-		void Execute()
+		static void GenerateAbs(const FAudioBuffer& Input, const float& Threshold, bool& bTriggered, FTriggerWriteRef& Out)
 		{
-			Generate(TLess<float>{});
-		}		
-	};
+			const float* InputBuffer = Input.GetData();
+			const int32 NumFrames = Input.Num();
+			const float ThresholdSqr = Threshold * Threshold;
 
-	struct FTriggerOnThresholdOperator_AbsThreshold final : public FTriggerOnThresholdOperator
-	{
-		using FTriggerOnThresholdOperator::FTriggerOnThresholdOperator;
-
-		static void ExecuteFunction(IOperator* InOperator) { static_cast<FTriggerOnThresholdOperator_AbsThreshold*>(InOperator)->Execute(); }
-		FExecuteFunction GetExecuteFunction() override { return &FTriggerOnThresholdOperator_AbsThreshold::ExecuteFunction; }
-
-		void Execute()
-		{
-			Out->AdvanceBlock();
-
-			const float* InputBuffer = In->GetData();
-			const int32 NumFrames = In->Num();
-			const float ThresholdSqr = *Threshold * *Threshold;
-					
 			for (int32 i = 0; i < NumFrames; ++i)
 			{
 				const float Current = *InputBuffer;
 				const float CurrentSqr = Current * Current;
-				 
+
+				// Just went above the threshold, so send the trigger
 				if (CurrentSqr > ThresholdSqr && !bTriggered)
 				{
 					bTriggered = true;
@@ -159,112 +173,170 @@ namespace Metasound
 				++InputBuffer;
 			}
 		}
-	};
 
-	FTriggerOnThresholdOperator::FTriggerOnThresholdOperator(const FOperatorSettings& InSettings, FAudioBufferReadRef InBuffer, FFloatReadRef InThreshold)
-		: In(InBuffer)
-		, Threshold(InThreshold)
-		, Out(FTriggerWriteRef::CreateNew(InSettings))
-	{}
+		static const FVertexInterface DeclareVertexInterface(const float& DefaultThreshold)
+		{
+			using namespace TriggerOnThresholdVertexNames;
+
+			static const FVertexInterface Interface(
+				FInputVertexInterface(
+					TInputDataVertex<FAudioBuffer>(METASOUND_GET_PARAM_NAME_AND_METADATA(InPin)),
+					TInputDataVertex<float>(METASOUND_GET_PARAM_NAME_AND_METADATA(InThresholdPin), DefaultThreshold),
+					TInputDataVertex<FEnumBufferTriggerType>(METASOUND_GET_PARAM_NAME_AND_METADATA(InTriggerType))
+				),
+				FOutputVertexInterface(
+					TOutputDataVertex<FTrigger>(METASOUND_GET_PARAM_NAME_AND_METADATA(OutPin))
+				)
+			);
+			return Interface;
+		}
 
-	FDataReferenceCollection FTriggerOnThresholdOperator::GetInputs() const
-	{
-		using namespace TriggerOnThresholdVertexNames;
-		
-		FDataReferenceCollection InputDataReferences;
-		InputDataReferences.AddDataReadReference(METASOUND_GET_PARAM_NAME(InThresholdPin), FFloatReadRef(Threshold));
-		InputDataReferences.AddDataReadReference(METASOUND_GET_PARAM_NAME(InPin), FAudioBufferReadRef(In));
-		return InputDataReferences;
-	}
+		static FAudioBufferReadRef CreateInput(const FCreateOperatorParams& InParams)
+		{
+			using namespace TriggerOnThresholdVertexNames;
 
-	FDataReferenceCollection FTriggerOnThresholdOperator::GetOutputs() const
-	{
-		using namespace TriggerOnThresholdVertexNames; 
+			const FDataReferenceCollection& InputCol = InParams.InputDataReferences;
+			const FOperatorSettings& Settings = InParams.OperatorSettings;
 
-		FDataReferenceCollection OutputDataReferences;
-		OutputDataReferences.AddDataReadReference(METASOUND_GET_PARAM_NAME(OutPin), FTriggerWriteRef(Out));
-		return OutputDataReferences;
-	}
+			return InputCol.GetDataReadReferenceOrConstruct<FAudioBuffer>(METASOUND_GET_PARAM_NAME(InPin), Settings);
+		}
+	};
 
-	FVertexInterface FTriggerOnThresholdOperator::DeclareVertexInterface()
+	// The regular operator class. It's pretty standard from here, just templatized for all the types.
+	template <typename ValueType, typename ThresholdType>
+	class TTriggerOnThresholdOperator : public TExecutableOperator<TTriggerOnThresholdOperator<ValueType, ThresholdType>>
 	{
-		using namespace TriggerOnThresholdVertexNames; 
-
-		static const FVertexInterface Interface(
-			FInputVertexInterface(
-				TInputDataVertex<FAudioBuffer>(METASOUND_GET_PARAM_NAME_AND_METADATA(InPin)),
-				TInputDataVertex<float>(METASOUND_GET_PARAM_NAME_AND_METADATA(InThresholdPin), DefaultThreshold),
-				TInputDataVertex<FEnumBufferTriggerType>(METASOUND_GET_PARAM_NAME_AND_METADATA(InTriggerType))
-			),
-			FOutputVertexInterface(
-				TOutputDataVertex<FTrigger>(METASOUND_GET_PARAM_NAME_AND_METADATA(OutPin))
-			)
-		);
-		return Interface;
-	}
+	public:
+		static constexpr ThresholdType DefaultThreshold = 0.85f;
 
-	const FNodeClassMetadata& FTriggerOnThresholdOperator::GetNodeInfo()
-	{
-		auto InitNodeInfo = []() -> FNodeClassMetadata
+		TTriggerOnThresholdOperator(const FOperatorSettings& InSettings, const TDataReadReference<ValueType>& InBuffer, const TDataReadReference<ThresholdType>& InThreshold, const FBufferTriggerTypeReadRef& InTriggerType)
+			: In(InBuffer)
+			, Threshold(InThreshold)
+			, TriggerType(InTriggerType)
+			, Out(FTriggerWriteRef::CreateNew(InSettings))
+		{}
+
+		void Execute()
+		{
+			Out->AdvanceBlock();
+
+			switch (*TriggerType)
+			{
+			case EBufferTriggerType::RisingEdge:
+			default:
+				TTriggerOnThresholdHelper<ValueType, ThresholdType>::Generate(TGreater<ThresholdType> {}, *In, *Threshold, LastSample, Out);
+				break;
+			case EBufferTriggerType::FallingEdge:
+				TTriggerOnThresholdHelper<ValueType, ThresholdType>::Generate(TLess<ThresholdType> {}, *In, *Threshold, LastSample, Out);
+				break;
+			case EBufferTriggerType::AbsThreshold:
+				TTriggerOnThresholdHelper<ValueType, ThresholdType>::GenerateAbs(*In, *Threshold, bTriggered, Out);
+				break;
+			}
+		}
+
+		FDataReferenceCollection GetInputs() const
+		{
+			using namespace TriggerOnThresholdVertexNames;
+
+			FDataReferenceCollection InputDataReferences;
+			InputDataReferences.AddDataReadReference(METASOUND_GET_PARAM_NAME(InThresholdPin), Threshold);
+			InputDataReferences.AddDataReadReference(METASOUND_GET_PARAM_NAME(InPin), In);
+			InputDataReferences.AddDataReadReference(METASOUND_GET_PARAM_NAME(InTriggerType), TriggerType);
+			return InputDataReferences;
+		}
+
+		FDataReferenceCollection GetOutputs() const
 		{
-			FNodeClassMetadata Info;
-			Info.ClassName = {StandardNodes::Namespace, TEXT("TriggerOnThreshold"), StandardNodes::AudioVariant};
-			Info.MajorVersion = 1;
-			Info.MinorVersion = 0;
-			Info.DisplayName = METASOUND_LOCTEXT("Metasound_TriggerOnThresholdNodeDisplayName", "Trigger On Threshold");
-			Info.Description = METASOUND_LOCTEXT("Metasound_TriggerOnThresholdNodeDescription", "Trigger based on a audio buffer input");
-			Info.Author = PluginAuthor;
-			Info.PromptIfMissing = PluginNodeMissingPrompt;
-			Info.DefaultInterface = DeclareVertexInterface(); 
-			Info.CategoryHierarchy.Emplace(NodeCategories::Trigger);
+			using namespace TriggerOnThresholdVertexNames;
+
+			FDataReferenceCollection OutputDataReferences;
+			OutputDataReferences.AddDataReadReference(METASOUND_GET_PARAM_NAME(OutPin), FTriggerWriteRef(Out));
+			return OutputDataReferences;
+		}
 
+		static const FVertexInterface DeclareVertexInterface()
+		{
+			return TTriggerOnThresholdHelper<ValueType, ThresholdType>::DeclareVertexInterface(DefaultThreshold);
+		}
+
+		static const FNodeClassMetadata& GetNodeInfo()
+		{
+			auto InitNodeInfo = []() -> FNodeClassMetadata
+			{
+				const FText NodeDisplayName = METASOUND_LOCTEXT_FORMAT("TriggerOnThreshold_DisplayNamePattern", "Trigger On Threshold ({0})", GetMetasoundDataTypeDisplayText<ValueType>());
+
+				FNodeClassMetadata Info;
+				Info.ClassName = { StandardNodes::Namespace, TEXT("TriggerOnThreshold"), GetMetasoundDataTypeName<ValueType>() };
+				Info.MajorVersion = 1;
+				Info.MinorVersion = 0;
+				Info.DisplayName = NodeDisplayName;
+				Info.Description = METASOUND_LOCTEXT("TriggerOnThresholdNode_Description", "Trigger when input passes a given threshold.");
+				Info.Author = PluginAuthor;
+				Info.PromptIfMissing = PluginNodeMissingPrompt;
+				Info.DefaultInterface = DeclareVertexInterface();
+				Info.CategoryHierarchy.Emplace(NodeCategories::Trigger);
+
+				return Info;
+			};
+			static const FNodeClassMetadata Info = InitNodeInfo();
 			return Info;
-		};
-		static const FNodeClassMetadata Info = InitNodeInfo();
-		return Info;
-	}
+		}
 
-	FTriggerOnThresholdNode::FTriggerOnThresholdNode(const FNodeInitData& InInitData)
-		: FTriggerOnThresholdNode(InInitData.InstanceName, InInitData.InstanceID, FTriggerOnThresholdOperator::DefaultThreshold)
-	{
-	}
+		static TUniquePtr<IOperator> CreateOperator(const FCreateOperatorParams& InParams, FBuildErrorArray& OutErrors)
+		{
+			using namespace TriggerOnThresholdVertexNames;
 
-	FTriggerOnThresholdNode::FTriggerOnThresholdNode(const FVertexName& InName, const FGuid& InInstanceID, float InDefaultThreshold)
-		: FNodeFacade(InName, InInstanceID, TFacadeOperatorClass<FTriggerOnThresholdOperator>())
-		, DefaultThreshold(InDefaultThreshold)
-	{
-	}
+			const FDataReferenceCollection& InputCol = InParams.InputDataReferences;
+			const FOperatorSettings& Settings = InParams.OperatorSettings;
+			const FInputVertexInterface& InputInterface = InParams.Node.GetVertexInterface().GetInputInterface();
 
-	TUniquePtr<IOperator> FTriggerOnThresholdOperator::CreateOperator(const FCreateOperatorParams& InParams, FBuildErrorArray& OutErrors)
-	{
-		using namespace TriggerOnThresholdVertexNames; 
+			FBufferTriggerTypeReadRef Type = InputCol.GetDataReadReferenceOrConstruct<FEnumBufferTriggerType>(METASOUND_GET_PARAM_NAME(InTriggerType));
+			TDataReadReference<ValueType> Input = TTriggerOnThresholdHelper<ValueType>::CreateInput(InParams);
+			TDataReadReference<ThresholdType> Threshold = InputCol.GetDataReadReferenceOrConstructWithVertexDefault<ThresholdType>(InputInterface, METASOUND_GET_PARAM_NAME(InThresholdPin), Settings);
+
+			return MakeUnique<TTriggerOnThresholdOperator<ValueType, ThresholdType>>(InParams.OperatorSettings, Input, Threshold, Type);
+		}		
 
-		const FTriggerOnThresholdNode& Node = static_cast<const FTriggerOnThresholdNode&>(InParams.Node);
-		const FDataReferenceCollection& InputCol = InParams.InputDataReferences;
-		const FOperatorSettings& Settings = InParams.OperatorSettings;
-		const FInputVertexInterface& InputInterface = InParams.Node.GetVertexInterface().GetInputInterface();
+	protected:
+		TDataReadReference<ValueType> In;
+		TDataReadReference<ThresholdType> Threshold;
+		FBufferTriggerTypeReadRef TriggerType;
+		
+		FTriggerWriteRef Out;
 		
-		// Static property pin, only used for factory.
-		FBufferTriggerTypeReadRef Type = InputCol.GetDataReadReferenceOrConstruct<FEnumBufferTriggerType>(METASOUND_GET_PARAM_NAME(InTriggerType));
+		bool bTriggered = false;
+		ThresholdType LastSample = 0;
+	};
+
+	// Mac Clang requires linkage on constexpr
+	template <>
+	constexpr float TTriggerOnThresholdOperator<float, float>::DefaultThreshold = 0.85f;
 
-		FAudioBufferReadRef InputBuffer = InputCol.GetDataReadReferenceOrConstruct<FAudioBuffer>(METASOUND_GET_PARAM_NAME(InPin), Settings);
-		FFloatReadRef Threshold = InputCol.GetDataReadReferenceOrConstructWithVertexDefault<float>(InputInterface, METASOUND_GET_PARAM_NAME(InThresholdPin), Settings);
+	template <>
+	constexpr float TTriggerOnThresholdOperator<FAudioBuffer, float>::DefaultThreshold = 0.85f;
 
-		switch (*Type)
+	template <>
+	constexpr int32 TTriggerOnThresholdOperator<int32, int32>::DefaultThreshold = 1;
+	
+	template <typename ValueType, typename ThresholdType>
+	class TTriggerOnThresholdNode : public FNodeFacade
+	{
+	public:
+		TTriggerOnThresholdNode(const FNodeInitData& InitData)
+			: FNodeFacade(InitData.InstanceName, InitData.InstanceID, TFacadeOperatorClass<TTriggerOnThresholdOperator<ValueType, ThresholdType>>())
 		{
-		default:
-		case EBufferTriggerType::RisingEdge:
-			return MakeUnique<FTriggerOnThresholdOperator_RisingEdge>(InParams.OperatorSettings, InputBuffer, Threshold);
-		case EBufferTriggerType::FallingEdge:
-			return MakeUnique<FTriggerOnThresholdOperator_FallingEdge>(InParams.OperatorSettings, InputBuffer, Threshold);
-		case EBufferTriggerType::AbsThreshold:
-			return MakeUnique<FTriggerOnThresholdOperator_AbsThreshold>(InParams.OperatorSettings, InputBuffer, Threshold);		
 		}
-		checkNoEntry();
-		return nullptr;
-	}
-	
-	METASOUND_REGISTER_NODE(FTriggerOnThresholdNode);
+	};
+
+	using FTriggerOnThresholdAudioNode = TTriggerOnThresholdNode<FAudioBuffer, float>;
+	METASOUND_REGISTER_NODE(FTriggerOnThresholdAudioNode);
+
+	using FTriggerOnThresholdFloatNode = TTriggerOnThresholdNode<float, float>;
+	METASOUND_REGISTER_NODE(FTriggerOnThresholdFloatNode);
+
+	using FTriggerOnThresholdInt32Node = TTriggerOnThresholdNode<int32, int32>;
+	METASOUND_REGISTER_NODE(FTriggerOnThresholdInt32Node);
 }
 #undef LOCTEXT_NAMESPACE //MetasoundStandardNodes
 
diff --git a/Engine/Plugins/Runtime/Metasound/Source/MetasoundStandardNodes/Public/MetasoundTriggerOnThresholdNode.h b/Engine/Plugins/Runtime/Metasound/Source/MetasoundStandardNodes/Public/MetasoundTriggerOnThresholdNode.h
--- a/Engine/Plugins/Runtime/Metasound/Source/MetasoundStandardNodes/Public/MetasoundTriggerOnThresholdNode.h
+++ b/Engine/Plugins/Runtime/Metasound/Source/MetasoundStandardNodes/Public/MetasoundTriggerOnThresholdNode.h
@@ -1,8 +1,6 @@
 // Copyright Epic Games, Inc. All Rights Reserved.
 #pragma once
 
-#include "CoreMinimal.h"
-
 #include "MetasoundAudioBuffer.h"
 #include "MetasoundBuilderInterface.h"
 #include "MetasoundDataReferenceCollection.h"
@@ -11,18 +9,3 @@
 #include "MetasoundNode.h"
 #include "MetasoundOperatorInterface.h"
 #include "MetasoundVertex.h"
-
-
-namespace Metasound
-{
-	class METASOUNDSTANDARDNODES_API FTriggerOnThresholdNode : public FNodeFacade
-	{
-	public:
-		FTriggerOnThresholdNode(const FVertexName& InName, const FGuid& InInstanceID, float InDefaultThreshold);
-		FTriggerOnThresholdNode(const FNodeInitData& InInitData);
-
-		float GetDefaultThreshold() const { return DefaultThreshold; }
-	private:
-		float DefaultThreshold;
-	};
-}
