Prevent the MetaSound and Dataflow plugins from hijacking comment widget creation for other graph types

diff --git a/Engine/Plugins/Experimental/Dataflow/Source/DataflowEditor/Private/Dataflow/DataflowSNodeFactories.cpp b/Engine/Plugins/Experimental/Dataflow/Source/DataflowEditor/Private/Dataflow/DataflowSNodeFactories.cpp
--- a/Engine/Plugins/Experimental/Dataflow/Source/DataflowEditor/Private/Dataflow/DataflowSNodeFactories.cpp
+++ b/Engine/Plugins/Experimental/Dataflow/Source/DataflowEditor/Private/Dataflow/DataflowSNodeFactories.cpp
@@ -5,6 +5,7 @@
 #include "Dataflow/DataflowEdNode.h"
 #include "EdGraphNode_Comment.h"
 #include "Dataflow/DataflowCommentNode.h"
+#include "Dataflow/DataflowSchema.h"
 
 TSharedPtr<class SGraphNode> FDataflowSNodeFactory::CreateNode(UEdGraphNode* InNode) const
 {
@@ -14,7 +15,10 @@ TSharedPtr<class SGraphNode> FDataflowSNodeFactory::CreateNode(UEdGraphNode* InN
 	}	
 	else if (UEdGraphNode_Comment* CommentNode = Cast<UEdGraphNode_Comment>(InNode))
 	{
-		return SNew(SDataflowEdNodeComment, CommentNode);
+		if (CommentNode->GetSchema()->IsA(UDataflowSchema::StaticClass()))
+		{
+			return SNew(SDataflowEdNodeComment, CommentNode);
+		}
 	}
 	return NULL;
 }
diff --git a/Engine/Plugins/Runtime/Metasound/Source/MetasoundEditor/Private/MetasoundEditorGraphNodeFactory.h b/Engine/Plugins/Runtime/Metasound/Source/MetasoundEditor/Private/MetasoundEditorGraphNodeFactory.h
--- a/Engine/Plugins/Runtime/Metasound/Source/MetasoundEditor/Private/MetasoundEditorGraphNodeFactory.h
+++ b/Engine/Plugins/Runtime/Metasound/Source/MetasoundEditor/Private/MetasoundEditorGraphNodeFactory.h
@@ -4,6 +4,7 @@
 #include "EdGraphUtilities.h"
 #include "EdGraph/EdGraphNode.h"
 #include "MetasoundEditorGraphNode.h"
+#include "MetasoundEditorGraphSchema.h"
 #include "NodeTemplates/MetasoundFrontendNodeTemplateReroute.h"
 #include "SGraphNode.h"
 #include "SMetasoundGraphNode.h"
@@ -34,7 +35,10 @@ class FMetasoundGraphNodeFactory : public FGraphPanelNodeFactory
 		}
 		else if (UEdGraphNode_Comment* CommentNode = Cast<UEdGraphNode_Comment>(InNode))
 		{
-			return SNew(SMetasoundGraphNodeComment, CommentNode);
+			if (CommentNode->GetSchema()->IsA(UMetasoundEditorGraphSchema::StaticClass()))
+			{
+				return SNew(SMetasoundGraphNodeComment, CommentNode);
+			}
 		}
 
 		return nullptr;
