Retain Metasound literal default values when toggling constructor pin state, fix Metasound -> Metasound Patch editor setting name

#rnx

diff --git a/Engine/Plugins/Runtime/Metasound/Config/BaseMetasound.ini b/Engine/Plugins/Runtime/Metasound/Config/BaseMetasound.ini
--- a/Engine/Plugins/Runtime/Metasound/Config/BaseMetasound.ini
+++ b/Engine/Plugins/Runtime/Metasound/Config/BaseMetasound.ini
@@ -22,3 +22,4 @@
 
 ; 5.1
 +ClassRedirects=(OldName="MetaSound",NewName="/Script/MetasoundEngine.MetaSoundPatch")
++PropertyRedirects=(OldName="MetasoundEditorSettings.bPinMetaSoundInAssetMenu",NewName="bPinMetaSoundPatchInAssetMenu")
\ No newline at end of file
diff --git a/Engine/Plugins/Runtime/Metasound/Source/MetasoundEditor/Private/MetasoundAssetTypeActions.cpp b/Engine/Plugins/Runtime/Metasound/Source/MetasoundEditor/Private/MetasoundAssetTypeActions.cpp
--- a/Engine/Plugins/Runtime/Metasound/Source/MetasoundEditor/Private/MetasoundAssetTypeActions.cpp
+++ b/Engine/Plugins/Runtime/Metasound/Source/MetasoundEditor/Private/MetasoundAssetTypeActions.cpp
@@ -204,7 +204,7 @@ namespace Metasound
 
 		const TArray<FText>& FAssetTypeActions_MetaSoundPatch::GetSubMenus() const
 		{
-			if (GetDefault<UMetasoundEditorSettings>()->bPinMetaSoundInAssetMenu)
+			if (GetDefault<UMetasoundEditorSettings>()->bPinMetaSoundPatchInAssetMenu)
 			{
 				static const TArray<FText> SubMenus;
 				return SubMenus;
diff --git a/Engine/Plugins/Runtime/Metasound/Source/MetasoundEditor/Private/MetasoundEditorGraph.cpp b/Engine/Plugins/Runtime/Metasound/Source/MetasoundEditor/Private/MetasoundEditorGraph.cpp
--- a/Engine/Plugins/Runtime/Metasound/Source/MetasoundEditor/Private/MetasoundEditorGraph.cpp
+++ b/Engine/Plugins/Runtime/Metasound/Source/MetasoundEditor/Private/MetasoundEditorGraph.cpp
@@ -388,7 +388,7 @@ void UMetasoundEditorGraphVertex::SetVertexAccessType(EMetasoundFrontendVertexAc
 		return;
 	}
 
-	const FScopedTransaction Transaction(LOCTEXT("SetGraphVertexType", "Set MetaSound GraphVertex Type"), bPostTransaction);
+	const FScopedTransaction Transaction(LOCTEXT("SetVertexConstructorPinState", "Set MetaSound Vertex Constructor Pin State"), bPostTransaction);
 	Graph->GetMetasoundChecked().Modify();
 	Graph->Modify();
 	Modify();
@@ -445,7 +445,8 @@ void UMetasoundEditorGraphVertex::SetVertexAccessType(EMetasoundFrontendVertexAc
 	ClassName = NewNodeHandle->GetClassMetadata().GetClassName();
 	NodeID = NewNodeHandle->GetID();
 
-	InitializeLiteral();
+	constexpr bool bUpdateLiteralPostTransaction = false;
+	UpdateFrontendDefaultLiteral(bUpdateLiteralPostTransaction);
 
 	// 6. Create new node references in the same locations as the old locations
 	// and reconnect linked pins 
diff --git a/Engine/Plugins/Runtime/Metasound/Source/MetasoundEditor/Private/MetasoundEditorSettings.h b/Engine/Plugins/Runtime/Metasound/Source/MetasoundEditor/Private/MetasoundEditorSettings.h
--- a/Engine/Plugins/Runtime/Metasound/Source/MetasoundEditor/Private/MetasoundEditorSettings.h
+++ b/Engine/Plugins/Runtime/Metasound/Source/MetasoundEditor/Private/MetasoundEditorSettings.h
@@ -29,9 +29,9 @@ class METASOUNDEDITOR_API UMetasoundEditorSettings : public UObject
 	GENERATED_UCLASS_BODY()
 
 public:
-	/** Whether to pin the MetaSound asset type when creating new assets. */
-	UPROPERTY(EditAnywhere, config, DisplayName = "Pin MetaSound in Asset Menu", Category = AssetMenu)
-	bool bPinMetaSoundInAssetMenu = false;
+	/** Whether to pin the MetaSound Patch asset type when creating new assets. */
+	UPROPERTY(EditAnywhere, config, DisplayName = "Pin MetaSound Patch in Asset Menu", Category = AssetMenu)
+	bool bPinMetaSoundPatchInAssetMenu = false;
 
 	/** Whether to pin the MetaSound Source asset type when creating new assets. */
 	UPROPERTY(EditAnywhere, config, DisplayName = "Pin MetaSound Source in Asset Menu", Category = AssetMenu)
