Rename MetaSound cvars for consistency, fix typos

diff --git a/Engine/Plugins/Runtime/Metasound/Source/MetasoundEditor/Private/MetasoundEditorGraphConnectionDrawingPolicy.cpp b/Engine/Plugins/Runtime/Metasound/Source/MetasoundEditor/Private/MetasoundEditorGraphConnectionDrawingPolicy.cpp
--- a/Engine/Plugins/Runtime/Metasound/Source/MetasoundEditor/Private/MetasoundEditorGraphConnectionDrawingPolicy.cpp
+++ b/Engine/Plugins/Runtime/Metasound/Source/MetasoundEditor/Private/MetasoundEditorGraphConnectionDrawingPolicy.cpp
@@ -55,21 +55,21 @@ namespace Metasound
 			}
 
 			FAutoConsoleVariableRef CVarMetaSoundEditorAudioConnectSpeed(
-				TEXT("au.MetaSounds.Editor.EnvelopeConnection.Speed"),
+				TEXT("au.MetaSound.Editor.EnvelopeConnection.Speed"),
 				DrawingPolicyPrivate::EnvelopeConnectionSpeedCVar,
 				TEXT("Speed of lines drawn in audio connections in the MetaSoundEditor.\n")
 				TEXT("Default: 1.0f"),
 				ECVF_Default);
 
 			FAutoConsoleVariableRef CVarMetaSoundEditorAudioConnectSpacing(
-				TEXT("au.MetaSounds.Editor.EnvelopeConnection.Spacing"),
+				TEXT("au.MetaSound.Editor.EnvelopeConnection.Spacing"),
 				DrawingPolicyPrivate::EnvelopeConnectionSpacingCVar,
 				TEXT("Spacing of lines drawn in audio connections in the MetaSoundEditor.\n")
 				TEXT("Default: 3.0f"),
 				ECVF_Default);
 
 			FAutoConsoleVariableRef CVarEnvelopeConnectionMaxPointsPerConnection(
-				TEXT("au.MetaSounds.Editor.EnvelopeConnection.MaxPointsPerConnection"),
+				TEXT("au.MetaSound.Editor.EnvelopeConnection.MaxPointsPerConnection"),
 				DrawingPolicyPrivate::EnvelopeConnectionMaxPointsPerConnection,
 				TEXT("Max number of draw points per connection for animated connections.\n")
 				TEXT("Default: 1024"),
diff --git a/Engine/Plugins/Runtime/Metasound/Source/MetasoundEditor/Private/MetasoundEditorGraphSchema.cpp b/Engine/Plugins/Runtime/Metasound/Source/MetasoundEditor/Private/MetasoundEditorGraphSchema.cpp
--- a/Engine/Plugins/Runtime/Metasound/Source/MetasoundEditor/Private/MetasoundEditorGraphSchema.cpp
+++ b/Engine/Plugins/Runtime/Metasound/Source/MetasoundEditor/Private/MetasoundEditorGraphSchema.cpp
@@ -50,9 +50,9 @@
 
 static int32 EnableAllVersionsMetaSoundNodeClassCreationCVar = 0;
 FAutoConsoleVariableRef CVarEnableAllVersionsMetaSoundNodeClassCreation(
-	TEXT("au.Debug.Editor.MetaSounds.EnableAllVersionsNodeClassCreation"),
+	TEXT("au.MetaSound.EnableAllVersionsNodeClassCreation"),
 	EnableAllVersionsMetaSoundNodeClassCreationCVar,
-	TEXT("Enable creating nodes major versions of deprecated MetaSound classes in the Editor.\n")
+	TEXT("Enable creating nodes for major versions of deprecated MetaSound classes in the Editor.\n")
 	TEXT("0: Disabled (default), !0: Enabled"),
 	ECVF_Default);
 
diff --git a/Engine/Plugins/Runtime/Metasound/Source/MetasoundEditor/Private/MetasoundEditorModule.cpp b/Engine/Plugins/Runtime/Metasound/Source/MetasoundEditor/Private/MetasoundEditorModule.cpp
--- a/Engine/Plugins/Runtime/Metasound/Source/MetasoundEditor/Private/MetasoundEditorModule.cpp
+++ b/Engine/Plugins/Runtime/Metasound/Source/MetasoundEditor/Private/MetasoundEditorModule.cpp
@@ -54,7 +54,7 @@ DEFINE_LOG_CATEGORY(LogMetasoundEditor);
 
 static int32 MetaSoundEditorAsyncRegistrationEnabledCVar = 1;
 FAutoConsoleVariableRef CVarMetaSoundEditorAsyncRegistrationEnabled(
-	TEXT("au.MetaSounds.Editor.AsyncRegistrationEnabled"),
+	TEXT("au.MetaSound.Editor.AsyncRegistrationEnabled"),
 	MetaSoundEditorAsyncRegistrationEnabledCVar,
 	TEXT("Enable registering all MetaSound asset classes asyncronously on editor load.\n")
 	TEXT("0: Disabled, !0: Enabled (default)"),
diff --git a/Engine/Plugins/Runtime/Metasound/Source/MetasoundEngine/Private/MetasoundWaveProxyReader.cpp b/Engine/Plugins/Runtime/Metasound/Source/MetasoundEngine/Private/MetasoundWaveProxyReader.cpp
--- a/Engine/Plugins/Runtime/Metasound/Source/MetasoundEngine/Private/MetasoundWaveProxyReader.cpp
+++ b/Engine/Plugins/Runtime/Metasound/Source/MetasoundEngine/Private/MetasoundWaveProxyReader.cpp
@@ -22,7 +22,7 @@
 
 static int32 GMetaSoundWaveProxyReaderSimulateSeekOnNonSeekable = 0;
 FAutoConsoleVariableRef CVarMetaSoundWaveProxyReaderSimulateSeekOnNonSeekable(
-	TEXT("au.MetaSounds.WavePlayer.SimulateSeek"),
+	TEXT("au.MetaSound.WavePlayer.SimulateSeek"),
 	GMetaSoundWaveProxyReaderSimulateSeekOnNonSeekable ,
 	TEXT("If true, SoundWaves which are not of a seekable format will simulate seek calls by reading and discarding samples.\n")
 	TEXT("0: Do not simulate seek, !0: Simulate seek"),
diff --git a/Engine/Plugins/Runtime/Metasound/Source/MetasoundFrontend/Private/MetasoundFrontendNodeController.cpp b/Engine/Plugins/Runtime/Metasound/Source/MetasoundFrontend/Private/MetasoundFrontendNodeController.cpp
--- a/Engine/Plugins/Runtime/Metasound/Source/MetasoundFrontend/Private/MetasoundFrontendNodeController.cpp
+++ b/Engine/Plugins/Runtime/Metasound/Source/MetasoundFrontend/Private/MetasoundFrontendNodeController.cpp
@@ -17,9 +17,9 @@
 
 static int32 MetaSoundAutoUpdateNativeClassesOfEqualVersionCVar = 1;
 FAutoConsoleVariableRef CVarMetaSoundAutoUpdateNativeClass(
-	TEXT("au.MetaSounds.AutoUpdate.NativeClassesOfEqualVersion"),
+	TEXT("au.MetaSound.AutoUpdate.NativeClassesOfEqualVersion"),
 	MetaSoundAutoUpdateNativeClassesOfEqualVersionCVar,
-	TEXT("If true, node references to native class that share a version number will attempt to auto-update if the interface is different, which results in slower graph load times.\n")
+	TEXT("If true, node references to native classes that share a version number will attempt to auto-update if the interface is different, which results in slower graph load times.\n")
 	TEXT("0: Don't auto-update native classes of the same version with interface discrepancies, !0: Auto-update native classes of the same version with interface discrepancies (default)"),
 	ECVF_Default);
 
diff --git a/Engine/Plugins/Runtime/Metasound/Source/MetasoundFrontend/Private/MetasoundFrontendRegistryTransaction.cpp b/Engine/Plugins/Runtime/Metasound/Source/MetasoundFrontend/Private/MetasoundFrontendRegistryTransaction.cpp
--- a/Engine/Plugins/Runtime/Metasound/Source/MetasoundFrontend/Private/MetasoundFrontendRegistryTransaction.cpp
+++ b/Engine/Plugins/Runtime/Metasound/Source/MetasoundFrontend/Private/MetasoundFrontendRegistryTransaction.cpp
@@ -10,7 +10,7 @@ namespace Metasound
 		int32 MetaSoundFrontendDiscardStreamedRegistryTransactionsCVar = 1;
 
 		FAutoConsoleVariableRef CVarMetaSoundFrontendDiscardStreamedRegistryTransactions(
-			TEXT("au.MetaSounds.Frontend.DiscardStreamedRegistryTransactions"),
+			TEXT("au.MetaSound.Frontend.DiscardStreamedRegistryTransactions"),
 			MetaSoundFrontendDiscardStreamedRegistryTransactionsCVar,
 			TEXT("If enabled, MetaSound registry transactions are discarded after they have been streamed.\n")
 			TEXT("0: Disabled, !0: Enabled (default)"),
