Fix typos in Grain Delay node

diff --git a/Engine/Plugins/Runtime/Metasound/Source/MetasoundStandardNodes/Private/MetasoundGrainDelayNode.cpp b/Engine/Plugins/Runtime/Metasound/Source/MetasoundStandardNodes/Private/MetasoundGrainDelayNode.cpp
--- a/Engine/Plugins/Runtime/Metasound/Source/MetasoundStandardNodes/Private/MetasoundGrainDelayNode.cpp
+++ b/Engine/Plugins/Runtime/Metasound/Source/MetasoundStandardNodes/Private/MetasoundGrainDelayNode.cpp
@@ -31,7 +31,7 @@ namespace Metasound
 		METASOUND_PARAM(InGrainPitchShiftDelta, "Pitch Shift Range", "A random pitch shift delta (in semitones) randomly chosen when a grain is spawned.")
 		METASOUND_PARAM(InGrainDelayEnvelope, "Grain Envelope", "The type of envelope to use for the grains.")
 		METASOUND_PARAM(InGrainMaxCount, "Max Grain Count", "The maximum number of grains to render at a time (between 1 and 100). More grains will cost more CPU and potentially clip.")
-		METASOUND_PARAM(InGrainDelayFeedbackAmount, "Feedback Amount", "The amount of feedback of each grain. The grain delay will feed it's audio output back into itself.")
+		METASOUND_PARAM(InGrainDelayFeedbackAmount, "Feedback Amount", "The amount of feedback of each grain. The grain delay will feed its audio output back into itself.")
 
 		METASOUND_PARAM(OutAudio, "Out Audio", "Output audio buffer which has been grain-delayed.")
 	}
@@ -143,7 +143,7 @@ namespace Metasound
 			Info.MajorVersion = 1;
 			Info.MinorVersion = 0;
 			Info.DisplayName = METASOUND_LOCTEXT("GrainDelayNodeDisplayName", "Grain Delay");
-			Info.Description = METASOUND_LOCTEXT("Metasound_GrainDelayNodeDescription", "Performs delayed audio granulation on stereo input audio buffers.");
+			Info.Description = METASOUND_LOCTEXT("Metasound_GrainDelayNodeDescription", "Performs delayed audio granulation on an input audio buffer.");
 			Info.Author = PluginAuthor;
 			Info.PromptIfMissing = PluginNodeMissingPrompt;
 			Info.DefaultInterface = GetVertexInterface();
