Add editor setting to resample and apply window size ratio for MetaSound connection envelopes to be able to shorten or expand envelope history visually

diff --git a/Engine/Plugins/Runtime/Metasound/Source/MetasoundEditor/Private/MetasoundEditorGraphConnectionDrawingPolicy.cpp b/Engine/Plugins/Runtime/Metasound/Source/MetasoundEditor/Private/MetasoundEditorGraphConnectionDrawingPolicy.cpp
--- a/Engine/Plugins/Runtime/Metasound/Source/MetasoundEditor/Private/MetasoundEditorGraphConnectionDrawingPolicy.cpp
+++ b/Engine/Plugins/Runtime/Metasound/Source/MetasoundEditor/Private/MetasoundEditorGraphConnectionDrawingPolicy.cpp
@@ -9,11 +9,11 @@
 #include "MetasoundEditorGraph.h"
 #include "MetasoundEditorGraphBuilder.h"
 #include "MetasoundEditorGraphSchema.h"
-#include "MetasoundEditorSettings.h"
 #include "MetasoundTrace.h"
 #include "Misc/App.h"
 #include "NodeTemplates/MetasoundFrontendNodeTemplateReroute.h"
 #include "Templates/Function.h"
+#include "WaveTableSampler.h"
 
 #define METASOUND_EDITOR_DEBUG_CONNECTIONS 0
 
@@ -295,17 +295,11 @@ namespace Metasound
 			ArrowRadius = FVector2D::ZeroVector;
 
 			WireAnimationLayerID = (InFrontLayerID + InBackLayerID) / 2;
-			ActiveWireThickness = Settings->TraceAttackWireThickness;
-			InactiveWireThickness = Settings->TraceReleaseWireThickness;
 
 			// Numeric wires are inflated up to account for interior animation
 			if (const UMetasoundEditorSettings* MetasoundSettings = GetDefault<UMetasoundEditorSettings>())
 			{
-				bAnimateConnections = MetasoundSettings->bAnimateConnections;
-				ActiveAnalyzerNumericWireThickness = MetasoundSettings->ActiveAnalyzerNumericWireThickness;
-				ActiveAnalyzerEnvelopeWireThickness = MetasoundSettings->ActiveAnalyzerEnvelopeWireThickness;
-				ActiveAnalyzerWireScalarMin = MetasoundSettings->ActiveAnalyzerWireScalarMin;
-				ActiveAnalyzerWireScalarMax = MetasoundSettings->ActiveAnalyzerWireScalarMax;
+				AnalyzerSettings = MetasoundSettings->AnalyzerAnimationSettings;
 			}
 
 			if (GraphObj)
@@ -318,11 +312,11 @@ namespace Metasound
 		{
 			if (OutParams.bDrawBubbles)
 			{
-				OutParams.WireThickness = FMath::Lerp(InactiveWireThickness * ActiveAnalyzerWireScalarMin, ActiveAnalyzerEnvelopeWireThickness * ActiveAnalyzerWireScalarMax, OutParams.WireThickness);
+				OutParams.WireThickness = FMath::Lerp(Settings->TraceReleaseWireThickness * AnalyzerSettings.WireScalarMin, AnalyzerSettings.EnvelopeWireThickness * AnalyzerSettings.WireScalarMax, OutParams.WireThickness);
 			}
 			else
 			{
-				OutParams.WireThickness = InactiveWireThickness;
+				OutParams.WireThickness = Settings->TraceReleaseWireThickness;
 			}
 		}
 
@@ -383,7 +377,7 @@ namespace Metasound
 				OutParams.WireColor = FGraphBuilder::GetPinCategoryColor(AnyPin->PinType);
 			}
 
-			if (AnyPin && bAnimateConnections)
+			if (AnyPin && AnalyzerSettings.bAnimateConnections)
 			{
 				if (AnyPin->PinType.PinCategory == FGraphBuilder::PinCategoryTrigger)
 				{
@@ -405,17 +399,17 @@ namespace Metasound
 
 					if (bEdgeStyleValid && EditorPtr->IsPlaying())
 					{
-						OutParams.WireThickness = ActiveAnalyzerNumericWireThickness;
+						OutParams.WireThickness = AnalyzerSettings.NumericWireThickness;
 					}
 					else
 					{
-						OutParams.WireThickness = InactiveWireThickness;
+						OutParams.WireThickness = Settings->TraceReleaseWireThickness;
 					}
 				}
 			}
 			else
 			{
-				OutParams.WireThickness = InactiveWireThickness;
+				OutParams.WireThickness = Settings->TraceReleaseWireThickness;
 			}
 		}
 
@@ -432,7 +426,7 @@ namespace Metasound
 				return;
 			}
 
-			if (!GraphObj || !EditorPtr->IsPlaying() || !bAnimateConnections)
+			if (!GraphObj || !EditorPtr->IsPlaying() || !AnalyzerSettings.bAnimateConnections)
 			{
 				FConnectionDrawingPolicy::DrawConnection(LayerId, Start, End, Params);
 				return;
@@ -464,14 +458,14 @@ namespace Metasound
 					if (DataType == GetMetasoundDataTypeName<FAudioBuffer>())
 					{
 						SignalParams.SpacingFactor = DrawingPolicyPrivate::EnvelopeConnectionSpacingCVar;
-						SignalParams.SpeedFactor = DrawingPolicyPrivate::EnvelopeConnectionSpeedCVar;
+						SignalParams.SpeedFactor = AnalyzerSettings.EnvelopeSpeed * DrawingPolicyPrivate::EnvelopeConnectionSpeedCVar;
 						const FName AnalyzerName = Frontend::FVertexAnalyzerEnvelopeFollower::GetAnalyzerName();
 						DrawConnectionSignal_Envelope(SignalParams, AnalyzerName);
 					}
 					else if (DataType == GetMetasoundDataTypeName<FTrigger>())
 					{
 						SignalParams.SpacingFactor = DrawingPolicyPrivate::EnvelopeConnectionSpacingCVar;
-						SignalParams.SpeedFactor = DrawingPolicyPrivate::EnvelopeConnectionSpeedCVar;
+						SignalParams.SpeedFactor = AnalyzerSettings.EnvelopeSpeed * DrawingPolicyPrivate::EnvelopeConnectionSpeedCVar;
 						const FName AnalyzerName = Frontend::FVertexAnalyzerTriggerDensity::GetAnalyzerName();
 						DrawConnectionSignal_Envelope(SignalParams, AnalyzerName);
 					}
@@ -600,9 +594,8 @@ namespace Metasound
 			}
 
 			const float PointSpacing = FMath::Max(1.0f, InParams.SpacingFactor * ZoomFactor);
-			float PointSpeed = InParams.SpeedFactor * ZoomFactor;
 			const float PointSize = ZoomFactor * InParams.ConnectionData.Params.WireThickness;
-			const float PointOffset = FMath::Fmod(InParams.AppTime * PointSpeed, PointSpacing);
+			const float PointOffset = FMath::Fmod(InParams.AppTime * ZoomFactor, PointSpacing);
 			const int32 NumPoints = FMath::Min(FMath::CeilToInt(InParams.SplineLength / PointSpacing), DrawingPolicyPrivate::EnvelopeConnectionMaxPointsPerConnection);
 
 			TArray<float> EnvMagnitudes;
@@ -617,24 +610,72 @@ namespace Metasound
 						const FGuid NodeId = OutputHandle->GetOwningNodeID();
 						FName OutputName = OutputHandle->GetName();
 
-						// Only set to max as multiple edges can exist and each may have a different window size.
+						const float SpeedFactor = FMath::Clamp(1.0f - InParams.SpeedFactor, 0.0f, 1.0f);
+						const int32 NumTrackedPoints = FMath::Max(1, (int32)(SpeedFactor * NumPoints));
+
 						FGraphConnectionManager& ConnectionManager = Editor->GetConnectionManager();
 						if (const FFloatMovingWindow* Window = ConnectionManager.GetValueWindow(NodeId, OutputName, InAnalyzerName))
 						{
-							if (NumPoints > Window->Num())
+							if (NumTrackedPoints > Window->Num())
 							{
-								ConnectionManager.TrackValue(NodeId, OutputName, InAnalyzerName, NumPoints);
+								ConnectionManager.TrackValue(NodeId, OutputName, InAnalyzerName, NumTrackedPoints);
 							}
 
-							for (int32 i = 0; i < EnvMagnitudes.Num(); ++i)
+							auto SampleWindowFunction = [&Window, this](TArrayView<float> SampleEnvelope)
+							{
+								switch(AnalyzerSettings.EnvelopeDirection)
+								{
+									case EMetasoundActiveAnalyzerEnvelopeDirection::FromSourceOutput:
+									{
+										for (int32 i = 0; i < SampleEnvelope.Num(); ++i)
+										{
+											const float ClampedScalar = FMath::Clamp(Window->GetValueWrapped(i), 0.0f, 1.0f) * AnalyzerSettings.WireScalarMax;
+											SampleEnvelope[i] *= ClampedScalar;
+										}
+									}
+									break;
+
+									case EMetasoundActiveAnalyzerEnvelopeDirection::FromDestinationInput:
+									{
+										for (int32 i = 0; i < SampleEnvelope.Num(); ++i)
+										{
+											const float ClampedScalar = FMath::Clamp(Window->GetValueWrapped(SampleEnvelope.Num() - 1 - i), 0.0f, 1.0f) * AnalyzerSettings.WireScalarMax;
+											SampleEnvelope[i] *= ClampedScalar;
+										}
+									}
+									break;
+
+									default:
+									{
+										checkNoEntry();
+									}
+									break;
+								}
+							};
+
+							if (FMath::IsNearlyEqual(SpeedFactor, 1.0f))
 							{
-								const float ClampedScalar = FMath::Clamp(Window->GetValueWrapped(i), 0.0f, 1.0f) * ActiveAnalyzerWireScalarMax;
-								EnvMagnitudes[i] *= ClampedScalar;
+								SampleWindowFunction(EnvMagnitudes);
+							}
+							else
+							{
+								TArray<float> EnvSampled;
+								EnvSampled.Init(PointSize, NumTrackedPoints);
+
+								SampleWindowFunction(EnvSampled);
+
+								const float StepRatio = ((float)EnvSampled.Num()) / EnvMagnitudes.Num();
+								for (int32 i = 0; i < EnvMagnitudes.Num(); ++i)
+								{
+									EnvMagnitudes[i] = StepRatio * i;
+								}
+								constexpr auto InterpMode = WaveTable::FWaveTableSampler::EInterpolationMode::Cubic;
+								WaveTable::FWaveTableSampler::Interpolate(EnvSampled, EnvMagnitudes, InterpMode);
 							}
 						}
 						else
 						{
-							ConnectionManager.TrackValue(NodeId, OutputName, InAnalyzerName, NumPoints);
+							ConnectionManager.TrackValue(NodeId, OutputName, InAnalyzerName, NumTrackedPoints);
 						}
 					}
 				}
diff --git a/Engine/Plugins/Runtime/Metasound/Source/MetasoundEditor/Private/MetasoundEditorGraphConnectionDrawingPolicy.h b/Engine/Plugins/Runtime/Metasound/Source/MetasoundEditor/Private/MetasoundEditorGraphConnectionDrawingPolicy.h
--- a/Engine/Plugins/Runtime/Metasound/Source/MetasoundEditor/Private/MetasoundEditorGraphConnectionDrawingPolicy.h
+++ b/Engine/Plugins/Runtime/Metasound/Source/MetasoundEditor/Private/MetasoundEditorGraphConnectionDrawingPolicy.h
@@ -6,6 +6,7 @@
 #include "EdGraphUtilities.h"
 #include "Layout/ArrangedWidget.h"
 #include "MetasoundEditorGraphBuilder.h"
+#include "MetasoundEditorSettings.h"
 #include "MetasoundFrontendController.h"
 #include "Misc/CoreMiscDefines.h"
 #include "Rendering/DrawElements.h"
@@ -102,14 +103,7 @@ namespace Metasound
 
 			int32 WireAnimationLayerID = 0;
 
-			bool bAnimateConnections = false;
-			float ActiveWireThickness = 1.0f;
-			float ActiveAnalyzerEnvelopeWireThickness = 1.0f;
-			float ActiveAnalyzerNumericWireThickness = 1.0f;
-			float ActiveAnalyzerWireScalarMin = 1.0f;
-			float ActiveAnalyzerWireScalarMax = 1.0f;
-
-			float InactiveWireThickness = 1.0f;
+			FMetasoundAnalyzerAnimationSettings AnalyzerSettings;
 
 			TWeakPtr<FEditor> MetasoundEditor;
 
diff --git a/Engine/Plugins/Runtime/Metasound/Source/MetasoundEditor/Private/MetasoundEditorSettings.cpp b/Engine/Plugins/Runtime/Metasound/Source/MetasoundEditor/Private/MetasoundEditorSettings.cpp
--- a/Engine/Plugins/Runtime/Metasound/Source/MetasoundEditor/Private/MetasoundEditorSettings.cpp
+++ b/Engine/Plugins/Runtime/Metasound/Source/MetasoundEditor/Private/MetasoundEditorSettings.cpp
@@ -26,10 +26,4 @@ UMetasoundEditorSettings::UMetasoundEditorSettings(const FObjectInitializer& Obj
 	InputNodeTitleColor = FLinearColor(0.168f, 1.0f, 0.7294f);					// sea foam
 	OutputNodeTitleColor = FLinearColor(1.0f, 0.878f, 0.1686f);					// yellow
 	VariableNodeTitleColor = FLinearColor(0.211f, 0.513f, 0.035f);				// copper
-
-	bAnimateConnections = true;
-	ActiveAnalyzerNumericWireThickness = 5.0f;
-	ActiveAnalyzerEnvelopeWireThickness = 1.0f;
-	ActiveAnalyzerWireScalarMin = 1.0f;
-	ActiveAnalyzerWireScalarMax = 4.5f;
 }
diff --git a/Engine/Plugins/Runtime/Metasound/Source/MetasoundEditor/Private/MetasoundEditorSettings.h b/Engine/Plugins/Runtime/Metasound/Source/MetasoundEditor/Private/MetasoundEditorSettings.h
--- a/Engine/Plugins/Runtime/Metasound/Source/MetasoundEditor/Private/MetasoundEditorSettings.h
+++ b/Engine/Plugins/Runtime/Metasound/Source/MetasoundEditor/Private/MetasoundEditorSettings.h
@@ -8,6 +8,13 @@
 
 #include "MetasoundEditorSettings.generated.h"
 
+UENUM()
+enum class EMetasoundActiveAnalyzerEnvelopeDirection : uint8
+{
+	FromSourceOutput,
+	FromDestinationInput
+};
+
 UENUM()
 enum class EMetasoundMemberDefaultWidget : uint8
 {
@@ -23,6 +30,42 @@ enum class EMetasoundActiveDetailView : uint8
 	General
 };
 
+USTRUCT()
+struct FMetasoundAnalyzerAnimationSettings
+{
+	GENERATED_BODY()
+
+	/** Whether or not animated connections are enabled. */
+	UPROPERTY(EditAnywhere, config, Category = GraphAnimation, meta = (DisplayName = "Animate Connections (Beta)"))
+	bool bAnimateConnections = true;
+
+	/** Thickness of default envelope analyzer wire thickness when connection analyzer is active. */
+	UPROPERTY(EditAnywhere, config, Category = GraphAnimation, meta = (EditCondition = "bAnimateConnections", UIMin = 1, UIMax = 2, ClampMin = 1))
+	float EnvelopeWireThickness = 1.0f;
+
+	/** Speed of default envelope analyzer drawing over wire when connection analyzer is active, where 0 is full visual history (slowest progress) and 1 is no visual history (fastest progress). */
+	UPROPERTY(EditAnywhere, config, Category = GraphAnimation, meta = (EditCondition = "bAnimateConnections", ClampMin = 0, ClampMax = 1))
+	float EnvelopeSpeed = 0.95f;
+
+	/** Whether analyzer envelopes draw from a source output (default) or from the destination input. From the destination input may not
+	  * give the expected illusion of audio processing flowing left-to-right, but results in a waveform with earlier events on the left
+	  * and later on the right (like a traditional timeline with a moving play head). */
+	UPROPERTY(EditAnywhere, config, Category = GraphAnimation, meta = (EditCondition = "bAnimateConnections", ClampMin = 0, ClampMax = 1))
+	EMetasoundActiveAnalyzerEnvelopeDirection EnvelopeDirection = EMetasoundActiveAnalyzerEnvelopeDirection::FromSourceOutput;
+
+	/** Thickness of default numeric analyzer wire thickness when connection analyzer is active. */
+	UPROPERTY(EditAnywhere, config, Category = GraphAnimation, meta = (EditCondition = "bAnimateConnections", UIMin = 1, UIMax = 10, ClampMin = 1))
+	float NumericWireThickness = 5.0f;
+
+	/** Minimum height scalar of wire signal analyzers (ex. audio, triggers). */
+	UPROPERTY(EditAnywhere, config, Category = GraphAnimation, meta = (EditCondition = "bAnimateConnections", UIMin = 1, UIMax = 5, ClampMin = 1))
+	float WireScalarMin = 1.0f;
+
+	/** Maximum height scalar of wire signal analyzers (ex. audio, triggers). */
+	UPROPERTY(EditAnywhere, config, Category = GraphAnimation, meta = (EditCondition = "bAnimateConnections", UIMin = 1, UIMax = 5, ClampMin = 1))
+	float WireScalarMax = 4.5f;
+};
+
 UCLASS(config=EditorPerProjectUserSettings)
 class METASOUNDEDITOR_API UMetasoundEditorSettings : public UObject
 {
@@ -103,30 +146,14 @@ public:
 	UPROPERTY(EditAnywhere, config, Category = NodeTitleColors)
 	FLinearColor VariableNodeTitleColor;
 
-	/** Whether or not animated connections are enabled. */
-	UPROPERTY(EditAnywhere, config, Category = GraphAnimation, meta = (DisplayName = "Animate Connections (Beta)"))
-	bool bAnimateConnections;
-
-	/** Thickness of default envelope analyzer wire thickness when connection analyzer is active. */
-	UPROPERTY(EditAnywhere, config, Category = GraphAnimation, meta = (EditCondition = "bAnimateConnections", UIMin = 1, UIMax = 2, ClampMin = 1))
-	float ActiveAnalyzerEnvelopeWireThickness;
-
-	/** Thickness of default numeric analyzer wire thickness when connection analyzer is active. */
-	UPROPERTY(EditAnywhere, config, Category = GraphAnimation, meta = (EditCondition = "bAnimateConnections", UIMin = 1, UIMax = 10, ClampMin = 1))
-	float ActiveAnalyzerNumericWireThickness;
-
-	/** Minimum height scalar of wire signal analyzers (ex. audio, triggers). */
-	UPROPERTY(EditAnywhere, config, Category = GraphAnimation, meta = (EditCondition = "bAnimateConnections", UIMin = 1, UIMax = 5, ClampMin = 1))
-	float ActiveAnalyzerWireScalarMin;
-
-	/** Maximum height scalar of wire signal analyzers (ex. audio, triggers). */
-	UPROPERTY(EditAnywhere, config, Category = GraphAnimation, meta = (EditCondition = "bAnimateConnections", UIMin = 1, UIMax = 5, ClampMin = 1))
-	float ActiveAnalyzerWireScalarMax;
-
 	/** Widget type to show on input nodes by default */
 	UPROPERTY(EditAnywhere, config, Category = General)
 	EMetasoundMemberDefaultWidget DefaultInputWidgetType = EMetasoundMemberDefaultWidget::RadialSlider;
 
+	/** Settings for visualizing analyzed MetaSound connections */
+	UPROPERTY(EditAnywhere, config, Category = GraphAnimation, meta = (ShowOnlyInnerProperties))
+	FMetasoundAnalyzerAnimationSettings AnalyzerAnimationSettings;
+
 	/** Determines which details view to show in Metasounds Editor */
 	UPROPERTY(Transient)
 	EMetasoundActiveDetailView DetailView = EMetasoundActiveDetailView::Metasound;
diff --git a/Engine/Plugins/Runtime/Metasound/Source/MetasoundEditor/Private/SMetasoundPinValueInspector.cpp b/Engine/Plugins/Runtime/Metasound/Source/MetasoundEditor/Private/SMetasoundPinValueInspector.cpp
--- a/Engine/Plugins/Runtime/Metasound/Source/MetasoundEditor/Private/SMetasoundPinValueInspector.cpp
+++ b/Engine/Plugins/Runtime/Metasound/Source/MetasoundEditor/Private/SMetasoundPinValueInspector.cpp
@@ -462,7 +462,7 @@ namespace Metasound
 				];
 
 			const FName DataType = GetOutputHandle()->GetDataType();
-			const bool bConnectionAnimationsEnabled = GetDefault<UMetasoundEditorSettings>()->bAnimateConnections;
+			const bool bConnectionAnimationsEnabled = GetDefault<UMetasoundEditorSettings>()->AnalyzerAnimationSettings.bAnimateConnections;
 			const bool bIsTypeSupported = PinValueInspectorPrivate::IsMetasoundPrimitiveValueColorizationSupported(DataType);
 			if (bConnectionAnimationsEnabled && bIsTypeSupported)
 			{
diff --git a/Engine/Plugins/Runtime/WaveTable/Source/WaveTable/Private/WaveTableSampler.cpp b/Engine/Plugins/Runtime/WaveTable/Source/WaveTable/Private/WaveTableSampler.cpp
--- a/Engine/Plugins/Runtime/WaveTable/Source/WaveTable/Private/WaveTableSampler.cpp
+++ b/Engine/Plugins/Runtime/WaveTable/Source/WaveTable/Private/WaveTableSampler.cpp
@@ -12,7 +12,9 @@ namespace WaveTable
 		UE_NODISCARD FORCEINLINE FVector2f GetTangentP0(const float* InTableView, const float NextValue, float InIndex, int32 InArraySize)
 		{
 			const int32 Index = FMath::TruncToInt32(InIndex);
-			const float LastValue = InTableView[(Index - 1) % InArraySize];
+
+			// Addition of array size guarantees that if index is 0, we aren't modding a negative number and returning a negative index
+			const float LastValue = InTableView[(Index - 1 + InArraySize) % InArraySize];
 			return FVector2f(2.0f, NextValue - LastValue).GetSafeNormal();
 		};
 
