Allow Attenuation Visualization for MetaSounds

diff --git a/Engine/Plugins/Runtime/Metasound/Source/MetasoundEditor/Private/MetasoundDetailCustomization.cpp b/Engine/Plugins/Runtime/Metasound/Source/MetasoundEditor/Private/MetasoundDetailCustomization.cpp
--- a/Engine/Plugins/Runtime/Metasound/Source/MetasoundEditor/Private/MetasoundDetailCustomization.cpp
+++ b/Engine/Plugins/Runtime/Metasound/Source/MetasoundEditor/Private/MetasoundDetailCustomization.cpp
@@ -248,6 +248,7 @@ namespace Metasound
 					GeneralCategoryBuilder.AddProperty(KeywordsHandle);
 
 					DetailLayout.HideCategory("Attenuation");
+					DetailLayout.HideCategory("Developer");
 					DetailLayout.HideCategory("Effects");
 					DetailLayout.HideCategory("Loading");
 					DetailLayout.HideCategory("Modulation");
@@ -261,29 +262,45 @@ namespace Metasound
 					DetailLayout.HideCategory("MetaSound");
 
 					const bool bShouldBeInitiallyCollapsed = true;
-					IDetailCategoryBuilder& SoundCategory = DetailLayout.EditCategory("Sound");
-					SoundCategory.InitiallyCollapsed(bShouldBeInitiallyCollapsed);
+					TArray<TSharedRef<IPropertyHandle>>DeveloperProperties;
+					TArray<TSharedRef<IPropertyHandle>>SoundProperties;
+
+					DetailLayout.EditCategory("Sound")
+						.InitiallyCollapsed(bShouldBeInitiallyCollapsed)
+						.GetDefaultProperties(SoundProperties);
+					DetailLayout.EditCategory("Attenuation").InitiallyCollapsed(bShouldBeInitiallyCollapsed);
+					DetailLayout.EditCategory("Developer")
+						.InitiallyCollapsed(bShouldBeInitiallyCollapsed)
+						.GetDefaultProperties(DeveloperProperties);
+					DetailLayout.EditCategory("Effects").InitiallyCollapsed(bShouldBeInitiallyCollapsed);
+					DetailLayout.EditCategory("Modulation").InitiallyCollapsed(bShouldBeInitiallyCollapsed);
+					DetailLayout.EditCategory("Voice Management").InitiallyCollapsed(bShouldBeInitiallyCollapsed);
+
+					auto HideProperties = [](const TSet<FName>& PropsToHide, const TArray<TSharedRef<IPropertyHandle>>& Properties)
+					{
+						for (TSharedRef<IPropertyHandle> Property : Properties)
+						{
+							if (PropsToHide.Contains(Property->GetProperty()->GetFName()))
+							{
+								Property->MarkHiddenByCustomization();
+							}
+						}
+					};
 
 					static const TSet<FName> SoundPropsToHide =
 					{
 						GET_MEMBER_NAME_CHECKED(USoundWave, bLooping),
 						GET_MEMBER_NAME_CHECKED(USoundWave, SoundGroup)
 					};
+					HideProperties(SoundPropsToHide, SoundProperties);
 
-					TArray<TSharedRef<IPropertyHandle>>SoundProperties;
-					SoundCategory.GetDefaultProperties(SoundProperties);
-					for (TSharedRef<IPropertyHandle> Property : SoundProperties)
+					static const TSet<FName> DeveloperPropsToHide =
 					{
-						if (SoundPropsToHide.Contains(Property->GetProperty()->GetFName()))
-						{
-							Property->MarkHiddenByCustomization();
-						}
-					}
-
-					DetailLayout.EditCategory("Attenuation").InitiallyCollapsed(bShouldBeInitiallyCollapsed);
-					DetailLayout.EditCategory("Effects").InitiallyCollapsed(bShouldBeInitiallyCollapsed);
-					DetailLayout.EditCategory("Modulation").InitiallyCollapsed(bShouldBeInitiallyCollapsed);
-					DetailLayout.EditCategory("Voice Management").InitiallyCollapsed(bShouldBeInitiallyCollapsed);
+						GET_MEMBER_NAME_CHECKED(USoundBase, Duration),
+						GET_MEMBER_NAME_CHECKED(USoundBase, MaxDistance),
+						GET_MEMBER_NAME_CHECKED(USoundBase, TotalSamples)
+					};
+					HideProperties(DeveloperPropsToHide, DeveloperProperties);
 
 					break;
 			}
@@ -291,16 +308,16 @@ namespace Metasound
 			// Hack to hide parent structs for nested metadata properties
 			DetailLayout.HideCategory("CustomView");
 
+			DetailLayout.HideCategory("Advanced");
+			DetailLayout.HideCategory("Analysis");
 			DetailLayout.HideCategory("Curves");
-			DetailLayout.HideCategory("Developer");
 			DetailLayout.HideCategory("File Path");
 			DetailLayout.HideCategory("Format");
 			DetailLayout.HideCategory("Info");
 			DetailLayout.HideCategory("Loading");
 			DetailLayout.HideCategory("Playback");
 			DetailLayout.HideCategory("Subtitles");
-			DetailLayout.HideCategory("Analysis");
-			DetailLayout.HideCategory("Advanced");
+			DetailLayout.HideCategory("Transformations");
 		}
 
 		FMetasoundInterfacesDetailCustomization::FMetasoundInterfacesDetailCustomization()
diff --git a/Engine/Source/Runtime/Engine/Classes/Sound/SoundBase.h b/Engine/Source/Runtime/Engine/Classes/Sound/SoundBase.h
--- a/Engine/Source/Runtime/Engine/Classes/Sound/SoundBase.h
+++ b/Engine/Source/Runtime/Engine/Classes/Sound/SoundBase.h
@@ -74,8 +74,8 @@ public:
 	UPROPERTY(EditAnywhere, BlueprintReadOnly, Category = Sound, meta = (DisplayName = "Class"), AssetRegistrySearchable)
 	TObjectPtr<USoundClass> SoundClassObject;
 
-	/** When "au.3dVisualize.Attenuation" has been specified, draw this sound's attenuation shape when the sound is audible. For debugging purpose only. */
-	UPROPERTY(EditAnywhere, Category = Developer)
+	/** When "au.3dVisualize.Attenuation" has been specified, draw this sound's attenuation shape when the sound is audible. For debugging purposes only. */
+	UPROPERTY(EditAnywhere, Category = Developer, meta = (DisplayName = "Enable Attenuation Debug"))
 	uint8 bDebug : 1;
 
 	/** Whether or not to override the sound concurrency object with local concurrency settings. */
