Update details on MetaSound Editor synchronization
- Minor clean-up: remove unimplemented GetBlockRate declaration cruft

diff --git a/Engine/Plugins/Runtime/Metasound/Source/MetasoundEngine/Public/Metasound.h b/Engine/Plugins/Runtime/Metasound/Source/MetasoundEngine/Public/Metasound.h
--- a/Engine/Plugins/Runtime/Metasound/Source/MetasoundEngine/Public/Metasound.h
+++ b/Engine/Plugins/Runtime/Metasound/Source/MetasoundEngine/Public/Metasound.h
@@ -43,16 +43,11 @@ public:
 
 namespace Metasound
 {
-	namespace ConsoleVariables
-	{
-		static float GetBlockRate();
-	} // namespace ConsoleVariables
-
 #if WITH_EDITOR
 	template <typename TMetaSoundObject>
 	void PostEditUndo(TMetaSoundObject& InMetaSound)
 	{
-		InMetaSound.SetSynchronizationRequired();
+		InMetaSound.SetUpdateDetailsOnSynchronization();
 		if (UMetasoundEditorGraphBase* Graph = Cast<UMetasoundEditorGraphBase>(InMetaSound.GetGraph()))
 		{
 			Graph->RegisterGraphWithFrontend();
