Fix for assetpath not being generated properly from displayed filepath indirectly in SGraphPinObject locations that attempt to extract short display name from the pin's DefaultValue
#rnx

diff --git a/Engine/Plugins/Runtime/Metasound/Source/MetasoundFrontend/Private/MetasoundFrontendLiteral.cpp b/Engine/Plugins/Runtime/Metasound/Source/MetasoundFrontend/Private/MetasoundFrontendLiteral.cpp
--- a/Engine/Plugins/Runtime/Metasound/Source/MetasoundFrontend/Private/MetasoundFrontendLiteral.cpp
+++ b/Engine/Plugins/Runtime/Metasound/Source/MetasoundFrontend/Private/MetasoundFrontendLiteral.cpp
@@ -40,7 +40,7 @@ namespace MetasoundFrontendLiteralPrivate
 			{
 				return FString();
 			}
-			return InObject->GetFullName();
+			return InObject->GetPathName();
 		}
 	};
 
