Fix adding/removing interfaces in editor transaction construction not being cached for duration of scope
- Add interface name to scoped transaction display text

diff --git a/Engine/Plugins/Runtime/Metasound/Source/MetasoundEditor/Private/MetasoundDetailCustomization.cpp b/Engine/Plugins/Runtime/Metasound/Source/MetasoundEditor/Private/MetasoundDetailCustomization.cpp
--- a/Engine/Plugins/Runtime/Metasound/Source/MetasoundEditor/Private/MetasoundDetailCustomization.cpp
+++ b/Engine/Plugins/Runtime/Metasound/Source/MetasoundEditor/Private/MetasoundDetailCustomization.cpp
@@ -377,7 +377,7 @@ namespace Metasound
 						FMetasoundFrontendInterface InterfaceToAdd;
 						if (ensure(ISearchEngine::Get().FindInterfaceWithHighestVersion(FName(*NameToAdd.Get()), InterfaceToAdd)))
 						{
-							FScopedTransaction(LOCTEXT("AddInterfaceTransaction", "Add MetaSound Interface"));
+							const FScopedTransaction Transaction(FText::Format(LOCTEXT("AddInterfaceTransactionFormat", "Add MetaSound Interface '{0}'"), FText::FromString(InterfaceToAdd.Version.ToString())));
 							MetaSound.Get()->Modify();
 							MetaSoundAsset->GetGraphChecked().Modify();
 
@@ -435,7 +435,7 @@ namespace Metasound
 						});
 
 						{
-							FScopedTransaction(LOCTEXT("RemoveAllInterfacesTransaction", "Remove All MetaSound Interfaces"));
+							const FScopedTransaction Transaction(LOCTEXT("RemoveAllInterfacesTransaction", "Remove All MetaSound Interfaces"));
 							MetaSound.Get()->Modify();
 							MetaSoundAsset->GetGraphChecked().Modify();
 
@@ -482,7 +482,7 @@ namespace Metasound
 					}
 
 					{
-						FScopedTransaction(LOCTEXT("RemoveInterfaceTransaction", "Remove MetaSound Interface"));
+						const FScopedTransaction Transaction(FText::Format(LOCTEXT("RemoveInterfaceTransactionFormat", "Remove MetaSound Interface '{0}'"), FText::FromString(InterfaceEntry.Version.ToString())));
 						MetaSound.Get()->Modify();
 						MetaSoundAsset->GetGraphChecked().Modify();
 
