Fix for time pins not showing default value editor
#rnx

diff --git a/Engine/Plugins/Runtime/Metasound/Source/MetasoundEditor/Private/SMetasoundGraphNode.cpp b/Engine/Plugins/Runtime/Metasound/Source/MetasoundEditor/Private/SMetasoundGraphNode.cpp
--- a/Engine/Plugins/Runtime/Metasound/Source/MetasoundEditor/Private/SMetasoundGraphNode.cpp
+++ b/Engine/Plugins/Runtime/Metasound/Source/MetasoundEditor/Private/SMetasoundGraphNode.cpp
@@ -299,8 +299,9 @@ namespace Metasound
 				{
 					PinWidget = SNew(SMetasoundGraphPinBool, InPin);
 				}
-
-				else if (InPin->PinType.PinCategory == FGraphBuilder::PinCategoryFloat)
+				
+				else if (InPin->PinType.PinCategory == FGraphBuilder::PinCategoryFloat
+					|| InPin->PinType.PinCategory == FGraphBuilder::PinCategoryTime)
 				{
 					PinWidget = SNew(SMetasoundGraphPinFloat, InPin);
 				}
