Fix crash when no output is supplied on a reroute and user attempts undo
#rnx

diff --git a/Engine/Plugins/Runtime/Metasound/Source/MetasoundFrontend/Private/NodeTemplates/MetasoundFrontendNodeTemplateReroute.cpp b/Engine/Plugins/Runtime/Metasound/Source/MetasoundFrontend/Private/NodeTemplates/MetasoundFrontendNodeTemplateReroute.cpp
--- a/Engine/Plugins/Runtime/Metasound/Source/MetasoundFrontend/Private/NodeTemplates/MetasoundFrontendNodeTemplateReroute.cpp
+++ b/Engine/Plugins/Runtime/Metasound/Source/MetasoundFrontend/Private/NodeTemplates/MetasoundFrontendNodeTemplateReroute.cpp
@@ -129,7 +129,7 @@ namespace Metasound
 
 					const FMetasoundFrontendVertex& OutputVertex = InOutNode.Interface.Outputs.Last();
 					FNodeVertexGuidPair OutputNodeVertexPair{ InOutNode.GetID(), OutputVertex.VertexID };
-					OutputEdges = &OutputEdgeMap.FindChecked(OutputNodeVertexPair);
+					OutputEdges = OutputEdgeMap.Find(OutputNodeVertexPair);
 
 					// This can happen if the reroute node isn't provided any outputs to connect to, so its
 					// perfectly acceptable to just ignore this node as it ultimately provides no sourced input.
