FNames containing asset paths are deprecated. FSoftObjectPath or FTopLevelAssetPath should be used instead.
Fixups for Metasound.

diff --git a/Engine/Plugins/Runtime/Metasound/Source/MetasoundEditor/Private/MetasoundAssetTypeActions.cpp b/Engine/Plugins/Runtime/Metasound/Source/MetasoundEditor/Private/MetasoundAssetTypeActions.cpp
--- a/Engine/Plugins/Runtime/Metasound/Source/MetasoundEditor/Private/MetasoundAssetTypeActions.cpp
+++ b/Engine/Plugins/Runtime/Metasound/Source/MetasoundEditor/Private/MetasoundAssetTypeActions.cpp
@@ -121,8 +121,7 @@ namespace Metasound
 					{
 						if (const USoundBase* Sound = PreviewComponent->Sound)
 						{
-							const FName SoundName = *Sound->GetPathName();
-							return SoundName == InSourcePath.GetAssetPathName();
+							return FSoftObjectPath(Sound) == InSourcePath;
 						}
 					}
 				}
diff --git a/Engine/Plugins/Runtime/Metasound/Source/MetasoundEditor/Private/MetasoundEditorGraphBuilder.cpp b/Engine/Plugins/Runtime/Metasound/Source/MetasoundEditor/Private/MetasoundEditorGraphBuilder.cpp
--- a/Engine/Plugins/Runtime/Metasound/Source/MetasoundEditor/Private/MetasoundEditorGraphBuilder.cpp
+++ b/Engine/Plugins/Runtime/Metasound/Source/MetasoundEditor/Private/MetasoundEditorGraphBuilder.cpp
@@ -1006,7 +1006,7 @@ namespace Metasound
 
 							FARFilter Filter;
 							Filter.bRecursiveClasses = false;
-							Filter.ObjectPaths.Add(*ObjectPath);
+							Filter.SoftObjectPaths.Add(FSoftObjectPath(ObjectPath));
 
 							TArray<FAssetData> AssetData;
 							AssetRegistryModule.Get().GetAssets(Filter, AssetData);
diff --git a/Engine/Plugins/Runtime/Metasound/Source/MetasoundEditor/Private/MetasoundEditorGraphSchema.cpp b/Engine/Plugins/Runtime/Metasound/Source/MetasoundEditor/Private/MetasoundEditorGraphSchema.cpp
--- a/Engine/Plugins/Runtime/Metasound/Source/MetasoundEditor/Private/MetasoundEditorGraphSchema.cpp
+++ b/Engine/Plugins/Runtime/Metasound/Source/MetasoundEditor/Private/MetasoundEditorGraphSchema.cpp
@@ -1666,7 +1666,7 @@ void UMetasoundEditorGraphSchema::GetAssetsGraphHoverMessage(const TArray<FAsset
 				break;
 			}
 
-			if (MetaSoundAsset->AddingReferenceCausesLoop(FSoftObjectPath(Data.ObjectPath)))
+			if (MetaSoundAsset->AddingReferenceCausesLoop(Data.GetSoftObjectPath()))
 			{
 				OutOkIcon = false;
 				OutTooltipText = TEXT("Cannot add an asset that would create a reference loop.");
@@ -1749,7 +1749,7 @@ void UMetasoundEditorGraphSchema::DroppedAssetsOnGraph(const TArray<FAssetData>&
 				continue;
 			}
 
-			if (MetaSoundAsset->AddingReferenceCausesLoop(FSoftObjectPath(DroppedAsset.ObjectPath)))
+			if (MetaSoundAsset->AddingReferenceCausesLoop(DroppedAsset.GetSoftObjectPath()))
 			{
 				continue;
 			}
diff --git a/Engine/Plugins/Runtime/Metasound/Source/MetasoundEngine/Private/Metasound.cpp b/Engine/Plugins/Runtime/Metasound/Source/MetasoundEngine/Private/Metasound.cpp
--- a/Engine/Plugins/Runtime/Metasound/Source/MetasoundEngine/Private/Metasound.cpp
+++ b/Engine/Plugins/Runtime/Metasound/Source/MetasoundEngine/Private/Metasound.cpp
@@ -113,7 +113,7 @@ void UMetaSoundPatch::SetRegistryAssetClassInfo(const Metasound::Frontend::FNode
 
 Metasound::Frontend::FNodeClassInfo UMetaSoundPatch::GetAssetClassInfo() const
 {
-	return { GetDocumentChecked().RootGraph, *GetPathName() };
+	return { GetDocumentChecked().RootGraph, FSoftObjectPath(this) };
 }
 
 void UMetaSoundPatch::SetReferencedAssetClassKeys(TSet<Metasound::Frontend::FNodeRegistryKey>&& InKeys)
diff --git a/Engine/Plugins/Runtime/Metasound/Source/MetasoundEngine/Private/MetasoundAssetSubsystem.cpp b/Engine/Plugins/Runtime/Metasound/Source/MetasoundEngine/Private/MetasoundAssetSubsystem.cpp
--- a/Engine/Plugins/Runtime/Metasound/Source/MetasoundEngine/Private/MetasoundAssetSubsystem.cpp
+++ b/Engine/Plugins/Runtime/Metasound/Source/MetasoundEngine/Private/MetasoundAssetSubsystem.cpp
@@ -30,7 +30,7 @@ namespace Metasound
 			bool bSuccess = true;
 
 			OutInfo.Type = EMetasoundFrontendClassType::External;
-			OutInfo.AssetPath = InAssetData.ObjectPath;
+			OutInfo.AssetPath = InAssetData.GetSoftObjectPath();
 
 			FString AssetClassID;
 			bSuccess &= InAssetData.GetTagValue(AssetTags::AssetClassID, AssetClassID);
@@ -207,8 +207,8 @@ Metasound::Frontend::FNodeRegistryKey UMetaSoundAssetSubsystem::AddOrUpdateAsset
 	{
 		UObject* Object = nullptr;
 
-		FSoftObjectPath Path(InAssetData.ObjectPath);
-		if (!FPackageName::GetPackageMountPoint(InAssetData.ObjectPath.ToString()).IsNone())
+		FSoftObjectPath Path = InAssetData.ToSoftObjectPath();
+		if (!FPackageName::GetPackageMountPoint(InAssetData.GetObjectPathString()).IsNone())
 		{
 			if (InAssetData.IsAssetLoaded())
 			{
@@ -233,7 +233,7 @@ Metasound::Frontend::FNodeRegistryKey UMetaSoundAssetSubsystem::AddOrUpdateAsset
 		const FNodeRegistryKey RegistryKey = NodeRegistryKey::CreateKey(ClassInfo);
 		if (NodeRegistryKey::IsValid(RegistryKey))
 		{
-			PathMap.FindOrAdd(RegistryKey) = InAssetData.ObjectPath;
+			PathMap.FindOrAdd(RegistryKey) = InAssetData.GetSoftObjectPath();
 		}
 
 		return RegistryKey;
@@ -421,7 +421,7 @@ void UMetaSoundAssetSubsystem::RenameAsset(const FAssetData& InAssetData, bool b
 void UMetaSoundAssetSubsystem::ResetAssetClassDisplayName(const FAssetData& InAssetData)
 {
 	UObject* Object = nullptr;
-	FSoftObjectPath Path(InAssetData.ObjectPath);
+	FSoftObjectPath Path = InAssetData.GetSoftObjectPath();
 	if (InAssetData.IsAssetLoaded())
 	{
 		Object = Path.ResolveObject();
diff --git a/Engine/Plugins/Runtime/Metasound/Source/MetasoundEngine/Private/MetasoundSource.cpp b/Engine/Plugins/Runtime/Metasound/Source/MetasoundEngine/Private/MetasoundSource.cpp
--- a/Engine/Plugins/Runtime/Metasound/Source/MetasoundEngine/Private/MetasoundSource.cpp
+++ b/Engine/Plugins/Runtime/Metasound/Source/MetasoundEngine/Private/MetasoundSource.cpp
@@ -575,7 +575,7 @@ void UMetaSoundSource::InitResources()
 
 Metasound::Frontend::FNodeClassInfo UMetaSoundSource::GetAssetClassInfo() const
 {
-	return { GetDocumentChecked().RootGraph, *GetPathName() };
+	return { GetDocumentChecked().RootGraph, FSoftObjectPath(this) };
 }
 
 bool UMetaSoundSource::IsPlayable() const
diff --git a/Engine/Plugins/Runtime/Metasound/Source/MetasoundFrontend/Private/MetasoundAssetBase.cpp b/Engine/Plugins/Runtime/Metasound/Source/MetasoundFrontend/Private/MetasoundAssetBase.cpp
--- a/Engine/Plugins/Runtime/Metasound/Source/MetasoundFrontend/Private/MetasoundAssetBase.cpp
+++ b/Engine/Plugins/Runtime/Metasound/Source/MetasoundFrontend/Private/MetasoundAssetBase.cpp
@@ -159,7 +159,7 @@ void FMetasoundAssetBase::RegisterGraphWithFrontend(Metasound::Frontend::FMetaSo
 	class FNodeRegistryEntry : public INodeRegistryEntry
 	{
 	public:
-		FNodeRegistryEntry(const FString& InName, TSharedPtr<FMetasoundFrontendDocument> InPreprocessedDoc, FName InAssetPath)
+		FNodeRegistryEntry(const FString& InName, TSharedPtr<FMetasoundFrontendDocument> InPreprocessedDoc, const FSoftObjectPath& InAssetPath)
 		: Name(InName)
 		, PreprocessedDoc(InPreprocessedDoc)
 		{
diff --git a/Engine/Plugins/Runtime/Metasound/Source/MetasoundFrontend/Private/MetasoundFrontendRegistries.cpp b/Engine/Plugins/Runtime/Metasound/Source/MetasoundFrontend/Private/MetasoundFrontendRegistries.cpp
--- a/Engine/Plugins/Runtime/Metasound/Source/MetasoundFrontend/Private/MetasoundFrontendRegistries.cpp
+++ b/Engine/Plugins/Runtime/Metasound/Source/MetasoundFrontend/Private/MetasoundFrontendRegistries.cpp
@@ -608,14 +608,14 @@ namespace Metasound
 		{
 		}
 
-		FNodeClassInfo::FNodeClassInfo(const FMetasoundFrontendGraphClass& InClass, FName InAssetPath)
+		FNodeClassInfo::FNodeClassInfo(const FMetasoundFrontendGraphClass& InClass, const FSoftObjectPath& InAssetPath)
 			: ClassName(InClass.Metadata.GetClassName())
 			, Type(EMetasoundFrontendClassType::External) // Overridden as it is considered the same as an external class in the registry
 			, AssetClassID(FGuid(ClassName.Name.ToString()))
 			, AssetPath(InAssetPath)
 			, Version(InClass.Metadata.GetVersion())
 		{
-			ensure(!AssetPath.IsNone());
+			ensure(!AssetPath.IsNull());
 #if WITH_EDITORONLY_DATA
 			for (const FMetasoundFrontendClassInput& Input : InClass.Interface.Inputs)
 			{
diff --git a/Engine/Plugins/Runtime/Metasound/Source/MetasoundFrontend/Public/MetasoundFrontendRegistries.h b/Engine/Plugins/Runtime/Metasound/Source/MetasoundFrontend/Public/MetasoundFrontendRegistries.h
--- a/Engine/Plugins/Runtime/Metasound/Source/MetasoundFrontend/Public/MetasoundFrontendRegistries.h
+++ b/Engine/Plugins/Runtime/Metasound/Source/MetasoundFrontend/Public/MetasoundFrontendRegistries.h
@@ -66,7 +66,7 @@ namespace Metasound
 			FGuid AssetClassID;
 
 			// Path to asset containing graph if external type and references asset class.
-			FName AssetPath;
+			FSoftObjectPath AssetPath;
 
 			// Version of the registered class
 			FMetasoundFrontendVersionNumber Version;
@@ -85,7 +85,7 @@ namespace Metasound
 			FNodeClassInfo(const FMetasoundFrontendClassMetadata& InMetadata);
 
 			// Constructor used to generate NodeClassInfo from an asset
-			FNodeClassInfo(const FMetasoundFrontendGraphClass& InClass, FName InAssetPath);
+			FNodeClassInfo(const FMetasoundFrontendGraphClass& InClass, const FSoftObjectPath& InAssetPath);
 
 			// Loads the asset from the provided path, ensuring that the class is of type graph.
 			UObject* LoadAsset() const
