- Fix doubleclick to add connected reroute not capturing full transaction
- Fix pasted reroute ed nodes not generating reroute template frontend node
- Fix for "promote to" transactions not capturing graph changes
#rnx

diff --git a/Engine/Plugins/Runtime/Metasound/Source/MetasoundEditor/Private/MetasoundEditor.cpp b/Engine/Plugins/Runtime/Metasound/Source/MetasoundEditor/Private/MetasoundEditor.cpp
--- a/Engine/Plugins/Runtime/Metasound/Source/MetasoundEditor/Private/MetasoundEditor.cpp
+++ b/Engine/Plugins/Runtime/Metasound/Source/MetasoundEditor/Private/MetasoundEditor.cpp
@@ -47,6 +47,7 @@
 #include "MetasoundUObjectRegistry.h"
 #include "Misc/Attribute.h"
 #include "Modules/ModuleManager.h"
+#include "NodeTemplates/MetasoundFrontendNodeTemplateReroute.h"
 #include "PropertyCustomizationHelpers.h"
 #include "PropertyEditorModule.h"
 #include "ScopedTransaction.h"
@@ -67,6 +68,7 @@ struct FGraphActionNode;
 
 #define LOCTEXT_NAMESPACE "MetaSoundEditor"
 
+
 namespace Metasound
 {
 	namespace Editor
@@ -2410,10 +2412,36 @@ namespace Metasound
 					}
 					else
 					{
-						FNodeHandle NewHandle = FGraphBuilder::AddNodeHandle(*Metasound, *ExternalNode);
-						if (!NewHandle->IsValid())
+						FMetasoundFrontendClass FrontendClass;
+						const FMetasoundFrontendClassName ClassName = ExternalNode->GetClassName();
+						if (ClassName == FRerouteNodeTemplate::ClassName)
 						{
-							NodesToRemove.Add(GraphNode);
+							bool bIsValid = false;
+							if (!ExternalNode->Pins.IsEmpty())
+							{
+								const FName DataType = FGraphBuilder::GetPinDataType(ExternalNode->Pins.Last());
+								if (!DataType.IsNone())
+								{
+									const FNodeRegistryKey& RerouteTemplateKey = FRerouteNodeTemplate::GetRegistryKey();
+									FMetasoundFrontendNodeInterface NodeInterface = FRerouteNodeTemplate::CreateNodeInterfaceFromDataType(DataType);
+
+									FNodeHandle NodeHandle = MetasoundAsset->GetRootGraphHandle()->AddTemplateNode(RerouteTemplateKey, MoveTemp(NodeInterface));
+									bIsValid = NodeHandle->IsValid();
+								}
+							}
+
+							if (!bIsValid)
+							{
+								NodesToRemove.Add(GraphNode);
+							}
+						}
+						else
+						{
+							FNodeHandle NewHandle = FGraphBuilder::AddNodeHandle(*Metasound, *ExternalNode);
+							if (!NewHandle->IsValid())
+							{
+								NodesToRemove.Add(GraphNode);
+							}
 						}
 					}
 				}
diff --git a/Engine/Plugins/Runtime/Metasound/Source/MetasoundEditor/Private/MetasoundEditorGraphSchema.cpp b/Engine/Plugins/Runtime/Metasound/Source/MetasoundEditor/Private/MetasoundEditorGraphSchema.cpp
--- a/Engine/Plugins/Runtime/Metasound/Source/MetasoundEditor/Private/MetasoundEditorGraphSchema.cpp
+++ b/Engine/Plugins/Runtime/Metasound/Source/MetasoundEditor/Private/MetasoundEditorGraphSchema.cpp
@@ -539,6 +539,7 @@ UEdGraphNode* FMetasoundGraphSchemaAction_PromoteToInput::PerformAction(UEdGraph
 	UMetasoundEditorGraph* MetasoundGraph = CastChecked<UMetasoundEditorGraph>(ParentGraph);
 	UObject& ParentMetasound = MetasoundGraph->GetMetasoundChecked();
 	ParentMetasound.Modify();
+	MetasoundGraph->Modify();
 
 	FMetasoundFrontendLiteral DefaultValue;
 	FGraphBuilder::GetPinLiteral(*FromPin, DefaultValue);
@@ -566,6 +567,7 @@ UEdGraphNode* FMetasoundGraphSchemaAction_PromoteToInput::PerformAction(UEdGraph
 					}
 
 					FGraphBuilder::RegisterGraphWithFrontend(ParentMetasound);
+					MetasoundEditor->ClearSelectionAndSelectNode(EdGraphNode);
 					return EdGraphNode;
 				}
 			}
@@ -757,6 +759,7 @@ UEdGraphNode* FMetasoundGraphSchemaAction_PromoteToOutput::PerformAction(UEdGrap
 	UMetasoundEditorGraph* MetasoundGraph = CastChecked<UMetasoundEditorGraph>(ParentGraph);
 	UObject& ParentMetasound = MetasoundGraph->GetMetasoundChecked();
 	ParentMetasound.Modify();
+	MetasoundGraph->Modify();
 
 	const FString OutputName = OutputHandle->GetName().ToString();
 	const FVertexName NewNodeName = FGraphBuilder::GenerateUniqueNameByClassType(ParentMetasound, EMetasoundFrontendClassType::Output, OutputName);
@@ -891,13 +894,14 @@ UEdGraphNode* FMetasoundGraphSchemaAction_NewFromSelected::PerformAction(UEdGrap
 	return nullptr;
 }
 
-FMetasoundGraphSchemaAction_NewReroute::FMetasoundGraphSchemaAction_NewReroute(const FLinearColor& InIconColor)
+FMetasoundGraphSchemaAction_NewReroute::FMetasoundGraphSchemaAction_NewReroute(const FLinearColor* InIconColor, bool bInShouldTransact /* = true */)
 	: FMetasoundGraphSchemaAction(
 		FText(),
 		LOCTEXT("RerouteName", "Add Reroute Node..."),
 		LOCTEXT("RerouteTooltip", "Reroute Node (reroutes wires)"),
 		Metasound::Editor::EPrimaryContextGroup::Common)
-	, IconColor(InIconColor)
+	, IconColor(InIconColor ? *InIconColor : FLinearColor::White)
+	, bShouldTransact(bInShouldTransact)
 {
 }
 
@@ -1102,7 +1106,7 @@ void UMetasoundEditorGraphSchema::GetGraphContextActions(FGraphContextMenuBuilde
 			}
 
 			const FLinearColor IconColor = GetPinTypeColor(FromPin->PinType);
-			ActionMenuBuilder.AddAction(MakeShared<FMetasoundGraphSchemaAction_NewReroute>(IconColor));
+			ActionMenuBuilder.AddAction(MakeShared<FMetasoundGraphSchemaAction_NewReroute>(&IconColor));
 		}
 
 		if (FromPin->Direction == EGPD_Output)
@@ -1131,7 +1135,7 @@ void UMetasoundEditorGraphSchema::GetGraphContextActions(FGraphContextMenuBuilde
 			ActionMenuBuilder.AddAction(MakeShared<FMetasoundGraphSchemaAction_PromoteToVariable_MutatorNode>());
 
 			const FLinearColor IconColor = GetPinTypeColor(FromPin->PinType);
-			ActionMenuBuilder.AddAction(MakeShared<FMetasoundGraphSchemaAction_NewReroute>(IconColor));
+			ActionMenuBuilder.AddAction(MakeShared<FMetasoundGraphSchemaAction_NewReroute>(&IconColor));
 		}
 	}
 	else
@@ -1325,6 +1329,8 @@ const FPinConnectionResponse UMetasoundEditorGraphSchema::CanCreateConnection(co
 
 void UMetasoundEditorGraphSchema::OnPinConnectionDoubleCicked(UEdGraphPin* PinA, UEdGraphPin* PinB, const FVector2D& GraphPosition) const
 {
+	using namespace Metasound::Editor;
+
 	if (!PinA || !PinB)
 	{
 		return;
@@ -1337,9 +1343,19 @@ void UMetasoundEditorGraphSchema::OnPinConnectionDoubleCicked(UEdGraphPin* PinA,
 	UMetasoundEditorGraph* ParentGraph = Cast<UMetasoundEditorGraph>(PinA->GetOwningNode()->GetGraph());
 	if (ParentGraph->IsEditable())
 	{
+		const FName DataType = FGraphBuilder::GetPinDataType(PinA);
+		const FScopedTransaction Transaction(FText::Format(LOCTEXT("AddConnectNewRerouteNode", "Add & Connect {0} Reroute Node"), FText::FromName(DataType)));
+
+		UMetasoundEditorGraph* MetaSoundGraph = CastChecked<UMetasoundEditorGraph>(ParentGraph);
+		UObject& ParentMetasound = MetaSoundGraph->GetMetasoundChecked();
+		ParentMetasound.Modify();
+		ParentGraph->Modify();
+
 		UEdGraphPin* OutputPin = PinA->Direction == EGPD_Output ? PinA : PinB;
-		const FLinearColor IconColor = GetPinTypeColor(OutputPin->PinType);
-		TSharedPtr<FMetasoundGraphSchemaAction_NewReroute> RerouteAction = MakeShared<FMetasoundGraphSchemaAction_NewReroute>(IconColor);
+
+		const FLinearColor* IconColor = nullptr;
+		constexpr bool bShouldTransact = false;
+		TSharedPtr<FMetasoundGraphSchemaAction_NewReroute> RerouteAction = MakeShared<FMetasoundGraphSchemaAction_NewReroute>(IconColor, bShouldTransact);
 
 		UEdGraphNode& Node = *PinA->GetOwningNode();
 		UEdGraphNode* NewNode = RerouteAction->PerformAction(Node.GetGraph(), OutputPin, GraphPosition, true);
diff --git a/Engine/Plugins/Runtime/Metasound/Source/MetasoundEditor/Public/MetasoundEditorGraphSchema.h b/Engine/Plugins/Runtime/Metasound/Source/MetasoundEditor/Public/MetasoundEditorGraphSchema.h
--- a/Engine/Plugins/Runtime/Metasound/Source/MetasoundEditor/Public/MetasoundEditorGraphSchema.h
+++ b/Engine/Plugins/Runtime/Metasound/Source/MetasoundEditor/Public/MetasoundEditorGraphSchema.h
@@ -389,8 +389,7 @@ struct METASOUNDEDITOR_API FMetasoundGraphSchemaAction_NewReroute : public FMeta
 	GENERATED_USTRUCT_BODY();
 
 	FMetasoundGraphSchemaAction_NewReroute() = default;
-
-	FMetasoundGraphSchemaAction_NewReroute(const FLinearColor& InIconColor);
+	FMetasoundGraphSchemaAction_NewReroute(const FLinearColor* InIconColor, bool bInShouldTransact = true);
 
 	virtual ~FMetasoundGraphSchemaAction_NewReroute() = default;
 
@@ -406,6 +405,7 @@ struct METASOUNDEDITOR_API FMetasoundGraphSchemaAction_NewReroute : public FMeta
 
 private:
 	FLinearColor IconColor;
+	bool bShouldTransact = true;
 };
 
 /** Action to create new comment */
diff --git a/Engine/Plugins/Runtime/Metasound/Source/MetasoundFrontend/Private/MetasoundFrontendGraphController.cpp b/Engine/Plugins/Runtime/Metasound/Source/MetasoundFrontend/Private/MetasoundFrontendGraphController.cpp
--- a/Engine/Plugins/Runtime/Metasound/Source/MetasoundFrontend/Private/MetasoundFrontendGraphController.cpp
+++ b/Engine/Plugins/Runtime/Metasound/Source/MetasoundFrontend/Private/MetasoundFrontendGraphController.cpp
@@ -1376,6 +1376,10 @@ namespace Metasound
 
 		FNodeHandle FGraphController::AddNode(const FMetasoundFrontendClassMetadata& InClassMetadata, FGuid InNodeGuid)
 		{
+			ensureAlwaysMsgf(InClassMetadata.GetType() != EMetasoundFrontendClassType::Template,
+				TEXT("Cannot implement '%s' template node using 'AddNode'. Template nodes must always "
+				"be added using AddTemplateNode function and supply the interface to be implemented"),
+				*InClassMetadata.GetClassName().ToString());
 			return AddNode(NodeRegistryKey::CreateKey(InClassMetadata), InNodeGuid);
 		}
 
