Remove MetaSound delete node input pin context action, port emphasize graph connection color on hover functionality, clamp output/variable float literals

diff --git a/Engine/Plugins/Runtime/Metasound/Source/MetasoundEditor/Private/MetasoundEditor.cpp b/Engine/Plugins/Runtime/Metasound/Source/MetasoundEditor/Private/MetasoundEditor.cpp
--- a/Engine/Plugins/Runtime/Metasound/Source/MetasoundEditor/Private/MetasoundEditor.cpp
+++ b/Engine/Plugins/Runtime/Metasound/Source/MetasoundEditor/Private/MetasoundEditor.cpp
@@ -1832,15 +1832,6 @@ namespace Metasound
 			return MetasoundGraphEditor->GetSelectedNodes().Num() == 1;
 		}
 
-		void FEditor::DeleteInput()
-		{
-		}
-
-		bool FEditor::CanDeleteInput() const
-		{
-			return true;
-		}
-
 		void FEditor::OnCreateComment()
 		{
 			if (MetasoundGraphEditor.IsValid())
@@ -1875,10 +1866,6 @@ namespace Metasound
 					FExecuteAction::CreateSP(this, &FEditor::AddInput),
 					FCanExecuteAction::CreateSP(this, &FEditor::CanAddInput));
 
-				GraphEditorCommands->MapAction(FEditorCommands::Get().DeleteInput,
-					FExecuteAction::CreateSP(this, &FEditor::DeleteInput),
-					FCanExecuteAction::CreateSP(this, &FEditor::CanDeleteInput));
-
 				// Editing Commands
 				GraphEditorCommands->MapAction(FGenericCommands::Get().SelectAll,
 					FExecuteAction::CreateLambda([this]() { MetasoundGraphEditor->SelectAllNodes(); }));
diff --git a/Engine/Plugins/Runtime/Metasound/Source/MetasoundEditor/Private/MetasoundEditor.h b/Engine/Plugins/Runtime/Metasound/Source/MetasoundEditor/Private/MetasoundEditor.h
--- a/Engine/Plugins/Runtime/Metasound/Source/MetasoundEditor/Private/MetasoundEditor.h
+++ b/Engine/Plugins/Runtime/Metasound/Source/MetasoundEditor/Private/MetasoundEditor.h
@@ -380,12 +380,6 @@ namespace Metasound
 			/** Whether we can add an input to the currently selected node */
 			bool CanAddInput() const;
 
-			/** Delete an input from the currently selected node */
-			void DeleteInput();
-
-			/** Whether we can delete an input from the currently selected node */
-			bool CanDeleteInput() const;
-
 			/* Create comment node on graph */
 			void OnCreateComment();
 
diff --git a/Engine/Plugins/Runtime/Metasound/Source/MetasoundEditor/Private/MetasoundEditorCommands.cpp b/Engine/Plugins/Runtime/Metasound/Source/MetasoundEditor/Private/MetasoundEditorCommands.cpp
--- a/Engine/Plugins/Runtime/Metasound/Source/MetasoundEditor/Private/MetasoundEditorCommands.cpp
+++ b/Engine/Plugins/Runtime/Metasound/Source/MetasoundEditor/Private/MetasoundEditorCommands.cpp
@@ -18,7 +18,6 @@ namespace Metasound
 
 			UI_COMMAND(BrowserSync, "Browse", "Selects the MetaSound in the content browser. If referencing MetaSound nodes are selected, selects referenced assets instead.", EUserInterfaceActionType::Button, FInputChord());
 			UI_COMMAND(AddInput, "Add Input", "Adds an input to the node", EUserInterfaceActionType::Button, FInputChord());
-			UI_COMMAND(DeleteInput, "Delete Input", "Removes an input from the node", EUserInterfaceActionType::Button, FInputChord());
 
 			UI_COMMAND(EditMetasoundSettings, "MetaSound", "Edit MetaSound settings", EUserInterfaceActionType::Button, FInputChord());
 			UI_COMMAND(EditSourceSettings, "Source", "Edit source specific settings", EUserInterfaceActionType::Button, FInputChord());
diff --git a/Engine/Plugins/Runtime/Metasound/Source/MetasoundEditor/Private/MetasoundEditorGraphConnectionDrawingPolicy.cpp b/Engine/Plugins/Runtime/Metasound/Source/MetasoundEditor/Private/MetasoundEditorGraphConnectionDrawingPolicy.cpp
--- a/Engine/Plugins/Runtime/Metasound/Source/MetasoundEditor/Private/MetasoundEditorGraphConnectionDrawingPolicy.cpp
+++ b/Engine/Plugins/Runtime/Metasound/Source/MetasoundEditor/Private/MetasoundEditorGraphConnectionDrawingPolicy.cpp
@@ -411,6 +411,41 @@ namespace Metasound
 			{
 				OutParams.WireThickness = Settings->TraceReleaseWireThickness;
 			}
+
+			const bool bDeemphasizeUnhoveredPins = HoveredPins.Num() > 0;
+			if (bDeemphasizeUnhoveredPins)
+			{
+				const bool bIsPlaying = GetEditor()->IsPlaying();
+				ApplyHoverDeemphasisMetaSound(OutputPin, InputPin, bIsPlaying, /*inout*/ OutParams.WireThickness, /*inout*/ OutParams.WireColor);
+			}
+		}
+
+		void FGraphConnectionDrawingPolicy::ApplyHoverDeemphasisMetaSound(UEdGraphPin* OutputPin, UEdGraphPin* InputPin, const bool bIsPlaying, /*inout*/ float& Thickness, /*inout*/ FLinearColor& WireColor)
+		{
+			//@TODO: Move these parameters into the settings object
+			const float FadeInBias = 0.75f; // Time in seconds before the fading starts to occur
+			const float FadeInPeriod = 0.6f; // Time in seconds after the bias before the fade is fully complete
+			const float TimeFraction = FMath::SmoothStep(0.0f, FadeInPeriod, (float)(FSlateApplication::Get().GetCurrentTime() - LastHoverTimeEvent - FadeInBias));
+
+			const float LightFraction = 0.25f;
+			const FLinearColor DarkenedColor(0.0f, 0.0f, 0.0f, 0.5f);
+			const FLinearColor LightenedColor(1.0f, 1.0f, 1.0f, 1.0f);
+
+			const bool bContainsBoth = HoveredPins.Contains(InputPin) && HoveredPins.Contains(OutputPin);
+			const bool bContainsOutput = HoveredPins.Contains(OutputPin);
+			const bool bEmphasize = bContainsBoth || (bContainsOutput && (InputPin == nullptr));
+			if (bEmphasize)
+			{
+				if (!bIsPlaying)
+				{
+					Thickness = FMath::Lerp(Thickness, Thickness * ((Thickness < 2.5f) ? 3.5f : 2.5f), TimeFraction);
+				}
+				WireColor = FMath::Lerp<FLinearColor>(WireColor, LightenedColor, LightFraction * TimeFraction);
+			}
+			else
+			{
+				WireColor = FMath::Lerp<FLinearColor>(WireColor, DarkenedColor, HoverDeemphasisDarkFraction * TimeFraction);
+			}
 		}
 
 		void FGraphConnectionDrawingPolicy::DrawConnection(int32 LayerId, const FVector2D& Start, const FVector2D& End, const FConnectionParams& Params)
diff --git a/Engine/Plugins/Runtime/Metasound/Source/MetasoundEditor/Private/MetasoundEditorGraphConnectionDrawingPolicy.h b/Engine/Plugins/Runtime/Metasound/Source/MetasoundEditor/Private/MetasoundEditorGraphConnectionDrawingPolicy.h
--- a/Engine/Plugins/Runtime/Metasound/Source/MetasoundEditor/Private/MetasoundEditorGraphConnectionDrawingPolicy.h
+++ b/Engine/Plugins/Runtime/Metasound/Source/MetasoundEditor/Private/MetasoundEditorGraphConnectionDrawingPolicy.h
@@ -119,6 +119,9 @@ namespace Metasound
 			virtual void DrawSpline(const FDrawConnectionData& InData);
 			virtual void DrawConnectionSpline(int32 InLayerId, const FDrawConnectionData& InData);
 
+			// Copied from FConnectionDrawingPolicy::ApplyHoverDeemphasis, with modification to only apply thickness emphasis when not playing
+			void ApplyHoverDeemphasisMetaSound(UEdGraphPin* OutputPin, UEdGraphPin* InputPin, const bool bIsPlaying, /*inout*/ float& Thickness, /*inout*/ FLinearColor& WireColor);
+
 			const TSharedPtr<const FEditor> GetEditor() const;
 			TSharedPtr<FEditor> GetEditor();
 
diff --git a/Engine/Plugins/Runtime/Metasound/Source/MetasoundEditor/Private/MetasoundEditorGraphNode.cpp b/Engine/Plugins/Runtime/Metasound/Source/MetasoundEditor/Private/MetasoundEditorGraphNode.cpp
--- a/Engine/Plugins/Runtime/Metasound/Source/MetasoundEditor/Private/MetasoundEditorGraphNode.cpp
+++ b/Engine/Plugins/Runtime/Metasound/Source/MetasoundEditor/Private/MetasoundEditorGraphNode.cpp
@@ -13,6 +13,7 @@
 #include "MetasoundEditorCommands.h"
 #include "MetasoundEditorGraph.h"
 #include "MetasoundEditorGraphBuilder.h"
+#include "MetasoundEditorGraphMemberDefaults.h"
 #include "MetasoundEditorGraphSchema.h"
 #include "MetasoundEditorGraphValidation.h"
 #include "MetasoundEditorModule.h"
@@ -532,16 +533,7 @@ void UMetasoundEditorGraphNode::GetNodeContextMenuActions(UToolMenu* Menu, UGrap
 {
 	using namespace Metasound::Editor;
 
-	if (Context->Pin)
-	{
-		// If on an input that can be deleted, show option
-		if (Context->Pin->Direction == EGPD_Input)
-		{
-			FToolMenuSection& Section = Menu->AddSection("MetasoundEditorGraphDeleteInput");
-			Section.AddMenuEntry(FEditorCommands::Get().DeleteInput);
-		}
-	}
-	else if (Context->Node)
+	if (Context->Node)
 	{
 		{
 			FToolMenuSection& Section = Menu->AddSection("MetasoundEditorGraphNodeAlignment");
@@ -586,6 +578,22 @@ FString UMetasoundEditorGraphNode::GetDocumentationExcerptName() const
 	return FString::Printf(TEXT("%s%s"), UMetaSoundPatch::StaticClass()->GetPrefixCPP(), *UMetaSoundPatch::StaticClass()->GetName());
 }
 
+bool UMetasoundEditorGraphMemberNode::ClampFloatLiteral(const UMetasoundEditorGraphMemberDefaultFloat* DefaultFloatLiteral, FMetasoundFrontendLiteral& LiteralValue)
+{
+	bool bClampedFloatLiteral = false;
+	if (DefaultFloatLiteral->ClampDefault)
+	{
+		float LiteralFloatValue = 0.0f;
+		float ClampedFloatValue = 0.0f;
+
+		LiteralValue.TryGet(LiteralFloatValue);
+		ClampedFloatValue = FMath::Clamp(LiteralFloatValue, DefaultFloatLiteral->Range.X, DefaultFloatLiteral->Range.Y);
+		bClampedFloatLiteral = !FMath::IsNearlyEqual(ClampedFloatValue, LiteralFloatValue);
+		LiteralValue.Set(ClampedFloatValue);
+	}
+	return bClampedFloatLiteral;
+}
+
 void UMetasoundEditorGraphOutputNode::PinDefaultValueChanged(UEdGraphPin* InPin)
 {
 	using namespace Metasound::Editor;
@@ -593,27 +601,43 @@ void UMetasoundEditorGraphOutputNode::PinDefaultValueChanged(UEdGraphPin* InPin)
 
 	if (InPin && InPin->Direction == EGPD_Input)
 	{
-		GetMetasoundChecked().Modify();
-
+		UObject& MetaSound = GetMetasoundChecked();
+		MetaSound.Modify();
+		
 		FInputHandle InputHandle = FGraphBuilder::GetInputHandleFromPin(InPin);
 		if (InputHandle->IsValid())
 		{
 			FMetasoundFrontendLiteral LiteralValue;
 			if (FGraphBuilder::GetPinLiteral(*InPin, LiteralValue))
 			{
-				InputHandle->SetLiteral(LiteralValue);
-			}
-
-			if (Output)
-			{
-				UMetasoundEditorGraphMemberDefaultLiteral* Literal = Output->GetLiteral();
-				if (ensure(Literal))
+				if (Output)
 				{
-					Literal->SetFromLiteral(LiteralValue);
-				}
+					UMetasoundEditorGraphMemberDefaultLiteral* Literal = Output->GetLiteral();
+					if (ensure(Literal))
+					{
+						// Clamp float literal if necessary 
+						bool bClampedFloatLiteral = false;
+						if (const UMetasoundEditorGraphMemberDefaultFloat* DefaultFloatLiteral = Cast<UMetasoundEditorGraphMemberDefaultFloat>(Literal))
+						{
+							bClampedFloatLiteral = ClampFloatLiteral(DefaultFloatLiteral, LiteralValue);
+						}
+
+						Literal->SetFromLiteral(LiteralValue);
 
-				constexpr bool bPostTransaction = false;
-				Output->UpdateFrontendDefaultLiteral(bPostTransaction);
+						constexpr bool bPostTransaction = false;
+						Output->UpdateFrontendDefaultLiteral(bPostTransaction);
+
+						// Update graph node if it was clamped
+						if (bClampedFloatLiteral)
+						{
+							FGraphBuilder::RegisterGraphWithFrontend(MetaSound);
+							if (FMetasoundAssetBase* MetaSoundAsset = Metasound::IMetasoundUObjectRegistry::Get().GetObjectAsAssetBase(&MetaSound))
+							{
+								MetaSoundAsset->SetUpdateDetailsOnSynchronization();
+							}
+						}
+					}
+				}
 			}
 		}
 	}
@@ -1141,7 +1165,8 @@ void UMetasoundEditorGraphVariableNode::PinDefaultValueChanged(UEdGraphPin* Pin)
 
 	if (Pin && Pin->Direction == EGPD_Input)
 	{
-		GetMetasoundChecked().Modify();
+		UObject& MetaSound = GetMetasoundChecked();
+		MetaSound.Modify();
 
 		FInputHandle InputHandle = FGraphBuilder::GetInputHandleFromPin(Pin);
 		if (InputHandle->IsValid())
@@ -1149,23 +1174,37 @@ void UMetasoundEditorGraphVariableNode::PinDefaultValueChanged(UEdGraphPin* Pin)
 			FMetasoundFrontendLiteral LiteralValue;
 			if (FGraphBuilder::GetPinLiteral(*Pin, LiteralValue))
 			{
-				InputHandle->SetLiteral(LiteralValue);
-			}
-
-			// If this is the mutator node, synchronize the variable default literal with this default.
-			FNodeHandle MutatorNode = Variable->GetVariableHandle()->FindMutatorNode();
-			if (MutatorNode->IsValid())
-			{
-				if (MutatorNode->GetID() == NodeID)
+				// If this is the mutator node, synchronize the variable default literal with this default.
+				FNodeHandle MutatorNode = Variable->GetVariableHandle()->FindMutatorNode();
+				if (MutatorNode->IsValid())
 				{
-					UMetasoundEditorGraphMemberDefaultLiteral* Literal = Variable->GetLiteral();
-					if (ensure(Literal))
+					if (MutatorNode->GetID() == NodeID)
 					{
-						Literal->SetFromLiteral(LiteralValue);
-					}
+						UMetasoundEditorGraphMemberDefaultLiteral* Literal = Variable->GetLiteral();
+						if (ensure(Literal))
+						{
+							// Clamp float literal if necessary 
+							bool bClampedFloatLiteral = false;
+							if (const UMetasoundEditorGraphMemberDefaultFloat* DefaultFloatLiteral = Cast<UMetasoundEditorGraphMemberDefaultFloat>(Literal))
+							{
+								bClampedFloatLiteral = ClampFloatLiteral(DefaultFloatLiteral, LiteralValue);
+							}
+							Literal->SetFromLiteral(LiteralValue);
 
-					constexpr bool bPostTransaction = false;
-					Variable->UpdateFrontendDefaultLiteral(bPostTransaction);
+							constexpr bool bPostTransaction = false;
+							Variable->UpdateFrontendDefaultLiteral(bPostTransaction);
+
+							if (bClampedFloatLiteral)
+							{
+								// Update graph node if it was clamped
+								FGraphBuilder::RegisterGraphWithFrontend(MetaSound);
+								if (FMetasoundAssetBase* MetaSoundAsset = Metasound::IMetasoundUObjectRegistry::Get().GetObjectAsAssetBase(&MetaSound))
+								{
+									MetaSoundAsset->SetUpdateDetailsOnSynchronization();
+								}
+							}
+						}
+					}
 				}
 			}
 		}
diff --git a/Engine/Plugins/Runtime/Metasound/Source/MetasoundEditor/Public/MetasoundEditorCommands.h b/Engine/Plugins/Runtime/Metasound/Source/MetasoundEditor/Public/MetasoundEditorCommands.h
--- a/Engine/Plugins/Runtime/Metasound/Source/MetasoundEditor/Public/MetasoundEditorCommands.h
+++ b/Engine/Plugins/Runtime/Metasound/Source/MetasoundEditor/Public/MetasoundEditorCommands.h
@@ -43,9 +43,6 @@ namespace Metasound
 			/** Adds an input to the node */
 			TSharedPtr<FUICommandInfo> AddInput;
 
-			/** Removes an input from the node */
-			TSharedPtr<FUICommandInfo> DeleteInput;
-
 			/** Shows the MetasoundSource's specific settings in the Details panel (if of respective type) */
 			TSharedPtr<FUICommandInfo> EditSourceSettings;
 
diff --git a/Engine/Plugins/Runtime/Metasound/Source/MetasoundEditor/Public/MetasoundEditorGraphNode.h b/Engine/Plugins/Runtime/Metasound/Source/MetasoundEditor/Public/MetasoundEditorGraphNode.h
--- a/Engine/Plugins/Runtime/Metasound/Source/MetasoundEditor/Public/MetasoundEditorGraphNode.h
+++ b/Engine/Plugins/Runtime/Metasound/Source/MetasoundEditor/Public/MetasoundEditorGraphNode.h
@@ -21,6 +21,7 @@ class UMetaSoundPatch;
 class UMetasoundEditorGraphOutput;
 class UMetasoundEditorGraphMember;
 class UMetasoundEditorGraphVariable;
+class UMetasoundEditorGraphMemberDefaultFloat;
 
 namespace Metasound
 {
@@ -147,6 +148,10 @@ public:
 	{
 		return true;
 	}
+protected:
+	// Clamp float literal value based on the given default float literal. 
+	// Returns whether the literal was clamped. 
+	static bool ClampFloatLiteral(const UMetasoundEditorGraphMemberDefaultFloat* DefaultFloatLiteral, FMetasoundFrontendLiteral& LiteralValue);
 };
 
 /** Node that represents a graph output */
