Fix missing rerouted implicit accessor behavior causing various bugs with constructor pins
#rnx

diff --git a/Engine/Plugins/Runtime/Metasound/Source/MetasoundEditor/Private/MetasoundEditorGraph.cpp b/Engine/Plugins/Runtime/Metasound/Source/MetasoundEditor/Private/MetasoundEditorGraph.cpp
--- a/Engine/Plugins/Runtime/Metasound/Source/MetasoundEditor/Private/MetasoundEditorGraph.cpp
+++ b/Engine/Plugins/Runtime/Metasound/Source/MetasoundEditor/Private/MetasoundEditorGraph.cpp
@@ -377,6 +377,8 @@ void UMetasoundEditorGraphVertex::SetVertexAccessType(EMetasoundFrontendVertexAc
 	using namespace Metasound::Editor;
 	using namespace Metasound::Frontend;
 
+	ensureAlwaysMsgf(InNewAccessType != EMetasoundFrontendVertexAccessType::Unset, TEXT("Vertices do not support ambiguous access type."));
+
 	if (InNewAccessType == GetVertexAccessType())
 	{
 		return;
diff --git a/Engine/Plugins/Runtime/Metasound/Source/MetasoundEditor/Private/MetasoundEditorGraphBuilder.cpp b/Engine/Plugins/Runtime/Metasound/Source/MetasoundEditor/Private/MetasoundEditorGraphBuilder.cpp
--- a/Engine/Plugins/Runtime/Metasound/Source/MetasoundEditor/Private/MetasoundEditorGraphBuilder.cpp
+++ b/Engine/Plugins/Runtime/Metasound/Source/MetasoundEditor/Private/MetasoundEditorGraphBuilder.cpp
@@ -1179,7 +1179,9 @@ namespace Metasound
 			ClassInput.Name = GenerateUniqueNameByClassType(InMetaSound, EMetasoundFrontendClassType::Input, InNameBase ? InNameBase->ToString() : TEXT("Input"));
 			ClassInput.TypeName = InParams.DataType;
 			ClassInput.VertexID = VertexID;
-			ClassInput.AccessType = InParams.AccessType;
+
+			// Can be unset if attempting to mirror parameters from a reroute, so default to reference
+			ClassInput.AccessType = InParams.AccessType == EMetasoundFrontendVertexAccessType::Unset ? EMetasoundFrontendVertexAccessType::Reference : InParams.AccessType;
 
 			if (nullptr != InDefaultValue)
 			{
@@ -1212,7 +1214,9 @@ namespace Metasound
 			ClassOutput.Name = NewName;
 			ClassOutput.TypeName = InParams.DataType;
 			ClassOutput.VertexID = VertexID;
-			ClassOutput.AccessType = InParams.AccessType;
+
+			// Can be unset if attempting to mirror parameters from a reroute, so default to reference
+			ClassOutput.AccessType = InParams.AccessType == EMetasoundFrontendVertexAccessType::Unset ? EMetasoundFrontendVertexAccessType::Reference : InParams.AccessType;
 
 			return MetaSoundAsset->GetRootGraphHandle()->AddOutputVertex(ClassOutput);
 		}
diff --git a/Engine/Plugins/Runtime/Metasound/Source/MetasoundEditor/Private/MetasoundEditorGraphSchema.cpp b/Engine/Plugins/Runtime/Metasound/Source/MetasoundEditor/Private/MetasoundEditorGraphSchema.cpp
--- a/Engine/Plugins/Runtime/Metasound/Source/MetasoundEditor/Private/MetasoundEditorGraphSchema.cpp
+++ b/Engine/Plugins/Runtime/Metasound/Source/MetasoundEditor/Private/MetasoundEditorGraphSchema.cpp
@@ -529,7 +529,7 @@ UEdGraphNode* FMetasoundGraphSchemaAction_PromoteToInput::PerformAction(UEdGraph
 	using namespace Metasound::Editor;
 	using namespace Metasound::Frontend;
 
-	FInputHandle InputHandle = FGraphBuilder::GetInputHandleFromPin(FromPin);
+	FConstInputHandle InputHandle = FGraphBuilder::GetConstInputHandleFromPin(FromPin);
 	if (!ensure(InputHandle->IsValid()))
 	{
 		return nullptr;
@@ -1069,13 +1069,13 @@ void UMetasoundEditorGraphSchema::GetGraphContextActions(FGraphContextMenuBuilde
 
 	FActionClassFilters ClassFilters;
 	FConstGraphHandle GraphHandle = IGraphController::GetInvalidHandle();
-	bool bIsConstructorOutput = false;
+	EMetasoundFrontendVertexAccessType OutputAccessType = EMetasoundFrontendVertexAccessType::Unset;
 	if (const UEdGraphPin* FromPin = ContextMenuBuilder.FromPin)
 	{
 		if (FromPin->Direction == EGPD_Input)
 		{
 			FConstInputHandle InputHandle = FGraphBuilder::GetConstInputHandleFromPin(ContextMenuBuilder.FromPin);
-			bIsConstructorOutput = InputHandle->GetVertexAccessType() == EMetasoundFrontendVertexAccessType::Value;
+			OutputAccessType = InputHandle->GetVertexAccessType();
 
 			ClassFilters.OutputFilterFunction = [InputHandle](const FMetasoundFrontendClassOutput& InOutput)
 			{
@@ -1099,7 +1099,7 @@ void UMetasoundEditorGraphSchema::GetGraphContextActions(FGraphContextMenuBuilde
 			ActionMenuBuilder.AddAction(MakeShared<FMetasoundGraphSchemaAction_PromoteToInput>());
 
 			// Constructor outputs cannot be promoted to variables
-			if (!bIsConstructorOutput)
+			if (OutputAccessType != EMetasoundFrontendVertexAccessType::Value)
 			{
 				ActionMenuBuilder.AddAction(MakeShared<FMetasoundGraphSchemaAction_PromoteToVariable_AccessorNode>());
 				ActionMenuBuilder.AddAction(MakeShared<FMetasoundGraphSchemaAction_PromoteToVariable_DeferredAccessorNode>());
@@ -1162,7 +1162,7 @@ void UMetasoundEditorGraphSchema::GetGraphContextActions(FGraphContextMenuBuilde
 	GetFunctionActions(ContextMenuBuilder, ClassFilters, true /* bShowSelectedActions */, GraphHandle);
 
 	// Variable and conversion actions are always by reference so are incompatible with constructor outputs 
-	if (!bIsConstructorOutput)
+	if (OutputAccessType != EMetasoundFrontendVertexAccessType::Value)
 	{
 		GetVariableActions(ContextMenuBuilder, ClassFilters, true /* bShowSelectedActions */, GraphHandle);
 		GetConversionActions(ContextMenuBuilder, ClassFilters);
diff --git a/Engine/Plugins/Runtime/Metasound/Source/MetasoundEditor/Private/SMetasoundGraphPin.h b/Engine/Plugins/Runtime/Metasound/Source/MetasoundEditor/Private/SMetasoundGraphPin.h
--- a/Engine/Plugins/Runtime/Metasound/Source/MetasoundEditor/Private/SMetasoundGraphPin.h
+++ b/Engine/Plugins/Runtime/Metasound/Source/MetasoundEditor/Private/SMetasoundGraphPin.h
@@ -211,7 +211,7 @@ namespace Metasound
 			virtual const FSlateBrush* GetPinIcon() const override
 			{
 				const bool bIsConnected = ParentPinType::IsConnected();
-				EMetasoundFrontendVertexAccessType AccessType = EMetasoundFrontendVertexAccessType::Reference;
+				EMetasoundFrontendVertexAccessType AccessType = EMetasoundFrontendVertexAccessType::Unset;
 
 				if (const UEdGraphPin* Pin = ParentPinType::GetPinObj())
 				{
diff --git a/Engine/Plugins/Runtime/Metasound/Source/MetasoundFrontend/Private/MetasoundFrontendController.cpp b/Engine/Plugins/Runtime/Metasound/Source/MetasoundFrontend/Private/MetasoundFrontendController.cpp
--- a/Engine/Plugins/Runtime/Metasound/Source/MetasoundFrontend/Private/MetasoundFrontendController.cpp
+++ b/Engine/Plugins/Runtime/Metasound/Source/MetasoundFrontend/Private/MetasoundFrontendController.cpp
@@ -4,6 +4,8 @@
 
 #include "MetasoundFrontendInvalidController.h"
 #include "MetasoundFrontendDocumentController.h"
+#include "NodeTemplates/MetasoundFrontendNodeTemplateReroute.h"
+
 
 namespace Metasound
 {
@@ -76,5 +78,93 @@ namespace Metasound
 		{
 			return InDocumentAccessor.ShareAccess();
 		}
-	}
-}
+
+		FConstOutputHandle FindReroutedOutput(FConstOutputHandle InOutputHandle)
+		{
+			using namespace Frontend;
+
+			if (InOutputHandle->IsValid())
+			{
+				FConstNodeHandle NodeHandle = InOutputHandle->GetOwningNode();
+				if (NodeHandle->IsValid())
+				{
+					if (NodeHandle->GetClassMetadata().GetClassName() == FRerouteNodeTemplate::ClassName)
+					{
+						TArray<FConstInputHandle> Inputs = NodeHandle->GetConstInputs();
+						if (!Inputs.IsEmpty())
+						{
+							FConstInputHandle RerouteInputHandle = Inputs.Last();
+							if (RerouteInputHandle->IsValid())
+							{
+								FConstOutputHandle ConnectedOutputHandle = RerouteInputHandle->GetConnectedOutput();
+								return FindReroutedOutput(ConnectedOutputHandle);
+							}
+						}
+					}
+				}
+			}
+
+			return InOutputHandle;
+		}
+
+		void FindReroutedInputs(FConstInputHandle InHandleToCheck, TArray<FConstInputHandle>& InOutInputHandles)
+		{
+			using namespace Frontend;
+
+			if (InHandleToCheck->IsValid())
+			{
+				FConstNodeHandle NodeHandle = InHandleToCheck->GetOwningNode();
+				if (NodeHandle->IsValid())
+				{
+					if (NodeHandle->GetClassMetadata().GetClassName() == FRerouteNodeTemplate::ClassName)
+					{
+						TArray<FConstOutputHandle> Outputs = NodeHandle->GetConstOutputs();
+						for (FConstOutputHandle& OutputHandle : Outputs)
+						{
+							TArray<FConstInputHandle> LinkedInputs = OutputHandle->GetConstConnectedInputs();
+							for (FConstInputHandle LinkedInput : LinkedInputs)
+							{
+								FindReroutedInputs(LinkedInput, InOutInputHandles);
+							}
+						}
+
+						return;
+					}
+
+					InOutInputHandles.Add(InHandleToCheck);
+				}
+
+			}
+		}
+
+		void IterateReroutedInputs(FConstInputHandle InHandleToCheck, TFunctionRef<void(FConstInputHandle)> Func)
+		{
+			using namespace Frontend;
+
+			if (InHandleToCheck->IsValid())
+			{
+				FConstNodeHandle NodeHandle = InHandleToCheck->GetOwningNode();
+				if (NodeHandle->IsValid())
+				{
+					if (NodeHandle->GetClassMetadata().GetClassName() == FRerouteNodeTemplate::ClassName)
+					{
+						TArray<FConstOutputHandle> Outputs = NodeHandle->GetConstOutputs();
+						for (FConstOutputHandle& OutputHandle : Outputs)
+						{
+							TArray<FConstInputHandle> LinkedInputs = OutputHandle->GetConstConnectedInputs();
+							for (FConstInputHandle LinkedInput : LinkedInputs)
+							{
+								IterateReroutedInputs(LinkedInput, Func);
+							}
+						}
+
+						return;
+					}
+
+					Func(InHandleToCheck);
+				}
+
+			}
+		}
+	} // namespace Frontend
+} // namespace Metasound
diff --git a/Engine/Plugins/Runtime/Metasound/Source/MetasoundFrontend/Private/MetasoundFrontendDocument.cpp b/Engine/Plugins/Runtime/Metasound/Source/MetasoundFrontend/Private/MetasoundFrontendDocument.cpp
--- a/Engine/Plugins/Runtime/Metasound/Source/MetasoundFrontend/Private/MetasoundFrontendDocument.cpp
+++ b/Engine/Plugins/Runtime/Metasound/Source/MetasoundFrontend/Private/MetasoundFrontendDocument.cpp
@@ -167,14 +167,19 @@ bool FMetasoundFrontendClassVertex::IsFunctionalEquivalent(const FMetasoundFront
 
 bool FMetasoundFrontendClassVertex::CanConnectVertexAccessTypes(EMetasoundFrontendVertexAccessType InFromType, EMetasoundFrontendVertexAccessType InToType)
 {
-	if (EMetasoundFrontendVertexAccessType::Value == InToType)
+	// Reroute nodes can have undefined access type, so if either is unset, then connection is valid.
+	if (EMetasoundFrontendVertexAccessType::Unset != InFromType && EMetasoundFrontendVertexAccessType::Unset != InToType)
 	{
-		// If the input vertex accesses by "Value" then the output vertex 
-		// must also access by "Value" to enforce unexpected consequences 
-		// of connecting data which varies over time to an input which only
-		// evaluates the data during operator initialization.
-		return (EMetasoundFrontendVertexAccessType::Value == InFromType);
+		if (EMetasoundFrontendVertexAccessType::Value == InToType)
+		{
+			// If the input vertex accesses by "Value" then the output vertex 
+			// must also access by "Value" to enforce unexpected consequences 
+			// of connecting data which varies over time to an input which only
+			// evaluates the data during operator initialization.
+			return (EMetasoundFrontendVertexAccessType::Value == InFromType);
+		}
 	}
+
 	return true;
 }
 
@@ -594,7 +599,10 @@ const TCHAR* LexToString(EMetasoundFrontendVertexAccessType InVertexAccess)
 			return TEXT("Value");
 			
 		case EMetasoundFrontendVertexAccessType::Reference:
-		default:
 			return TEXT("Reference");
+
+		case EMetasoundFrontendVertexAccessType::Unset:
+		default:
+			return TEXT("Unset");
 	}
 }
diff --git a/Engine/Plugins/Runtime/Metasound/Source/MetasoundFrontend/Private/MetasoundFrontendGraph.cpp b/Engine/Plugins/Runtime/Metasound/Source/MetasoundFrontend/Private/MetasoundFrontendGraph.cpp
--- a/Engine/Plugins/Runtime/Metasound/Source/MetasoundFrontend/Private/MetasoundFrontendGraph.cpp
+++ b/Engine/Plugins/Runtime/Metasound/Source/MetasoundFrontend/Private/MetasoundFrontendGraph.cpp
@@ -223,6 +223,7 @@ namespace Metasound
 						MoveTemp(Literal)
 					};
 
+					ensureAlwaysMsgf(InOwningGraphClassInput.AccessType != EMetasoundFrontendVertexAccessType::Unset, TEXT("Graph Class Input cannot be set to access type of 'Unset'"));
 					if (InOwningGraphClassInput.AccessType == EMetasoundFrontendVertexAccessType::Reference)
 					{
 						return DataTypeRegistry.CreateInputNode(InputVertex.TypeName, MoveTemp(InitParams));
diff --git a/Engine/Plugins/Runtime/Metasound/Source/MetasoundFrontend/Private/MetasoundFrontendInputController.cpp b/Engine/Plugins/Runtime/Metasound/Source/MetasoundFrontend/Private/MetasoundFrontendInputController.cpp
--- a/Engine/Plugins/Runtime/Metasound/Source/MetasoundFrontend/Private/MetasoundFrontendInputController.cpp
+++ b/Engine/Plugins/Runtime/Metasound/Source/MetasoundFrontend/Private/MetasoundFrontendInputController.cpp
@@ -57,14 +57,39 @@ namespace Metasound
 			return Invalid::GetInvalidName();
 		}
 
-		EMetasoundFrontendVertexAccessType FBaseInputController::GetVertexAccessType() const 
+		EMetasoundFrontendVertexAccessType FBaseInputController::GetVertexAccessType() const
 		{
-			if (const FMetasoundFrontendClassVertex* ClassInput = ClassInputPtr.Get())
+			EMetasoundFrontendVertexAccessType AccessType = EMetasoundFrontendVertexAccessType::Unset;
+			bool bIsRerouted = false;
+
+			Frontend::IterateReroutedInputs(AsShared(), [this, &bIsRerouted, &AccessType](const FConstInputHandle& ReroutedInput)
 			{
-				return ClassInput->AccessType;
-			}
-			
-			return EMetasoundFrontendVertexAccessType::Reference;
+				bIsRerouted = true;
+
+				if (AccessType != EMetasoundFrontendVertexAccessType::Value)
+				{
+					if (ReroutedInput->IsValid())
+					{
+						// If ReroutedInput is top-level controller, iterator function is returning self, so just report if set to value.
+						if (ReroutedInput->GetID() == GetID() && ReroutedInput->GetOwningNodeID() == GetOwningNodeID())
+						{
+							if (const FMetasoundFrontendClassVertex* ClassInput = ClassInputPtr.Get())
+							{
+								AccessType = ClassInput->AccessType;
+								return;
+							}
+						}
+
+						const EMetasoundFrontendVertexAccessType RerouteAccessType = ReroutedInput->GetVertexAccessType();
+						if (RerouteAccessType == EMetasoundFrontendVertexAccessType::Value)
+						{
+							AccessType = RerouteAccessType;
+						}
+					}
+				}
+			});
+
+			return bIsRerouted ? AccessType : EMetasoundFrontendVertexAccessType::Reference;
 		}
 
 #if WITH_EDITOR
diff --git a/Engine/Plugins/Runtime/Metasound/Source/MetasoundFrontend/Private/MetasoundFrontendOutputController.cpp b/Engine/Plugins/Runtime/Metasound/Source/MetasoundFrontend/Private/MetasoundFrontendOutputController.cpp
--- a/Engine/Plugins/Runtime/Metasound/Source/MetasoundFrontend/Private/MetasoundFrontendOutputController.cpp
+++ b/Engine/Plugins/Runtime/Metasound/Source/MetasoundFrontend/Private/MetasoundFrontendOutputController.cpp
@@ -49,13 +49,23 @@ namespace Metasound
 			return Invalid::GetInvalidName();
 		}
 
-		EMetasoundFrontendVertexAccessType FBaseOutputController::GetVertexAccessType() const 
+		EMetasoundFrontendVertexAccessType FBaseOutputController::GetVertexAccessType() const
 		{
+			FConstOutputHandle ReroutedOutput = Frontend::FindReroutedOutput(AsShared());
+			if (ReroutedOutput->IsValid())
+			{
+				if (ReroutedOutput->GetOwningNodeID() != GetOwningNodeID() || ReroutedOutput->GetID() != GetID())
+				{
+					return ReroutedOutput->GetVertexAccessType();
+				}
+			}
+
 			if (const FMetasoundFrontendClassVertex* ClassOutput = ClassOutputPtr.Get())
 			{
 				return ClassOutput->AccessType;
 			}
-			return EMetasoundFrontendVertexAccessType::Reference;
+
+			return EMetasoundFrontendVertexAccessType::Unset;
 		}
 
 		FGuid FBaseOutputController::GetOwningNodeID() const
diff --git a/Engine/Plugins/Runtime/Metasound/Source/MetasoundFrontend/Private/MetasoundFrontendRegistries.cpp b/Engine/Plugins/Runtime/Metasound/Source/MetasoundFrontend/Private/MetasoundFrontendRegistries.cpp
--- a/Engine/Plugins/Runtime/Metasound/Source/MetasoundFrontend/Private/MetasoundFrontendRegistries.cpp
+++ b/Engine/Plugins/Runtime/Metasound/Source/MetasoundFrontend/Private/MetasoundFrontendRegistries.cpp
@@ -395,22 +395,36 @@ namespace Metasound
 		bool FRegistryContainerImpl::FindInputNodeRegistryKeyForDataType(const FName& InDataTypeName, const EMetasoundFrontendVertexAccessType InAccessType, FNodeRegistryKey& OutKey)
 		{
 			FMetasoundFrontendClass Class;
-			if (InAccessType == EMetasoundFrontendVertexAccessType::Reference)
+			switch (InAccessType)
 			{
-				if (IDataTypeRegistry::Get().GetFrontendInputClass(InDataTypeName, Class))
+				case EMetasoundFrontendVertexAccessType::Reference:
 				{
-					OutKey = NodeRegistryKey::CreateKey(Class.Metadata);
-					return true;
+					if (IDataTypeRegistry::Get().GetFrontendInputClass(InDataTypeName, Class))
+					{
+						OutKey = NodeRegistryKey::CreateKey(Class.Metadata);
+						return true;
+					}
 				}
-			}
-			else // InAccessType == EMetasoundFrontendVertexAccessType::Value
-			{
-				if (IDataTypeRegistry::Get().GetFrontendConstructorInputClass(InDataTypeName, Class))
+				break;
+
+				case EMetasoundFrontendVertexAccessType::Value:
 				{
-					OutKey = NodeRegistryKey::CreateKey(Class.Metadata);
-					return true;
+					if (IDataTypeRegistry::Get().GetFrontendConstructorInputClass(InDataTypeName, Class))
+					{
+						OutKey = NodeRegistryKey::CreateKey(Class.Metadata);
+						return true;
+					}
+				}
+				break;
+
+				default:
+				case EMetasoundFrontendVertexAccessType::Unset:
+				{
+					return false;
 				}
+				break;
 			}
+
 			return false;
 		}
 
@@ -434,21 +448,27 @@ namespace Metasound
 		bool FRegistryContainerImpl::FindOutputNodeRegistryKeyForDataType(const FName& InDataTypeName, const EMetasoundFrontendVertexAccessType InAccessType, FNodeRegistryKey& OutKey)
 		{
 			FMetasoundFrontendClass Class;
-			if (InAccessType == EMetasoundFrontendVertexAccessType::Reference) 
+			switch (InAccessType)
 			{
-				if (IDataTypeRegistry::Get().GetFrontendOutputClass(InDataTypeName, Class))
+				case EMetasoundFrontendVertexAccessType::Reference:
 				{
-					OutKey = NodeRegistryKey::CreateKey(Class.Metadata);
-					return true;
+					if (IDataTypeRegistry::Get().GetFrontendOutputClass(InDataTypeName, Class))
+					{
+						OutKey = NodeRegistryKey::CreateKey(Class.Metadata);
+						return true;
+					}
 				}
-			}
-			else // InAccessType == EMetasoundFrontendVertexAccessType::Value
-			{
-				if (IDataTypeRegistry::Get().GetFrontendConstructorOutputClass(InDataTypeName, Class))
+				break;
+
+				case EMetasoundFrontendVertexAccessType::Value:
 				{
-					OutKey = NodeRegistryKey::CreateKey(Class.Metadata);
-					return true;
+					if (IDataTypeRegistry::Get().GetFrontendConstructorOutputClass(InDataTypeName, Class))
+					{
+						OutKey = NodeRegistryKey::CreateKey(Class.Metadata);
+						return true;
+					}
 				}
+				break;
 			}
 
 			return false;
diff --git a/Engine/Plugins/Runtime/Metasound/Source/MetasoundFrontend/Private/NodeTemplates/MetasoundFrontendNodeTemplateReroute.cpp b/Engine/Plugins/Runtime/Metasound/Source/MetasoundFrontend/Private/NodeTemplates/MetasoundFrontendNodeTemplateReroute.cpp
--- a/Engine/Plugins/Runtime/Metasound/Source/MetasoundFrontend/Private/NodeTemplates/MetasoundFrontendNodeTemplateReroute.cpp
+++ b/Engine/Plugins/Runtime/Metasound/Source/MetasoundFrontend/Private/NodeTemplates/MetasoundFrontendNodeTemplateReroute.cpp
@@ -14,63 +14,6 @@ namespace Metasound
 	{
 		namespace ReroutePrivate
 		{
-			FConstOutputHandle FindReroutedOutput(FConstOutputHandle InOutputHandle)
-			{
-				using namespace Frontend;
-
-				if (InOutputHandle->IsValid())
-				{
-					FConstNodeHandle NodeHandle = InOutputHandle->GetOwningNode();
-					if (NodeHandle->IsValid())
-					{
-						if (NodeHandle->GetClassMetadata().GetClassName() == FRerouteNodeTemplate::ClassName)
-						{
-							TArray<FConstInputHandle> Inputs = NodeHandle->GetConstInputs();
-							if (!Inputs.IsEmpty())
-							{
-								FConstInputHandle RerouteInputHandle = Inputs.Last();
-								if (RerouteInputHandle->IsValid())
-								{
-									FConstOutputHandle ConnectedOutputHandle = RerouteInputHandle->GetConnectedOutput();
-									return FindReroutedOutput(ConnectedOutputHandle);
-								}
-							}
-						}
-					}
-				}
-
-				return InOutputHandle;
-			}
-
-			void FindReroutedInputs(FConstInputHandle InHandleToCheck, TArray<FConstInputHandle>& InOutInputHandles)
-			{
-				using namespace Frontend;
-
-				if (InHandleToCheck->IsValid())
-				{
-					FConstNodeHandle NodeHandle = InHandleToCheck->GetOwningNode();
-					if (NodeHandle->IsValid())
-					{
-						if (NodeHandle->GetClassMetadata().GetClassName() == FRerouteNodeTemplate::ClassName)
-						{
-							TArray<FConstOutputHandle> Outputs = NodeHandle->GetConstOutputs();
-							for (FConstOutputHandle& OutputHandle : Outputs)
-							{
-								TArray<FConstInputHandle> LinkedInputs = OutputHandle->GetConstConnectedInputs();
-								for (FConstInputHandle LinkedInput : LinkedInputs)
-								{
-									FindReroutedInputs(LinkedInput, InOutInputHandles);
-								}
-							}
-
-							return;
-						}
-					}
-
-					InOutInputHandles.Add(InHandleToCheck);
-				}
-			}
-
 			class FRerouteNodeTemplatePreprocessTransform : public FNodeTemplatePreprocessTransformBase
 			{
 			public:
@@ -222,11 +165,12 @@ namespace Metasound
 		{
 			TArray<FConstOutputHandle> Outputs = InNodeHandle->GetConstOutputs();
 			TArray<FConstInputHandle> Inputs = InNodeHandle->GetConstInputs();
-			const bool bConnectedToNonRerouteOutputs = Algo::AnyOf(Outputs, [](const FConstOutputHandle& OutputHandle) { return ReroutePrivate::FindReroutedOutput(OutputHandle)->IsValid(); });
+
+			const bool bConnectedToNonRerouteOutputs = Algo::AnyOf(Outputs, [](const FConstOutputHandle& OutputHandle) { return Frontend::FindReroutedOutput(OutputHandle)->IsValid(); });
 			const bool bConnectedToNonRerouteInputs = Algo::AnyOf(Inputs, [](const FConstInputHandle& InputHandle)
 			{
 				TArray<FConstInputHandle> Inputs;
-				ReroutePrivate::FindReroutedInputs(InputHandle, Inputs);
+				Frontend::FindReroutedInputs(InputHandle, Inputs);
 				return !Inputs.IsEmpty();
 			});
 
diff --git a/Engine/Plugins/Runtime/Metasound/Source/MetasoundFrontend/Public/MetasoundFrontendController.h b/Engine/Plugins/Runtime/Metasound/Source/MetasoundFrontend/Public/MetasoundFrontendController.h
--- a/Engine/Plugins/Runtime/Metasound/Source/MetasoundFrontend/Public/MetasoundFrontendController.h
+++ b/Engine/Plugins/Runtime/Metasound/Source/MetasoundFrontend/Public/MetasoundFrontendController.h
@@ -10,6 +10,8 @@
 #include "MetasoundFrontendRegistries.h"
 #include "MetasoundGraph.h"
 #include "MetasoundVertex.h"
+#include "Templates/Function.h"
+#include "Templates/SharedPointer.h"
 #include "Templates/Tuple.h"
 
 
@@ -118,7 +120,6 @@ namespace Metasound
 			FMetasoundFrontendClass RegistryClass;
 		};
 
-
 		/** IDocumentAccessor describes an interface for various I*Controllers to interact with
 		 * each other without exposing that interface publicly or requiring friendship 
 		 * between various controller implementation classes. 
@@ -1118,5 +1119,9 @@ namespace Metasound
 			/** Exports the document to a json formatted string. */
 			virtual FString ExportToJSON() const = 0;
 		};
-	}
-}
+
+		METASOUNDFRONTEND_API FConstOutputHandle FindReroutedOutput(FConstOutputHandle InOutputHandle);
+		METASOUNDFRONTEND_API void FindReroutedInputs(FConstInputHandle InHandleToCheck, TArray<FConstInputHandle>& InOutInputHandles);
+		METASOUNDFRONTEND_API void IterateReroutedInputs(FConstInputHandle InHandleToCheck, TFunctionRef<void(FConstInputHandle)> Func);
+	} // namespace Frontend
+} // namespace Metasound
diff --git a/Engine/Plugins/Runtime/Metasound/Source/MetasoundFrontend/Public/MetasoundFrontendDocument.h b/Engine/Plugins/Runtime/Metasound/Source/MetasoundFrontend/Public/MetasoundFrontendDocument.h
--- a/Engine/Plugins/Runtime/Metasound/Source/MetasoundFrontend/Public/MetasoundFrontendDocument.h
+++ b/Engine/Plugins/Runtime/Metasound/Source/MetasoundFrontend/Public/MetasoundFrontendDocument.h
@@ -52,7 +52,11 @@ UENUM()
 enum class EMetasoundFrontendVertexAccessType
 {
 	Reference,	//< The vertex accesses data by reference.
-	Value 		//< The vertex accesses data by value.
+	Value,		//< The vertex accesses data by value.
+
+	Unset		//< The vertex access level is unset (ex. vertex on an unconnected reroute node).
+				//< Not reflected as a graph core access type as core does not deal with reroutes
+				//< or ambiguous accessor level (it is resolved during document pre-processing).
 };
 
 UENUM()
