Set invalid controller's default access type to unset
#rnx

diff --git a/Engine/Plugins/Runtime/Metasound/Source/MetasoundFrontend/Private/MetasoundFrontendInvalidController.h b/Engine/Plugins/Runtime/Metasound/Source/MetasoundFrontend/Private/MetasoundFrontendInvalidController.h
--- a/Engine/Plugins/Runtime/Metasound/Source/MetasoundFrontend/Private/MetasoundFrontendInvalidController.h
+++ b/Engine/Plugins/Runtime/Metasound/Source/MetasoundFrontend/Private/MetasoundFrontendInvalidController.h
@@ -60,7 +60,7 @@ namespace Metasound
 			virtual FGuid GetID() const override { return Metasound::FrontendInvalidID; }
 			virtual const FName& GetDataType() const override { return Invalid::GetInvalidName(); }
 			virtual const FVertexName& GetName() const override { return Invalid::GetInvalidName(); }
-			virtual EMetasoundFrontendVertexAccessType GetVertexAccessType() const override { return EMetasoundFrontendVertexAccessType::Reference; }
+			virtual EMetasoundFrontendVertexAccessType GetVertexAccessType() const override { return EMetasoundFrontendVertexAccessType::Unset; }
 
 #if WITH_EDITOR
 			virtual FText GetDisplayName() const override { return FText::GetEmpty(); }
