Fix bug with synchronization not being completed prior to validation for referenced MetaSound graphs

TLBSRA (Too long but still read anyway)
This was due to validation moving to recursive behavior but sync not to support bubbling up referenced errors/warnings in 5.1.  This resulted in missing members after interface transforms are applied and potentially hitting ensures in validation.
...
But wait there's more!!! Synchronization on tick is now more efficient as it doesn't attempt to resynchronize/validate references already synchronized/validated in prior pass in same frame due to referencing graph needing to sync

#rnx

diff --git a/Engine/Plugins/Runtime/Metasound/Source/MetasoundEditor/Private/MetasoundEditor.cpp b/Engine/Plugins/Runtime/Metasound/Source/MetasoundEditor/Private/MetasoundEditor.cpp
--- a/Engine/Plugins/Runtime/Metasound/Source/MetasoundEditor/Private/MetasoundEditor.cpp
+++ b/Engine/Plugins/Runtime/Metasound/Source/MetasoundEditor/Private/MetasoundEditor.cpp
@@ -55,6 +55,7 @@
 #include "SMetasoundPalette.h"
 #include "SNodePanel.h"
 #include "Stats/Stats.h"
+#include "Templates/Function.h"
 #include "Templates/SharedPointer.h"
 #include "Widgets/Docking/SDockTab.h"
 #include "Widgets/Images/SImage.h"
@@ -81,6 +82,30 @@ namespace Metasound
 			LOCTEXT("NodeSectionName_Variables", "Variables")
 		};
 
+		namespace EditorPrivate
+		{
+			void ForEachReferencedMetaSoundEditor(const UObject& InObject, TFunctionRef<void(TSharedRef<FEditor>)> Func)
+			{
+				using namespace Frontend;
+
+				if (const FMetasoundAssetBase* MetaSoundAsset = IMetasoundUObjectRegistry::Get().GetObjectAsAssetBase(&InObject))
+				{
+					TArray<FMetasoundAssetBase*> References;
+					ensureAlways(IMetaSoundAssetManager::GetChecked().TryLoadReferencedAssets(*MetaSoundAsset, References));
+					for (FMetasoundAssetBase* Reference : References)
+					{
+						UObject* RefMetaSound = Reference->GetOwningAsset();
+						check(RefMetaSound);
+						TSharedPtr<FEditor> RefEditor = FGraphBuilder::GetEditorForMetasound(*RefMetaSound);
+						if (RefEditor.IsValid())
+						{
+							Func(RefEditor.ToSharedRef());
+						}
+					}
+				}
+			}
+		} // namespace EditorPrivate
+
 		class FMetasoundGraphMemberSchemaAction : public FEdGraphSchemaAction
 		{
 		public:
@@ -2758,16 +2783,36 @@ namespace Metasound
 			}
 		}
 
-		void FEditor::RefreshDetails()
+		void FEditor::RefreshDetails(bool bRefreshReferencedAssetEditors)
 		{
+			using namespace Frontend;
+
+			if (bRefreshReferencedAssetEditors && Metasound)
+			{
+				EditorPrivate::ForEachReferencedMetaSoundEditor(*Metasound, [bRefreshReferencedAssetEditors](TSharedRef<FEditor> RefEditor)
+				{
+					RefEditor->RefreshDetails(bRefreshReferencedAssetEditors);
+				});
+			}
+
 			if (MetasoundDetails.IsValid())
 			{
 				MetasoundDetails->ForceRefresh();
 			}
 		}
 
-		void FEditor::RefreshInterfaces()
+		void FEditor::RefreshInterfaces(bool bRefreshReferencedAssetEditors)
 		{
+			using namespace Frontend;
+
+			if (bRefreshReferencedAssetEditors && Metasound)
+			{
+				EditorPrivate::ForEachReferencedMetaSoundEditor(*Metasound, [bRefreshReferencedAssetEditors](TSharedRef<FEditor> RefEditor)
+				{
+					RefEditor->RefreshInterfaces(bRefreshReferencedAssetEditors);
+				});
+			}
+
 			if (InterfacesDetails.IsValid())
 			{
 				InterfacesDetails->ForceRefresh();
@@ -3297,17 +3342,14 @@ namespace Metasound
 
 				// Presets always update interfaces
 				const FMetasoundFrontendGraphClass& RootGraphClass = MetasoundAsset->GetDocumentHandle()->GetRootGraphClass();
-				const bool bIsPreset = RootGraphClass.PresetOptions.bIsPreset;
-				if (bIsPreset)
-				{
-					RefreshInterfaces();
-				}
 
+				const bool bIsPreset = RootGraphClass.PresetOptions.bIsPreset;
+				constexpr bool bRefreshReferencedAssets = true;
 				if (bShouldRefreshDetails || bIsPreset)
 				{
 					// TODO: Break up this synchronization flag
-					RefreshDetails();
-					RefreshInterfaces();
+					RefreshDetails(bRefreshReferencedAssets);
+					RefreshInterfaces(bRefreshReferencedAssets);
 				}
 				else
 				{
@@ -3324,7 +3366,7 @@ namespace Metasound
 
 					if (Algo::AnyOf(MetasoundDetails->GetSelectedObjects(), ShouldRefreshDetails))
 					{
-						RefreshDetails();
+						RefreshDetails(bRefreshReferencedAssets);
 					}
 				}
 			}
diff --git a/Engine/Plugins/Runtime/Metasound/Source/MetasoundEditor/Private/MetasoundEditor.h b/Engine/Plugins/Runtime/Metasound/Source/MetasoundEditor/Private/MetasoundEditor.h
--- a/Engine/Plugins/Runtime/Metasound/Source/MetasoundEditor/Private/MetasoundEditor.h
+++ b/Engine/Plugins/Runtime/Metasound/Source/MetasoundEditor/Private/MetasoundEditor.h
@@ -154,7 +154,7 @@ namespace Metasound
 			void DuplicateNodes();
 
 			/** Forces all UX pertaining to the root graph's details panel to be refreshed. */
-			void RefreshDetails();
+			void RefreshDetails(bool bRefreshReferencedAssetEditors = false);
 			
 			/** Pastes node(s) from the clipboard to the graph */
 			void PasteNodes(const FVector2D* InLocation = nullptr);
@@ -168,7 +168,7 @@ namespace Metasound
 			void RefreshGraphMemberMenu();
 
 			/** Forces refresh of interfaces list. */
-			void RefreshInterfaces();
+			void RefreshInterfaces(bool bRefreshReferencedAssetEditors = false);
 
 			/** Updates selected node classes to highest class found in the MetaSound Class Registry. */
 			void UpdateSelectedNodeClasses();
diff --git a/Engine/Plugins/Runtime/Metasound/Source/MetasoundEditor/Private/MetasoundEditorGraphBuilder.cpp b/Engine/Plugins/Runtime/Metasound/Source/MetasoundEditor/Private/MetasoundEditorGraphBuilder.cpp
--- a/Engine/Plugins/Runtime/Metasound/Source/MetasoundEditor/Private/MetasoundEditorGraphBuilder.cpp
+++ b/Engine/Plugins/Runtime/Metasound/Source/MetasoundEditor/Private/MetasoundEditorGraphBuilder.cpp
@@ -79,6 +79,59 @@ namespace Metasound
 
 				return FName(*NewName);
 			}
+
+			bool SynchronizeGraphRecursively(UObject& InMetaSound, bool bForceRefreshNodes, bool bEditorGraphModified = false)
+			{
+				using namespace Frontend;
+
+				FMetasoundAssetBase* MetaSoundAsset = IMetasoundUObjectRegistry::Get().GetObjectAsAssetBase(&InMetaSound);
+				check(MetaSoundAsset);
+
+				// Synchronize referenced graphs first to ensure all editor data
+				// is up-to-date prior to synchronizing this referencing graph.
+				TArray<FMetasoundAssetBase*> References;
+				ensureAlways(IMetaSoundAssetManager::GetChecked().TryLoadReferencedAssets(*MetaSoundAsset, References));
+				for (FMetasoundAssetBase* Reference : References)
+				{
+					check(Reference);
+					SynchronizeGraphRecursively(*Reference->GetOwningAsset(), bForceRefreshNodes, bEditorGraphModified);
+				}
+
+				// If no graph is set, MetaSound has been created outside of asset factory, so initialize it here.
+				// TODO: Move factory initialization and this code to single builder function (in header so cannot move
+				// until 5.1+).
+				if (!MetaSoundAsset->GetGraph())
+				{
+					FString Author = UKismetSystemLibrary::GetPlatformUserName();
+					if (const UMetasoundEditorSettings* EditorSettings = GetDefault<UMetasoundEditorSettings>())
+					{
+						if (!EditorSettings->DefaultAuthor.IsEmpty())
+						{
+							Author = EditorSettings->DefaultAuthor;
+						}
+					}
+
+					FGraphBuilder::InitMetaSound(InMetaSound, Author);
+
+					// Initial graph generation is not something to be managed by the transaction
+					// stack, so don't track dirty state until after initial setup if necessary.
+					UMetasoundEditorGraph* Graph = NewObject<UMetasoundEditorGraph>(&InMetaSound, FName(), RF_Transactional);
+					Graph->Schema = UMetasoundEditorGraphSchema::StaticClass();
+					MetaSoundAsset->SetGraph(Graph);
+				}
+
+				bEditorGraphModified |= FGraphBuilder::SynchronizeGraphMembers(InMetaSound);
+				bEditorGraphModified |= FGraphBuilder::SynchronizeNodeMembers(InMetaSound);
+				bEditorGraphModified |= FGraphBuilder::SynchronizeNodes(InMetaSound);
+				bEditorGraphModified |= FGraphBuilder::SynchronizeConnections(InMetaSound);
+
+				if (bEditorGraphModified)
+				{
+					InMetaSound.MarkPackageDirty();
+				}
+
+				return bEditorGraphModified;
+			}
 		} // namespace GraphBuilderPrivate
 
 		// Categories corresponding with POD DataTypes
@@ -2023,47 +2076,26 @@ namespace Metasound
 
 		int32 FGraphBuilder::SynchronizeGraph(UObject& InMetaSound, bool bForceRefreshNodes)
 		{
+			using namespace Frontend;
+
+			int32 HighestMessageSeverity = EMessageSeverity::Info;
+
 			FMetasoundAssetBase* MetaSoundAsset = IMetasoundUObjectRegistry::Get().GetObjectAsAssetBase(&InMetaSound);
 			check(MetaSoundAsset);
-
-			// If no graph is set, MetaSound has been created outside of asset factory, so initialize it here.
-			// TODO: Move factory initialization and this code to single builder function (in header so cannot move
-			// until 5.1+).
-			if(!MetaSoundAsset->GetGraph())
+			if (MetaSoundAsset->GetSynchronizationRequired())
 			{
-				FString Author = UKismetSystemLibrary::GetPlatformUserName();
-				if (const UMetasoundEditorSettings* EditorSettings = GetDefault<UMetasoundEditorSettings>())
-				{
-					if (!EditorSettings->DefaultAuthor.IsEmpty())
-					{
-						Author = EditorSettings->DefaultAuthor;
-					}
-				}
+				const bool bEditorGraphModified = GraphBuilderPrivate::SynchronizeGraphRecursively(InMetaSound, bForceRefreshNodes);
+				bForceRefreshNodes |= bEditorGraphModified;
+				HighestMessageSeverity = ValidateGraph(InMetaSound, bForceRefreshNodes);
 
-				FGraphBuilder::InitMetaSound(InMetaSound, Author);
-
-				// Initial graph generation is not something to be managed by the transaction
-				// stack, so don't track dirty state until after initial setup if necessary.
-				UMetasoundEditorGraph* Graph = NewObject<UMetasoundEditorGraph>(&InMetaSound, FName(), RF_Transactional);
-				Graph->Schema = UMetasoundEditorGraphSchema::StaticClass();
-				MetaSoundAsset->SetGraph(Graph);
+				MetaSoundAsset->ResetSynchronizationState();
 			}
-
-			bool bEditorGraphModified = SynchronizeGraphMembers(InMetaSound);
-			bEditorGraphModified |= SynchronizeNodeMembers(InMetaSound);
-			bEditorGraphModified |= SynchronizeNodes(InMetaSound);
-			bEditorGraphModified |= SynchronizeConnections(InMetaSound);
-
-			if (bEditorGraphModified)
+			else
 			{
-				InMetaSound.MarkPackageDirty();
+				const UMetasoundEditorGraph* Graph = CastChecked<UMetasoundEditorGraph>(MetaSoundAsset->GetGraph());
+				HighestMessageSeverity = Graph->GetHighestMessageSeverity();
 			}
 
-			bForceRefreshNodes |= bEditorGraphModified;
-			const int32 HighestMessageSeverity = FGraphBuilder::ValidateGraph(InMetaSound, bForceRefreshNodes);
-
-			MetaSoundAsset->ResetSynchronizationState();
-
 			return HighestMessageSeverity;
 		}
 
diff --git a/Engine/Plugins/Runtime/Metasound/Source/MetasoundEditor/Private/MetasoundEditorGraphBuilder.h b/Engine/Plugins/Runtime/Metasound/Source/MetasoundEditor/Private/MetasoundEditorGraphBuilder.h
--- a/Engine/Plugins/Runtime/Metasound/Source/MetasoundEditor/Private/MetasoundEditorGraphBuilder.h
+++ b/Engine/Plugins/Runtime/Metasound/Source/MetasoundEditor/Private/MetasoundEditorGraphBuilder.h
@@ -248,8 +248,10 @@ namespace Metasound
 			// Refreshes pin state from class FrontendClassVertexMetadata
 			static void RefreshPinMetadata(UEdGraphPin& InPin, const FMetasoundFrontendVertexMetadata& InMetadata);
 
-			// Adds and removes nodes, pins and connections so that the UEdGraph of the MetaSound matches the FMetasoundFrontendDocumentModel
+			// Adds and removes nodes, pins and connections so that the UEdGraph of the MetaSound matches the
+			// FMetasoundFrontendDocument model. Validates the graph (and those referenced recursively).
 			//
+			// @param InMetaSound - MetaSound to synchronize and optionally validate.
 			// @param bForceRefreshNodes - Refreshes all transient data stored on editor nodes & redraws them
 			// @return Highest MessageSeverity integer value listed on validated nodes post synchronization.
 			static int32 SynchronizeGraph(UObject& InMetaSound, bool bForceRefreshNodes = false);
diff --git a/Engine/Plugins/Runtime/Metasound/Source/MetasoundEngine/Private/Metasound.cpp b/Engine/Plugins/Runtime/Metasound/Source/MetasoundEngine/Private/Metasound.cpp
--- a/Engine/Plugins/Runtime/Metasound/Source/MetasoundEngine/Private/Metasound.cpp
+++ b/Engine/Plugins/Runtime/Metasound/Source/MetasoundEngine/Private/Metasound.cpp
@@ -2,8 +2,8 @@
 #include "Metasound.h"
 
 #include "AssetRegistry/AssetRegistryModule.h"
-#include "CoreMinimal.h"
 #include "Internationalization/Text.h"
+#include "Logging/TokenizedMessage.h"
 #include "MetasoundAssetBase.h"
 #include "MetasoundAudioFormats.h"
 #include "MetasoundEngineArchetypes.h"
@@ -31,6 +31,23 @@
 #define LOCTEXT_NAMESPACE "MetaSound"
 
 
+int32 UMetasoundEditorGraphBase::GetHighestMessageSeverity() const
+{
+	int32 HighestMessageSeverity = EMessageSeverity::Info;
+
+	for (const UEdGraphNode* Node : Nodes)
+	{
+		// Lower integer value is "higher severity"
+		if (Node->ErrorType < HighestMessageSeverity)
+		{
+			HighestMessageSeverity = Node->ErrorType;
+		}
+	}
+
+	return HighestMessageSeverity;
+}
+
+
 UMetaSoundPatch::UMetaSoundPatch(const FObjectInitializer& ObjectInitializer)
 	: Super(ObjectInitializer)
 	, FMetasoundAssetBase()
diff --git a/Engine/Plugins/Runtime/Metasound/Source/MetasoundEngine/Public/Metasound.h b/Engine/Plugins/Runtime/Metasound/Source/MetasoundEngine/Public/Metasound.h
--- a/Engine/Plugins/Runtime/Metasound/Source/MetasoundEngine/Public/Metasound.h
+++ b/Engine/Plugins/Runtime/Metasound/Source/MetasoundEngine/Public/Metasound.h
@@ -38,6 +38,8 @@ public:
 	virtual bool GetVersionedOnLoad() const PURE_VIRTUAL(UMetasoundEditorGraphBase::GetVersionedOnLoad, return false; )
 	virtual void SetVersionedOnLoad() PURE_VIRTUAL(UMetasoundEditorGraphBase::SetVersionedOnLoad, )
 #endif // WITH_EDITORONLY_DATA
+
+	int32 GetHighestMessageSeverity() const;
 };
 
 
diff --git a/Engine/Plugins/Runtime/Metasound/Source/MetasoundFrontend/Private/MetasoundAssetBase.cpp b/Engine/Plugins/Runtime/Metasound/Source/MetasoundFrontend/Private/MetasoundAssetBase.cpp
--- a/Engine/Plugins/Runtime/Metasound/Source/MetasoundFrontend/Private/MetasoundAssetBase.cpp
+++ b/Engine/Plugins/Runtime/Metasound/Source/MetasoundFrontend/Private/MetasoundAssetBase.cpp
@@ -536,12 +536,48 @@ void FMetasoundAssetBase::CacheRegistryMetadata()
 
 bool FMetasoundAssetBase::GetSynchronizationRequired() const
 {
-	return bSynchronizationRequired;
+	using namespace Metasound::Frontend;
+
+	if (bSynchronizationRequired)
+	{
+		return true;
+	}
+
+	TArray<FMetasoundAssetBase*> References;
+	ensureAlways(IMetaSoundAssetManager::GetChecked().TryLoadReferencedAssets(*this, References));
+	for (FMetasoundAssetBase* Reference : References)
+	{
+		const bool bReferenceSyncRequired = Reference->GetSynchronizationRequired();
+		if (bReferenceSyncRequired)
+		{
+			return true;
+		}
+	}
+
+	return false;
 }
 
 bool FMetasoundAssetBase::GetSynchronizationUpdateDetails() const
 {
-	return bSynchronizationUpdateDetails;
+	using namespace Metasound::Frontend;
+
+	if (bSynchronizationUpdateDetails)
+	{
+		return true;
+	}
+
+	TArray<FMetasoundAssetBase*> References;
+	ensureAlways(IMetaSoundAssetManager::GetChecked().TryLoadReferencedAssets(*this, References));
+	for (FMetasoundAssetBase* Reference : References)
+	{
+		const bool bReferenceSyncRequired = Reference->GetSynchronizationUpdateDetails();
+		if (bReferenceSyncRequired)
+		{
+			return true;
+		}
+	}
+
+	return false;
 }
 
 void FMetasoundAssetBase::SetSynchronizationRequired()
