Fix for alt-left clicking reroute nodes not being deleted
#rnx

diff --git a/Engine/Plugins/Runtime/Metasound/Source/MetasoundEditor/Private/MetasoundEditorGraphSchema.cpp b/Engine/Plugins/Runtime/Metasound/Source/MetasoundEditor/Private/MetasoundEditorGraphSchema.cpp
--- a/Engine/Plugins/Runtime/Metasound/Source/MetasoundEditor/Private/MetasoundEditorGraphSchema.cpp
+++ b/Engine/Plugins/Runtime/Metasound/Source/MetasoundEditor/Private/MetasoundEditorGraphSchema.cpp
@@ -1462,6 +1462,24 @@ void UMetasoundEditorGraphSchema::TrySetDefaultValue(UEdGraphPin& Pin, const FSt
 	return Super::TrySetDefaultValue(Pin, InNewDefaultValue, bInMarkAsModified);
 }
 
+bool UMetasoundEditorGraphSchema::SafeDeleteNodeFromGraph(UEdGraph* Graph, UEdGraphNode* InNodeToDelete) const
+{
+	using namespace Metasound::Editor;
+
+	if (!InNodeToDelete || !Graph || InNodeToDelete->GetGraph() != Graph)
+	{
+		return false;
+	}
+
+	if (UObject* MetaSound = Graph->GetOutermostObject())
+	{
+		MetaSound->Modify();
+	}
+	Graph->Modify();
+
+	return FGraphBuilder::DeleteNode(*InNodeToDelete);
+}
+
 bool UMetasoundEditorGraphSchema::ShouldHidePinDefaultValue(UEdGraphPin* Pin) const
 {
 	using namespace Metasound::Editor;
diff --git a/Engine/Plugins/Runtime/Metasound/Source/MetasoundEditor/Public/MetasoundEditorGraphSchema.h b/Engine/Plugins/Runtime/Metasound/Source/MetasoundEditor/Public/MetasoundEditorGraphSchema.h
--- a/Engine/Plugins/Runtime/Metasound/Source/MetasoundEditor/Public/MetasoundEditorGraphSchema.h
+++ b/Engine/Plugins/Runtime/Metasound/Source/MetasoundEditor/Public/MetasoundEditorGraphSchema.h
@@ -473,6 +473,7 @@ public:
 	virtual bool TryCreateConnection(UEdGraphPin* A, UEdGraphPin* B) const override;
 	virtual void TrySetDefaultObject(UEdGraphPin& Pin, UObject* NewDefaultObject, bool bInMarkAsModified) const override;
 	virtual void TrySetDefaultValue(UEdGraphPin& Pin, const FString& InNewDefaultValue, bool bInMarkAsModified) const override;
+	virtual bool SafeDeleteNodeFromGraph(UEdGraph* Graph, UEdGraphNode* InNodeToDelete) const override;
 	virtual bool ShouldHidePinDefaultValue(UEdGraphPin* Pin) const override;
 	virtual FLinearColor GetPinTypeColor(const FEdGraphPinType& PinType) const override;
 	virtual void BreakNodeLinks(UEdGraphNode& TargetNode) const override;
