Missing recursive call to reset sync state in prior CL

diff --git a/Engine/Plugins/Runtime/Metasound/Source/MetasoundFrontend/Private/MetasoundAssetBase.cpp b/Engine/Plugins/Runtime/Metasound/Source/MetasoundFrontend/Private/MetasoundAssetBase.cpp
--- a/Engine/Plugins/Runtime/Metasound/Source/MetasoundFrontend/Private/MetasoundAssetBase.cpp
+++ b/Engine/Plugins/Runtime/Metasound/Source/MetasoundFrontend/Private/MetasoundAssetBase.cpp
@@ -593,8 +593,17 @@ void FMetasoundAssetBase::SetUpdateDetailsOnSynchronization()
 
 void FMetasoundAssetBase::ResetSynchronizationState()
 {
+	using namespace Metasound::Frontend;
+
 	bSynchronizationUpdateDetails = false;
 	bSynchronizationRequired = false;
+
+	TArray<FMetasoundAssetBase*> References;
+	ensureAlways(IMetaSoundAssetManager::GetChecked().TryLoadReferencedAssets(*this, References));
+	for (FMetasoundAssetBase* Reference : References)
+	{
+		Reference->ResetSynchronizationState();
+	}
 }
 #endif // WITH_EDITOR
 
