Make MetaSound parameter error logs into warnings
Add log warning when constructor input value is ignored on runnign instance.

diff --git a/Engine/Plugins/Runtime/Metasound/Source/MetasoundEngine/Private/MetasoundSource.cpp b/Engine/Plugins/Runtime/Metasound/Source/MetasoundEngine/Private/MetasoundSource.cpp
--- a/Engine/Plugins/Runtime/Metasound/Source/MetasoundEngine/Private/MetasoundSource.cpp
+++ b/Engine/Plugins/Runtime/Metasound/Source/MetasoundEngine/Private/MetasoundSource.cpp
@@ -545,7 +545,10 @@ void UMetaSoundSource::InitParameters(TArray<FAudioParameter>& InParametersToIni
 			else
 			{
 #if !NO_LOGGING
-				UE_LOG(LogMetaSound, Error, TEXT("Failed to set invalid parameter '%s' in asset '%s': Either does not exist or is unsupported type"), *Parameter.ParamName.ToString(), *AssetName);
+				if (::Metasound::MetaSoundParameterEnableWarningOnIgnoredParameterCVar)
+				{
+					UE_LOG(LogMetaSound, Warning, TEXT("Failed to set invalid parameter '%s' in asset '%s': Either does not exist or is unsupported type"), *Parameter.ParamName.ToString(), *AssetName);
+				}
 #endif // !NO_LOGGING
 				constexpr bool bAllowShrinking = false;
 				InParametersToInit.RemoveAtSwap(i, 1, bAllowShrinking);
@@ -555,8 +558,12 @@ void UMetaSoundSource::InitParameters(TArray<FAudioParameter>& InParametersToIni
 		{
 			constexpr bool bAllowShrinking = false;
 			InParametersToInit.RemoveAtSwap(i, 1, bAllowShrinking);
+
 #if !NO_LOGGING
-			UE_LOG(LogMetaSound, Error, TEXT("Failed to set parameter '%s' in asset '%s': No name specified, no transmittable input found, or type mismatch."), *Parameter.ParamName.ToString(), *AssetName);
+			if (::Metasound::MetaSoundParameterEnableWarningOnIgnoredParameterCVar)
+			{
+				UE_LOG(LogMetaSound, Warning, TEXT("Failed to set parameter '%s' in asset '%s': No name specified, no transmittable input found, or type mismatch."), *Parameter.ParamName.ToString(), *AssetName);
+			}
 #endif // !NO_LOGGING
 		}
 	}
@@ -910,6 +917,8 @@ TUniquePtr<Audio::IParameterTransmitter> UMetaSoundSource::CreateParameterTransm
 		InitParams.Infos.Add(InfoAndName.SendInfo);
 	}
 
+	InitParams.DebugMetaSoundName = GetFName();
+
 	TUniquePtr<Audio::IParameterTransmitter> NewTransmitter = MakeUnique<Metasound::FMetaSoundParameterTransmitter>(InitParams);
 
 	return NewTransmitter;
diff --git a/Engine/Plugins/Runtime/Metasound/Source/MetasoundFrontend/Private/MetasoundParameterTransmitter.cpp b/Engine/Plugins/Runtime/Metasound/Source/MetasoundFrontend/Private/MetasoundParameterTransmitter.cpp
--- a/Engine/Plugins/Runtime/Metasound/Source/MetasoundFrontend/Private/MetasoundParameterTransmitter.cpp
+++ b/Engine/Plugins/Runtime/Metasound/Source/MetasoundFrontend/Private/MetasoundParameterTransmitter.cpp
@@ -9,6 +9,15 @@
 
 namespace Metasound
 {
+	int32 MetaSoundParameterEnableWarningOnIgnoredParameterCVar = 0;
+
+	FAutoConsoleVariableRef CVarMetaSoundParameterEnableWarningOnIgnoredParameter(
+		TEXT("au.MetaSound.Parameter.EnableWarningOnIgnoredParameter"),
+		MetaSoundParameterEnableWarningOnIgnoredParameterCVar,
+		TEXT("If enabled, a warning will be logged when a parameters sent to a MetaSound is ignored.\n")
+		TEXT("0: Disabled (default), !0: Enabled"),
+		ECVF_Default);
+
 	namespace Frontend
 	{
 		FLiteral ConvertParameterToLiteral(FAudioParameter&& InValue)
@@ -146,9 +155,10 @@ namespace Metasound
 	}
 
 	FMetaSoundParameterTransmitter::FMetaSoundParameterTransmitter(const FMetaSoundParameterTransmitter::FInitParams& InInitParams)
-	: SendInfos(InInitParams.Infos)
-	, OperatorSettings(InInitParams.OperatorSettings)
+	: OperatorSettings(InInitParams.OperatorSettings)
 	, InstanceID(InInitParams.InstanceID)
+	, DebugMetaSoundName(InInitParams.DebugMetaSoundName)
+	, SendInfos(InInitParams.Infos)
 	{
 	}
 
@@ -207,6 +217,12 @@ namespace Metasound
 			}
 		}
 
+		// Enable / disable via CVAR to avoid log spam.
+		if (MetaSoundParameterEnableWarningOnIgnoredParameterCVar)
+		{
+			UE_LOG(LogMetaSound, Warning, TEXT("Failed to set parameter %s in asset %s on instance %d with value %s. No runtime modifiable input with that name exists on instance."), *InParameterName.ToString(), *DebugMetaSoundName.ToString(), InstanceID, *LexToString(InLiteral));
+		}
+
 		return false;
 	}
 
diff --git a/Engine/Plugins/Runtime/Metasound/Source/MetasoundFrontend/Public/MetasoundParameterTransmitter.h b/Engine/Plugins/Runtime/Metasound/Source/MetasoundFrontend/Public/MetasoundParameterTransmitter.h
--- a/Engine/Plugins/Runtime/Metasound/Source/MetasoundFrontend/Public/MetasoundParameterTransmitter.h
+++ b/Engine/Plugins/Runtime/Metasound/Source/MetasoundFrontend/Public/MetasoundParameterTransmitter.h
@@ -12,9 +12,10 @@
 
 struct FMetasoundFrontendLiteral;
 
-
 namespace Metasound
 {
+	extern METASOUNDFRONTEND_API int32 MetaSoundParameterEnableWarningOnIgnoredParameterCVar;
+
 	namespace Frontend
 	{
 		METASOUNDFRONTEND_API FLiteral ConvertParameterToLiteral(FAudioParameter&& InValue);
@@ -62,13 +63,15 @@ namespace Metasound
 			/** Available input parameters on MetaSound instance. */
 			TArray<FSendInfo> Infos;
 
+			/** Name of MetaSound used to log parameter related errors. */
+			FName DebugMetaSoundName;
+
 			FInitParams(const FOperatorSettings& InSettings, uint64 InInstanceID, const TArray<FSendInfo>& InInfos=TArray<FSendInfo>())
 			: OperatorSettings(InSettings)
 			, InstanceID(InInstanceID)
 			, Infos(InInfos)
 			{
 			}
-
 		};
 
 		/** Creates a unique send address using the given MetaSound environment. */
@@ -114,9 +117,10 @@ namespace Metasound
 		// Create and store a new ISender for the given FSendInfo.
 		ISender* AddSender(const FSendInfo& InInfo);
 
-		TArray<FSendInfo> SendInfos;
 		FOperatorSettings OperatorSettings;
 		uint64 InstanceID;
+		FName DebugMetaSoundName;
+		TArray<FSendInfo> SendInfos;
 
 		TMap<FName, TUniquePtr<ISender>> InputSends;
 	};
