[MetaSounds] Editor Graph Syncronization Upgrade
- Move synchronize/details panel flag to more generic ModifyContext
- Do not use MarkAsGarbage as means of enforcing not to show object in details panel to avoid leaving objects in invalid state if undo called. Check ed graph membership instead
- Get rid of leaky namechanged delegates
    - Allows for migration to isolated synchronization in the future
    - Removes conflation of having to synchronize and update views when modifying frontend document
    - Decouples concept of a specific view to update from generic frontend graph modifications

diff --git a/Engine/Plugins/Runtime/Metasound/Source/MetasoundEditor/Private/MetasoundDetailCustomization.cpp b/Engine/Plugins/Runtime/Metasound/Source/MetasoundEditor/Private/MetasoundDetailCustomization.cpp
--- a/Engine/Plugins/Runtime/Metasound/Source/MetasoundEditor/Private/MetasoundDetailCustomization.cpp
+++ b/Engine/Plugins/Runtime/Metasound/Source/MetasoundEditor/Private/MetasoundDetailCustomization.cpp
@@ -375,7 +375,8 @@ namespace Metasound
 					if (InSelectInfo != ESelectInfo::OnNavigation)
 					{
 						FMetasoundFrontendInterface InterfaceToAdd;
-						if (ensure(ISearchEngine::Get().FindInterfaceWithHighestVersion(FName(*NameToAdd.Get()), InterfaceToAdd)))
+						const FName InterfaceName { *NameToAdd.Get() };
+						if (ensure(ISearchEngine::Get().FindInterfaceWithHighestVersion(InterfaceName, InterfaceToAdd)))
 						{
 							const FScopedTransaction Transaction(FText::Format(LOCTEXT("AddInterfaceTransactionFormat", "Add MetaSound Interface '{0}'"), FText::FromString(InterfaceToAdd.Version.ToString())));
 							MetaSound.Get()->Modify();
@@ -385,12 +386,10 @@ namespace Metasound
 							FModifyRootGraphInterfaces ModifyTransform({ }, { InterfaceToAdd });
 							ModifyTransform.SetDefaultNodeLocations(false); // Don't automatically add nodes to ed graph
 							ModifyTransform.Transform(DocumentHandle);
-							MetaSoundAsset->SetUpdateDetailsOnSynchronization();
 						}
 
 						UpdateInterfaceNames();
 						InterfaceComboBox->RefreshOptions();
-						MetaSoundAsset->SetSynchronizationRequired();
 						FGraphBuilder::RegisterGraphWithFrontend(*MetaSound.Get());
 					}
 				})
@@ -445,8 +444,6 @@ namespace Metasound
 
 						UpdateInterfaceNames();
 						InterfaceComboBox->RefreshOptions();
-						MetaSoundAsset->SetUpdateDetailsOnSynchronization();
-						MetaSoundAsset->SetSynchronizationRequired();
 						FGraphBuilder::RegisterGraphWithFrontend(*MetaSound.Get());
 
 					}), LOCTEXT("RemoveInterfaceTooltip1", "Removes all interfaces from the given MetaSound."))
@@ -471,7 +468,7 @@ namespace Metasound
 					return SNullWidget::NullWidget;
 				}
 
-				TSharedRef<SWidget> RemoveButtonWidget = PropertyCustomizationHelpers::MakeDeleteButton(FSimpleDelegate::CreateLambda([this, InterfaceEntry]()
+				TSharedRef<SWidget> RemoveButtonWidget = PropertyCustomizationHelpers::MakeDeleteButton(FSimpleDelegate::CreateLambda([this, InInterfaceName, InterfaceEntry]()
 				{
 					using namespace Frontend;
 
@@ -492,8 +489,6 @@ namespace Metasound
 
 					UpdateInterfaceNames();
 					InterfaceComboBox->RefreshOptions();
-					MetaSoundAsset->SetUpdateDetailsOnSynchronization();
-					MetaSoundAsset->SetSynchronizationRequired();
 					FGraphBuilder::RegisterGraphWithFrontend(*MetaSound.Get());
 
 				}), LOCTEXT("RemoveInterfaceTooltip2", "Removes the associated interface from the MetaSound."));
diff --git a/Engine/Plugins/Runtime/Metasound/Source/MetasoundEditor/Private/MetasoundEditor.cpp b/Engine/Plugins/Runtime/Metasound/Source/MetasoundEditor/Private/MetasoundEditor.cpp
--- a/Engine/Plugins/Runtime/Metasound/Source/MetasoundEditor/Private/MetasoundEditor.cpp
+++ b/Engine/Plugins/Runtime/Metasound/Source/MetasoundEditor/Private/MetasoundEditor.cpp
@@ -82,30 +82,6 @@ namespace Metasound
 			LOCTEXT("NodeSectionName_Variables", "Variables")
 		};
 
-		namespace EditorPrivate
-		{
-			void ForEachReferencedMetaSoundEditor(const UObject& InObject, TFunctionRef<void(TSharedRef<FEditor>)> Func)
-			{
-				using namespace Frontend;
-
-				if (const FMetasoundAssetBase* MetaSoundAsset = IMetasoundUObjectRegistry::Get().GetObjectAsAssetBase(&InObject))
-				{
-					TArray<FMetasoundAssetBase*> References;
-					ensureAlways(IMetaSoundAssetManager::GetChecked().TryLoadReferencedAssets(*MetaSoundAsset, References));
-					for (FMetasoundAssetBase* Reference : References)
-					{
-						UObject* RefMetaSound = Reference->GetOwningAsset();
-						check(RefMetaSound);
-						TSharedPtr<FEditor> RefEditor = FGraphBuilder::GetEditorForMetasound(*RefMetaSound);
-						if (RefEditor.IsValid())
-						{
-							Func(RefEditor.ToSharedRef());
-						}
-					}
-				}
-			}
-		} // namespace EditorPrivate
-
 		class FMetasoundGraphMemberSchemaAction : public FEdGraphSchemaAction
 		{
 		public:
@@ -595,41 +571,6 @@ namespace Metasound
 					TextWidget
 				];
 
-				// TODO: Move to right-click menu. Currently, too easy to accidentally hit & clutters interface
-// 				if (!InterfaceVersion.IsValid())
-// 				{
-// 					LayoutWidget->AddSlot()
-// 					.AutoWidth()
-// 					.HAlign(HAlign_Left)
-// 					.VAlign(VAlign_Center)
-// 					.Padding(2, 0, 0, 0)
-// 					[
-// 						PropertyCustomizationHelpers::MakeDeleteButton(FSimpleDelegate::CreateLambda([this]()
-// 						{
-// 							UMetasoundEditorGraphMember* GraphMember = MetasoundAction->GetGraphMember();
-// 							if (ensure(GraphMember))
-// 							{
-// 								if (UMetasoundEditorGraph* Graph = Cast<UMetasoundEditorGraph>(MetasoundAction->Graph))
-// 								{
-// 									if (UObject* MetaSound = Graph->GetMetasound())
-// 									{
-// 										const FScopedTransaction Transaction(LOCTEXT("MetaSoundEditorDeleteOnClick", "Delete MetaSound Graph Member"));
-// 										MetaSound->Modify();
-// 										Graph->Modify();
-// 										Graph->RemoveMember(*GraphMember);
-// 
-// 										//Synchronize will update the interface
-// 										if (FMetasoundAssetBase* MetasoundAsset = IMetasoundUObjectRegistry::Get().GetObjectAsAssetBase(MetaSound))
-// 										{
-// 											MetasoundAsset->SetSynchronizationRequired();
-// 										}
-// 									}
-// 								}
-// 							}
-// 						}))
-// 					];
-// 				}
-
 				return LayoutWidget;
 			}
 
@@ -779,32 +720,6 @@ namespace Metasound
 				Stop();
 			}
 
-			if (Metasound)
-			{
-				FMetasoundAssetBase* MetasoundAsset = IMetasoundUObjectRegistry::Get().GetObjectAsAssetBase(Metasound);
-				check(MetasoundAsset);
-
-				if(UMetasoundEditorGraph* Graph = Cast<UMetasoundEditorGraph>(MetasoundAsset->GetGraph()))
-				{
-					for (TPair<FGuid, FDelegateHandle>& Pair : NameChangeDelegateHandles)
-					{
-						if (UMetasoundEditorGraphInput* Input = Graph->FindInput(Pair.Key))
-						{
-							Input->NameChanged.Remove(Pair.Value);
-						}
-						else if (UMetasoundEditorGraphOutput* Output = Graph->FindOutput(Pair.Key))
-						{
-							Output->NameChanged.Remove(Pair.Value);
-						}
-						else if (UMetasoundEditorGraphVariable* Variable = Graph->FindVariable(Pair.Key))
-						{
-							Variable->NameChanged.Remove(Pair.Value);
-						}
-					}
-				}
-				NameChangeDelegateHandles.Reset();
-			}
-
 			InterfacesView.Reset();
 			DestroyAnalyzers();
 			check(GEditor);
@@ -841,11 +756,9 @@ namespace Metasound
 			const bool bRestartSound = IsPlaying();
 			if (!bRestartSound)
 			{
-				constexpr bool bForceRefreshNodes = true;
 				FGraphBuilder::RegisterGraphWithFrontend(*Metasound);
-				FGraphBuilder::SynchronizeGraph(*Metasound, bForceRefreshNodes);
 			}
-
+			RefreshEditorContext();
 
 			BindGraphCommands();
 			CreateInternalWidgets();
@@ -1129,7 +1042,7 @@ namespace Metasound
 				{
 					if (FMetasoundAssetBase* Asset = IMetasoundUObjectRegistry::Get().GetObjectAsAssetBase(Metasound))
 					{
-						Asset->SetSynchronizationRequired();
+						Asset->GetModifyContext().SetDocumentModified();
 					}
 				}
 			}
@@ -1565,7 +1478,7 @@ namespace Metasound
 		{
 			if (USoundBase* MetasoundToPlay = Cast<USoundBase>(Metasound))
 			{
-				HighestMessageSeverity = FGraphBuilder::SynchronizeGraph(*Metasound);
+				HighestMessageSeverity = GetMetaSoundGraphChecked().GetHighestMessageSeverity();
 				if (HighestMessageSeverity == EMessageSeverity::Error)
 				{
 					return;
@@ -2042,6 +1955,7 @@ namespace Metasound
 			UMetasoundEditorGraphMember* GraphMember = ActionToDelete->GetGraphMember();
 			if (ensure(GraphMember))
 			{
+				const FGuid MemberID = GraphMember->GetMemberID();
 				UMetasoundEditorGraph& Graph = GetMetaSoundGraphChecked();
 				UMetasoundEditorGraphMember* NextToSelect = Graph.FindAdjacentMember(*GraphMember);
 
@@ -2053,11 +1967,6 @@ namespace Metasound
 					Graph.RemoveMember(*GraphMember);
 				}
 
-				//Synchronize will update the interface
-				FMetasoundAssetBase* MetasoundAsset = IMetasoundUObjectRegistry::Get().GetObjectAsAssetBase(Metasound);
-				check(MetasoundAsset);
-				MetasoundAsset->SetSynchronizationRequired();
-
 				if (NextToSelect)
 				{
 					if (GraphMembersMenu->SelectItemByName(NextToSelect->GetMemberName(), ESelectInfo::Direct, static_cast<int32>(NextToSelect->GetSectionID())))
@@ -2068,7 +1977,9 @@ namespace Metasound
 				}
 				else
 				{
-					MetasoundAsset->SetUpdateDetailsOnSynchronization();
+					FMetasoundAssetBase* MetasoundAsset = IMetasoundUObjectRegistry::Get().GetObjectAsAssetBase(Metasound);
+					check(MetasoundAsset);
+					MetasoundAsset->GetModifyContext().AddMemberIDsModified({ MemberID });
 				}
 			}
 
@@ -2100,7 +2011,6 @@ namespace Metasound
 
 			FMetasoundAssetBase* MetasoundAsset = IMetasoundUObjectRegistry::Get().GetObjectAsAssetBase(Metasound);
 			check(MetasoundAsset);
-			MetasoundAsset->SetUpdateDetailsOnSynchronization();
 
 			const FScopedTransaction Transaction(LOCTEXT("MetaSoundEditorDeleteSelectedNode2", "Delete Selected MetaSound Node(s)"));
 			check(Metasound);
@@ -2627,7 +2537,6 @@ namespace Metasound
 				MetasoundGraphEditor->SetNodeSelection(GraphNode, true);
 			}
 
-			MetasoundAsset->SetSynchronizationRequired();
 
 			if (bNotifyReferenceLoop)
 			{
@@ -2783,36 +2692,18 @@ namespace Metasound
 			}
 		}
 
-		void FEditor::RefreshDetails(bool bRefreshReferencedAssetEditors)
+		void FEditor::RefreshDetails()
 		{
 			using namespace Frontend;
 
-			if (bRefreshReferencedAssetEditors && Metasound)
-			{
-				EditorPrivate::ForEachReferencedMetaSoundEditor(*Metasound, [bRefreshReferencedAssetEditors](TSharedRef<FEditor> RefEditor)
-				{
-					RefEditor->RefreshDetails(bRefreshReferencedAssetEditors);
-				});
-			}
-
 			if (MetasoundDetails.IsValid())
 			{
 				MetasoundDetails->ForceRefresh();
 			}
 		}
 
-		void FEditor::RefreshInterfaces(bool bRefreshReferencedAssetEditors)
+		void FEditor::RefreshInterfaces()
 		{
-			using namespace Frontend;
-
-			if (bRefreshReferencedAssetEditors && Metasound)
-			{
-				EditorPrivate::ForEachReferencedMetaSoundEditor(*Metasound, [bRefreshReferencedAssetEditors](TSharedRef<FEditor> RefEditor)
-				{
-					RefEditor->RefreshInterfaces(bRefreshReferencedAssetEditors);
-				});
-			}
-
 			if (InterfacesDetails.IsValid())
 			{
 				InterfacesDetails->ForceRefresh();
@@ -2823,7 +2714,23 @@ namespace Metasound
 		{
 			if (GraphMembersMenu.IsValid())
 			{
+				TArray<TSharedPtr<FEdGraphSchemaAction>> SelectedActions;
+				GraphMembersMenu->GetSelectedActions(SelectedActions);
+
 				GraphMembersMenu->RefreshAllActions(true /* bPreserveExpansion */, false /*bHandleOnSelectionEvent*/);
+
+				for (const TSharedPtr<FEdGraphSchemaAction>& Action : SelectedActions)
+				{
+					TSharedPtr<FMetasoundGraphMemberSchemaAction> MetasoundAction = StaticCastSharedPtr<FMetasoundGraphMemberSchemaAction>(Action);
+					if (MetasoundAction.IsValid())
+					{
+						if (UMetasoundEditorGraphMember* Member = MetasoundAction->GetGraphMember())
+						{
+							const FName ActionName = Member->GetMemberName();
+							GraphMembersMenu->SelectItemByName(ActionName, ESelectInfo::Direct, Action->GetSectionID());
+						}
+					}
+				}
 			}
 		}
 
@@ -2873,7 +2780,7 @@ namespace Metasound
 				FDocumentHandle DocumentHandle = Graph.GetDocumentHandle();
 				DocumentHandle->RemoveUnreferencedDependencies();
 				DocumentHandle->SynchronizeDependencyMetadata();
-				Graph.SetSynchronizationRequired();
+				Graph.GetModifyContext().SetDocumentModified();
 			}
 		}
 
@@ -2910,96 +2817,6 @@ namespace Metasound
 			GEditor->RedoTransaction();
 		}
 
-		void FEditor::OnInputNameChanged(FGuid InNodeID)
-		{
-			if (!GraphMembersMenu.IsValid() || !Metasound)
-			{
-				return;
-			}
-
-			TArray<TSharedPtr<FEdGraphSchemaAction>> SelectedActions;
-			GraphMembersMenu->GetSelectedActions(SelectedActions);
-			GraphMembersMenu->RefreshAllActions(/* bPreserveExpansion */ true);
-
-			for(const TSharedPtr<FEdGraphSchemaAction>& Action : SelectedActions)
-			{
-				TSharedPtr<FMetasoundGraphMemberSchemaAction> MetasoundAction = StaticCastSharedPtr<FMetasoundGraphMemberSchemaAction>(Action);
-				if (MetasoundAction.IsValid())
-				{
-					if (UMetasoundEditorGraphMember* Member = MetasoundAction->GetGraphMember())
-					{
-						if (InNodeID == Member->GetMemberID())
-						{
-							const FName ActionName = Member->GetMemberName();
-							GraphMembersMenu->SelectItemByName(ActionName, ESelectInfo::Direct, Action->GetSectionID());
-							break;
-						}
-					}
-				}
-			}
-
-			Metasound::Editor::FGraphBuilder::RegisterGraphWithFrontend(*Metasound);
-		}
-
-		void FEditor::OnOutputNameChanged(FGuid InNodeID)
-		{
-			if (!GraphMembersMenu.IsValid())
-			{
-				return;
-			}
-
-			TArray<TSharedPtr<FEdGraphSchemaAction>> SelectedActions;
-			GraphMembersMenu->GetSelectedActions(SelectedActions);
-			GraphMembersMenu->RefreshAllActions(/* bPreserveExpansion */ true);
-
-			for (const TSharedPtr<FEdGraphSchemaAction>& Action : SelectedActions)
-			{
-				TSharedPtr<FMetasoundGraphMemberSchemaAction> MetasoundAction = StaticCastSharedPtr<FMetasoundGraphMemberSchemaAction>(Action);
-				if (MetasoundAction.IsValid())
-				{
-					if (UMetasoundEditorGraphMember* Member = MetasoundAction->GetGraphMember())
-					{
-						if (InNodeID == Member->GetMemberID())
-						{
-							const FName ActionName = Member->GetMemberName();
-							GraphMembersMenu->SelectItemByName(ActionName, ESelectInfo::Direct, Action->GetSectionID());
-							break;
-						}
-					}
-				}
-			}
-
-			Metasound::Editor::FGraphBuilder::RegisterGraphWithFrontend(*Metasound);
-		}
-
-		void FEditor::OnVariableNameChanged(FGuid InVariableID)
-		{
-			if (!GraphMembersMenu.IsValid())
-			{
-				return;
-			}
-
-			TArray<TSharedPtr<FEdGraphSchemaAction>> SelectedActions;
-			GraphMembersMenu->GetSelectedActions(SelectedActions);
-			GraphMembersMenu->RefreshAllActions(/* bPreserveExpansion */ true);
-
-			for (const TSharedPtr<FEdGraphSchemaAction>& Action : SelectedActions)
-			{
-				TSharedPtr<FMetasoundGraphMemberSchemaAction> MetasoundAction = StaticCastSharedPtr<FMetasoundGraphMemberSchemaAction>(Action);
-				if (MetasoundAction.IsValid())
-				{
-					if (UMetasoundEditorGraphVariable* Variable = Cast<UMetasoundEditorGraphVariable>(MetasoundAction->GetGraphMember()))
-					{
-						if (InVariableID == Variable->GetVariableID())
-						{
-							GraphMembersMenu->SelectItemByName(Variable->GetMemberName(), ESelectInfo::Direct, Action->GetSectionID());
-							break;
-						}
-					}
-				}
-			}
-		}
-
 		void FEditor::CollectAllActions(FGraphActionListBuilderBase& OutAllActions)
 		{
 			const FMetasoundAssetBase* MetasoundAsset = IMetasoundUObjectRegistry::Get().GetObjectAsAssetBase(Metasound);
@@ -3037,16 +2854,6 @@ namespace Metasound
 				NewFuncAction->MemberID = NodeID;
 
 				ActionList->AddAction(NewFuncAction);
-
-				UMetasoundEditorGraphInput * EdGraphInput = EdGraph.FindInput(NodeID);
-				if (EdGraphInput)
-				{
-					if (FDelegateHandle* NameChangeDelegate = NameChangeDelegateHandles.Find(NodeID))
-					{
-						EdGraphInput->NameChanged.Remove(*NameChangeDelegate);
-					}
-					NameChangeDelegateHandles.FindOrAdd(NodeID) = EdGraphInput->NameChanged.AddSP(this, &FEditor::OnInputNameChanged);
-				}
 			}, EMetasoundFrontendClassType::Input);
 
 			FrontendGraph->IterateConstNodes([this, &GetMemberCategory, &EdGraph, ActionList = &OutAllActions](const Frontend::FConstNodeHandle& Output)
@@ -3062,16 +2869,6 @@ namespace Metasound
 				NewFuncAction->Graph = &EdGraph;
 				NewFuncAction->MemberID = Output->GetID();
 				ActionList->AddAction(NewFuncAction);
-
-				UMetasoundEditorGraphOutput* EdGraphOutput = EdGraph.FindOutput(NodeID);
-				if (ensure(EdGraphOutput))
-				{
-					if (FDelegateHandle* NameChangeDelegate = NameChangeDelegateHandles.Find(NodeID))
-					{
-						EdGraphOutput->NameChanged.Remove(*NameChangeDelegate);
-					}
-					NameChangeDelegateHandles.FindOrAdd(NodeID) = EdGraphOutput->NameChanged.AddSP(this, &FEditor::OnOutputNameChanged);
-				}
 			}, EMetasoundFrontendClassType::Output);
 
 			TArray<Frontend::FConstVariableHandle> Variables = FrontendGraph->GetVariables();
@@ -3085,35 +2882,8 @@ namespace Metasound
 				NewFuncAction->Graph = &EdGraph;
 				NewFuncAction->MemberID = VariableID; 
 				OutAllActions.AddAction(NewFuncAction);
-
-				UMetasoundEditorGraphVariable* EdGraphVariable = EdGraph.FindVariable(VariableID);
-				if (ensure(EdGraphVariable))
-				{
-					if (FDelegateHandle* NameChangeDelegate = NameChangeDelegateHandles.Find(VariableID))
-					{
-						EdGraphVariable->NameChanged.Remove(*NameChangeDelegate);
-					}
-					NameChangeDelegateHandles.FindOrAdd(VariableID) = EdGraphVariable->NameChanged.AddSP(this, &FEditor::OnVariableNameChanged);
-				}
-			}
-
-			// In certain cases, while synchronizing the editor layer with the frontend, nodes
-			// associated with delegates are orphaned, but can still have stale handles
-			// associated.  Clear them out to avoid them being fired.
-			TArray<FGuid> StaleNodeGuids;
-			UMetasoundEditorGraph& Graph = GetMetaSoundGraphChecked();
-			for (const TPair<FGuid, FDelegateHandle>& Pair : NameChangeDelegateHandles)
-			{
-				if (!Graph.FindMember(Pair.Key))
-				{
-					StaleNodeGuids.Add(Pair.Key);
-				}
 			}
 
-			for (const FGuid& StaleNodeGuid : StaleNodeGuids)
-			{
-				NameChangeDelegateHandles.Remove(StaleNodeGuid);
-			}
 		}
 
 		void FEditor::CollectStaticSections(TArray<int32>& StaticSectionIDs)
@@ -3312,6 +3082,41 @@ namespace Metasound
 			return MenuBuilder.MakeWidget();
 		}
 
+		void FEditor::RemoveInvalidSelection()
+		{
+			if (MetasoundDetails.IsValid())
+			{
+				const TArray<TWeakObjectPtr<UObject>>& Objects = MetasoundDetails->GetSelectedObjects();
+				TArray<UObject*> NewSelection;
+
+				TSet<const UMetasoundEditorGraphMember*> GraphMembers;
+				GetMetaSoundGraphChecked().IterateMembers([&GraphMembers](UMetasoundEditorGraphMember& GraphMember) { GraphMembers.Add(&GraphMember); });
+
+				for (const TWeakObjectPtr<UObject>& Object : Objects)
+				{
+					if (Object.IsValid())
+					{
+						if (const UMetasoundEditorGraphMember* Member = Cast<UMetasoundEditorGraphMember>(Object.Get()))
+						{
+							if (GraphMembers.Contains(Member))
+							{
+								NewSelection.Add(Object.Get());
+							}
+						}
+						else
+						{
+							NewSelection.Add(Object.Get());
+						}
+					}
+				}
+
+				if (NewSelection.Num() != Objects.Num())
+				{
+					SetSelection(NewSelection);
+				}
+			}
+		}
+
 		void FEditor::Tick(float DeltaTime)
 		{
 			if (!Metasound)
@@ -3319,9 +3124,6 @@ namespace Metasound
 				return;
 			}
 
-			FMetasoundAssetBase* MetasoundAsset = IMetasoundUObjectRegistry::Get().GetObjectAsAssetBase(Metasound);
-			check(MetasoundAsset);
-
 			if (bPrimingRegistry)
 			{
 				IMetasoundEditorModule& MetaSoundEditorModule = FModuleManager::GetModuleChecked<IMetasoundEditorModule>("MetaSoundEditor");
@@ -3332,46 +3134,148 @@ namespace Metasound
 				}
 			}
 
-			if (MetasoundAsset->GetSynchronizationRequired())
-			{
-				MetasoundAsset->CacheRegistryMetadata();
+			RefreshEditorContext();
 
-				// Capture before synchronizing as the flag is cleared therein.
-				const bool bShouldRefreshDetails = MetasoundAsset->GetSynchronizationUpdateDetails();
-				HighestMessageSeverity = FGraphBuilder::SynchronizeGraph(*Metasound);
+			GraphConnectionManager.Update(DeltaTime);
+		}
+
+		void FEditor::RefreshEditorContext()
+		{
+			const bool bSynchronizedGraph = FGraphBuilder::SynchronizeGraph(*Metasound);
 
-				// Presets always update interfaces
-				const FMetasoundFrontendGraphClass& RootGraphClass = MetasoundAsset->GetDocumentHandle()->GetRootGraphClass();
+			FMetasoundAssetBase* MetasoundAsset = IMetasoundUObjectRegistry::Get().GetObjectAsAssetBase(Metasound);
+			check(MetasoundAsset);
+
+			// Capture after synchronizing as the modification state may be modified therein
+			const FMetasoundFrontendDocumentModifyContext& ModifyContext = MetasoundAsset->GetModifyContext();
+			const bool bForceRefreshViews = ModifyContext.GetForceRefreshViews();
+			const TSet<FName>& InterfacesModified = ModifyContext.GetInterfacesModified();
+			const TSet<FGuid>& MembersModified = ModifyContext.GetMemberIDsModified();
+			const TSet<FGuid>& NodesModified = ModifyContext.GetNodeIDsModified();
+			if (bSynchronizedGraph || bForceRefreshViews || !InterfacesModified.IsEmpty() || !NodesModified.IsEmpty() || !MembersModified.IsEmpty())
+			{
+				FGraphValidationResults Results = FGraphBuilder::ValidateGraph(*Metasound);
 
-				const bool bIsPreset = RootGraphClass.PresetOptions.bIsPreset;
-				constexpr bool bRefreshReferencedAssets = true;
-				if (bShouldRefreshDetails || bIsPreset)
+				for (const FGraphNodeValidationResult& Result : Results.GetResults())
 				{
-					// TODO: Break up this synchronization flag
-					RefreshDetails(bRefreshReferencedAssets);
-					RefreshInterfaces(bRefreshReferencedAssets);
+					UMetasoundEditorGraphNode& Node = Result.GetNodeChecked();
+					const bool bClassChanged = Node.ContainsClassChange();
+					const FText Title = Node.GetCachedTitle();
+					Node.CacheTitle();
+					const bool bTitleUpdated = !Title.IdenticalTo(Node.GetCachedTitle());
+					const bool bRefreshNode = NodesModified.Contains(Node.GetNodeID());
+					if (Result.GetHasDirtiedNode() || bTitleUpdated || bClassChanged || bForceRefreshViews || bRefreshNode)
+					{
+						Node.SyncChangeIDs();
+						if (MetasoundGraphEditor.IsValid())
+						{
+							MetasoundGraphEditor->RefreshNode(Node);
+						}
+					}
 				}
-				else
+
+				TArray<UObject*> Selection;
+
+				if (!MembersModified.IsEmpty() || bForceRefreshViews)
 				{
-					// Also refresh details if the object in the panel has gone invalid
-					auto ShouldRefreshDetails = [](TWeakObjectPtr<UObject> Obj)
+					RefreshGraphMemberMenu();
+
+					UMetasoundEditorGraph& Graph = GetMetaSoundGraphChecked();
+					for (const FGuid& MemberGuid : MembersModified)
 					{
-						if (!Obj.IsValid() || !IsValid(Obj.Get()))
+						if (UObject* Member = Graph.FindMember(MemberGuid))
 						{
-							return true;
+							// Currently only one member can be selected at a time, so only first found is added
+							Selection.Add(Member);
+							break;
 						}
+					}
+				}
 
-						return IMetasoundUObjectRegistry::Get().IsRegisteredClass(Obj.Get());
-					};
-
-					if (Algo::AnyOf(MetasoundDetails->GetSelectedObjects(), ShouldRefreshDetails))
+				// Only refresh details panel if
+				// 1. Forcing refresh with modify context option
+				// 2. The currently selected object(s) is/are modified.
+				// 3. If the selection is changed via the modify context, it will automatically dirty & refresh via 'SetSelection' below
+				if (bForceRefreshViews)
+				{
+					RefreshDetails();
+				}
+				else if (!NodesModified.IsEmpty() || !MembersModified.IsEmpty())
+				{
+					TArray<TWeakObjectPtr<UObject>> SelectedObjects = MetasoundDetails->GetSelectedObjects();
+					const bool bShouldRefreshDetails = Algo::AnyOf(SelectedObjects, [&NodesModified, &MembersModified](const TWeakObjectPtr<UObject>& Obj)
 					{
-						RefreshDetails(bRefreshReferencedAssets);
+						if (const UMetasoundEditorGraphNode* Node = Cast<const UMetasoundEditorGraphNode>(Obj.Get()))
+						{
+							return NodesModified.Contains(Node->GetNodeID());
+						}
+						if (const UMetasoundEditorGraphMember* Member = Cast<const UMetasoundEditorGraphMember>(Obj.Get()))
+						{
+							return MembersModified.Contains(Member->GetMemberID());
+						}
+						return false;
+					});
+					if (bShouldRefreshDetails)
+					{
+						RefreshDetails();
 					}
+
+					// TODO: Because input editor nodes are not one-to-one, this can cause multi-selection
+					// when not desired.  Once input alias templates are complete, this can be switched on.
+					// For now, callsites are just directly setting selection via ClearSelectionAndSelectNode.
+// 					MetasoundGraphEditor->ClearSelectionSet();
+// 
+// 					const UMetasoundEditorGraph& Graph = GetMetaSoundGraphChecked();
+// 
+// 					TArray<UObject*> NodeSelection;
+// 					for (UEdGraphNode* Node : Graph.Nodes)
+// 					{
+// 						if (UMetasoundEditorGraphNode* MetaSoundNode = Cast<UMetasoundEditorGraphNode>(Node))
+// 						{
+// 							if (NodesModified.Contains(MetaSoundNode->GetNodeID()))
+// 							{
+// 								if (UMetasoundEditorGraphMemberNode* MemberNode = Cast<UMetasoundEditorGraphMemberNode>(Node))
+// 								{
+// 									NodeSelection.Add(MemberNode->GetMember());
+// 								}
+// 								else
+// 								{
+// 									NodeSelection.Add(Node);
+// 								}
+// 
+// 								MetasoundGraphEditor->SetNodeSelection(Node, /*bSelect=*/true);
+// 
+// 							}
+// 						}
+// 					}
+// 
+// 					if (!NodeSelection.IsEmpty())
+// 					{
+// 						Selection = NodeSelection;
+// 					}
 				}
-			}
 
-			GraphConnectionManager.Update(DeltaTime);
+				if (!InterfacesModified.IsEmpty() || bForceRefreshViews)
+				{
+					RefreshInterfaces();
+				}
+
+				HighestMessageSeverity = GetMetaSoundGraphChecked().GetHighestMessageSeverity();
+
+				// Modify data has been observed both from synchronization & by
+				// updating views by this point, so full reset is completed here.
+				MetasoundAsset->GetModifyContext().Reset();
+
+				if (!Selection.IsEmpty())
+				{
+					SetSelection(Selection);
+
+				}
+
+				// Avoids details panel displaying
+				// removed members in certain cases.
+				RemoveInvalidSelection();
+			}
 		}
 
 		TStatId FEditor::GetStatId() const
@@ -3500,12 +3404,6 @@ namespace Metasound
 						TObjectPtr<UMetasoundEditorGraphInput> Input = Graph.FindOrAddInput(NodeHandle);
 						if (ensure(Input))
 						{
-							FGuid NodeID = NodeHandle->GetID();
-							if (FDelegateHandle* NameChangeDelegate = NameChangeDelegateHandles.Find(NodeID))
-							{
-								Input->NameChanged.Remove(*NameChangeDelegate);
-							}
-							NameChangeDelegateHandles.FindOrAdd(NodeID) = Input->NameChanged.AddSP(this, &FEditor::OnInputNameChanged);
 							SelectedObjects.Add(Input);
 						}
 					}
@@ -3529,12 +3427,6 @@ namespace Metasound
 						TObjectPtr<UMetasoundEditorGraphOutput> Output = Graph.FindOrAddOutput(NodeHandle);
 						if (ensure(Output))
 						{
-							FGuid NodeID = NodeHandle->GetID();
-							if (FDelegateHandle* NameChangeDelegate = NameChangeDelegateHandles.Find(NodeID))
-							{
-								Output->NameChanged.Remove(*NameChangeDelegate);
-							}
-							NameChangeDelegateHandles.FindOrAdd(NodeID) = Output->NameChanged.AddSP(this, &FEditor::OnOutputNameChanged);
 							SelectedObjects.Add(Output);
 						}
 					}
@@ -3554,12 +3446,6 @@ namespace Metasound
 						TObjectPtr<UMetasoundEditorGraphVariable> EditorVariable = Graph.FindOrAddVariable(FrontendVariable);
 						if (ensure(EditorVariable))
 						{
-							FGuid VariableID = FrontendVariable->GetID();
-							if (FDelegateHandle* NameChangeDelegate = NameChangeDelegateHandles.Find(VariableID))
-							{
-								EditorVariable->NameChanged.Remove(*NameChangeDelegate);
-							}
-							NameChangeDelegateHandles.FindOrAdd(VariableID) = EditorVariable->NameChanged.AddSP(this, &FEditor::OnVariableNameChanged);
 							SelectedObjects.Add(EditorVariable);
 							NameToSelect = EditorVariable->GetMemberName();
 						}
diff --git a/Engine/Plugins/Runtime/Metasound/Source/MetasoundEditor/Private/MetasoundEditor.h b/Engine/Plugins/Runtime/Metasound/Source/MetasoundEditor/Private/MetasoundEditor.h
--- a/Engine/Plugins/Runtime/Metasound/Source/MetasoundEditor/Private/MetasoundEditor.h
+++ b/Engine/Plugins/Runtime/Metasound/Source/MetasoundEditor/Private/MetasoundEditor.h
@@ -154,7 +154,7 @@ namespace Metasound
 			void DuplicateNodes();
 
 			/** Forces all UX pertaining to the root graph's details panel to be refreshed. */
-			void RefreshDetails(bool bRefreshReferencedAssetEditors = false);
+			void RefreshDetails();
 			
 			/** Pastes node(s) from the clipboard to the graph */
 			void PasteNodes(const FVector2D* InLocation = nullptr);
@@ -168,7 +168,7 @@ namespace Metasound
 			void RefreshGraphMemberMenu();
 
 			/** Forces refresh of interfaces list. */
-			void RefreshInterfaces(bool bRefreshReferencedAssetEditors = false);
+			void RefreshInterfaces();
 
 			/** Updates selected node classes to highest class found in the MetaSound Class Registry. */
 			void UpdateSelectedNodeClasses();
@@ -183,10 +183,6 @@ namespace Metasound
 				return MetasoundGraphEditor->GetSelectedNodes().Num();
 			}
 
-			void OnInputNameChanged(FGuid InNodeID);
-			void OnOutputNameChanged(FGuid InNodeID);
-			void OnVariableNameChanged(FGuid InVariableID);
-
 			/** Creates analyzers */
 			void CreateAnalyzers();
 
@@ -325,7 +321,11 @@ namespace Metasound
 			/** Called to redo the last undone action */
 			void RedoGraphAction();
 
+			void RefreshEditorContext();
+
 		private:
+			void RemoveInvalidSelection();
+
 			void SetPreviewID(uint32 InPreviewID);
 
 			/** FNotifyHook interface */
@@ -430,8 +430,6 @@ namespace Metasound
 			/** The Metasound asset being edited */
 			UObject* Metasound = nullptr;
 
-			TMap<FGuid, FDelegateHandle> NameChangeDelegateHandles;
-
 			/** Whether or not metasound being edited is valid */
 			bool bPassedValidation = true;
 
diff --git a/Engine/Plugins/Runtime/Metasound/Source/MetasoundEditor/Private/MetasoundEditorGraph.cpp b/Engine/Plugins/Runtime/Metasound/Source/MetasoundEditor/Private/MetasoundEditorGraph.cpp
--- a/Engine/Plugins/Runtime/Metasound/Source/MetasoundEditor/Private/MetasoundEditorGraph.cpp
+++ b/Engine/Plugins/Runtime/Metasound/Source/MetasoundEditor/Private/MetasoundEditorGraph.cpp
@@ -40,7 +40,7 @@ namespace Metasound
 						FGraphBuilder::RegisterGraphWithFrontend(*MetaSound);
 						if (FMetasoundAssetBase* MetaSoundAsset = IMetasoundUObjectRegistry::Get().GetObjectAsAssetBase(MetaSound))
 						{
-							MetaSoundAsset->SetUpdateDetailsOnSynchronization();
+							MetaSoundAsset->GetModifyContext().AddMemberIDsModified({ InMember.GetMemberID() });
 						}
 					}
 				}
@@ -148,8 +148,7 @@ void UMetasoundEditorGraphMember::PostEditUndo()
 				{
 					FMetasoundAssetBase* MetasoundAsset = IMetasoundUObjectRegistry::Get().GetObjectAsAssetBase(MetaSound);
 					check(MetasoundAsset);
-					MetasoundAsset->SetUpdateDetailsOnSynchronization();
-					MetasoundAsset->SetSynchronizationRequired();
+					MetasoundAsset->GetModifyContext().SetDocumentModified();
 				}
 			}
 		}
@@ -239,25 +238,27 @@ void UMetasoundEditorGraphVertex::SetMemberName(const FName& InNewName, bool bPo
 	using namespace Metasound::Editor;
 	using namespace Metasound::Frontend;
 
+	UMetasoundEditorGraph* Graph = GetOwningGraph();
+	if (!ensure(Graph))
 	{
-		const FText TransactionLabel = FText::Format(LOCTEXT("RenameGraphVertexMemberNameFormat", "Set Metasound {0} MemberName"), GetGraphMemberLabel());
-		const FScopedTransaction Transaction(TransactionLabel, bPostTransaction);
+		return;
+	}
 
-		UMetasoundEditorGraph* Graph = GetOwningGraph();
-		if (Graph)
-		{
-			Graph->Modify();
-			Graph->GetMetasoundChecked().Modify();
-		}
+	FNodeHandle NodeHandle = GetNodeHandle();
+	if (NodeHandle->GetNodeName() == InNewName)
+	{
+		return;
+	}
 
-		GetNodeHandle()->SetNodeName(InNewName);
+	const FText TransactionLabel = FText::Format(LOCTEXT("RenameGraphVertexMemberNameFormat", "Set Metasound {0} MemberName"), GetGraphMemberLabel());
+	const FScopedTransaction Transaction(TransactionLabel, bPostTransaction);
 
-		FMetasoundAssetBase* MetasoundAsset = IMetasoundUObjectRegistry::Get().GetObjectAsAssetBase(&Graph->GetMetasoundChecked());
-		check(MetasoundAsset);
-		MetasoundAsset->SetSynchronizationRequired();
-	}
+	Graph->Modify();
+	Graph->GetMetasoundChecked().Modify();
 
-	NameChanged.Broadcast(NodeID);
+	NodeHandle->SetNodeName(InNewName);
+	Graph->RegisterGraphWithFrontend();
+	Graph->GetModifyContext().AddMemberIDsModified({ GetMemberID() });
 }
 
 FText UMetasoundEditorGraphVertex::GetDisplayName() const
@@ -270,33 +271,41 @@ void UMetasoundEditorGraphVertex::SetDisplayName(const FText& InNewName, bool bP
 {
 	using namespace Metasound::Frontend;
 
+	UMetasoundEditorGraph* Graph = Cast<UMetasoundEditorGraph>(GetOuter());
+	if (!ensure(Graph))
 	{
-		const FText TransactionLabel = FText::Format(LOCTEXT("RenameGraphVertexDisplayNameFormat", "Set Metasound {0} DisplayName"), GetGraphMemberLabel());
-		const FScopedTransaction Transaction(TransactionLabel, bPostTransaction);
+		return;
+	}
 
-		if (UMetasoundEditorGraph* Graph = Cast<UMetasoundEditorGraph>(GetOuter()))
-		{
-			Graph->Modify();
-			Graph->GetMetasoundChecked().Modify();
-		}
 
-		FNodeHandle NodeHandle = GetNodeHandle();
-		NodeHandle->SetDisplayName(InNewName);
+	FNodeHandle NodeHandle = GetNodeHandle();
+	if (NodeHandle->GetDisplayName().EqualTo(InNewName))
+	{
+		return;
+	}
 
-		const TArray<UMetasoundEditorGraphMemberNode*> Nodes = GetNodes();
-		for (UMetasoundEditorGraphMemberNode* Node : Nodes)
-		{
-			const TArray<UEdGraphPin*>& Pins = Node->GetAllPins();
-			ensure(Pins.Num() == 1);
+	const FText TransactionLabel = FText::Format(LOCTEXT("RenameGraphVertexDisplayNameFormat", "Set Metasound {0} DisplayName"), GetGraphMemberLabel());
+	const FScopedTransaction Transaction(TransactionLabel, bPostTransaction);
 
-			for (UEdGraphPin* Pin : Pins)
-			{
-				Pin->PinFriendlyName = InNewName;
-			}
+	Graph->Modify();
+	Graph->GetMetasoundChecked().Modify();
+
+	NodeHandle->SetDisplayName(InNewName);
+
+	const TArray<UMetasoundEditorGraphMemberNode*> Nodes = GetNodes();
+	for (UMetasoundEditorGraphMemberNode* Node : Nodes)
+	{
+		const TArray<UEdGraphPin*>& Pins = Node->GetAllPins();
+		ensure(Pins.Num() == 1);
+
+		for (UEdGraphPin* Pin : Pins)
+		{
+			Pin->PinFriendlyName = InNewName;
 		}
 	}
 
-	NameChanged.Broadcast(NodeID);
+	Graph->RegisterGraphWithFrontend();
+	Graph->GetModifyContext().AddMemberIDsModified({ GetMemberID() });
 }
 
 void UMetasoundEditorGraphVertex::SetDataType(FName InNewType, bool bPostTransaction)
@@ -367,9 +376,8 @@ void UMetasoundEditorGraphVertex::SetDataType(FName InNewType, bool bPostTransac
 		FGraphBuilder::AddNode(Metasound, NewNodeHandle, Location, false /* bInSelectNewNode */);
 	}
 
-	// Notify now that the node has a new ID (doing so before creating & syncing Frontend Node &
-	// EdGraph variable can result in refreshing editors while in a desync'ed state)
-	NameChanged.Broadcast(NodeID);
+	Graph->GetModifyContext().AddMemberIDsModified({ GetMemberID() });
+	Graph->RegisterGraphWithFrontend();
 }
 
 void UMetasoundEditorGraphVertex::SetVertexAccessType(EMetasoundFrontendVertexAccessType InNewAccessType, bool bPostTransaction)
@@ -465,9 +473,7 @@ void UMetasoundEditorGraphVertex::SetVertexAccessType(EMetasoundFrontendVertexAc
 		}
 	}
 
-	// Notify now that the node has a new ID (doing so before creating & syncing Frontend Node &
-	// EdGraph variable can result in refreshing editors while in a desync'ed state)
-	NameChanged.Broadcast(NodeID);
+	Graph->RegisterGraphWithFrontend();
 }
 
 Metasound::Frontend::FNodeHandle UMetasoundEditorGraphVertex::GetNodeHandle()
@@ -586,7 +592,7 @@ void UMetasoundEditorGraphMemberDefaultLiteral::PostEditUndo()
 
 	if (UMetasoundEditorGraph* Graph = Member->GetOwningGraph())
 	{
-		Graph->SetSynchronizationRequired();
+		Graph->GetModifyContext().SetDocumentModified();
 	}
 }
 #endif // WITH_EDITOR
@@ -650,7 +656,7 @@ void UMetasoundEditorGraphInput::SetSortOrderIndex(int32 InSortOrderIndex)
 	const Metasound::FVertexName& NodeName = NodeHandle->GetNodeName();
 
 	GraphHandle->SetSortOrderIndexForInput(NodeName, InSortOrderIndex);
-	MetaSoundGraph->SetSynchronizationRequired();
+	MetaSoundGraph->GetModifyContext().AddMemberIDsModified({ GetMemberID() });
 }
 
 void UMetasoundEditorGraphInput::ResetToClassDefault()
@@ -685,7 +691,7 @@ void UMetasoundEditorGraphInput::ResetToClassDefault()
 
 	if (UMetasoundEditorGraph* Graph = GetOwningGraph())
 	{
-		Graph->SetSynchronizationRequired();
+		Graph->GetModifyContext().AddMemberIDsModified({ GetMemberID() });
 	}
 }
 
@@ -814,7 +820,7 @@ void UMetasoundEditorGraphOutput::SetSortOrderIndex(int32 InSortOrderIndex)
 	const Metasound::FVertexName& NodeName = NodeHandle->GetNodeName();
 
 	GraphHandle->SetSortOrderIndexForOutput(NodeName, InSortOrderIndex);
-	MetaSoundGraph->SetSynchronizationRequired();
+	MetaSoundGraph->GetModifyContext().AddMemberIDsModified({ GetMemberID()});
 }
 
 const FText& UMetasoundEditorGraphOutput::GetGraphMemberLabel() const
@@ -846,7 +852,7 @@ void UMetasoundEditorGraphOutput::ResetToClassDefault()
 
 	if (UMetasoundEditorGraph* Graph = GetOwningGraph())
 	{
-		Graph->SetSynchronizationRequired();
+		Graph->GetModifyContext().AddMemberIDsModified({ GetMemberID() });
 	}
 }
 
@@ -965,20 +971,22 @@ Metasound::Frontend::FConstVariableHandle UMetasoundEditorGraphVariable::GetCons
 
 void UMetasoundEditorGraphVariable::SetMemberName(const FName& InNewName, bool bPostTransaction)
 {
+	UMetasoundEditorGraph* Graph = GetOwningGraph();
+	if (!ensure(Graph))
 	{
-		const FText TransactionLabel = FText::Format(LOCTEXT("RenameGraphVariableMemberNameFormat", "Set Metasound {0} Name"), GetGraphMemberLabel());
-		const FScopedTransaction Transaction(TransactionLabel, bPostTransaction);
+		return;
+	}
 
-		if (UMetasoundEditorGraph* Graph = GetOwningGraph())
-		{
-			Graph->Modify();
-			Graph->GetMetasoundChecked().Modify();
-		}
+	const FText TransactionLabel = FText::Format(LOCTEXT("RenameGraphVariableMemberNameFormat", "Set Metasound {0} Name"), GetGraphMemberLabel());
+	const FScopedTransaction Transaction(TransactionLabel, bPostTransaction);
 
-		GetVariableHandle()->SetName(InNewName);
-	}
+	Graph->Modify();
+	Graph->GetMetasoundChecked().Modify();
+
+	GetVariableHandle()->SetName(InNewName);
 
-	NameChanged.Broadcast(VariableID);
+	Graph->RegisterGraphWithFrontend();
+	Graph->GetModifyContext().AddMemberIDsModified({ GetMemberID() });
 }
 
 FGuid UMetasoundEditorGraphVariable::GetMemberID() const 
@@ -1079,20 +1087,23 @@ void UMetasoundEditorGraphVariable::SetDisplayName(const FText& InNewName, bool
 {
 	using namespace Metasound::Frontend;
 
+	UMetasoundEditorGraph* Graph = GetOwningGraph();
+	if (!ensure(Graph))
 	{
-		const FText TransactionLabel = FText::Format(LOCTEXT("RenameGraphVariableDisplayNameFormat", "Set Metasound {0} DisplayName"), GetGraphMemberLabel());
-		const FScopedTransaction Transaction(TransactionLabel, bPostTransaction);
-		if (UMetasoundEditorGraph* Graph = GetOwningGraph())
-		{
-			Graph->Modify();
-			Graph->GetMetasoundChecked().Modify();
-		}
+		return;
+	}
 
-		FVariableHandle VariableHandle = GetVariableHandle();
-		VariableHandle->SetDisplayName(InNewName);
+	const FText TransactionLabel = FText::Format(LOCTEXT("RenameGraphVariableDisplayNameFormat", "Set Metasound {0} DisplayName"), GetGraphMemberLabel());
+	const FScopedTransaction Transaction(TransactionLabel, bPostTransaction);
+	{
+		Graph->Modify();
+		Graph->GetMetasoundChecked().Modify();
 	}
 
-	NameChanged.Broadcast(VariableID);
+	FVariableHandle VariableHandle = GetVariableHandle();
+	VariableHandle->SetDisplayName(InNewName);
+
+	Graph->GetModifyContext().AddMemberIDsModified({ GetMemberID() });
 }
 
 void UMetasoundEditorGraphVariable::SetDataType(FName InNewType, bool bPostTransaction)
@@ -1169,9 +1180,6 @@ void UMetasoundEditorGraphVariable::SetDataType(FName InNewType, bool bPostTrans
 		AddVariableNodes(Metasound, FrontendGraph, CachedData.Locations);
 	}
 
-	// Notify now that the variable has a new ID (doing so before creating & syncing Frontend Node &
-	// EdGraph variable can result in refreshing editors while in a desync'ed state)
-	NameChanged.Broadcast(VariableID);
 }
 
 UMetasoundEditorGraphVariable::FVariableEditorNodes UMetasoundEditorGraphVariable::GetVariableNodes() const
@@ -1315,7 +1323,7 @@ void UMetasoundEditorGraphVariable::ResetToClassDefault()
 
 	if (UMetasoundEditorGraph* Graph = GetOwningGraph())
 	{
-		Graph->SetSynchronizationRequired();
+		Graph->GetModifyContext().AddMemberIDsModified({ GetMemberID() });
 	}
 }
 
@@ -1491,16 +1499,28 @@ void UMetasoundEditorGraph::SetVersionedOnLoad()
 	bVersionedOnLoad = true;
 }
 
-void UMetasoundEditorGraph::SetSynchronizationRequired()
+FMetasoundFrontendDocumentModifyContext& UMetasoundEditorGraph::GetModifyContext()
 {
 	using namespace Metasound;
 
-	if (UObject* ParentMetasound = GetOuter())
-	{
-		FMetasoundAssetBase* MetasoundAsset = IMetasoundUObjectRegistry::Get().GetObjectAsAssetBase(ParentMetasound);
-		check(MetasoundAsset);
-		MetasoundAsset->SetSynchronizationRequired();
-	}
+	UObject* ParentMetasound = GetOuter();
+	check(ParentMetasound);
+
+	FMetasoundAssetBase* MetasoundAsset = IMetasoundUObjectRegistry::Get().GetObjectAsAssetBase(ParentMetasound);
+	check(MetasoundAsset);
+	return MetasoundAsset->GetModifyContext();
+}
+
+const FMetasoundFrontendDocumentModifyContext& UMetasoundEditorGraph::GetModifyContext() const
+{
+	using namespace Metasound;
+
+	UObject* ParentMetasound = GetOuter();
+	check(ParentMetasound);
+
+	const FMetasoundAssetBase* MetasoundAsset = IMetasoundUObjectRegistry::Get().GetObjectAsAssetBase(ParentMetasound);
+	check(MetasoundAsset);
+	return MetasoundAsset->GetModifyContext();
 }
 
 UMetasoundEditorGraphInput* UMetasoundEditorGraph::FindInput(FGuid InNodeID) const
@@ -1761,17 +1781,6 @@ bool UMetasoundEditorGraph::ContainsVariable(const UMetasoundEditorGraphVariable
 	return Variables.Contains(&InVariable);
 }
 
-void UMetasoundEditorGraph::IterateInputs(TUniqueFunction<void(UMetasoundEditorGraphInput&)> InFunction) const
-{
-	for (UMetasoundEditorGraphInput* Input : Inputs)
-	{
-		if (Input)
-		{
-			InFunction(*Input);
-		}
-	}
-}
-
 void UMetasoundEditorGraph::SetPreviewID(uint32 InPreviewID)
 {
 	PreviewID = InPreviewID;
@@ -1803,23 +1812,49 @@ bool UMetasoundEditorGraph::IsEditable() const
 	return GetGraphHandle()->GetGraphStyle().bIsGraphEditable;
 }
 
-void UMetasoundEditorGraph::SetForceRefreshNodes()
+void UMetasoundEditorGraph::IterateInputs(TFunctionRef<void(UMetasoundEditorGraphInput&)> InFunction) const
 {
-	bForceRefreshNodes = true;
+	for (UMetasoundEditorGraphInput* Input : Inputs)
+	{
+		if (Input)
+		{
+			InFunction(*Input);
+		}
+	}
 }
 
-void UMetasoundEditorGraph::ClearForceRefreshNodes()
+void UMetasoundEditorGraph::IterateOutputs(TFunctionRef<void(UMetasoundEditorGraphOutput&)> InFunction) const
 {
-	bForceRefreshNodes = false;
+	for (UMetasoundEditorGraphOutput* Output : Outputs)
+	{
+		if (ensure(Output))
+		{
+			InFunction(*Output);
+		}
+	}
 }
 
-bool UMetasoundEditorGraph::RequiresForceRefreshNodes() const
+void UMetasoundEditorGraph::IterateVariables(TFunctionRef<void(UMetasoundEditorGraphVariable&)> InFunction) const
 {
-	return bForceRefreshNodes;
+	for (UMetasoundEditorGraphVariable* Variable : Variables)
+	{
+		if (ensure(Variable))
+		{
+			InFunction(*Variable);
+		}
+	}
 }
 
-void UMetasoundEditorGraph::IterateOutputs(TUniqueFunction<void(UMetasoundEditorGraphOutput&)> InFunction) const
+void UMetasoundEditorGraph::IterateMembers(TFunctionRef<void(UMetasoundEditorGraphMember&)> InFunction) const
 {
+	for (UMetasoundEditorGraphInput* Input : Inputs)
+	{
+		if (Input)
+		{
+			InFunction(*Input);
+		}
+	}
+
 	for (UMetasoundEditorGraphOutput* Output : Outputs)
 	{
 		if (ensure(Output))
@@ -1827,6 +1862,14 @@ void UMetasoundEditorGraph::IterateOutputs(TUniqueFunction<void(UMetasoundEditor
 			InFunction(*Output);
 		}
 	}
+
+	for (UMetasoundEditorGraphVariable* Variable : Variables)
+	{
+		if (ensure(Variable))
+		{
+			InFunction(*Variable);
+		}
+	}
 }
 
 void UMetasoundEditorGraph::ValidateInternal(Metasound::Editor::FGraphValidationResults& OutResults)
@@ -1871,12 +1914,6 @@ bool UMetasoundEditorGraph::RemoveMember(UMetasoundEditorGraphMember& InGraphMem
 		bSuccess = false;
 	}
 
-	// Required to ensure member is removed from the editor views on next tick when refreshing details.
-	if (NumRemoved > 0)
-	{
-		InGraphMember.MarkAsGarbage();
-	}
-
 	bSuccess = bSuccess && (NumRemoved > 0);
 
 	return bSuccess;
diff --git a/Engine/Plugins/Runtime/Metasound/Source/MetasoundEditor/Private/MetasoundEditorGraphBuilder.cpp b/Engine/Plugins/Runtime/Metasound/Source/MetasoundEditor/Private/MetasoundEditorGraphBuilder.cpp
--- a/Engine/Plugins/Runtime/Metasound/Source/MetasoundEditor/Private/MetasoundEditorGraphBuilder.cpp
+++ b/Engine/Plugins/Runtime/Metasound/Source/MetasoundEditor/Private/MetasoundEditorGraphBuilder.cpp
@@ -80,7 +80,23 @@ namespace Metasound
 				return FName(*NewName);
 			}
 
-			bool SynchronizeGraphRecursively(UObject& InMetaSound, bool bForceRefreshNodes, bool bEditorGraphModified = false)
+			void RecurseClearDocumentModified(FMetasoundAssetBase& InAssetBase)
+			{
+				using namespace Metasound::Frontend;
+
+				InAssetBase.GetModifyContext().ClearDocumentModified();
+
+				TArray<FMetasoundAssetBase*> References;
+				ensureAlways(IMetaSoundAssetManager::GetChecked().TryLoadReferencedAssets(InAssetBase, References));
+				for (FMetasoundAssetBase* Reference : References)
+				{
+					check(Reference);
+					Reference->GetModifyContext().ClearDocumentModified();
+					RecurseClearDocumentModified(*Reference);
+				}
+			};
+
+			bool SynchronizeGraphRecursively(UObject& InMetaSound, bool bEditorGraphModified = false)
 			{
 				using namespace Frontend;
 
@@ -94,7 +110,12 @@ namespace Metasound
 				for (FMetasoundAssetBase* Reference : References)
 				{
 					check(Reference);
-					SynchronizeGraphRecursively(*Reference->GetOwningAsset(), bForceRefreshNodes, bEditorGraphModified);
+					bEditorGraphModified |= SynchronizeGraphRecursively(*Reference->GetOwningAsset(), bEditorGraphModified);
+				}
+
+				if (!MetaSoundAsset->GetModifyContext().GetDocumentModified())
+				{
+					return bEditorGraphModified;
 				}
 
 				// If no graph is set, MetaSound has been created outside of asset factory, so initialize it here.
@@ -494,7 +515,7 @@ namespace Metasound
 							Variable->UpdateFrontendDefaultLiteral(bPostTransaction);
 						}
 
-						MetasoundGraph->SetSynchronizationRequired();
+						MetasoundGraph->GetModifyContext().AddNodeIDsModified({ NewGraphNode->GetNodeID() });
 						NewGraphNode->SetNodeLocation(InLocation);
 					}
 				}
@@ -534,7 +555,7 @@ namespace Metasound
 					constexpr bool bPostTransaction = false;
 					Output->UpdateFrontendDefaultLiteral(bPostTransaction);
 
-					MetasoundGraph->SetSynchronizationRequired();
+					MetasoundGraph->GetModifyContext().AddNodeIDsModified({ NewGraphNode->GetNodeID() });
 				}
 
 				NewGraphNode->CacheTitle();
@@ -550,56 +571,28 @@ namespace Metasound
 			RebuildNodePins(*NewGraphNode);
 		}
 
-		EMessageSeverity::Type FGraphBuilder::ValidateGraph(UObject& InMetaSound, bool bForceRefreshNodes)
+		FGraphValidationResults FGraphBuilder::ValidateGraph(UObject& InMetaSound)
 		{
 			using namespace Frontend;
 
-			TSharedPtr<SGraphEditor> GraphEditor;
-			TSharedPtr<FEditor> MetaSoundEditor = GetEditorForMetasound(InMetaSound);
-			if (MetaSoundEditor.IsValid())
-			{
-				GraphEditor = MetaSoundEditor->GetGraphEditor();
-			}
-
 			FMetasoundAssetBase* MetaSoundAsset = IMetasoundUObjectRegistry::Get().GetObjectAsAssetBase(&InMetaSound);
 			check(MetaSoundAsset);
-			UMetasoundEditorGraph& Graph = *CastChecked<UMetasoundEditorGraph>(&MetaSoundAsset->GetGraphChecked());
-
-			bForceRefreshNodes |= Graph.RequiresForceRefreshNodes();
-			Graph.ClearForceRefreshNodes();
-
-			int8 HighestMessageSeverity = EMessageSeverity::Info;
 
-			FGraphValidationResults Results;
-			Graph.ValidateInternal(Results);
-			for (const FGraphNodeValidationResult& Result : Results.GetResults())
-			{
-				UMetasoundEditorGraphNode& Node = Result.GetNodeChecked();
-				const bool bInterfaceChange = Node.ContainsInterfaceChange();
-				const bool bMetadataChange = Node.ContainsMetadataChange();
-				const bool bStyleChange = Node.ContainsStyleChange();
-
-				const FText Title = Node.GetCachedTitle();
-				Node.CacheTitle();
-				const bool bTitleUpdated = !Title.IdenticalTo(Node.GetCachedTitle());
-
-				if (Result.GetHasDirtiedNode() || bTitleUpdated || bMetadataChange || bInterfaceChange || bStyleChange || bForceRefreshNodes)
-				{
-					Node.SyncChangeIDs();
-
-					if (GraphEditor.IsValid())
-					{
-						GraphEditor->RefreshNode(Node);
-					}
-				}
-			}
-
-			if (MetaSoundEditor.IsValid())
+			// Validate referenced graphs first to ensure all editor data
+			// is up-to-date prior to validating this referencing graph to
+			// allow errors to bubble up.
+			TArray<FMetasoundAssetBase*> References;
+			ensureAlways(IMetaSoundAssetManager::GetChecked().TryLoadReferencedAssets(*MetaSoundAsset, References));
+			for (FMetasoundAssetBase* Reference : References)
 			{
-				MetaSoundEditor->RefreshGraphMemberMenu();
+				check(Reference);
+				ValidateGraph(*Reference->GetOwningAsset());
 			}
 
-			return Results.GetHighestMessageSeverity();
+			FGraphValidationResults Results;
+			UMetasoundEditorGraph& Graph = *CastChecked<UMetasoundEditorGraph>(&MetaSoundAsset->GetGraphChecked());
+			Graph.ValidateInternal(Results);
+			return Results;
 		}
 
 		TArray<FString> FGraphBuilder::GetDataTypeNameCategories(const FName& InDataTypeName)
@@ -990,7 +983,7 @@ namespace Metasound
 			{
 				NewGraphNode->SetNodeLocation(InLocation);
 				RebuildNodePins(*NewGraphNode);
-				MetasoundGraph->SetSynchronizationRequired();
+				MetasoundGraph->GetModifyContext().AddNodeIDsModified({ NewGraphNode->GetNodeID() });
 				return NewGraphNode;
 			}
 
@@ -1511,7 +1504,7 @@ namespace Metasound
 
 			UObject& MetaSound = Node->GetMetasoundChecked();
 			FMetasoundAssetBase* MetaSoundAsset = IMetasoundUObjectRegistry::Get().GetObjectAsAssetBase(&MetaSound);
-			MetaSoundAsset->SetSynchronizationRequired();
+			MetaSoundAsset->GetModifyContext().SetDocumentModified();
 		}
 
 		void FGraphBuilder::InitMetaSound(UObject& InMetaSound, const FString& InAuthor)
@@ -1533,9 +1526,9 @@ namespace Metasound
 
 			// 2. Set default doc version Metadata
 			FDocumentHandle DocumentHandle = MetaSoundAsset->GetDocumentHandle();
-			FMetasoundFrontendDocumentMetadata DocMetadata = DocumentHandle->GetMetadata();
-			DocMetadata.Version.Number = FMetasoundFrontendDocument::GetMaxVersion();
-			DocumentHandle->SetMetadata(DocMetadata);
+			FMetasoundFrontendDocumentMetadata* DocMetadata = DocumentHandle->GetMetadata();
+			check(DocMetadata);
+			DocMetadata->Version.Number = FMetasoundFrontendDocument::GetMaxVersion();
 
 			MetaSoundAsset->AddDefaultInterfaces();
 
@@ -1826,7 +1819,7 @@ namespace Metasound
 							EditedMetaSound->RebuildReferencedAssetClassKeys();
 							if (EditedMetaSound->IsReferencedAsset(*MetaSoundAsset))
 							{
-								EditedMetaSound->SetSynchronizationRequired();
+								EditedMetaSound->GetModifyContext().SetDocumentModified();
 							}
 						}
 					}
@@ -1935,6 +1928,30 @@ namespace Metasound
 			return NewPin;
 		}
 
+		bool FGraphBuilder::RecurseGetDocumentModified(FMetasoundAssetBase& InAssetBase)
+		{
+			using namespace Metasound::Frontend;
+
+			if (InAssetBase.GetModifyContext().GetDocumentModified())
+			{
+				return true;
+			}
+
+			TArray<FMetasoundAssetBase*> References;
+			ensureAlways(IMetaSoundAssetManager::GetChecked().TryLoadReferencedAssets(InAssetBase, References));
+			for (FMetasoundAssetBase* Reference : References)
+			{
+				check(Reference);
+				const bool bReferenceDocumentModified = RecurseGetDocumentModified(*Reference);
+				if (bReferenceDocumentModified)
+				{
+					return true;
+				}
+			}
+
+			return false;
+		}
+
 		bool FGraphBuilder::SynchronizePinType(const IMetasoundEditorModule& InEditorModule, UEdGraphPin& InPin, const FName InDataType)
 		{
 			FEdGraphPinType PinType;
@@ -2074,29 +2091,55 @@ namespace Metasound
 			return bIsGraphDirty;
 		}
 
-		int32 FGraphBuilder::SynchronizeGraph(UObject& InMetaSound, bool bForceRefreshNodes)
+		bool FGraphBuilder::SynchronizeGraph(UObject& InMetaSound)
 		{
 			using namespace Frontend;
 
-			int32 HighestMessageSeverity = EMessageSeverity::Info;
-
 			FMetasoundAssetBase* MetaSoundAsset = IMetasoundUObjectRegistry::Get().GetObjectAsAssetBase(&InMetaSound);
 			check(MetaSoundAsset);
-			if (MetaSoundAsset->GetSynchronizationRequired())
-			{
-				const bool bEditorGraphModified = GraphBuilderPrivate::SynchronizeGraphRecursively(InMetaSound, bForceRefreshNodes);
-				bForceRefreshNodes |= bEditorGraphModified;
-				HighestMessageSeverity = ValidateGraph(InMetaSound, bForceRefreshNodes);
 
-				MetaSoundAsset->ResetSynchronizationState();
-			}
-			else
+			if (RecurseGetDocumentModified(*MetaSoundAsset))
 			{
-				const UMetasoundEditorGraph* Graph = CastChecked<UMetasoundEditorGraph>(MetaSoundAsset->GetGraph());
-				HighestMessageSeverity = Graph->GetHighestMessageSeverity();
+				TArray<FMetasoundAssetBase*> EditedReferencingMetaSounds;
+				if (GEditor)
+				{
+					TArray<UObject*> EditedAssets = GEditor->GetEditorSubsystem<UAssetEditorSubsystem>()->GetAllEditedAssets();
+					for (UObject* Asset : EditedAssets)
+					{
+						if (Asset != &InMetaSound)
+						{
+							if (FMetasoundAssetBase* EditedMetaSound = IMetasoundUObjectRegistry::Get().GetObjectAsAssetBase(Asset))
+							{
+								if (EditedMetaSound->IsReferencedAsset(*MetaSoundAsset))
+								{
+									EditedReferencingMetaSounds.Add(EditedMetaSound);
+								}
+							}
+						}
+					}
+				}
+
+				if (EditedReferencingMetaSounds.IsEmpty())
+				{
+					MetaSoundAsset->CacheRegistryMetadata();
+					const bool bEditorGraphModified = GraphBuilderPrivate::SynchronizeGraphRecursively(InMetaSound);
+					GraphBuilderPrivate::RecurseClearDocumentModified(*MetaSoundAsset);
+				}
+				else
+				{
+					for (FMetasoundAssetBase* EditedMetaSound : EditedReferencingMetaSounds)
+					{
+						check(EditedMetaSound);
+						UObject* OwningMetaSound = EditedMetaSound->GetOwningAsset();
+						check(OwningMetaSound);
+						SynchronizeGraph(*OwningMetaSound);
+					}
+				}
+
+				return true;
 			}
 
-			return HighestMessageSeverity;
+			return false;
 		}
 
 		bool FGraphBuilder::SynchronizeNodeMembers(UObject& InMetaSound)
diff --git a/Engine/Plugins/Runtime/Metasound/Source/MetasoundEditor/Private/MetasoundEditorGraphBuilder.h b/Engine/Plugins/Runtime/Metasound/Source/MetasoundEditor/Private/MetasoundEditorGraphBuilder.h
--- a/Engine/Plugins/Runtime/Metasound/Source/MetasoundEditor/Private/MetasoundEditorGraphBuilder.h
+++ b/Engine/Plugins/Runtime/Metasound/Source/MetasoundEditor/Private/MetasoundEditorGraphBuilder.h
@@ -1,6 +1,7 @@
 // Copyright Epic Games, Inc. All Rights Reserved.
 #pragma once
 
+#include "MetasoundEditorGraphValidation.h"
 #include "MetasoundFrontend.h"
 #include "MetasoundFrontendController.h"
 #include "MetasoundFrontendDocument.h"
@@ -70,7 +71,11 @@ namespace Metasound
 			static TSharedPtr<FEditor> GetEditorForPin(const UEdGraphPin& InEdPin);
 
 			// Validates MetaSound graph, returning the highest EMessageSeverity integer value
-			static EMessageSeverity::Type ValidateGraph(UObject& InMetaSound, bool bForceRefreshNodes);
+			static FGraphValidationResults ValidateGraph(UObject& InMetaSound);
+
+			// Recursively checks whether the provided Asset's Document is marked as modified since last
+			// EdGraph synchronization, or if any of its referenced asset graphs have been marked as modified.
+			static bool RecurseGetDocumentModified(FMetasoundAssetBase& InAssetBase);
 
 			// Wraps RegisterGraphWithFrontend logic in Frontend with any additional logic required to refresh editor & respective editor object state.
 			// @param InMetaSound - MetaSound to register
@@ -252,9 +257,8 @@ namespace Metasound
 			// FMetasoundFrontendDocument model. Validates the graph (and those referenced recursively).
 			//
 			// @param InMetaSound - MetaSound to synchronize and optionally validate.
-			// @param bForceRefreshNodes - Refreshes all transient data stored on editor nodes & redraws them
-			// @return Highest MessageSeverity integer value listed on validated nodes post synchronization.
-			static int32 SynchronizeGraph(UObject& InMetaSound, bool bForceRefreshNodes = false);
+			// @return whether or not EditorGraph synchronization was performed.
+			static bool SynchronizeGraph(UObject& InMetaSound);
 
 			// Synchronizes editor nodes with frontend nodes, removing editor nodes that are not represented in the frontend, and adding editor nodes to represent missing frontend nodes.
 			//
diff --git a/Engine/Plugins/Runtime/Metasound/Source/MetasoundEditor/Private/MetasoundEditorGraphMemberDefaults.cpp b/Engine/Plugins/Runtime/Metasound/Source/MetasoundEditor/Private/MetasoundEditorGraphMemberDefaults.cpp
--- a/Engine/Plugins/Runtime/Metasound/Source/MetasoundEditor/Private/MetasoundEditorGraphMemberDefaults.cpp
+++ b/Engine/Plugins/Runtime/Metasound/Source/MetasoundEditor/Private/MetasoundEditorGraphMemberDefaults.cpp
@@ -219,13 +219,22 @@ void UMetasoundEditorGraphMemberDefaultFloat::PostEditChangeChainProperty(FPrope
 		OnClampChanged.Broadcast(ClampDefault);
 	}
 
-	// TODO: Remove this once widget Metadata is migrated to frontend
-	// style, which will inherently update the change guid and issue
-	// resync.
 	UMetasoundEditorGraphMember* Member = GetParentMember();
 	if (ensure(Member))
 	{
-		Member->GetOwningGraph()->SetForceRefreshNodes();
+		FMetasoundFrontendDocumentModifyContext& ModifyContext = Member->GetOwningGraph()->GetModifyContext();
+
+		// Only update member on non-interactive changes to avoid refreshing the details panel mid-update
+		if (PropertyChangedEvent.ChangeType != EPropertyChangeType::Interactive)
+		{
+			ModifyContext.AddMemberIDsModified({ Member->GetMemberID() });
+		}
+
+		// Mark all nodes as modified to refresh them on synchronization.  This ensures all corresponding widgets get updated.
+		const TArray<UMetasoundEditorGraphMemberNode*> MemberNodes = Member->GetNodes();
+		TSet<FGuid> NodesToRefresh;
+		Algo::Transform(MemberNodes, NodesToRefresh, [](const UMetasoundEditorGraphMemberNode* MemberNode) { return MemberNode->GetNodeID(); });
+		ModifyContext.AddNodeIDsModified({ NodesToRefresh });
 	}
 
 	Super::PostEditChangeChainProperty(PropertyChangedEvent);
diff --git a/Engine/Plugins/Runtime/Metasound/Source/MetasoundEditor/Private/MetasoundEditorGraphNode.cpp b/Engine/Plugins/Runtime/Metasound/Source/MetasoundEditor/Private/MetasoundEditorGraphNode.cpp
--- a/Engine/Plugins/Runtime/Metasound/Source/MetasoundEditor/Private/MetasoundEditorGraphNode.cpp
+++ b/Engine/Plugins/Runtime/Metasound/Source/MetasoundEditor/Private/MetasoundEditorGraphNode.cpp
@@ -242,28 +242,14 @@ void UMetasoundEditorGraphNode::Validate(Metasound::Editor::FGraphNodeValidation
 #endif // WITH_EDITOR
 }
 
-bool UMetasoundEditorGraphNode::ContainsMetadataChange() const
+bool UMetasoundEditorGraphNode::ContainsClassChange() const
 {
 	using namespace Metasound::Frontend;
 	FConstNodeHandle NodeHandle = GetConstNodeHandle();
 
-	return MetadataChangeID != NodeHandle->GetClassMetadata().GetChangeID();
-}
-
-bool UMetasoundEditorGraphNode::ContainsInterfaceChange() const
-{
-	using namespace Metasound::Frontend;
-	FConstNodeHandle NodeHandle = GetConstNodeHandle();
-
-	return InterfaceChangeID != NodeHandle->GetClassInterface().GetChangeID();
-}
-
-bool UMetasoundEditorGraphNode::ContainsStyleChange() const
-{
-	using namespace Metasound::Frontend;
-	FConstNodeHandle NodeHandle = GetConstNodeHandle();
-
-	return StyleChangeID != NodeHandle->GetClassStyle().GetChangeID();
+	return InterfaceChangeID != NodeHandle->GetClassInterface().GetChangeID()
+	|| StyleChangeID != NodeHandle->GetClassStyle().GetChangeID()
+	|| MetadataChangeID != NodeHandle->GetClassMetadata().GetChangeID();
 }
 
 void UMetasoundEditorGraphNode::ReconstructNode()
@@ -633,7 +619,7 @@ void UMetasoundEditorGraphOutputNode::PinDefaultValueChanged(UEdGraphPin* InPin)
 							FGraphBuilder::RegisterGraphWithFrontend(MetaSound);
 							if (FMetasoundAssetBase* MetaSoundAsset = Metasound::IMetasoundUObjectRegistry::Get().GetObjectAsAssetBase(&MetaSound))
 							{
-								MetaSoundAsset->SetUpdateDetailsOnSynchronization();
+								MetaSoundAsset->GetModifyContext().AddMemberIDsModified({ Output->GetMemberID() });
 							}
 						}
 					}
@@ -821,7 +807,10 @@ void UMetasoundEditorGraphExternalNode::Validate(Metasound::Editor::FGraphNodeVa
 		{
 			if (UObject* AssetObject = Path->ResolveObject())
 			{
-				const EMessageSeverity::Type MaxGraphMsg = FGraphBuilder::ValidateGraph(*AssetObject, false /* bForceRefreshNodes */);
+				FMetasoundAssetBase* MetaSoundAsset = Metasound::IMetasoundUObjectRegistry::Get().GetObjectAsAssetBase(AssetObject);
+				check(MetaSoundAsset);
+				const UMetasoundEditorGraph* NodeGraph = CastChecked<UMetasoundEditorGraph>(&MetaSoundAsset->GetGraphChecked());
+				const EMessageSeverity::Type MaxGraphMsg = static_cast<EMessageSeverity::Type>(NodeGraph->GetHighestMessageSeverity());
 				switch (MaxGraphMsg)
 				{
 					case EMessageSeverity::Error:
@@ -1200,7 +1189,7 @@ void UMetasoundEditorGraphVariableNode::PinDefaultValueChanged(UEdGraphPin* Pin)
 								FGraphBuilder::RegisterGraphWithFrontend(MetaSound);
 								if (FMetasoundAssetBase* MetaSoundAsset = Metasound::IMetasoundUObjectRegistry::Get().GetObjectAsAssetBase(&MetaSound))
 								{
-									MetaSoundAsset->SetUpdateDetailsOnSynchronization();
+									MetaSoundAsset->GetModifyContext().AddNodeIDsModified({ NodeID });
 								}
 							}
 						}
diff --git a/Engine/Plugins/Runtime/Metasound/Source/MetasoundEditor/Private/MetasoundEditorGraphSchema.cpp b/Engine/Plugins/Runtime/Metasound/Source/MetasoundEditor/Private/MetasoundEditorGraphSchema.cpp
--- a/Engine/Plugins/Runtime/Metasound/Source/MetasoundEditor/Private/MetasoundEditorGraphSchema.cpp
+++ b/Engine/Plugins/Runtime/Metasound/Source/MetasoundEditor/Private/MetasoundEditorGraphSchema.cpp
@@ -225,7 +225,16 @@ namespace Metasound
 				}
 			}
 
-			UEdGraphNode* PromoteToVariable(const FName BaseName, UEdGraphPin& FromPin, const FName DataType, const FMetasoundFrontendClass& InVariableClass, const FVector2D& InLocation)
+			void SelecteNodeInEditor(UMetasoundEditorGraph& InMetaSoundGraph, UMetasoundEditorGraphNode& InNode)
+			{
+				TSharedPtr<FEditor> MetasoundEditor = FGraphBuilder::GetEditorForGraph(InMetaSoundGraph);
+				if (MetasoundEditor.IsValid())
+				{
+					MetasoundEditor->ClearSelectionAndSelectNode(&InNode);
+				}
+			}
+
+			UEdGraphNode* PromoteToVariable(const FName BaseName, UEdGraphPin& FromPin, const FName DataType, const FMetasoundFrontendClass& InVariableClass, const FVector2D& InLocation, bool bSelectNode)
 			{
 				using namespace Frontend;
 
@@ -287,14 +296,7 @@ namespace Metasound
 							if (ensure(SchemaPrivate::TryConnectNewNodeToMatchingDataTypePin(*EdGraphNode, &FromPin)))
 							{
 								FGraphBuilder::RegisterGraphWithFrontend(ParentMetaSound);
-								MetaSoundGraph->SetSynchronizationRequired();
-
-								TSharedPtr<FEditor> ParentEditor = FGraphBuilder::GetEditorForMetasound(ParentMetaSound);
-								if (ParentEditor.IsValid())
-								{
-									ParentEditor->RefreshGraphMemberMenu();
-								}
-
+								SelecteNodeInEditor(*MetaSoundGraph, *NewGraphNode);
 								return EdGraphNode;
 							}
 						}
@@ -357,7 +359,6 @@ namespace Metasound
 
 				return Algo::AnyOf(AccessorNodes, IsSourceNodeReachableDownstream);
 			}
-
 		} // namespace SchemaPrivate
 	} // namespace Editor
 } // namespace Metasound
@@ -443,14 +444,7 @@ UEdGraphNode* FMetasoundGraphSchemaAction_NewNode::PerformAction(UEdGraph* Paren
 	{
 		NewGraphNode->Modify();
 		SchemaPrivate::TryConnectNewNodeToMatchingDataTypePin(*NewGraphNode, FromPin);
-		MetaSoundGraph->SetSynchronizationRequired();
-		
-		TSharedPtr<FEditor> ParentEditor = FGraphBuilder::GetEditorForMetasound(ParentMetasound);
-		if (ParentEditor.IsValid() && bSelectNewNode)
-		{
-			ParentEditor->ClearSelectionAndSelectNode(NewGraphNode);
-		}
-		
+		SchemaPrivate::SelecteNodeInEditor(*MetaSoundGraph, *NewGraphNode);
 		return NewGraphNode;
 	}
 
@@ -545,7 +539,7 @@ UEdGraphNode* FMetasoundGraphSchemaAction_PromoteToInput::PerformAction(UEdGraph
 	FGraphBuilder::GetPinLiteral(*FromPin, DefaultValue);
 
 	const FName InputName = InputHandle->GetName();
-	// The promoted input must have the same vertex access type in order for it to be connectable 
+	// The promoted input must have the same vertex access type in order for it to be connectible
 	const FCreateNodeVertexParams VertexParams = { InputHandle->GetDataType(), InputHandle->GetVertexAccessType() };
 
 	FNodeHandle NodeHandle = FGraphBuilder::AddInputNodeHandle(ParentMetasound, VertexParams, &DefaultValue, &InputName);
@@ -560,14 +554,8 @@ UEdGraphNode* FMetasoundGraphSchemaAction_PromoteToInput::PerformAction(UEdGraph
 
 				if (ensure(SchemaPrivate::TryConnectNewNodeToMatchingDataTypePin(*EdGraphNode, FromPin)))
 				{
-					TSharedPtr<FEditor> MetasoundEditor = FGraphBuilder::GetEditorForGraph(*ParentGraph);
-					if (MetasoundEditor.IsValid())
-					{
-						MetasoundEditor->OnInputNameChanged(NewGraphNode->GetNodeID());
-					}
-
 					FGraphBuilder::RegisterGraphWithFrontend(ParentMetasound);
-					MetasoundEditor->ClearSelectionAndSelectNode(EdGraphNode);
+					SchemaPrivate::SelecteNodeInEditor(*MetasoundGraph, *NewGraphNode);
 					return EdGraphNode;
 				}
 			}
@@ -603,7 +591,7 @@ UEdGraphNode* FMetasoundGraphSchemaAction_PromoteToVariable_AccessorNode::Perfor
 	FMetasoundFrontendClass VariableClass;
 	if (ensure(IDataTypeRegistry::Get().GetFrontendVariableAccessorClass(DataType, VariableClass)))
 	{
-		return SchemaPrivate::PromoteToVariable(NodeName, *FromPin, DataType, VariableClass, InLocation);
+		return SchemaPrivate::PromoteToVariable(NodeName, *FromPin, DataType, VariableClass, InLocation, bSelectNewNode);
 	}
 
 	return nullptr;
@@ -635,7 +623,7 @@ UEdGraphNode* FMetasoundGraphSchemaAction_PromoteToVariable_DeferredAccessorNode
 	FMetasoundFrontendClass VariableClass;
 	if (ensure(IDataTypeRegistry::Get().GetFrontendVariableDeferredAccessorClass(DataType, VariableClass)))
 	{
-		return SchemaPrivate::PromoteToVariable(NodeName, *FromPin, DataType, VariableClass, InLocation);
+		return SchemaPrivate::PromoteToVariable(NodeName, *FromPin, DataType, VariableClass, InLocation, bSelectNewNode);
 	}
 
 	return nullptr;
@@ -667,13 +655,7 @@ UEdGraphNode* FMetasoundGraphSchemaAction_PromoteToVariable_MutatorNode::Perform
 	FMetasoundFrontendClass VariableClass;
 	if (ensure(IDataTypeRegistry::Get().GetFrontendVariableMutatorClass(DataType, VariableClass)))
 	{
-		UEdGraphNode* NewGraphNode = SchemaPrivate::PromoteToVariable(NodeName, *FromPin, DataType, VariableClass, InLocation);
-		TSharedPtr<FEditor> MetasoundEditor = FGraphBuilder::GetEditorForGraph(*ParentGraph);
-		if (MetasoundEditor.IsValid() && bSelectNewNode)
-		{
-			MetasoundEditor->ClearSelectionAndSelectNode(NewGraphNode);
-		}
-		return NewGraphNode;
+		return SchemaPrivate::PromoteToVariable(NodeName, *FromPin, DataType, VariableClass, InLocation, bSelectNewNode);
 	}
 
 	return nullptr;
@@ -777,14 +759,8 @@ UEdGraphNode* FMetasoundGraphSchemaAction_PromoteToOutput::PerformAction(UEdGrap
 
 				if (ensure(SchemaPrivate::TryConnectNewNodeToMatchingDataTypePin(*EdGraphNode, FromPin)))
 				{
-					TSharedPtr<FEditor> MetasoundEditor = FGraphBuilder::GetEditorForGraph(*ParentGraph);
-					if (MetasoundEditor.IsValid())
-					{
-						MetasoundEditor->OnOutputNameChanged(NewGraphNode->GetNodeID());
-					}
-
 					FGraphBuilder::RegisterGraphWithFrontend(ParentMetasound);
-					MetasoundEditor->ClearSelectionAndSelectNode(EdGraphNode);
+					SchemaPrivate::SelecteNodeInEditor(*MetasoundGraph, *NewGraphNode);
 					return EdGraphNode;
 				}
 			}
@@ -933,7 +909,7 @@ UEdGraphNode* FMetasoundGraphSchemaAction_NewReroute::PerformAction(UEdGraph* Pa
 	{
 		NewGraphNode->Modify();
 		SchemaPrivate::TryConnectNewNodeToMatchingDataTypePin(*NewGraphNode, FromPin);
-		MetaSoundGraph->SetSynchronizationRequired();
+		MetaSoundGraph->GetModifyContext().AddNodeIDsModified({ NewGraphNode->GetNodeID()});
 
 		TSharedPtr<FEditor> ParentEditor = FGraphBuilder::GetEditorForMetasound(ParentMetasound);
 		if (ParentEditor.IsValid() && bSelectNewNode)
@@ -1417,7 +1393,7 @@ bool UMetasoundEditorGraphSchema::TryCreateConnection(UEdGraphPin* PinA, UEdGrap
 		return false;
 	}
 
-	Graph->SetSynchronizationRequired();
+	Graph->GetModifyContext().SetDocumentModified();
 
 	return true;
 }
diff --git a/Engine/Plugins/Runtime/Metasound/Source/MetasoundEditor/Private/MetasoundNodeDetailCustomization.cpp b/Engine/Plugins/Runtime/Metasound/Source/MetasoundEditor/Private/MetasoundNodeDetailCustomization.cpp
--- a/Engine/Plugins/Runtime/Metasound/Source/MetasoundEditor/Private/MetasoundNodeDetailCustomization.cpp
+++ b/Engine/Plugins/Runtime/Metasound/Source/MetasoundEditor/Private/MetasoundNodeDetailCustomization.cpp
@@ -169,7 +169,10 @@ namespace Metasound
 						FloatLiteral->OnClampChanged.Remove(OnClampChangedDelegateHandle);
 						if (FMetasoundAssetBase* MetasoundAsset = Metasound::IMetasoundUObjectRegistry::Get().GetObjectAsAssetBase(FloatLiteral->GetOutermostObject()))
 						{
-							MetasoundAsset->SetSynchronizationRequired();
+							if (const UMetasoundEditorGraphMember* Member = Cast<UMetasoundEditorGraphMember>(FloatLiteral->GetOuter()))
+							{
+								MetasoundAsset->GetModifyContext().AddMemberIDsModified({ Member->GetMemberID() });
+							}
 						}
 					}
 				});
@@ -675,11 +678,6 @@ namespace Metasound
 			if (GraphMember.IsValid())
 			{
 				GraphMember->SetDataType(NewDataTypeName);
-
-				if (FMetasoundAssetBase* MetasoundAsset = Metasound::IMetasoundUObjectRegistry::Get().GetObjectAsAssetBase(GraphMember->GetOutermostObject()))
-				{
-					MetasoundAsset->SetUpdateDetailsOnSynchronization();
-				}
 			}
 		}
 
@@ -838,11 +836,6 @@ namespace Metasound
 
 					const FName DataType = InNewState == ECheckBoxState::Checked ? ArrayTypeName : BaseTypeName;
 					InGraphMember->SetDataType(DataType);
-
-					if (FMetasoundAssetBase* MetasoundAsset = Metasound::IMetasoundUObjectRegistry::Get().GetObjectAsAssetBase(GraphMember->GetOutermostObject()))
-					{
-						MetasoundAsset->SetUpdateDetailsOnSynchronization();
-					}
 				}
 			}
 		}
@@ -1345,7 +1338,7 @@ namespace Metasound
 				
 				if (FMetasoundAssetBase* MetasoundAsset = Metasound::IMetasoundUObjectRegistry::Get().GetObjectAsAssetBase(GraphMember->GetOutermostObject()))
 				{
-					MetasoundAsset->SetUpdateDetailsOnSynchronization();
+					MetasoundAsset->GetModifyContext().AddMemberIDsModified({ GraphMember->GetMemberID() });
 				}
 			}
 		}
diff --git a/Engine/Plugins/Runtime/Metasound/Source/MetasoundEditor/Private/MetasoundNodeDetailCustomization.h b/Engine/Plugins/Runtime/Metasound/Source/MetasoundEditor/Private/MetasoundNodeDetailCustomization.h
--- a/Engine/Plugins/Runtime/Metasound/Source/MetasoundEditor/Private/MetasoundNodeDetailCustomization.h
+++ b/Engine/Plugins/Runtime/Metasound/Source/MetasoundEditor/Private/MetasoundNodeDetailCustomization.h
@@ -280,8 +280,6 @@ namespace Metasound
 			ECheckBoxState OnGetDataTypeArrayCheckState(TWeakObjectPtr<UMetasoundEditorGraphMember> InGraphMember) const;
 			void OnDataTypeSelected(FName InSelectedTypeName);
 			FName GetDataType() const;
-
-			TFunction<void()> OnDataTypeChanged;
 		
 		private:
 			TWeakObjectPtr<UMetasoundEditorGraphMember> GraphMember;
diff --git a/Engine/Plugins/Runtime/Metasound/Source/MetasoundEditor/Private/SMetasoundGraphNode.cpp b/Engine/Plugins/Runtime/Metasound/Source/MetasoundEditor/Private/SMetasoundGraphNode.cpp
--- a/Engine/Plugins/Runtime/Metasound/Source/MetasoundEditor/Private/SMetasoundGraphNode.cpp
+++ b/Engine/Plugins/Runtime/Metasound/Source/MetasoundEditor/Private/SMetasoundGraphNode.cpp
@@ -685,7 +685,7 @@ namespace Metasound
 
 								if (UMetasoundEditorGraph* Graph = GraphMember->GetOwningGraph())
 								{
-									Graph->SetSynchronizationRequired();
+									Graph->GetModifyContext().AddMemberIDsModified({ GraphMember->GetMemberID() });
 								}
 							}
 						};
diff --git a/Engine/Plugins/Runtime/Metasound/Source/MetasoundEditor/Private/SMetasoundGraphPin.h b/Engine/Plugins/Runtime/Metasound/Source/MetasoundEditor/Private/SMetasoundGraphPin.h
--- a/Engine/Plugins/Runtime/Metasound/Source/MetasoundEditor/Private/SMetasoundGraphPin.h
+++ b/Engine/Plugins/Runtime/Metasound/Source/MetasoundEditor/Private/SMetasoundGraphPin.h
@@ -182,18 +182,20 @@ namespace Metasound
 												if (ensure(Member))
 												{
 													Member->ResetToClassDefault();
+													MetaSoundGraph->GetModifyContext().AddMemberIDsModified({ Member->GetMemberID() });
+												}
+												else
+												{
+													MetaSoundGraph->GetModifyContext().SetDocumentModified();
 												}
 											}
 											else
 											{
 												FInputHandle InputHandle = GetInputHandle();
 												InputHandle->ClearLiteral();
+												MetaSoundGraph->GetModifyContext().SetDocumentModified();
 											}
 										}
-
-										// Full node synchronization required as custom
-										// node-level widgets may need to be refreshed
-										MetaSoundGraph->SetSynchronizationRequired();
 									}
 								}
 							}
diff --git a/Engine/Plugins/Runtime/Metasound/Source/MetasoundEditor/Public/MetasoundEditorGraph.h b/Engine/Plugins/Runtime/Metasound/Source/MetasoundEditor/Public/MetasoundEditorGraph.h
--- a/Engine/Plugins/Runtime/Metasound/Source/MetasoundEditor/Public/MetasoundEditorGraph.h
+++ b/Engine/Plugins/Runtime/Metasound/Source/MetasoundEditor/Public/MetasoundEditorGraph.h
@@ -99,9 +99,6 @@ class METASOUNDEDITOR_API UMetasoundEditorGraphMember : public UObject
 	GENERATED_BODY()
 
 public:
-	/** Delegate called when the name of the associated Frontend Node is changed */
-	FOnMetasoundMemberNameChanged NameChanged;
-
 	/** Delegate called when a rename is requested on a renameable member node. */
 	FOnMetasoundMemberRenameRequested OnRenameRequested;
 
@@ -383,8 +380,10 @@ public:
 	UObject& GetMetasoundChecked();
 	const UObject& GetMetasoundChecked() const;
 
-	void IterateInputs(TUniqueFunction<void(UMetasoundEditorGraphInput&)> InFunction) const;
-	void IterateOutputs(TUniqueFunction<void(UMetasoundEditorGraphOutput&)> InFunction) const;
+	void IterateInputs(TFunctionRef<void(UMetasoundEditorGraphInput&)> InFunction) const;
+	void IterateOutputs(TFunctionRef<void(UMetasoundEditorGraphOutput&)> InFunction) const;
+	void IterateVariables(TFunctionRef<void(UMetasoundEditorGraphVariable&)> InFunction) const;
+	void IterateMembers(TFunctionRef<void(UMetasoundEditorGraphMember&)> InFunction) const;
 
 	bool ContainsInput(const UMetasoundEditorGraphInput& InInput) const;
 	bool ContainsOutput(const UMetasoundEditorGraphOutput& InOutput) const;
@@ -394,13 +393,10 @@ public:
 	bool IsPreviewing() const;
 	bool IsEditable() const;
 
-	void SetForceRefreshNodes();
-	void ClearForceRefreshNodes();
-	bool RequiresForceRefreshNodes() const;
-
 	// UMetasoundEditorGraphBase Implementation
 	virtual void RegisterGraphWithFrontend() override;
-	virtual void SetSynchronizationRequired() override;
+	virtual FMetasoundFrontendDocumentModifyContext& GetModifyContext() override;
+	virtual const FMetasoundFrontendDocumentModifyContext& GetModifyContext() const override;
 
 	virtual void ClearVersionedOnLoad() override;
 	virtual bool GetVersionedOnLoad() const override;
diff --git a/Engine/Plugins/Runtime/Metasound/Source/MetasoundEditor/Public/MetasoundEditorGraphNode.h b/Engine/Plugins/Runtime/Metasound/Source/MetasoundEditor/Public/MetasoundEditorGraphNode.h
--- a/Engine/Plugins/Runtime/Metasound/Source/MetasoundEditor/Public/MetasoundEditorGraphNode.h
+++ b/Engine/Plugins/Runtime/Metasound/Source/MetasoundEditor/Public/MetasoundEditorGraphNode.h
@@ -110,14 +110,8 @@ public:
 
 	FText GetCachedTitle() const { return CachedTitle; }
 
-	// Returns whether or not Metadata has been changed since the last node refresh
-	bool ContainsMetadataChange() const;
-
-	// Returns whether or not Interface has been changed since the last node refresh
-	bool ContainsInterfaceChange() const;
-
-	// Returns whether or not Style has been changed since the last node refresh
-	bool ContainsStyleChange() const;
+	// Returns whether or not the class interface, metadata, or style has been changed since the last node refresh
+	bool ContainsClassChange() const;
 
 protected:
 	FGuid InterfaceChangeID;
diff --git a/Engine/Plugins/Runtime/Metasound/Source/MetasoundEngine/Public/Metasound.h b/Engine/Plugins/Runtime/Metasound/Source/MetasoundEngine/Public/Metasound.h
--- a/Engine/Plugins/Runtime/Metasound/Source/MetasoundEngine/Public/Metasound.h
+++ b/Engine/Plugins/Runtime/Metasound/Source/MetasoundEngine/Public/Metasound.h
@@ -30,10 +30,11 @@ public:
 	virtual bool IsEditorOnly() const override { return true; }
 	virtual bool NeedsLoadForEditorGame() const override { return false; }
 
-	virtual void SetSynchronizationRequired() PURE_VIRTUAL(UMetasoundEditorGraphBase::SetSynchronizationRequired, )
 	virtual void RegisterGraphWithFrontend() PURE_VIRTUAL(UMetasoundEditorGraphBase::RegisterGraphWithFrontend(), )
 
 #if WITH_EDITORONLY_DATA
+	virtual FMetasoundFrontendDocumentModifyContext& GetModifyContext() PURE_VIRTUAL(UMetasoundEditorGraphBase::GetModifyContext, static FMetasoundFrontendDocumentModifyContext InvalidModifyData; return InvalidModifyData; )
+	virtual const FMetasoundFrontendDocumentModifyContext& GetModifyContext() const PURE_VIRTUAL(UMetasoundEditorGraphBase::GetModifyContext, static const FMetasoundFrontendDocumentModifyContext InvalidModifyData; return InvalidModifyData; )
 	virtual void ClearVersionedOnLoad() PURE_VIRTUAL(UMetasoundEditorGraphBase::ClearVersionedOnLoad, )
 	virtual bool GetVersionedOnLoad() const PURE_VIRTUAL(UMetasoundEditorGraphBase::GetVersionedOnLoad, return false; )
 	virtual void SetVersionedOnLoad() PURE_VIRTUAL(UMetasoundEditorGraphBase::SetVersionedOnLoad, )
@@ -49,7 +50,7 @@ namespace Metasound
 	template <typename TMetaSoundObject>
 	void PostEditUndo(TMetaSoundObject& InMetaSound)
 	{
-		InMetaSound.SetUpdateDetailsOnSynchronization();
+		InMetaSound.GetModifyContext().SetForceRefreshViews();
 		if (UMetasoundEditorGraphBase* Graph = Cast<UMetasoundEditorGraphBase>(InMetaSound.GetGraph()))
 		{
 			Graph->RegisterGraphWithFrontend();
@@ -70,7 +71,7 @@ namespace Metasound
 			if (!InSaveContext.IsCooking())
 			{
 				MetaSoundGraph->RegisterGraphWithFrontend();
-				MetaSoundGraph->SetSynchronizationRequired();
+				MetaSoundGraph->GetModifyContext().SetForceRefreshViews();
 			}
 		}
 #endif // WITH_EDITORONLY_DATA
diff --git a/Engine/Plugins/Runtime/Metasound/Source/MetasoundFrontend/Private/MetasoundAssetBase.cpp b/Engine/Plugins/Runtime/Metasound/Source/MetasoundFrontend/Private/MetasoundAssetBase.cpp
--- a/Engine/Plugins/Runtime/Metasound/Source/MetasoundFrontend/Private/MetasoundAssetBase.cpp
+++ b/Engine/Plugins/Runtime/Metasound/Source/MetasoundFrontend/Private/MetasoundAssetBase.cpp
@@ -135,7 +135,7 @@ void FMetasoundAssetBase::RegisterGraphWithFrontend(Metasound::Frontend::FMetaSo
 #if WITH_EDITOR
 		if (bAutoUpdated || InRegistrationOptions.bForceViewSynchronization)
 		{
- 			SetSynchronizationRequired();
+			GetModifyContext().SetForceRefreshViews();
 		}
 #endif // WITH_EDITOR
 	}
@@ -144,7 +144,7 @@ void FMetasoundAssetBase::RegisterGraphWithFrontend(Metasound::Frontend::FMetaSo
 #if WITH_EDITOR
 		if (InRegistrationOptions.bForceViewSynchronization)
 		{
-			SetSynchronizationRequired();
+			GetModifyContext().SetForceRefreshViews();
 		}
 #endif // WITH_EDITOR
 	}
@@ -534,76 +534,14 @@ void FMetasoundAssetBase::CacheRegistryMetadata()
 	}
 }
 
-bool FMetasoundAssetBase::GetSynchronizationRequired() const
+FMetasoundFrontendDocumentModifyContext& FMetasoundAssetBase::GetModifyContext()
 {
-	using namespace Metasound::Frontend;
-
-	if (bSynchronizationRequired)
-	{
-		return true;
-	}
-
-	TArray<FMetasoundAssetBase*> References;
-	ensureAlways(IMetaSoundAssetManager::GetChecked().TryLoadReferencedAssets(*this, References));
-	for (FMetasoundAssetBase* Reference : References)
-	{
-		const bool bReferenceSyncRequired = Reference->GetSynchronizationRequired();
-		if (bReferenceSyncRequired)
-		{
-			return true;
-		}
-	}
-
-	return false;
-}
-
-bool FMetasoundAssetBase::GetSynchronizationUpdateDetails() const
-{
-	using namespace Metasound::Frontend;
-
-	if (bSynchronizationUpdateDetails)
-	{
-		return true;
-	}
-
-	TArray<FMetasoundAssetBase*> References;
-	ensureAlways(IMetaSoundAssetManager::GetChecked().TryLoadReferencedAssets(*this, References));
-	for (FMetasoundAssetBase* Reference : References)
-	{
-		const bool bReferenceSyncRequired = Reference->GetSynchronizationUpdateDetails();
-		if (bReferenceSyncRequired)
-		{
-			return true;
-		}
-	}
-
-	return false;
+	return GetDocumentChecked().Metadata.ModifyContext;
 }
 
-void FMetasoundAssetBase::SetSynchronizationRequired()
+const FMetasoundFrontendDocumentModifyContext& FMetasoundAssetBase::GetModifyContext() const
 {
-	bSynchronizationRequired = true;
-}
-
-void FMetasoundAssetBase::SetUpdateDetailsOnSynchronization()
-{
-	bSynchronizationUpdateDetails = true;
-	bSynchronizationRequired = true;
-}
-
-void FMetasoundAssetBase::ResetSynchronizationState()
-{
-	using namespace Metasound::Frontend;
-
-	bSynchronizationUpdateDetails = false;
-	bSynchronizationRequired = false;
-
-	TArray<FMetasoundAssetBase*> References;
-	ensureAlways(IMetaSoundAssetManager::GetChecked().TryLoadReferencedAssets(*this, References));
-	for (FMetasoundAssetBase* Reference : References)
-	{
-		Reference->ResetSynchronizationState();
-	}
+	return GetDocumentChecked().Metadata.ModifyContext;
 }
 #endif // WITH_EDITOR
 
diff --git a/Engine/Plugins/Runtime/Metasound/Source/MetasoundFrontend/Private/MetasoundFrontendDocument.cpp b/Engine/Plugins/Runtime/Metasound/Source/MetasoundFrontend/Private/MetasoundFrontendDocument.cpp
--- a/Engine/Plugins/Runtime/Metasound/Source/MetasoundFrontend/Private/MetasoundFrontendDocument.cpp
+++ b/Engine/Plugins/Runtime/Metasound/Source/MetasoundFrontend/Private/MetasoundFrontendDocument.cpp
@@ -104,6 +104,95 @@ namespace Metasound
 	} // namespace DocumentPrivate
 } // namespace Metasound
 
+#if WITH_EDITORONLY_DATA
+void FMetasoundFrontendDocumentModifyContext::ClearDocumentModified()
+{
+	bDocumentModified = false;
+}
+
+bool FMetasoundFrontendDocumentModifyContext::GetDocumentModified() const
+{
+	return bDocumentModified;
+}
+
+bool FMetasoundFrontendDocumentModifyContext::GetForceRefreshViews() const
+{
+	return bForceRefreshViews;
+}
+
+const TSet<FName>& FMetasoundFrontendDocumentModifyContext::GetInterfacesModified() const
+{
+	return InterfacesModified;
+}
+
+const TSet<FGuid>& FMetasoundFrontendDocumentModifyContext::GetMemberIDsModified() const
+{
+	return MemberIDsModified;
+}
+
+const TSet<FGuid>& FMetasoundFrontendDocumentModifyContext::GetNodeIDsModified() const
+{
+	return NodeIDsModified;
+}
+
+void FMetasoundFrontendDocumentModifyContext::Reset()
+{
+	bDocumentModified = false;
+	bForceRefreshViews = false;
+	InterfacesModified.Empty();
+	MemberIDsModified.Empty();
+	NodeIDsModified.Empty();
+}
+
+void FMetasoundFrontendDocumentModifyContext::SetDocumentModified()
+{
+	bDocumentModified = true;
+}
+
+void FMetasoundFrontendDocumentModifyContext::SetForceRefreshViews()
+{
+	bDocumentModified = true;
+	bForceRefreshViews = true;
+}
+
+void FMetasoundFrontendDocumentModifyContext::AddInterfaceModified(FName InInterfaceModified)
+{
+	bDocumentModified = true;
+	InterfacesModified.Add(InInterfaceModified);
+}
+
+void FMetasoundFrontendDocumentModifyContext::AddInterfacesModified(const TSet<FName>& InInterfacesModified)
+{
+	bDocumentModified = true;
+	InterfacesModified.Append(InInterfacesModified);
+}
+
+void FMetasoundFrontendDocumentModifyContext::AddMemberIDModified(const FGuid& InMemberIDModified)
+{
+	bDocumentModified = true;
+	MemberIDsModified.Add(InMemberIDModified);
+}
+
+void FMetasoundFrontendDocumentModifyContext::AddMemberIDsModified(const TSet<FGuid>& InMemberIDsModified)
+{
+	bDocumentModified = true;
+	MemberIDsModified.Append(InMemberIDsModified);
+}
+
+void FMetasoundFrontendDocumentModifyContext::AddNodeIDModified(const FGuid& InNodeIDModified)
+{
+	bDocumentModified = true;
+	NodeIDsModified.Add(InNodeIDModified);
+}
+
+void FMetasoundFrontendDocumentModifyContext::AddNodeIDsModified(const TSet<FGuid>& InNodesModified)
+{
+	bDocumentModified = true;
+	NodeIDsModified.Append(InNodesModified);
+}
+#endif // WITH_EDITORONLY_DATA
+
+
 FMetasoundFrontendNodeInterface::FMetasoundFrontendNodeInterface(const FMetasoundFrontendClassInterface& InClassInterface)
 {
 	for (const FMetasoundFrontendClassInput& Input : InClassInterface.Inputs)
diff --git a/Engine/Plugins/Runtime/Metasound/Source/MetasoundFrontend/Private/MetasoundFrontendDocumentController.cpp b/Engine/Plugins/Runtime/Metasound/Source/MetasoundFrontend/Private/MetasoundFrontendDocumentController.cpp
--- a/Engine/Plugins/Runtime/Metasound/Source/MetasoundFrontend/Private/MetasoundFrontendDocumentController.cpp
+++ b/Engine/Plugins/Runtime/Metasound/Source/MetasoundFrontend/Private/MetasoundFrontendDocumentController.cpp
@@ -156,6 +156,9 @@ namespace Metasound
 		{
 			if (FMetasoundFrontendDocument* Doc = DocumentPtr.Get())
 			{
+#if WITH_EDITOR
+				Doc->Metadata.ModifyContext.AddInterfaceModified({ InVersion.Name });
+#endif // WITH_EDITOR
 				Doc->Interfaces.Add(InVersion);
 			}
 		}
@@ -164,6 +167,9 @@ namespace Metasound
 		{
 			if (FMetasoundFrontendDocument* Doc = DocumentPtr.Get())
 			{
+#if WITH_EDITOR
+				Doc->Metadata.ModifyContext.AddInterfaceModified({ InVersion.Name });
+#endif // WITH_EDITOR
 				Doc->Interfaces.Remove(InVersion);
 			}
 		}
@@ -251,6 +257,16 @@ namespace Metasound
 			return Invalid::GetInvalidDocumentMetadata();
 		}
 
+		FMetasoundFrontendDocumentMetadata* FDocumentController::GetMetadata()
+		{
+			if (FMetasoundFrontendDocument* Document = DocumentPtr.Get())
+			{
+				return &Document->Metadata;
+			}
+
+			return nullptr;
+		}
+
 		FConstClassAccessPtr FDocumentController::FindClass(const FNodeRegistryKey& InKey) const
 		{
 			return DocumentPtr.GetClassWithRegistryKey(InKey);
diff --git a/Engine/Plugins/Runtime/Metasound/Source/MetasoundFrontend/Private/MetasoundFrontendDocumentController.h b/Engine/Plugins/Runtime/Metasound/Source/MetasoundFrontend/Private/MetasoundFrontendDocumentController.h
--- a/Engine/Plugins/Runtime/Metasound/Source/MetasoundFrontend/Private/MetasoundFrontendDocumentController.h
+++ b/Engine/Plugins/Runtime/Metasound/Source/MetasoundFrontend/Private/MetasoundFrontendDocumentController.h
@@ -61,6 +61,7 @@ namespace Metasound
 
 			void SetMetadata(const FMetasoundFrontendDocumentMetadata& InMetadata) override;
 			const FMetasoundFrontendDocumentMetadata& GetMetadata() const override;
+			FMetasoundFrontendDocumentMetadata* GetMetadata() override;
 
 			void RemoveUnreferencedDependencies() override;
 			TArray<FConstClassAccessPtr> SynchronizeDependencyMetadata() override;
diff --git a/Engine/Plugins/Runtime/Metasound/Source/MetasoundFrontend/Private/MetasoundFrontendDocumentVersioning.cpp b/Engine/Plugins/Runtime/Metasound/Source/MetasoundFrontend/Private/MetasoundFrontendDocumentVersioning.cpp
--- a/Engine/Plugins/Runtime/Metasound/Source/MetasoundFrontend/Private/MetasoundFrontendDocumentVersioning.cpp
+++ b/Engine/Plugins/Runtime/Metasound/Source/MetasoundFrontend/Private/MetasoundFrontendDocumentVersioning.cpp
@@ -26,18 +26,14 @@ namespace Metasound
 			public:
 				bool Transform(FDocumentHandle InDocument) const override
 				{
-					const FMetasoundFrontendDocumentMetadata& Metadata = InDocument->GetMetadata();
-					if (Metadata.Version.Number >= GetTargetVersion())
+					FMetasoundFrontendDocumentMetadata* Metadata = InDocument->GetMetadata();
+					if (Metadata && Metadata->Version.Number >= GetTargetVersion())
 					{
 						return false;
 					}
 
 					TransformInternal(InDocument);
-
-					FMetasoundFrontendDocumentMetadata NewMetadata = Metadata;
-					NewMetadata.Version.Number = GetTargetVersion();
-					InDocument->SetMetadata(NewMetadata);
-
+					Metadata->Version.Number = GetTargetVersion();
 					return true;
 				}
 		};
@@ -186,8 +182,10 @@ namespace Metasound
 
 			void TransformInternal(FDocumentHandle InDocument) const override
 			{
-				check(InDocument->GetMetadata().Version.Number.Major == 1);
-				check(InDocument->GetMetadata().Version.Number.Minor == 3);
+				FMetasoundFrontendDocumentMetadata* Metadata = InDocument->GetMetadata();
+				check(Metadata);
+				check(Metadata->Version.Number.Major == 1);
+				check(Metadata->Version.Number.Minor == 3);
 
 				const TSet<FMetasoundFrontendVersion>& Interfaces = InDocument->GetInterfaceVersions();
 
@@ -499,7 +497,9 @@ namespace Metasound
 
 			bool bWasUpdated = false;
 
-			const FMetasoundFrontendVersionNumber InitVersionNumber = InDocument->GetMetadata().Version.Number;
+			FMetasoundFrontendDocumentMetadata* Metadata = InDocument->GetMetadata();
+			check(Metadata);
+			const FMetasoundFrontendVersionNumber InitVersionNumber = Metadata->Version.Number;
 
 			// Add additional transforms here after defining them above, example below.
 			bWasUpdated |= FVersionDocument_1_1(Name, Path).Transform(InDocument);
@@ -515,7 +515,7 @@ namespace Metasound
 
 			if (bWasUpdated)
 			{
-				const FMetasoundFrontendVersionNumber NewVersionNumber = InDocument->GetMetadata().Version.Number;
+				const FMetasoundFrontendVersionNumber& NewVersionNumber = Metadata->Version.Number;
 				UE_LOG(LogMetaSound, Display, TEXT("MetaSound at '%s' Document Versioned: '%s' --> '%s'"), *Path, *InitVersionNumber.ToString(), *NewVersionNumber.ToString());
 			}
 
diff --git a/Engine/Plugins/Runtime/Metasound/Source/MetasoundFrontend/Private/MetasoundFrontendGraphController.cpp b/Engine/Plugins/Runtime/Metasound/Source/MetasoundFrontend/Private/MetasoundFrontendGraphController.cpp
--- a/Engine/Plugins/Runtime/Metasound/Source/MetasoundFrontend/Private/MetasoundFrontendGraphController.cpp
+++ b/Engine/Plugins/Runtime/Metasound/Source/MetasoundFrontend/Private/MetasoundFrontendGraphController.cpp
@@ -286,6 +286,12 @@ namespace Metasound
 					Variable.TypeName = Info.DataTypeName;
 					Variable.Literal.SetFromLiteral(Registry.CreateDefaultLiteral(InDataType));
 					Variable.ID = VariableID;
+#if WITH_EDITOR
+					if (FMetasoundFrontendDocumentMetadata* DocMetadata = OwningDocument->GetMetadata())
+					{
+						DocMetadata->ModifyContext.AddMemberIDModified(Variable.ID);
+					}
+#endif // WITH_EDITOR
 
 					FMetasoundFrontendClass VariableNodeClass;
 					if (IDataTypeRegistry::Get().GetFrontendVariableClass(Variable.TypeName, VariableNodeClass))
@@ -365,6 +371,13 @@ namespace Metasound
 						RemoveNode(*NodeHandle);
 					}
 
+#if WITH_EDITOR
+					if (FMetasoundFrontendDocumentMetadata* DocMetadata = OwningDocument->GetMetadata())
+					{
+						DocMetadata->ModifyContext.AddMemberIDModified(InVariableID);
+					}
+#endif // WITH_EDITOR
+
 					auto IsVariableWithID = [&](const FMetasoundFrontendVariable& InVariable)
 					{
 						return InVariable.ID == InVariableID;
@@ -1001,6 +1014,12 @@ namespace Metasound
 								NewInput.VertexID = FGuid::NewGuid();
 							}
 
+#if WITH_EDITOR
+							if (FMetasoundFrontendDocumentMetadata* DocMetadata = OwningDocument->GetMetadata())
+							{
+								DocMetadata->ModifyContext.AddMemberIDModified(NewInput.NodeID);
+							}
+#endif // WITH_EDITOR
 							GraphClass->Interface.UpdateChangeID();
 
 							FNodeAccessPtr NodePtr = GraphClassPtr.GetNodeWithNodeID(Node.GetID());
@@ -1136,6 +1155,13 @@ namespace Metasound
 								NewOutput.VertexID = FGuid::NewGuid();
 							}
 
+#if WITH_EDITOR
+							if (FMetasoundFrontendDocumentMetadata* DocMetadata = OwningDocument->GetMetadata())
+							{
+								DocMetadata->ModifyContext.AddMemberIDModified(NewOutput.VertexID);
+							}
+#endif // WITH_EDITO
+
 							// Mark interface as changed.
 							GraphClass->Interface.UpdateChangeID();
 
@@ -1621,6 +1647,12 @@ namespace Metasound
 					}
 					Node.UpdateID(InNodeGuid);
 
+#if WITH_EDITOR
+					if (FMetasoundFrontendDocumentMetadata* DocMetadata = OwningDocument->GetMetadata())
+					{
+						DocMetadata->ModifyContext.AddNodeIDModified(InNodeGuid);
+					}
+#endif // WITH_EDITOR
 					FNodeAccessPtr NodePtr = GraphClassPtr.GetNodeWithNodeID(Node.GetID());
 					return GetNodeHandle(FGraphController::FNodeAndClass { NodePtr, InExistingDependency });
 				}
@@ -1644,6 +1676,16 @@ namespace Metasound
 
 				OwningDocument->RemoveUnreferencedDependencies();
 
+#if WITH_EDITOR
+				if (NumRemoved > 0)
+				{
+					if (FMetasoundFrontendDocumentMetadata* DocMetadata = OwningDocument->GetMetadata())
+					{
+						DocMetadata->ModifyContext.AddNodeIDModified(InDesc.GetID());
+					}
+				}	
+#endif // WITH_EDITOR
+
 				return (NumRemoved > 0);
 			}
 			return false;
@@ -1659,6 +1701,12 @@ namespace Metasound
 
 				if (NumInputsRemoved > 0)
 				{
+#if WITH_EDITOR
+					if (FMetasoundFrontendDocumentMetadata* DocMetadata = OwningDocument->GetMetadata())
+					{
+						DocMetadata->ModifyContext.AddMemberIDModified(InNode.GetID());
+					}
+#endif // WITH_EDITOR
 					GraphClass->Interface.UpdateChangeID();
 				}
 
@@ -1680,6 +1728,12 @@ namespace Metasound
 
 				if (NumOutputsRemoved)
 				{
+#if WITH_EDITOR
+					if (FMetasoundFrontendDocumentMetadata* DocMetadata = OwningDocument->GetMetadata())
+					{
+						DocMetadata->ModifyContext.AddMemberIDModified(InNode.GetID());
+					}
+#endif // WITH_EDITOR
 					GraphClass->Interface.UpdateChangeID();
 				}
 
diff --git a/Engine/Plugins/Runtime/Metasound/Source/MetasoundFrontend/Private/MetasoundFrontendInvalidController.h b/Engine/Plugins/Runtime/Metasound/Source/MetasoundFrontend/Private/MetasoundFrontendInvalidController.h
--- a/Engine/Plugins/Runtime/Metasound/Source/MetasoundFrontend/Private/MetasoundFrontendInvalidController.h
+++ b/Engine/Plugins/Runtime/Metasound/Source/MetasoundFrontend/Private/MetasoundFrontendInvalidController.h
@@ -473,6 +473,7 @@ namespace Metasound
 
 				virtual void SetMetadata(const FMetasoundFrontendDocumentMetadata& InMetadata) override { }
 				virtual const FMetasoundFrontendDocumentMetadata& GetMetadata() const override { return Invalid::GetInvalidDocumentMetadata(); }
+				virtual FMetasoundFrontendDocumentMetadata* GetMetadata() override { return nullptr; }
 
 				virtual void RemoveUnreferencedDependencies() override { }
 				virtual TArray<FConstClassAccessPtr> SynchronizeDependencyMetadata() override { return { }; }
diff --git a/Engine/Plugins/Runtime/Metasound/Source/MetasoundFrontend/Public/MetasoundAssetBase.h b/Engine/Plugins/Runtime/Metasound/Source/MetasoundFrontend/Public/MetasoundAssetBase.h
--- a/Engine/Plugins/Runtime/Metasound/Source/MetasoundFrontend/Public/MetasoundAssetBase.h
+++ b/Engine/Plugins/Runtime/Metasound/Source/MetasoundFrontend/Public/MetasoundAssetBase.h
@@ -140,14 +140,8 @@ public:
 	 */
 	void CacheRegistryMetadata();
 
-	// TODO: These flags & associated functions are highly UE editor-specific.
-	// Split synchronization requirement flag into synchronization required &
-	// object type refresh or checking frontend class guids when synchronizing.
-	bool GetSynchronizationRequired() const;
-	bool GetSynchronizationUpdateDetails() const;
-	void ResetSynchronizationState();
-	void SetUpdateDetailsOnSynchronization();
-	void SetSynchronizationRequired();
+	FMetasoundFrontendDocumentModifyContext& GetModifyContext();
+	const FMetasoundFrontendDocumentModifyContext& GetModifyContext() const;
 #endif // WITH_EDITOR
 
 	// Calls the outermost package and marks it dirty.
@@ -183,11 +177,6 @@ protected:
 	// Returns an access pointer to the document.
 	virtual Metasound::Frontend::FConstDocumentAccessPtr GetDocument() const = 0;
 
-#if WITH_EDITORONLY_DATA
-	bool bSynchronizationRequired = true;
-	bool bSynchronizationUpdateDetails = false;
-#endif // WITH_EDITORONLY_DATA
-
 protected:
 	// Container for runtime data of MetaSound graph.
 	struct FRuntimeData
diff --git a/Engine/Plugins/Runtime/Metasound/Source/MetasoundFrontend/Public/MetasoundFrontendController.h b/Engine/Plugins/Runtime/Metasound/Source/MetasoundFrontend/Public/MetasoundFrontendController.h
--- a/Engine/Plugins/Runtime/Metasound/Source/MetasoundFrontend/Public/MetasoundFrontendController.h
+++ b/Engine/Plugins/Runtime/Metasound/Source/MetasoundFrontend/Public/MetasoundFrontendController.h
@@ -1035,6 +1035,7 @@ namespace Metasound
 
 			virtual void SetMetadata(const FMetasoundFrontendDocumentMetadata& InMetadata) = 0;
 			virtual const FMetasoundFrontendDocumentMetadata& GetMetadata() const = 0;
+			virtual FMetasoundFrontendDocumentMetadata* GetMetadata() = 0;
 
 			/** Returns an existing Metasound class description corresponding to 
 			 * a dependency which matches the provided class information.
diff --git a/Engine/Plugins/Runtime/Metasound/Source/MetasoundFrontend/Public/MetasoundFrontendDocument.h b/Engine/Plugins/Runtime/Metasound/Source/MetasoundFrontend/Public/MetasoundFrontendDocument.h
--- a/Engine/Plugins/Runtime/Metasound/Source/MetasoundFrontend/Public/MetasoundFrontendDocument.h
+++ b/Engine/Plugins/Runtime/Metasound/Source/MetasoundFrontend/Public/MetasoundFrontendDocument.h
@@ -47,6 +47,62 @@ namespace Metasound
 	} // namespace Frontend
 } // namespace Metasound
 
+
+#if WITH_EDITORONLY_DATA
+// Struct containing any modified data breadcrumbs to inform what the editor/view layer must synchronize or refresh.
+class METASOUNDFRONTEND_API FMetasoundFrontendDocumentModifyContext
+{
+private:
+	// Whether or not the owning asset's MetaSoundDocument has been modified. True by default to force refreshing views on loading/reloading asset.
+	bool bDocumentModified = true;
+
+	// Whether or not to force refresh all views. True by default to force refreshing views on loading/reloading asset.
+	bool bForceRefreshViews = true;
+
+	// Which Interfaces have been modified since the last editor graph synchronization
+	TSet<FName> InterfacesModified;
+
+	// Which MemberIDs have been modified since the last editor graph synchronization
+	TSet<FGuid> MemberIDsModified;
+
+	// Which NodeIDs have been modified since the last editor graph synchronization
+	TSet<FGuid> NodeIDsModified;
+
+public:
+	void ClearDocumentModified();
+
+	bool GetDocumentModified() const;
+	bool GetForceRefreshViews() const;
+	const TSet<FName>& GetInterfacesModified() const;
+	const TSet<FGuid>& GetNodeIDsModified() const;
+	const TSet<FGuid>& GetMemberIDsModified() const;
+
+	void Reset();
+
+	void SetDocumentModified();
+	void SetForceRefreshViews();
+
+	// Adds an interface name to the set of interfaces that have been modified since last context reset/construction
+	void AddInterfaceModified(FName InInterfaceModified);
+
+	// Performs union of provided interface set with the set of interfaces that have been modified since last context reset/construction
+	void AddInterfacesModified(const TSet<FName>& InInterfacesModified);
+
+	// Adds a MemberID to the set of interfaces that have been modified since last context reset/construction
+	void AddMemberIDModified(const FGuid& InMemberIDModified);
+
+	// Performs union of provided MemberIDs set with the set of MemberIDs that have been modified since last context reset/construction
+	void AddMemberIDsModified(const TSet<FGuid>& InMemberIDsModified);
+
+	// Performs union of provided NodeID set with the set of NodeIDs that have been modified since last context reset/construction
+	void AddNodeIDModified(const FGuid& InNodeIDModified);
+
+	// Performs union of provided interface set with the set of interfaces that have been modified since last context reset/construction
+	void AddNodeIDsModified(const TSet<FGuid>& InNodeIDsModified);
+};
+#endif // WITH_EDITORONLY_DATA
+
+
 // Describes how a vertex accesses the data connected to it. 
 UENUM()
 enum class EMetasoundFrontendVertexAccessType
@@ -1361,6 +1417,10 @@ struct METASOUNDFRONTEND_API FMetasoundFrontendDocumentMetadata
 
 	UPROPERTY()
 	FMetasoundFrontendVersion Version;
+
+#if WITH_EDITORONLY_DATA
+	FMetasoundFrontendDocumentModifyContext ModifyContext;
+#endif // WITH_EDITORONLY_DATA
 };
 
 USTRUCT()
diff --git a/Engine/Plugins/Runtime/Metasound/Source/MetasoundGraphCore/Private/MetasoundLiteralNode.cpp b/Engine/Plugins/Runtime/Metasound/Source/MetasoundGraphCore/Private/MetasoundLiteralNode.cpp
deleted file mode 100644
