Fix for include cycle
#rnx
#hashtag hashtagforthehashtagallthethings

diff --git a/Engine/Plugins/Runtime/Metasound/Source/MetasoundEditor/Public/MetasoundEditorGraph.h b/Engine/Plugins/Runtime/Metasound/Source/MetasoundEditor/Public/MetasoundEditorGraph.h
--- a/Engine/Plugins/Runtime/Metasound/Source/MetasoundEditor/Public/MetasoundEditorGraph.h
+++ b/Engine/Plugins/Runtime/Metasound/Source/MetasoundEditor/Public/MetasoundEditorGraph.h
@@ -3,7 +3,6 @@
 
 #include "AudioParameterControllerInterface.h"
 #include "MetasoundEditor.h"
-#include "MetasoundEditorGraphBuilder.h"
 #include "MetasoundEditorGraphNode.h"
 #include "MetasoundFrontendController.h"
 #include "MetasoundFrontendDocument.h"
