Adding graph name to warning log for metasound array nodes

diff --git a/Engine/Plugins/Runtime/Metasound/Source/MetasoundFrontend/Public/Interfaces/MetasoundFrontendSourceInterface.h b/Engine/Plugins/Runtime/Metasound/Source/MetasoundFrontend/Public/Interfaces/MetasoundFrontendSourceInterface.h
--- a/Engine/Plugins/Runtime/Metasound/Source/MetasoundFrontend/Public/Interfaces/MetasoundFrontendSourceInterface.h
+++ b/Engine/Plugins/Runtime/Metasound/Source/MetasoundFrontend/Public/Interfaces/MetasoundFrontendSourceInterface.h
@@ -5,8 +5,8 @@
 #include "MetasoundFrontendDocument.h"
 #include "MetasoundFrontendTransform.h"
 #include "UObject/NameTypes.h"
-//#include "MetasoundEngineArchetypes.h"
 #include "MetasoundFrontendController.h"
+
 // Forward Declarations
 struct FMetasoundFrontendVersion;
 
@@ -76,8 +76,5 @@ namespace Metasound
 				virtual bool Transform(Frontend::FDocumentHandle InDocument) const override;
 			};
 		}
-
-
-
 	} // namespace Frontend
 } // namespace Metasound
diff --git a/Engine/Plugins/Runtime/Metasound/Source/MetasoundFrontend/Public/MetasoundArrayNodes.h b/Engine/Plugins/Runtime/Metasound/Source/MetasoundFrontend/Public/MetasoundArrayNodes.h
--- a/Engine/Plugins/Runtime/Metasound/Source/MetasoundFrontend/Public/MetasoundArrayNodes.h
+++ b/Engine/Plugins/Runtime/Metasound/Source/MetasoundFrontend/Public/MetasoundArrayNodes.h
@@ -2,9 +2,11 @@
 #pragma once
 
 #include "CoreMinimal.h"
+#include "Interfaces/MetasoundFrontendSourceInterface.h"
 #include "Internationalization/Text.h"
 #include "MetasoundBuilderInterface.h"
 #include "MetasoundDataFactory.h"
+#include "MetasoundEnvironment.h"
 #include "MetasoundExecutableOperator.h"
 #include "MetasoundFacade.h"
 #include "MetasoundLog.h"
@@ -20,7 +22,6 @@
 
 #define LOCTEXT_NAMESPACE "MetasoundFrontend"
 
-
 namespace Metasound
 {
 	namespace MetasoundArrayNodesPrivate
@@ -220,6 +221,16 @@ namespace Metasound
 			return Metadata;
 		}
 
+		struct FInitParams
+		{
+			TDataReadReference<FTrigger> Trigger;
+			FArrayDataReadReference Array;
+			TDataReadReference<int32> Index;
+#if WITH_METASOUND_DEBUG_ENVIRONMENT
+			FString GraphName;
+#endif
+		};
+
 		static TUniquePtr<IOperator> CreateOperator(const FCreateOperatorParams& InParams, TArray<TUniquePtr<IOperatorBuildError>>& OutErrors)
 		{
 			using namespace ArrayNodeVertexNames;
@@ -236,15 +247,29 @@ namespace Metasound
 			// Input Index
 			TDataReadReference<int32> Index = InParams.InputDataReferences.GetDataReadReferenceOrConstructWithVertexDefault<int32>(Inputs, METASOUND_GET_PARAM_NAME(InputIndex), InParams.OperatorSettings);
 
-			return MakeUnique<TArrayGetOperator>(InParams.OperatorSettings, Trigger, Array, Index);
+			FInitParams OperatorInitParams
+			{
+				Trigger
+				, Array
+				, Index
+#if WITH_METASOUND_DEBUG_ENVIRONMENT
+				, *InParams.Environment.GetValue<FString>(Frontend::SourceInterface::Environment::GraphName)
+#endif // WITH_METASOUND_DEBUG_ENVIRONMENT
+			};
+
+			return MakeUnique<TArrayGetOperator>(InParams.OperatorSettings, MoveTemp(OperatorInitParams));
 		}
 
 
-		TArrayGetOperator(const FOperatorSettings& InSettings, TDataReadReference<FTrigger> InTrigger, FArrayDataReadReference InArray, TDataReadReference<int32> InIndex)
-		: Trigger(InTrigger)
-		, Array(InArray)
-		, Index(InIndex)
+
+		TArrayGetOperator(const FOperatorSettings& InSettings, FInitParams&& InParams)
+		: Trigger(InParams.Trigger)
+		, Array(InParams.Array)
+		, Index(InParams.Index)
 		, Value(TDataWriteReferenceFactory<ElementType>::CreateAny(InSettings))
+#if WITH_METASOUND_DEBUG_ENVIRONMENT
+		, GraphName(InParams.GraphName)
+#endif // WITH_METASOUND_DEBUG_ENVIRONMENT
 		{
 		}
 
@@ -286,10 +311,12 @@ namespace Metasound
 				{
 					*Value = ArrayRef[IndexValue];
 				}
+#if WITH_METASOUND_DEBUG_ENVIRONMENT
 				else
 				{
-					UE_LOG(LogMetaSound, Error, TEXT("Attempt to get value at invalid index [ArraySize:%d, Index:%d]"), ArrayRef.Num(), IndexValue);
+					UE_LOG(LogMetaSound, Warning, TEXT("Attempt to get value at invalid index [ArraySize:%d, Index:%d] in MetaSound Graph \"%s\"."), ArrayRef.Num(), IndexValue, *GraphName);
 				}
+#endif // WITH_METASOUND_DEBUG_ENVIRONMENT
 			}
 		}
 
@@ -299,6 +326,11 @@ namespace Metasound
 		FArrayDataReadReference Array;
 		TDataReadReference<int32> Index;
 		TDataWriteReference<ElementType> Value;
+
+#if WITH_METASOUND_DEBUG_ENVIRONMENT
+		FString GraphName;
+#endif // WITH_METASOUND_DEBUG_ENVIRONMENT
+
 	};
 
 	template<typename ArrayType>
@@ -358,6 +390,18 @@ namespace Metasound
 			return Metadata;
 		}
 
+		struct FInitParams
+		{
+			TDataReadReference<FTrigger> Trigger;
+			FArrayDataReadReference InitArray;
+			FArrayDataWriteReference Array;
+			TDataReadReference<int32> Index;
+			TDataReadReference<ElementType> Value;
+#if WITH_METASOUND_DEBUG_ENVIRONMENT
+			FString GraphName;
+#endif // WITH_METASOUND_DEBUG_ENVIRONMENT
+		};
+
 		static TUniquePtr<IOperator> CreateOperator(const FCreateOperatorParams& InParams, TArray<TUniquePtr<IOperatorBuildError>>& OutErrors)
 		{
 			using namespace ArrayNodeVertexNames;
@@ -374,17 +418,31 @@ namespace Metasound
 
 			TDataReadReference<ElementType> Value = InParams.InputDataReferences.GetDataReadReferenceOrConstructWithVertexDefault<ElementType>(Inputs, METASOUND_GET_PARAM_NAME(InputValue), InParams.OperatorSettings);
 
-			return MakeUnique<TArraySetOperator>(InParams.OperatorSettings, Trigger, InitArray, Array, Index, Value);
-		}
+			FInitParams OperatorInitParams 
+			{
+				Trigger
+				, InitArray
+				, Array 
+				, Index 
+				, Value
+#if WITH_METASOUND_DEBUG_ENVIRONMENT
+				, *InParams.Environment.GetValue<FString>(Frontend::SourceInterface::Environment::GraphName)
+#endif // WITH_METASOUND_DEBUG_ENVIRONMENT
+			};
 
+			return MakeUnique<TArraySetOperator>(InParams.OperatorSettings, MoveTemp(OperatorInitParams));
+		}
 
-		TArraySetOperator(const FOperatorSettings& InSettings, TDataReadReference<FTrigger> InTrigger, FArrayDataReadReference InInitArray, FArrayDataWriteReference InArray, TDataReadReference<int32> InIndex, TDataReadReference<ElementType> InValue)
+		TArraySetOperator(const FOperatorSettings& InSettings, FInitParams&& InParams)
 		: OperatorSettings(InSettings)
-		, Trigger(InTrigger)
-		, InitArray(InInitArray)
-		, Array(InArray)
-		, Index(InIndex)
-		, Value(InValue)
+		, Trigger(InParams.Trigger)
+		, InitArray(InParams.InitArray)
+		, Array(InParams.Array)
+		, Index(InParams.Index)
+		, Value(InParams.Value)
+#if WITH_METASOUND_DEBUG_ENVIRONMENT
+		, GraphName(InParams.GraphName)
+#endif // WITH_METASOUND_DEBUG_ENVIRONMENT
 		{
 		}
 
@@ -424,10 +482,12 @@ namespace Metasound
 				{
 					ArrayRef[IndexValue] = *Value;
 				}
+#if WITH_METASOUND_DEBUG_ENVIRONMENT
 				else
 				{
-					UE_LOG(LogMetaSound, Error, TEXT("Attempt to set value at invalid index [ArraySize:%d, Index:%d]"), ArrayRef.Num(), IndexValue);
+					UE_LOG(LogMetaSound, Warning, TEXT("Attempt to set value at invalid index [ArraySize:%d, Index:%d] in MetaSound Graph \"%s\"."), ArrayRef.Num(), IndexValue, *GraphName);
 				}
+#endif // WITH_METASOUND_DEBUG_ENVIRONMENT
 			}
 		}
 
@@ -439,6 +499,11 @@ namespace Metasound
 		FArrayDataWriteReference Array;
 		TDataReadReference<int32> Index;
 		TDataReadReference<ElementType> Value;
+
+#if WITH_METASOUND_DEBUG_ENVIRONMENT
+		FString GraphName;
+#endif // WITH_METASOUND_DEBUG_ENVIRONMENT
+
 	};
 
 	template<typename ArrayType>
