Fix clicking in bandsplitter node when changing crossover frequencies.

diff --git a/Engine/Plugins/Runtime/Metasound/Source/MetasoundStandardNodes/Private/MetasoundBandsplitterNode.cpp b/Engine/Plugins/Runtime/Metasound/Source/MetasoundStandardNodes/Private/MetasoundBandsplitterNode.cpp
--- a/Engine/Plugins/Runtime/Metasound/Source/MetasoundStandardNodes/Private/MetasoundBandsplitterNode.cpp
+++ b/Engine/Plugins/Runtime/Metasound/Source/MetasoundStandardNodes/Private/MetasoundBandsplitterNode.cpp
@@ -248,36 +248,44 @@ namespace Metasound
 
 		void UpdateFilterSettings()
 		{
-			bool bNeedsUpdate = false;
+			bool bNeedsReinitFilters = false;
+			bool bNeedsUpdateCrossovers = false;
 
 			for (int32 CrossoverIndex = 0; CrossoverIndex < Crossovers.Num(); CrossoverIndex++)
 			{
 				if (Crossovers[CrossoverIndex] != *CrossoverFrequencies[CrossoverIndex])
 				{
-					bNeedsUpdate = true;
+					bNeedsUpdateCrossovers = true;
 					Crossovers[CrossoverIndex] = *CrossoverFrequencies[CrossoverIndex];
 				}
 			}
 
 			if (*bPhaseCompensate != bPrevPhaseCompensate)
 			{
-				bNeedsUpdate = true;
+				bNeedsReinitFilters = true;
 				bPrevPhaseCompensate = *bPhaseCompensate;
 			}
 
 			if (PrevFilterOrder != GetFilterOrder(*FilterOrder))
 			{
-				bNeedsUpdate = true;
+				bNeedsReinitFilters = true;
 				PrevFilterOrder = GetFilterOrder(*FilterOrder);
 			}
 
-			if (bNeedsUpdate)
+			if (bNeedsReinitFilters)
 			{
 				for (int32 ChannelIndex = 0; ChannelIndex < NumChannels; ChannelIndex++)
 				{
 					Filters[ChannelIndex].Init(1, Settings.GetSampleRate(), PrevFilterOrder, bPrevPhaseCompensate, Crossovers);
 				}
 			}
+			else if (bNeedsUpdateCrossovers)
+			{
+				for (int32 ChannelIndex = 0; ChannelIndex < NumChannels; ChannelIndex++)
+				{
+					Filters[ChannelIndex].SetCrossovers(Crossovers);
+				}
+			}
 		}
 
 		void Execute()
