Fix null deref static analysis warning
#rnx

diff --git a/Engine/Plugins/Runtime/Metasound/Source/MetasoundFrontend/Private/MetasoundFrontendDocumentVersioning.cpp b/Engine/Plugins/Runtime/Metasound/Source/MetasoundFrontend/Private/MetasoundFrontendDocumentVersioning.cpp
--- a/Engine/Plugins/Runtime/Metasound/Source/MetasoundFrontend/Private/MetasoundFrontendDocumentVersioning.cpp
+++ b/Engine/Plugins/Runtime/Metasound/Source/MetasoundFrontend/Private/MetasoundFrontendDocumentVersioning.cpp
@@ -26,15 +26,18 @@ namespace Metasound
 			public:
 				bool Transform(FDocumentHandle InDocument) const override
 				{
-					FMetasoundFrontendDocumentMetadata* Metadata = InDocument->GetMetadata();
-					if (Metadata && Metadata->Version.Number >= GetTargetVersion())
+					if (FMetasoundFrontendDocumentMetadata* Metadata = InDocument->GetMetadata())
 					{
-						return false;
+						const FMetasoundFrontendVersionNumber TargetVersion = GetTargetVersion();
+						if (Metadata->Version.Number < TargetVersion)
+						{
+							TransformInternal(InDocument);
+							Metadata->Version.Number = TargetVersion;
+							return true;
+						}
 					}
 
-					TransformInternal(InDocument);
-					Metadata->Version.Number = GetTargetVersion();
-					return true;
+					return false;
 				}
 		};
 
