Fix for ensure on metasound audio tests accessing missing variable

diff --git a/Engine/Plugins/Runtime/Metasound/Source/MetasoundFrontend/Public/MetasoundArrayNodes.h b/Engine/Plugins/Runtime/Metasound/Source/MetasoundFrontend/Public/MetasoundArrayNodes.h
--- a/Engine/Plugins/Runtime/Metasound/Source/MetasoundFrontend/Public/MetasoundArrayNodes.h
+++ b/Engine/Plugins/Runtime/Metasound/Source/MetasoundFrontend/Public/MetasoundArrayNodes.h
@@ -246,6 +246,13 @@ namespace Metasound
 
 			// Input Index
 			TDataReadReference<int32> Index = InParams.InputDataReferences.GetDataReadReferenceOrConstructWithVertexDefault<int32>(Inputs, METASOUND_GET_PARAM_NAME(InputIndex), InParams.OperatorSettings);
+#if WITH_METASOUND_DEBUG_ENVIRONMENT
+			FString GraphName;
+			if (InParams.Environment.Contains<FString>(Frontend::SourceInterface::Environment::GraphName))
+			{
+				GraphName = InParams.Environment.GetValue<FString>(Frontend::SourceInterface::Environment::GraphName);
+			}
+#endif // WITH_METASOUND_DEBUG_ENVIRONMENT
 
 			FInitParams OperatorInitParams
 			{
@@ -253,7 +260,7 @@ namespace Metasound
 				, Array
 				, Index
 #if WITH_METASOUND_DEBUG_ENVIRONMENT
-				, *InParams.Environment.GetValue<FString>(Frontend::SourceInterface::Environment::GraphName)
+				, GraphName
 #endif // WITH_METASOUND_DEBUG_ENVIRONMENT
 			};
 
@@ -417,6 +424,13 @@ namespace Metasound
 			TDataReadReference<int32> Index = InParams.InputDataReferences.GetDataReadReferenceOrConstructWithVertexDefault<int32>(Inputs, METASOUND_GET_PARAM_NAME(InputIndex), InParams.OperatorSettings);
 
 			TDataReadReference<ElementType> Value = InParams.InputDataReferences.GetDataReadReferenceOrConstructWithVertexDefault<ElementType>(Inputs, METASOUND_GET_PARAM_NAME(InputValue), InParams.OperatorSettings);
+#if WITH_METASOUND_DEBUG_ENVIRONMENT
+			FString GraphName;
+			if (InParams.Environment.Contains<FString>(Frontend::SourceInterface::Environment::GraphName))
+			{
+				GraphName = InParams.Environment.GetValue<FString>(Frontend::SourceInterface::Environment::GraphName);
+			}
+#endif // WITH_METASOUND_DEBUG_ENVIRONMENT
 
 			FInitParams OperatorInitParams 
 			{
@@ -426,7 +440,7 @@ namespace Metasound
 				, Index 
 				, Value
 #if WITH_METASOUND_DEBUG_ENVIRONMENT
-				, *InParams.Environment.GetValue<FString>(Frontend::SourceInterface::Environment::GraphName)
+				, GraphName
 #endif // WITH_METASOUND_DEBUG_ENVIRONMENT
 			};
 
