Fix for inconsistent metasound wave seek functionality

diff --git a/Engine/Plugins/Runtime/Metasound/Source/MetasoundEngine/Private/MetasoundWaveProxyReader.cpp b/Engine/Plugins/Runtime/Metasound/Source/MetasoundEngine/Private/MetasoundWaveProxyReader.cpp
--- a/Engine/Plugins/Runtime/Metasound/Source/MetasoundEngine/Private/MetasoundWaveProxyReader.cpp
+++ b/Engine/Plugins/Runtime/Metasound/Source/MetasoundEngine/Private/MetasoundWaveProxyReader.cpp
@@ -410,13 +410,26 @@ namespace Metasound
 		}
 
 		// Seek input to start time.
-		const bool bSeekSucceeded = DecoderInput->SeekToTime(InStartTimeInSeconds);
-		if (!bSeekSucceeded)
+		if (!FMath::IsNearlyEqual(0.0f, InStartTimeInSeconds))
 		{
-			UE_LOG(LogMetaSound, Warning, TEXT("Failed to seek decoder input during initialization: (format:%s) for wave (package:%s) to time '%.6f'"),
-				*Format.ToString(),
-				*WaveProxy->GetPackageName().ToString(),
-				InStartTimeInSeconds);
+			if (WaveProxy->IsSeekable())
+			{
+				const bool bSeekSucceeded = DecoderInput->SeekToTime(InStartTimeInSeconds);
+				if (!bSeekSucceeded)
+				{
+					UE_LOG(LogMetaSound, Warning, TEXT("Failed to seek decoder input during initialization: (format:%s) for wave (package:%s) to time '%.6f'"),
+						*Format.ToString(),
+						*WaveProxy->GetPackageName().ToString(),
+						InStartTimeInSeconds);
+				}
+			}
+			else
+			{
+				UE_LOG(LogMetaSound, Warning, TEXT("Attempt to seek on non-seekable wave: (format:%s) for wave (package:%s) to time '%.6f'"),
+					*Format.ToString(),
+					*WaveProxy->GetPackageName().ToString(),
+					InStartTimeInSeconds);
+			}
 		}
 		CurrentFrameIndex = InStartTimeInSeconds * GetSampleRate();
 
diff --git a/Engine/Source/Runtime/Engine/Private/SoundWave.cpp b/Engine/Source/Runtime/Engine/Private/SoundWave.cpp
--- a/Engine/Source/Runtime/Engine/Private/SoundWave.cpp
+++ b/Engine/Source/Runtime/Engine/Private/SoundWave.cpp
@@ -177,6 +177,7 @@ void FSoundWaveData::InitializeDataFromSoundWave(USoundWave& InWave)
 	bIsLooping = InWave.IsLooping();
 	bIsTemplate = InWave.IsTemplate();
 	bIsStreaming = InWave.IsStreaming(nullptr);
+	bIsSeekable = InWave.IsSeekable();
 	SoundAssetCompressionType = InWave.GetSoundAssetCompressionType();
 	bShouldUseStreamCaching = InWave.ShouldUseStreamCaching();
 
@@ -3220,8 +3221,6 @@ bool USoundWave::IsSeekable() const
 		Result = false;
 	}
 
-	SoundWaveDataPtr->bIsSeekable = Result; // update shared flags
-
 	return Result;
 }
 
