Prevent MetaSound template nodes from showing up in the node class selector menu

#rnx

diff --git a/Engine/Plugins/Runtime/Metasound/Source/MetasoundEditor/Private/MetasoundEditorGraphSchema.cpp b/Engine/Plugins/Runtime/Metasound/Source/MetasoundEditor/Private/MetasoundEditorGraphSchema.cpp
--- a/Engine/Plugins/Runtime/Metasound/Source/MetasoundEditor/Private/MetasoundEditorGraphSchema.cpp
+++ b/Engine/Plugins/Runtime/Metasound/Source/MetasoundEditor/Private/MetasoundEditorGraphSchema.cpp
@@ -1975,8 +1975,12 @@ void UMetasoundEditorGraphSchema::GetFunctionActions(FGraphActionMenuBuilder& Ac
 			continue;
 		}
 
-		if (FrontendClass.Metadata.GetType() != EMetasoundFrontendClassType::External
-			&& FrontendClass.Metadata.GetType() != EMetasoundFrontendClassType::Template)
+		if (FrontendClass.Metadata.GetType() == EMetasoundFrontendClassType::Template)
+		{
+			continue;
+		}
+
+		if (FrontendClass.Metadata.GetType() != EMetasoundFrontendClassType::External)
 		{
 			continue;
 		}
