Fix MetaSound input nodes creating wrong connections when renaming

diff --git a/Engine/Plugins/Runtime/Metasound/Source/MetasoundEditor/Private/MetasoundEditorGraph.cpp b/Engine/Plugins/Runtime/Metasound/Source/MetasoundEditor/Private/MetasoundEditorGraph.cpp
--- a/Engine/Plugins/Runtime/Metasound/Source/MetasoundEditor/Private/MetasoundEditorGraph.cpp
+++ b/Engine/Plugins/Runtime/Metasound/Source/MetasoundEditor/Private/MetasoundEditorGraph.cpp
@@ -259,6 +259,20 @@ void UMetasoundEditorGraphVertex::SetMemberName(const FName& InNewName, bool bPo
 	Graph->GetMetasoundChecked().Modify();
 
 	NodeHandle->SetNodeName(InNewName);
+
+	const TArray<UMetasoundEditorGraphMemberNode*> Nodes = GetNodes();
+	for (UMetasoundEditorGraphMemberNode* Node : Nodes)
+	{
+		const TArray<UEdGraphPin*>& Pins = Node->GetAllPins();
+		ensure(Pins.Num() == 1);
+
+		for (UEdGraphPin* Pin : Pins)
+		{
+			Pin->Modify();
+			Pin->PinName = InNewName;
+		}
+	}
+
 	Graph->RegisterGraphWithFrontend();
 	Graph->GetModifyContext().AddMemberIDsModified({ GetMemberID() });
 }
