[MetaSounds] Param loop virtualization fixes
- Split default param transmitter to SoundCueTransmitter & ParameterTransmitterBase to avoid generating transmitters for SoundBase child classes that don't support/require one (ex. SoundWaves)
- Move param transmitter to sharedptrs to avoid a bunch of unnecessary copy/clone logic when virtualizing loops
- Remove DefaultParameters field on ActiveSound in favor of shared impl in FParameterTransmitterBase that caches current state of parameters
- Fix for move that clears out active sound's current parameters when sound could have multiple wave instances that need said parameters

diff --git a/Engine/Plugins/Runtime/Metasound/Source/MetasoundEngine/Private/MetasoundSource.cpp b/Engine/Plugins/Runtime/Metasound/Source/MetasoundEngine/Private/MetasoundSource.cpp
--- a/Engine/Plugins/Runtime/Metasound/Source/MetasoundEngine/Private/MetasoundSource.cpp
+++ b/Engine/Plugins/Runtime/Metasound/Source/MetasoundEngine/Private/MetasoundSource.cpp
@@ -373,7 +373,7 @@ void UMetaSoundSource::SetRegistryAssetClassInfo(const Metasound::Frontend::FNod
 }
 #endif // WITH_EDITORONLY_DATA
 
-void UMetaSoundSource::InitParameters(TArray<FAudioParameter>& InParametersToInit, FName InFeatureName)
+void UMetaSoundSource::InitParameters(TArray<FAudioParameter>& ParametersToInit, FName InFeatureName)
 {
 	using namespace Metasound::SourcePrivate;
 
@@ -529,9 +529,9 @@ void UMetaSoundSource::InitParameters(TArray<FAudioParameter>& InParametersToIni
 	};
 
 
-	for (int32 i = InParametersToInit.Num() - 1; i >= 0; --i)
+	for (int32 i = ParametersToInit.Num() - 1; i >= 0; --i)
 	{
-		FAudioParameter& Parameter = InParametersToInit[i];
+		FAudioParameter& Parameter = ParametersToInit[i];
 		
 #if !NO_LOGGING
 		// For logging in case of failure
@@ -553,13 +553,13 @@ void UMetaSoundSource::InitParameters(TArray<FAudioParameter>& InParametersToIni
 				}
 #endif // !NO_LOGGING
 				constexpr bool bAllowShrinking = false;
-				InParametersToInit.RemoveAtSwap(i, 1, bAllowShrinking);
+				ParametersToInit.RemoveAtSwap(i, 1, bAllowShrinking);
 			}
 		}
 		else
 		{
 			constexpr bool bAllowShrinking = false;
-			InParametersToInit.RemoveAtSwap(i, 1, bAllowShrinking);
+			ParametersToInit.RemoveAtSwap(i, 1, bAllowShrinking);
 
 #if !NO_LOGGING
 			if (::Metasound::MetaSoundParameterEnableWarningOnIgnoredParameterCVar)
@@ -908,22 +908,19 @@ bool UMetaSoundSource::IsOneShot() const
 	return IsInterfaceDeclared(SourceOneShotInterface::GetVersion());
 }
 
-TUniquePtr<Audio::IParameterTransmitter> UMetaSoundSource::CreateParameterTransmitter(Audio::FParameterTransmitterInitParams&& InParams) const
+TSharedPtr<Audio::IParameterTransmitter> UMetaSoundSource::CreateParameterTransmitter(Audio::FParameterTransmitterInitParams&& InParams) const
 {
 	METASOUND_LLM_SCOPE;
 
-	Metasound::FMetaSoundParameterTransmitter::FInitParams InitParams(GetOperatorSettings(InParams.SampleRate), InParams.InstanceID);
+	Metasound::FMetaSoundParameterTransmitter::FInitParams InitParams(GetOperatorSettings(InParams.SampleRate), InParams.InstanceID, MoveTemp(InParams.DefaultParams));
+	InitParams.DebugMetaSoundName = GetFName();
 
 	for (const FSendInfoAndVertexName& InfoAndName : FMetasoundAssetBase::GetSendInfos(InParams.InstanceID))
 	{
 		InitParams.Infos.Add(InfoAndName.SendInfo);
 	}
 
-	InitParams.DebugMetaSoundName = GetFName();
-
-	TUniquePtr<Audio::IParameterTransmitter> NewTransmitter = MakeUnique<Metasound::FMetaSoundParameterTransmitter>(InitParams);
-
-	return NewTransmitter;
+	return MakeShared<Metasound::FMetaSoundParameterTransmitter>(MoveTemp(InitParams));
 }
 
 Metasound::FOperatorSettings UMetaSoundSource::GetOperatorSettings(Metasound::FSampleRate InSampleRate) const
diff --git a/Engine/Plugins/Runtime/Metasound/Source/MetasoundEngine/Public/MetasoundSource.h b/Engine/Plugins/Runtime/Metasound/Source/MetasoundEngine/Public/MetasoundSource.h
--- a/Engine/Plugins/Runtime/Metasound/Source/MetasoundEngine/Public/MetasoundSource.h
+++ b/Engine/Plugins/Runtime/Metasound/Source/MetasoundEngine/Public/MetasoundSource.h
@@ -168,7 +168,7 @@ public:
 		return this;
 	}
 
-	virtual void InitParameters(TArray<FAudioParameter>& InParametersToInit, FName InFeatureName) override;
+	virtual void InitParameters(TArray<FAudioParameter>& ParametersToInit, FName InFeatureName) override;
 	virtual void InitResources() override;
 
 	virtual bool IsPlayable() const override;
@@ -176,7 +176,7 @@ public:
 	virtual float GetDuration() const override;
 	virtual bool ImplementsParameterInterface(Audio::FParameterInterfacePtr InInterface) const override;
 	virtual ISoundGeneratorPtr CreateSoundGenerator(const FSoundGeneratorInitParams& InParams, TArray<FAudioParameter>&& InDefaultParameters) override;
-	virtual TUniquePtr<Audio::IParameterTransmitter> CreateParameterTransmitter(Audio::FParameterTransmitterInitParams&& InParams) const override;
+	virtual TSharedPtr<Audio::IParameterTransmitter> CreateParameterTransmitter(Audio::FParameterTransmitterInitParams&& InParams) const override;
 	virtual bool IsParameterValid(const FAudioParameter& InParameter) const override;
 	virtual bool IsLooping() const override;
 	virtual bool IsOneShot() const override;
diff --git a/Engine/Plugins/Runtime/Metasound/Source/MetasoundFrontend/Private/MetasoundParameterTransmitter.cpp b/Engine/Plugins/Runtime/Metasound/Source/MetasoundFrontend/Private/MetasoundParameterTransmitter.cpp
--- a/Engine/Plugins/Runtime/Metasound/Source/MetasoundFrontend/Private/MetasoundParameterTransmitter.cpp
+++ b/Engine/Plugins/Runtime/Metasound/Source/MetasoundFrontend/Private/MetasoundParameterTransmitter.cpp
@@ -101,6 +101,93 @@ namespace Metasound
 			return FLiteral();
 		}
 
+		FLiteral ConvertParameterToLiteral(const FAudioParameter& InValue)
+		{
+			switch (InValue.ParamType)
+			{
+				case EAudioParameterType::Boolean:
+				{
+					return FLiteral(InValue.BoolParam);
+				}
+
+				case EAudioParameterType::BooleanArray:
+				{
+					return FLiteral(InValue.ArrayBoolParam);
+				}
+
+				case EAudioParameterType::Float:
+				{
+					return FLiteral(InValue.FloatParam);
+				}
+
+				case EAudioParameterType::FloatArray:
+				{
+					return FLiteral(InValue.ArrayFloatParam);
+				}
+
+				case EAudioParameterType::Integer:
+				{
+					return FLiteral(InValue.IntParam);
+				}
+
+				case EAudioParameterType::IntegerArray:
+				{
+					return FLiteral(InValue.ArrayIntParam);
+				}
+
+				case EAudioParameterType::None:
+				{
+					return FLiteral();
+				}
+
+				case EAudioParameterType::NoneArray:
+				{
+					TArray<FLiteral::FNone> InitArray;
+					InitArray.Init(FLiteral::FNone(), InValue.IntParam);
+					return FLiteral(InitArray);
+				}
+
+				case EAudioParameterType::Object:
+				{
+					if (InValue.ObjectProxies.IsEmpty())
+					{
+						return FLiteral();
+					}
+
+					Audio::IProxyDataPtr ObjectProxyClone = InValue.ObjectProxies.Last()->Clone();
+					return FLiteral(MoveTemp(ObjectProxyClone));
+				}
+
+				case EAudioParameterType::ObjectArray:
+				{
+					TArray<Audio::IProxyDataPtr> ObjectProxiesClone;
+					Algo::Transform(InValue.ObjectProxies, ObjectProxiesClone, [](const Audio::IProxyDataPtr& DataPtr)
+					{
+						return DataPtr->Clone();
+					});
+					return FLiteral(MoveTemp(ObjectProxiesClone));
+				}
+
+				case EAudioParameterType::String:
+				{
+					return FLiteral(InValue.StringParam);
+				}
+
+				case EAudioParameterType::StringArray:
+				{
+					return FLiteral(InValue.ArrayStringParam);
+				}
+
+				default:
+				{
+					static_assert(static_cast<int32>(EAudioParameterType::COUNT) == 12, "Possible missing switch case coverage");
+					checkNoEntry();
+				}
+			}
+
+			return FLiteral();
+		}
+
 		FName ConvertParameterToDataType(EAudioParameterType InParameterType)
 		{
 			switch (InParameterType)
@@ -154,11 +241,12 @@ namespace Metasound
 		return FSendAddress(InVertexName, InTypeName, InInstanceID);
 	}
 
-	FMetaSoundParameterTransmitter::FMetaSoundParameterTransmitter(const FMetaSoundParameterTransmitter::FInitParams& InInitParams)
-	: OperatorSettings(InInitParams.OperatorSettings)
-	, InstanceID(InInitParams.InstanceID)
-	, DebugMetaSoundName(InInitParams.DebugMetaSoundName)
-	, SendInfos(InInitParams.Infos)
+	FMetaSoundParameterTransmitter::FMetaSoundParameterTransmitter(FMetaSoundParameterTransmitter::FInitParams&& InInitParams)
+		: Audio::FParameterTransmitterBase(MoveTemp(InInitParams.DefaultParams))
+		, OperatorSettings(MoveTemp(InInitParams.OperatorSettings))
+		, InstanceID(InInitParams.InstanceID)
+		, DebugMetaSoundName(InInitParams.DebugMetaSoundName)
+		, SendInfos(MoveTemp(InInitParams.Infos))
 	{
 	}
 
@@ -179,22 +267,32 @@ namespace Metasound
 			}
 		}
 
-		return bSuccess;
-	}
+		bSuccess &= Audio::FParameterTransmitterBase::Reset();
 
-	uint64 FMetaSoundParameterTransmitter::GetInstanceID() const
-	{
-		return InstanceID;
+		return bSuccess;
 	}
 
 	bool FMetaSoundParameterTransmitter::SetParameters(TArray<FAudioParameter>&& InParameters)
 	{
 		bool bSuccess = true;
 
-		for (FAudioParameter& InParameter : InParameters)
+		TArray<FAudioParameter> NewParams;
+		for (FAudioParameter& Param : InParameters)
 		{
-			const FName ParamName = InParameter.ParamName;
-			bSuccess &= SetParameterWithLiteral(ParamName, Frontend::ConvertParameterToLiteral(MoveTemp(InParameter)));
+			const FName ParamName = Param.ParamName;
+			if (SetParameterWithLiteral(ParamName, Frontend::ConvertParameterToLiteral(Param)))
+			{
+				NewParams.Add(MoveTemp(Param));
+			}
+			else
+			{
+				bSuccess = false;
+			}
+		}
+
+		if (!NewParams.IsEmpty())
+		{
+			bSuccess &= FParameterTransmitterBase::SetParameters(MoveTemp(NewParams));
 		}
 
 		InParameters.Reset();
@@ -226,11 +324,6 @@ namespace Metasound
 		return false;
 	}
 
-	TUniquePtr<Audio::IParameterTransmitter> FMetaSoundParameterTransmitter::Clone() const
-	{
-		return MakeUnique<FMetaSoundParameterTransmitter>(FMetaSoundParameterTransmitter::FInitParams(OperatorSettings, InstanceID, SendInfos));
-	}
-
 	const FMetaSoundParameterTransmitter::FSendInfo* FMetaSoundParameterTransmitter::FindSendInfo(const FName& InParameterName) const
 	{
 		return SendInfos.FindByPredicate([&](const FSendInfo& Info) { return Info.ParameterName == InParameterName; });
diff --git a/Engine/Plugins/Runtime/Metasound/Source/MetasoundFrontend/Public/MetasoundParameterTransmitter.h b/Engine/Plugins/Runtime/Metasound/Source/MetasoundFrontend/Public/MetasoundParameterTransmitter.h
--- a/Engine/Plugins/Runtime/Metasound/Source/MetasoundFrontend/Public/MetasoundParameterTransmitter.h
+++ b/Engine/Plugins/Runtime/Metasound/Source/MetasoundFrontend/Public/MetasoundParameterTransmitter.h
@@ -19,6 +19,7 @@ namespace Metasound
 	namespace Frontend
 	{
 		METASOUNDFRONTEND_API FLiteral ConvertParameterToLiteral(FAudioParameter&& InValue);
+		METASOUNDFRONTEND_API FLiteral ConvertParameterToLiteral(const FAudioParameter& InValue);
 		METASOUNDFRONTEND_API FName ConvertParameterToDataType(EAudioParameterType InParameterType);
 	}
 
@@ -28,7 +29,7 @@ namespace Metasound
 	 * be safely ushered across thread boundaries in scenarios where the instance
 	 * transmitter and metasound instance live on different threads. 
 	 */
-	class METASOUNDFRONTEND_API FMetaSoundParameterTransmitter : public Audio::IParameterTransmitter
+	class METASOUNDFRONTEND_API FMetaSoundParameterTransmitter : public Audio::FParameterTransmitterBase
 	{
 		FMetaSoundParameterTransmitter(const FMetaSoundParameterTransmitter&) = delete;
 		FMetaSoundParameterTransmitter& operator=(const FMetaSoundParameterTransmitter&) = delete;
@@ -66,10 +67,14 @@ namespace Metasound
 			/** Name of MetaSound used to log parameter related errors. */
 			FName DebugMetaSoundName;
 
-			FInitParams(const FOperatorSettings& InSettings, uint64 InInstanceID, const TArray<FSendInfo>& InInfos=TArray<FSendInfo>())
+			/** Default Audio Parameters set when transmitter is initialized */
+			TArray<FAudioParameter> DefaultParams;
+
+			FInitParams(const FOperatorSettings& InSettings, uint64 InInstanceID, TArray<FAudioParameter>&& InDefaultParams, const TArray<FSendInfo>& InInfos=TArray<FSendInfo>())
 			: OperatorSettings(InSettings)
 			, InstanceID(InInstanceID)
 			, Infos(InInfos)
+			, DefaultParams(MoveTemp(InDefaultParams))
 			{
 			}
 		};
@@ -80,14 +85,11 @@ namespace Metasound
 		/** Creates a unique send address using the given InstanceID. */
 		static FSendAddress CreateSendAddressFromInstanceID(uint64 InInstanceID, const FVertexName& InVertexName, const FName& InTypeName);
 
-		FMetaSoundParameterTransmitter(const FMetaSoundParameterTransmitter::FInitParams& InInitParams);
+		FMetaSoundParameterTransmitter(FMetaSoundParameterTransmitter::FInitParams&& InInitParams);
 		virtual ~FMetaSoundParameterTransmitter() = default;
 
 		bool Reset() override;
 
-		/** Returns ID of the MetaSound instance associated with this transmitter. */
-		uint64 GetInstanceID() const override;
-
 		/** Sets parameters using array of AudioParameter structs
 		 *
 		 * @param InParameter - Parameter to set.
@@ -103,10 +105,6 @@ namespace Metasound
 		 */
 		bool SetParameterWithLiteral(FName InParameterName, const FLiteral& InValue);
 
-		/** Duplicate this transmitter interface. The transmitters association with
-		 * the MetaSound instance will be maintained. */
-		TUniquePtr<Audio::IParameterTransmitter> Clone() const override;
-
 	private:
 		// Find FSendInfo by parameter name. 
 		const FSendInfo* FindSendInfo(const FName& InParameterName) const;
diff --git a/Engine/Source/Runtime/AudioExtensions/Private/IAudioParameterTransmitter.cpp b/Engine/Source/Runtime/AudioExtensions/Private/IAudioParameterTransmitter.cpp
--- a/Engine/Source/Runtime/AudioExtensions/Private/IAudioParameterTransmitter.cpp
+++ b/Engine/Source/Runtime/AudioExtensions/Private/IAudioParameterTransmitter.cpp
@@ -7,13 +7,41 @@ namespace Audio
 {
 	const FName IParameterTransmitter::RouterName = "ParameterTransmitter";
 
-	bool ILegacyParameterTransmitter::GetParameter(FName InName, FAudioParameter& OutValue) const
+	TArray<UObject*> ILegacyParameterTransmitter::GetReferencedObjects() const
+	{
+		return { };
+	}
+
+	FParameterTransmitterBase::FParameterTransmitterBase(TArray<FAudioParameter>&& InDefaultParams)
+		: AudioParameters(MoveTemp(InDefaultParams))
 	{
+	}
+
+	bool FParameterTransmitterBase::GetParameter(FName InName, FAudioParameter& OutValue) const
+	{
+		if (const FAudioParameter* Param = FAudioParameter::FindParam(AudioParameters, InName))
+		{
+			OutValue = *Param;
+			return true;
+		}
+
 		return false;
 	}
 
-	TArray<UObject*> ILegacyParameterTransmitter::GetReferencedObjects() const
+	const TArray<FAudioParameter>& FParameterTransmitterBase::GetParameters() const
 	{
-		return { };
+		return AudioParameters;
+	}
+
+	bool FParameterTransmitterBase::Reset()
+	{
+		AudioParameters.Reset();
+		return true;
+	}
+
+	bool FParameterTransmitterBase::SetParameters(TArray<FAudioParameter>&& InParameters)
+	{
+		FAudioParameter::Merge(MoveTemp(InParameters), AudioParameters);
+		return true;
 	}
 } // namespace Audio
\ No newline at end of file
diff --git a/Engine/Source/Runtime/AudioExtensions/Public/IAudioParameterTransmitter.h b/Engine/Source/Runtime/AudioExtensions/Public/IAudioParameterTransmitter.h
--- a/Engine/Source/Runtime/AudioExtensions/Public/IAudioParameterTransmitter.h
+++ b/Engine/Source/Runtime/AudioExtensions/Public/IAudioParameterTransmitter.h
@@ -26,16 +26,17 @@ namespace Audio
 		TArray<FAudioParameter> DefaultParams;
 	};
 
-	// Parameter getter & reference collector for legacy parameter system.
+	// Reference collector functionality for legacy parameter system
 	// (i.e. backwards compatibility with the SoundCue system). None of this
-	// should be used by future systems (i.e. MetaSounds) as object references
-	// from parameters should NOT be cached on threads other than the GameThread.
+	// should be used by future assets supporting parameters (ex. MetaSounds)
+	// as object pointers within parameters should NOT be cached on threads
+	// other than the GameThread, utilizing a proxy methodology like MetaSounds
+	// that copies UObject data when and where necessary.
 	class AUDIOEXTENSIONS_API ILegacyParameterTransmitter
 	{
 		public:
 			virtual ~ILegacyParameterTransmitter() = default;
 
-			virtual bool GetParameter(FName InName, FAudioParameter& OutParam) const;
 			virtual TArray<UObject*> GetReferencedObjects() const;
 	};
 
@@ -52,13 +53,32 @@ namespace Audio
 
 			virtual bool Reset() = 0;
 
-			// Return the instance ID
-			virtual uint64 GetInstanceID() const = 0;
+			// Return the cached parameter with the given name if it exists
+			// @return False if param not found, true if found.
+			virtual bool GetParameter(FName InName, FAudioParameter& OutParam) const = 0;
+
+			// Return reference to the cached parameter array.
+			virtual const TArray<FAudioParameter>& GetParameters() const = 0;
 
 			// Parameter Setters
 			virtual bool SetParameters(TArray<FAudioParameter>&& InParameters) = 0;
+	};
+
+	/** Base implementation for the parameter transmitter, which caches parameters
+	  * and provides implementer to add additional logic to route parameter data accordingly.
+	  */
+	class AUDIOEXTENSIONS_API FParameterTransmitterBase : public IParameterTransmitter
+	{
+	public:
+		FParameterTransmitterBase(TArray<FAudioParameter>&& InDefaultParams);
+		virtual ~FParameterTransmitterBase() = default;
+
+		virtual bool GetParameter(FName InName, FAudioParameter& OutParam) const override;
+		virtual const TArray<FAudioParameter>& GetParameters() const override;
+		virtual bool Reset() override;
+		virtual bool SetParameters(TArray<FAudioParameter>&& InParameters) override;
 
-			// Create a copy of the instance transmitter.
-			virtual TUniquePtr<Audio::IParameterTransmitter> Clone() const = 0;
+	protected:
+		TArray<FAudioParameter> AudioParameters;
 	};
 } // namespace Audio
diff --git a/Engine/Source/Runtime/AudioMixer/Private/AudioMixerSource.cpp b/Engine/Source/Runtime/AudioMixer/Private/AudioMixerSource.cpp
--- a/Engine/Source/Runtime/AudioMixer/Private/AudioMixerSource.cpp
+++ b/Engine/Source/Runtime/AudioMixer/Private/AudioMixerSource.cpp
@@ -1061,11 +1061,15 @@ namespace Audio
 		uint64 InstanceID = 0;
 		bool bActiveSoundIsPreviewSound = false;
 		TArray<FAudioParameter> DefaultParameters;
-		if (WaveInstance->ActiveSound)
+		FActiveSound* ActiveSound = WaveInstance->ActiveSound;
+		if (ActiveSound)
 		{
-			InstanceID = WaveInstance->ActiveSound->GetAudioComponentID();
-			bActiveSoundIsPreviewSound = WaveInstance->ActiveSound->bIsPreviewSound;
-			DefaultParameters = MoveTemp(WaveInstance->ActiveSound->DefaultParameters);
+			InstanceID = ActiveSound->GetAudioComponentID();
+			bActiveSoundIsPreviewSound = ActiveSound->bIsPreviewSound;
+			if (Audio::IParameterTransmitter* Transmitter = ActiveSound->GetTransmitter())
+			{
+				DefaultParameters = Transmitter->GetParameters();
+			}
 		}
 
 		FMixerSourceBufferInitArgs BufferInitArgs;
@@ -1088,9 +1092,9 @@ namespace Audio
 			// Guarantee that this wave instance does not try to replay by disabling looping.
 			WaveInstance->LoopingMode = LOOP_Never;
 
-			if (ensure(WaveInstance->ActiveSound))
+			if (ensure(ActiveSound))
 			{
-				WaveInstance->ActiveSound->bShouldRemainActiveIfDropped = false;
+				ActiveSound->bShouldRemainActiveIfDropped = false;
 			}
 		}
 		
diff --git a/Engine/Source/Runtime/Engine/Classes/Sound/SoundBase.h b/Engine/Source/Runtime/Engine/Classes/Sound/SoundBase.h
--- a/Engine/Source/Runtime/Engine/Classes/Sound/SoundBase.h
+++ b/Engine/Source/Runtime/Engine/Classes/Sound/SoundBase.h
@@ -300,7 +300,7 @@ public:
 	//~ End IInterface_AssetUserData Interface
 
 	/** Called from the Game Thread prior to attempting to pass parameters to the ParameterTransmitter. */
-	virtual void InitParameters(TArray<FAudioParameter>& InParametersToInit, FName InFeatureName);
+	virtual void InitParameters(TArray<FAudioParameter>& ParametersToInit, FName InFeatureName);
 
 	/** Called from the Game Thread prior to attempting to initialize a sound instance. */
 	virtual void InitResources() { }
@@ -315,7 +315,7 @@ public:
 	virtual ISoundGeneratorPtr CreateSoundGenerator(const FSoundGeneratorInitParams& InParams, TArray<FAudioParameter>&& InDefaultParameters) { return CreateSoundGenerator(InParams); }
 
 	/** Creates a parameter transmitter for communicating with active sound instances. */
-	virtual TUniquePtr<Audio::IParameterTransmitter> CreateParameterTransmitter(Audio::FParameterTransmitterInitParams&& InParams) const;
+	virtual TSharedPtr<Audio::IParameterTransmitter> CreateParameterTransmitter(Audio::FParameterTransmitterInitParams&& InParams) const;
 
 	/** Returns whether parameter is valid input for the given sound */
 	virtual bool IsParameterValid(const FAudioParameter& InParameter) const;
diff --git a/Engine/Source/Runtime/Engine/Classes/Sound/SoundCue.h b/Engine/Source/Runtime/Engine/Classes/Sound/SoundCue.h
--- a/Engine/Source/Runtime/Engine/Classes/Sound/SoundCue.h
+++ b/Engine/Source/Runtime/Engine/Classes/Sound/SoundCue.h
@@ -2,6 +2,7 @@
 #pragma once
 
 #include "CoreMinimal.h"
+#include "IAudioParameterTransmitter.h"
 #include "UObject/ObjectMacros.h"
 #include "Templates/SubclassOf.h"
 #include "Sound/SoundAttenuation.h"
@@ -179,6 +180,7 @@ public:
 	virtual bool GetSoundWavesWithCookedAnalysisData(TArray<USoundWave*>& OutSoundWaves) override;
 	virtual bool HasCookedFFTData() const override;
 	virtual bool HasCookedAmplitudeEnvelopeData() const override;
+	virtual TSharedPtr<Audio::IParameterTransmitter> CreateParameterTransmitter(Audio::FParameterTransmitterInitParams&& InParams) const override;
 	//~ End USoundBase Interface.
 
 	/** Construct and initialize a node within this Cue */
@@ -249,7 +251,6 @@ public:
 		}
 	}
 
-
 	/** Find the path through the sound cue to a node identified by its hash */
 	bool FindPathToNode(const UPTRINT NodeHashToFind, TArray<USoundNode*>& OutPath) const;
 
diff --git a/Engine/Source/Runtime/Engine/Private/AudioDevice.cpp b/Engine/Source/Runtime/Engine/Private/AudioDevice.cpp
--- a/Engine/Source/Runtime/Engine/Private/AudioDevice.cpp
+++ b/Engine/Source/Runtime/Engine/Private/AudioDevice.cpp
@@ -3478,12 +3478,12 @@ void FAudioDevice::SetListenerAttenuationOverride(int32 ListenerIndex, const FVe
 			Listener.bUseAttenuationOverride = true;
 			Listener.AttenuationOverride = AttenuationPosition;
 
-		if (!bPrevAttenuationOverride)
-		{
+			if (!bPrevAttenuationOverride)
+			{
 				UpdateVirtualLoops(true);
 			}
 		}
-		}
+	}
 }
 
 void FAudioDevice::ClearListenerAttenuationOverride(int32 ListenerIndex)
@@ -5097,6 +5097,23 @@ void FAudioDevice::AddNewActiveSoundInternal(const FActiveSound& InNewActiveSoun
 	}
 #endif // !UE_BUILD_SHIPPING
 
+	auto InitSoundParams = [this, &Sound](FActiveSound& OutActiveSound, TArray<FAudioParameter>&& DefaultParams)
+	{
+		// Retriggering a virtualized ActiveSound which already have a transmitter
+		// should not be given a new transmitter.
+		if (!OutActiveSound.InstanceTransmitter.IsValid())
+		{
+			Audio::FParameterTransmitterInitParams TransmitterInitParams
+			{
+				OutActiveSound.GetAudioComponentID(),
+				GetSampleRate(),
+				MoveTemp(DefaultParams)
+			};
+
+			OutActiveSound.InstanceTransmitter = Sound->CreateParameterTransmitter(MoveTemp(TransmitterInitParams));
+		}
+	};
+
 	// Determine if sound is loop and eligible for virtualize prior to creating "live" active sound in next Concurrency check step
 	if (!InVirtualLoopToRetrigger)
 	{
@@ -5105,6 +5122,7 @@ void FAudioDevice::AddNewActiveSoundInternal(const FActiveSound& InNewActiveSoun
 		if (FAudioVirtualLoop::Virtualize(InNewActiveSound, *this, bDoRangeCheck, VirtualLoop))
 		{
 			UE_LOG(LogAudio, Verbose, TEXT("New ActiveSound %s Virtualizing: Failed to pass initial audible range check"), *Sound->GetName());
+			InitSoundParams(VirtualLoop.GetActiveSound(), MoveTemp(InDefaultParams));
 			AddVirtualLoop(VirtualLoop);
 			return;
 		}
@@ -5130,6 +5148,7 @@ void FAudioDevice::AddNewActiveSoundInternal(const FActiveSound& InNewActiveSoun
 			if (FAudioVirtualLoop::Virtualize(InNewActiveSound, *this, bDoRangeCheck, VirtualLoop))
 			{
 				UE_LOG(LogAudioConcurrency, Verbose, TEXT("New ActiveSound %s Virtualizing: Failed to pass concurrency"), *Sound->GetName());
+				InitSoundParams(VirtualLoop.GetActiveSound(), MoveTemp(InDefaultParams));
 				AddVirtualLoop(VirtualLoop);
 			}
 			else
@@ -5195,24 +5214,7 @@ void FAudioDevice::AddNewActiveSoundInternal(const FActiveSound& InNewActiveSoun
 		VirtualActiveSound.ClearAudioComponent();
 	}
 
-	// Retriggering a virtualized ActiveSound which already have a transmitter
-	// should not be given a new transmitter.
-	if (!ActiveSound->InstanceTransmitter.IsValid())
-	{
-		Audio::FParameterTransmitterInitParams TransmitterInitParams
-		{
-			ActiveSound->GetAudioComponentID(),
-			GetSampleRate(),
-			InDefaultParams
-		};
-
-		ActiveSound->InstanceTransmitter = Sound->CreateParameterTransmitter(MoveTemp(TransmitterInitParams));
-	}
-
-	// Set the default parameter used during creation of the active sound.
-	ActiveSound->DefaultParameters = MoveTemp(InDefaultParams);
-
-
+	InitSoundParams(*ActiveSound, MoveTemp(InDefaultParams));
 	ActiveSounds.Add(ActiveSound);
 
 	if (ActiveSound->GetAudioComponentID() > 0)
@@ -5997,6 +5999,7 @@ bool FAudioDevice::SoundIsAudible(const FActiveSound& NewActiveSound)
 
 	return false;
 }
+
 int32 FAudioDevice::FindClosestListenerIndex(const FTransform& SoundTransform, const TArray<FListener>& InListeners)
 {
 	check(IsInAudioThread());
@@ -6073,12 +6076,8 @@ void FAudioDevice::UnlinkActiveSoundFromComponent(const FActiveSound& InActiveSo
 				{
 					if (ActiveSound->GetInstanceID() == InActiveSound.GetInstanceID())
 					{
-						// Reset the transmitter prior to unlinking to avoid associated
-						// state potentially kept in plugins to not become stale.
-						if (Audio::IParameterTransmitter* Transmitter = ActiveSound->GetTransmitter())
-						{
-							Transmitter->Reset();
-						}
+						// Clears the transmitter ensuring additional parameter requests are ignored.
+						ActiveSound->ClearTransmitter();
 						ActiveSoundsInComponent->RemoveAtSwap(i, 1, false);
 						break;
 					}
diff --git a/Engine/Source/Runtime/Engine/Private/SoundBase.cpp b/Engine/Source/Runtime/Engine/Private/SoundBase.cpp
--- a/Engine/Source/Runtime/Engine/Private/SoundBase.cpp
+++ b/Engine/Source/Runtime/Engine/Private/SoundBase.cpp
@@ -288,88 +288,18 @@ const TArray<UAssetUserData*>* USoundBase::GetAssetUserDataArray() const
 	return &ToRawPtrTArrayUnsafe(AssetUserData);
 }
 
-TUniquePtr<Audio::IParameterTransmitter> USoundBase::CreateParameterTransmitter(Audio::FParameterTransmitterInitParams&& InParams) const
+TSharedPtr<Audio::IParameterTransmitter> USoundBase::CreateParameterTransmitter(Audio::FParameterTransmitterInitParams&& InParams) const
 {
-	class FDefaultParameterTransmitter : public Audio::IParameterTransmitter
-	{
-		uint64 InstanceID = INDEX_NONE;
-		TArray<FAudioParameter> AudioParameters;
-
-	public:
-		FDefaultParameterTransmitter(Audio::FParameterTransmitterInitParams&& InParams)
-			: InstanceID(MoveTemp(InParams.InstanceID))
-			, AudioParameters(MoveTemp(InParams.DefaultParams))
-		{
-		}
-
-		virtual ~FDefaultParameterTransmitter() = default;
-
-		bool Reset() override
-		{
-			AudioParameters.Reset();
-			return true;
-		}
-
-		bool SetParameters(TArray<FAudioParameter>&& InParameters) override
-		{
-			FAudioParameter::Merge(MoveTemp(InParameters), AudioParameters);
-			return true;
-		}
-
-		uint64 GetInstanceID() const override
-		{
-			return InstanceID;
-		}
-
-		bool GetParameter(FName InName, FAudioParameter& OutValue) const override
-		{
-			if (const FAudioParameter* Param = FAudioParameter::FindParam(AudioParameters, InName))
-			{
-				OutValue = *Param;
-				return true;
-			}
-
-			return false;
-		}
-
-		TArray<UObject*> GetReferencedObjects() const override
-		{
-			TArray<UObject*> Objects;
-			for (const FAudioParameter& Param : AudioParameters)
-			{
-				if (Param.ObjectParam)
-				{
-					Objects.Add(Param.ObjectParam);
-				}
-
-				for (UObject* Object : Param.ArrayObjectParam)
-				{
-					if (Object)
-					{
-						Objects.Add(Object);
-					}
-				}
-			}
-
-			return Objects;
-		}
-
-		TUniquePtr<Audio::IParameterTransmitter> Clone() const override
-		{
-			return MakeUnique<FDefaultParameterTransmitter>(*this);
-		}
-	};
-
-	return MakeUnique<FDefaultParameterTransmitter>(MoveTemp(InParams));
+	return nullptr;
 }
 
-void USoundBase::InitParameters(TArray<FAudioParameter>& InParametersToInit, FName InFeatureName)
+void USoundBase::InitParameters(TArray<FAudioParameter>& ParametersToInit, FName InFeatureName)
 {
-	for (int32 i = InParametersToInit.Num() - 1; i >= 0; --i)
+	for (int32 i = ParametersToInit.Num() - 1; i >= 0; --i)
 	{
-		if (!IsParameterValid(InParametersToInit[i]))
+		if (!IsParameterValid(ParametersToInit[i]))
 		{
-			InParametersToInit.RemoveAtSwap(i, 1, false /* bAllowShrinking */);
+			ParametersToInit.RemoveAtSwap(i, 1, false /* bAllowShrinking */);
 		}
 	}
 }
diff --git a/Engine/Source/Runtime/Engine/Private/SoundCue.cpp b/Engine/Source/Runtime/Engine/Private/SoundCue.cpp
--- a/Engine/Source/Runtime/Engine/Private/SoundCue.cpp
+++ b/Engine/Source/Runtime/Engine/Private/SoundCue.cpp
@@ -36,6 +36,7 @@
 #include "Interfaces/ITargetPlatform.h"
 #include "AudioCompressionSettings.h"
 #include "Sound/AudioSettings.h"
+#include "Templates/SharedPointer.h"
 
 #include UE_INLINE_GENERATED_CPP_BY_NAME(SoundCue)
 
@@ -711,6 +712,44 @@ bool USoundCue::HasCookedAmplitudeEnvelopeData() const
 	return false;
 }
 
+TSharedPtr<Audio::IParameterTransmitter> USoundCue::CreateParameterTransmitter(Audio::FParameterTransmitterInitParams&& InParams) const
+{
+	class FSoundCueParameterTransmitter : public Audio::FParameterTransmitterBase
+	{
+	public:
+		FSoundCueParameterTransmitter(Audio::FParameterTransmitterInitParams&& InParams)
+			: Audio::FParameterTransmitterBase(MoveTemp(InParams.DefaultParams))
+		{
+		}
+
+		virtual ~FSoundCueParameterTransmitter() = default;
+
+		TArray<UObject*> GetReferencedObjects() const override
+		{
+			TArray<UObject*> Objects;
+			for (const FAudioParameter& Param : AudioParameters)
+			{
+				if (Param.ObjectParam)
+				{
+					Objects.Add(Param.ObjectParam);
+				}
+
+				for (UObject* Object : Param.ArrayObjectParam)
+				{
+					if (Object)
+					{
+						Objects.Add(Object);
+					}
+				}
+			}
+
+			return Objects;
+		}
+	};
+
+	return MakeShared<FSoundCueParameterTransmitter>(MoveTemp(InParams));
+}
+
 #if WITH_EDITOR
 UEdGraph* USoundCue::GetGraph()
 {
diff --git a/Engine/Source/Runtime/Engine/Public/ActiveSound.h b/Engine/Source/Runtime/Engine/Public/ActiveSound.h
--- a/Engine/Source/Runtime/Engine/Public/ActiveSound.h
+++ b/Engine/Source/Runtime/Engine/Public/ActiveSound.h
@@ -33,7 +33,7 @@ struct FAttenuationListenerData;
  */
 struct FAttenuationFocusData
 {
-	/** Azimuth of the active sound relative to the listener. Used by sound focus. */
+	/** Azimuth of the active sound relative to the listener. Used by sound  focus. */
 	float Azimuth = 0.0f;
 
 	/** Absolute azimuth of the active sound relative to the listener. Used for 3d audio calculations. */
@@ -376,71 +376,7 @@ private:
 
 	TMap<UPTRINT, FWaveInstance*> WaveInstances;
 
-	class FParameterTransmitter
-	{
-	public:
-		FParameterTransmitter() = default;
-		FParameterTransmitter(FParameterTransmitter&&) = default;
-		FParameterTransmitter& operator=(FParameterTransmitter&&) = default;
-
-		FParameterTransmitter(TUniquePtr<Audio::IParameterTransmitter>&& InTransmitterImpl)
-		:	TransmitterImpl(MoveTemp(InTransmitterImpl))
-		{
-		}
-
-		FParameterTransmitter(const FParameterTransmitter& InOther)
-		{
-			if (InOther.IsValid())
-			{
-				TransmitterImpl = InOther->Clone();
-			}
-		}
-
-		FParameterTransmitter& operator=(const FParameterTransmitter& InOther)
-		{
-			TransmitterImpl.Reset();
-			if (InOther.IsValid())
-			{
-				TransmitterImpl = InOther->Clone();
-			}
-			return *this;
-		}
-
-		FParameterTransmitter& operator=(TUniquePtr<Audio::IParameterTransmitter>&& InTransmitterImpl)
-		{
-			TransmitterImpl = MoveTemp(InTransmitterImpl);
-			return *this;
-		}
-
-		bool IsValid() const
-		{
-			return TransmitterImpl.IsValid();
-		}
-
-		Audio::IParameterTransmitter* Get()
-		{
-			return TransmitterImpl.Get();
-		}
-		const Audio::IParameterTransmitter* Get() const
-		{
-			return TransmitterImpl.Get();
-		}
-
-		Audio::IParameterTransmitter* operator->()
-		{
-			return TransmitterImpl.Get();
-		}
-
-		const Audio::IParameterTransmitter* operator->() const
-		{
-			return TransmitterImpl.Get();
-		}
-
-	private:
-		TUniquePtr<Audio::IParameterTransmitter> TransmitterImpl;
-	};
-
-	FParameterTransmitter InstanceTransmitter;
+	TSharedPtr<Audio::IParameterTransmitter> InstanceTransmitter;
 
 public:
 	Audio::IParameterTransmitter* GetTransmitter()
@@ -453,6 +389,11 @@ public:
 		return InstanceTransmitter.Get();
 	}
 
+	void ClearTransmitter()
+	{
+		return InstanceTransmitter.Reset();
+	}
+
 	enum class EFadeOut : uint8
 	{
 		// Sound is not currently fading out
@@ -465,9 +406,6 @@ public:
 		Concurrency
 	};
 
-	/** The default initial parameters use when starting a new active sound. */
-	TArray<FAudioParameter> DefaultParameters;
-
 	/** Whether or not the sound has checked if it was occluded already. Used to initialize a sound as occluded and bypassing occlusion interpolation. */
 	uint8 bHasCheckedOcclusion:1;
 
