Fix dodgy move semantics
#rnx

diff --git a/Engine/Plugins/Runtime/Metasound/Source/MetasoundGenerator/Private/MetasoundGenerator.cpp b/Engine/Plugins/Runtime/Metasound/Source/MetasoundGenerator/Private/MetasoundGenerator.cpp
--- a/Engine/Plugins/Runtime/Metasound/Source/MetasoundGenerator/Private/MetasoundGenerator.cpp
+++ b/Engine/Plugins/Runtime/Metasound/Source/MetasoundGenerator/Private/MetasoundGenerator.cpp
@@ -40,6 +40,7 @@ namespace Metasound
 		Environment = {};
 		MetaSoundName = {};
 		AudioOutputNames = {};
+		DefaultParameters = {};
 	}
 
 	FAsyncMetaSoundBuilder::FAsyncMetaSoundBuilder(FMetasoundGenerator* InGenerator, FMetasoundGeneratorInitParams&& InInitParams, bool bInTriggerGenerator)
@@ -59,7 +60,7 @@ namespace Metasound
 		FBuildResults BuildResults;
 
 		// Create an instance of the new graph
-		TUniquePtr<IOperator> GraphOperator = BuildGraphOperator(InitParams.DefaultParameters, BuildResults);
+		TUniquePtr<IOperator> GraphOperator = BuildGraphOperator(MoveTemp(InitParams.DefaultParameters), BuildResults);
 		LogBuildErrors(BuildResults);
 
 		if (GraphOperator.IsValid())
@@ -110,7 +111,7 @@ namespace Metasound
 		InitParams.Release();
 	}
 
-	TUniquePtr<IOperator> FAsyncMetaSoundBuilder::BuildGraphOperator(TArray<FAudioParameter>& InParameters, FBuildResults& OutBuildResults) const
+	TUniquePtr<IOperator> FAsyncMetaSoundBuilder::BuildGraphOperator(TArray<FAudioParameter>&& InParameters, FBuildResults& OutBuildResults) const
 	{
 		using namespace Frontend;
 
@@ -121,23 +122,24 @@ namespace Metasound
 		IDataTypeRegistry& DataRegistry = IDataTypeRegistry::Get();
 		for (FAudioParameter& Parameter : InParameters)
 		{
-			if (const FInputDataVertex* InputVertex = InputInterface.Find(Parameter.ParamName))
+			const FName ParamName = Parameter.ParamName;
+			if (const FInputDataVertex* InputVertex = InputInterface.Find(ParamName))
 			{
 				FLiteral Literal = Frontend::ConvertParameterToLiteral(MoveTemp(Parameter));
 				TOptional<FAnyDataReference> DataReference = DataRegistry.CreateDataReference(InputVertex->DataTypeName, EDataReferenceAccessType::Value, Literal, InitParams.OperatorSettings);
 
 				if (DataReference)
 				{
-					InputData.BindVertex(Parameter.ParamName, *DataReference);
+					InputData.BindVertex(ParamName, *DataReference);
 				}
 				else
 				{
-					UE_LOG(LogMetaSound, Error, TEXT("Failed to create initial input data reference from parameter %s of type %s on graph in MetaSoundSource [%s]"), *Parameter.ParamName.ToString(), *InputVertex->DataTypeName.ToString(), *InitParams.MetaSoundName);
+					UE_LOG(LogMetaSound, Error, TEXT("Failed to create initial input data reference from parameter %s of type %s on graph in MetaSoundSource [%s]"), *ParamName.ToString(), *InputVertex->DataTypeName.ToString(), *InitParams.MetaSoundName);
 				}
 			}
 			else
 			{
-				UE_LOG(LogMetaSound, Error, TEXT("Failed to set initial input parameter %s on graph in MetaSoundSource [%s]"), *Parameter.ParamName.ToString(), *InitParams.MetaSoundName);
+				UE_LOG(LogMetaSound, Error, TEXT("Failed to set initial input parameter %s on graph in MetaSoundSource [%s]"), *ParamName.ToString(), *InitParams.MetaSoundName);
 			}
 
 		}
diff --git a/Engine/Plugins/Runtime/Metasound/Source/MetasoundGenerator/Public/MetasoundGenerator.h b/Engine/Plugins/Runtime/Metasound/Source/MetasoundGenerator/Public/MetasoundGenerator.h
--- a/Engine/Plugins/Runtime/Metasound/Source/MetasoundGenerator/Public/MetasoundGenerator.h
+++ b/Engine/Plugins/Runtime/Metasound/Source/MetasoundGenerator/Public/MetasoundGenerator.h
@@ -62,7 +62,7 @@ namespace Metasound
 		}
 
 	private:
-		TUniquePtr<IOperator> BuildGraphOperator(TArray<FAudioParameter>& InParameters, FBuildResults& OutBuildResults) const;
+		TUniquePtr<IOperator> BuildGraphOperator(TArray<FAudioParameter>&& InParameters, FBuildResults& OutBuildResults) const;
 		TUniquePtr<Frontend::FGraphAnalyzer> BuildGraphAnalyzer(TMap<FGuid, FDataReferenceCollection>&& InInternalDataReferences) const;
 		void LogBuildErrors(const FBuildResults& InBuildResults) const;
 		TArray<FAudioBufferReadRef> FindOutputAudioBuffers(const FVertexInterfaceData& InVertexData) const;
