Fix for not loading UMetaSound assets that were created prior to them being redirected as UMetaSoundPatch assets
#rnx

diff --git a/Engine/Plugins/Runtime/Metasound/Source/MetasoundEditor/Private/MetasoundEditorModule.cpp b/Engine/Plugins/Runtime/Metasound/Source/MetasoundEditor/Private/MetasoundEditorModule.cpp
--- a/Engine/Plugins/Runtime/Metasound/Source/MetasoundEditor/Private/MetasoundEditorModule.cpp
+++ b/Engine/Plugins/Runtime/Metasound/Source/MetasoundEditor/Private/MetasoundEditorModule.cpp
@@ -747,6 +747,9 @@ namespace Metasound
 				MetaSoundClassNames.Add(UMetaSoundPatch::StaticClass()->GetClassPathName());
 				MetaSoundClassNames.Add(UMetaSoundSource::StaticClass()->GetClassPathName());
 
+				// Required to query MetaSound assets (that have been redirected to MetaSoundPatch assets) created before UE Release 5.1
+				MetaSoundClassNames.Add(FTopLevelAssetPath(TEXT("/Script/MetasoundEngine.MetaSound")));
+
 				FAssetTypeActions_MetaSoundPatch::RegisterMenuActions();
 				FAssetTypeActions_MetaSoundSource::RegisterMenuActions();
 
