Metasound DSP fixes around ADSR and triggers
- Updated logic on trigger delay node
- Updated logic for signalling that metasound generator is done
- Added ability to clear out selection of triggers
- Fixed ADSR situation when attack is zero.

diff --git a/Engine/Plugins/Runtime/Metasound/Source/MetasoundFrontend/Private/MetasoundTrigger.cpp b/Engine/Plugins/Runtime/Metasound/Source/MetasoundFrontend/Private/MetasoundTrigger.cpp
--- a/Engine/Plugins/Runtime/Metasound/Source/MetasoundFrontend/Private/MetasoundTrigger.cpp
+++ b/Engine/Plugins/Runtime/Metasound/Source/MetasoundFrontend/Private/MetasoundTrigger.cpp
@@ -113,6 +113,27 @@ namespace Metasound
 		bHasTrigger = false;
 	}
 
+	void FTrigger::RemoveAfter(int32 InStartingFrame)
+	{
+		if (InStartingFrame <= 0)
+		{
+			Reset();
+		}
+		else if (TriggeredFrames.Num() > 0)
+		{
+			// Algo::UpperBound Returns index of first element > InStartingFrame
+			int32 BeginningIndex = Algo::UpperBound(TriggeredFrames, InStartingFrame);
+			int32 NumToRemove = TriggeredFrames.Num() - BeginningIndex;
+
+			if (NumToRemove > 0)
+			{
+				TriggeredFrames.RemoveAt(BeginningIndex, NumToRemove);
+				UpdateLastTriggerIndexInBlock();
+				bHasTrigger = TriggeredFrames.Num() > 0;
+			}
+		}
+	}
+
 	void FTrigger::UpdateLastTriggerIndexInBlock()
 	{
 		LastTriggerIndexInBlock = Algo::LowerBound(TriggeredFrames, NumFramesPerBlock);
diff --git a/Engine/Plugins/Runtime/Metasound/Source/MetasoundFrontend/Public/MetasoundTrigger.h b/Engine/Plugins/Runtime/Metasound/Source/MetasoundFrontend/Public/MetasoundTrigger.h
--- a/Engine/Plugins/Runtime/Metasound/Source/MetasoundFrontend/Public/MetasoundTrigger.h
+++ b/Engine/Plugins/Runtime/Metasound/Source/MetasoundFrontend/Public/MetasoundTrigger.h
@@ -82,6 +82,9 @@ namespace Metasound
 			/** Removes all triggered frames. */
 			void Reset();
 
+			/** Removes all triggers which occur after the frame index. */
+			void RemoveAfter(int32 InFrameIndex);
+
 			/** Executes one block of frames and calls underlying InPreTrigger and InOnTrigger
 			 * functions with frame indices.
 			 *
diff --git a/Engine/Plugins/Runtime/Metasound/Source/MetasoundGenerator/Private/MetasoundGenerator.cpp b/Engine/Plugins/Runtime/Metasound/Source/MetasoundGenerator/Private/MetasoundGenerator.cpp
--- a/Engine/Plugins/Runtime/Metasound/Source/MetasoundGenerator/Private/MetasoundGenerator.cpp
+++ b/Engine/Plugins/Runtime/Metasound/Source/MetasoundGenerator/Private/MetasoundGenerator.cpp
@@ -392,9 +392,16 @@ namespace Metasound
 				GraphAnalyzer->Execute();
 			}
 
+			// Check if generated finished during this execute call
+			if (*OnFinishedTriggerRef)
+			{
+				FinishSample = ((*OnFinishedTriggerRef)[0] * NumChannels);
+			}
+
 			// Interleave audio because ISoundGenerator interface expects interleaved audio.
 			InterleaveGeneratedAudio();
 
+
 			// Add audio generated during graph execution to the output buffer.
 			int32 ThisLoopNumSamplesWritten = FillWithBuffer(InterleavedAudioBuffer, &OutAudio[NumSamplesWritten], NumSamplesRemaining);
 
@@ -412,11 +419,6 @@ namespace Metasound
 
 				FMemory::Memcpy(OverflowBuffer.GetData(), &InterleavedAudioBuffer.GetData()[ThisLoopNumSamplesWritten], OverflowCount * sizeof(float));
 			}
-
-			if (*OnFinishedTriggerRef)
-			{
-				bIsFinishTriggered = true;
-			}
 		}
 
 		return NumSamplesWritten;
@@ -438,16 +440,20 @@ namespace Metasound
 
 		if (InNum > 0)
 		{
-			if (InNum < MaxNumOutputSamples)
-			{
-				FMemory::Memcpy(OutAudio, InBuffer.GetData(), InNum * sizeof(float));
-				return InNum;
-			}
-			else
+			int32 NumSamplesToCopy = FMath::Min(InNum, MaxNumOutputSamples);
+			FMemory::Memcpy(OutAudio, InBuffer.GetData(), NumSamplesToCopy * sizeof(float));
+
+			if (FinishSample != INDEX_NONE)
 			{
-				FMemory::Memcpy(OutAudio, InBuffer.GetData(), MaxNumOutputSamples * sizeof(float));
-				return MaxNumOutputSamples;
+				FinishSample -= NumSamplesToCopy;
+				if (FinishSample <= 0)
+				{
+					bIsFinishTriggered = true;
+					FinishSample = INDEX_NONE;
+				}
 			}
+
+			return NumSamplesToCopy;
 		}
 
 		return 0;
diff --git a/Engine/Plugins/Runtime/Metasound/Source/MetasoundGenerator/Public/MetasoundGenerator.h b/Engine/Plugins/Runtime/Metasound/Source/MetasoundGenerator/Public/MetasoundGenerator.h
--- a/Engine/Plugins/Runtime/Metasound/Source/MetasoundGenerator/Public/MetasoundGenerator.h
+++ b/Engine/Plugins/Runtime/Metasound/Source/MetasoundGenerator/Public/MetasoundGenerator.h
@@ -177,6 +177,7 @@ namespace Metasound
 		bool bIsFinishTriggered;
 		bool bIsFinished;
 
+		int32 FinishSample = INDEX_NONE;
 		int32 NumChannels;
 		int32 NumFramesPerExecute;
 		int32 NumSamplesPerExecute;
diff --git a/Engine/Plugins/Runtime/Metasound/Source/MetasoundStandardNodes/Private/MetasoundADSREnvelopeNode.cpp b/Engine/Plugins/Runtime/Metasound/Source/MetasoundStandardNodes/Private/MetasoundADSREnvelopeNode.cpp
--- a/Engine/Plugins/Runtime/Metasound/Source/MetasoundStandardNodes/Private/MetasoundADSREnvelopeNode.cpp
+++ b/Engine/Plugins/Runtime/Metasound/Source/MetasoundStandardNodes/Private/MetasoundADSREnvelopeNode.cpp
@@ -55,10 +55,10 @@ namespace Metasound
 			// Number of samples for Decay
 			int32 DecaySampleCount = 0;
 
-			// Number of samples for Decay
+			// Number of samples for Release 
 			int32 ReleaseSampleCount = 0;
 			
-			// Sustain leve
+			// Sustain level
 			float SustainLevel = 0.0f;
 
 			// Curve factors for attack/decay/release
@@ -108,10 +108,10 @@ namespace Metasound
 				// We are in attack
 				if (InState.CurrentSampleIndex < InState.AttackSampleCount)
 				{
-					float AttackFraction = (float)InState.CurrentSampleIndex++ / InState.AttackSampleCount;
+					float AttackFraction = (float)++(InState.CurrentSampleIndex) / InState.AttackSampleCount;
 					float EnvValue = FMath::Pow(AttackFraction, InState.AttackCurveFactor);
 					float TargeEnvelopeValue = InState.StartingEnvelopeValue + (1.0f - InState.StartingEnvelopeValue) * EnvValue;
-					InState.EnvEase.SetValue(TargeEnvelopeValue);
+					InState.EnvEase.SetValue(TargeEnvelopeValue, true /* bInit */);
 					InState.CurrentEnvelopeValue = InState.EnvEase.GetNextValue();
 					OutEnvelopeValue = InState.CurrentEnvelopeValue;
 
@@ -215,14 +215,14 @@ namespace Metasound
 					EndAttackFrame = FMath::Min(StartFrame + AttackSamplesLeft, EndFrame);
 					for (int32 i = StartFrame; i < EndAttackFrame; ++i)
 					{
-						float AttackFraction = (float)InState.CurrentSampleIndex++ / InState.AttackSampleCount;
+						float AttackFraction = (float)++(InState.CurrentSampleIndex) / InState.AttackSampleCount;
 						float EnvValue = FMath::Pow(AttackFraction, InState.AttackCurveFactor);
-						float TargeEnvelopeValue = InState.StartingEnvelopeValue + (1.0f - InState.StartingEnvelopeValue) * EnvValue;
-						InState.EnvEase.SetValue(TargeEnvelopeValue);
-						InState.CurrentEnvelopeValue = InState.EnvEase.GetNextValue();
-						OutEnvPtr[i] = InState.CurrentEnvelopeValue;
+						OutEnvPtr[i] = InState.StartingEnvelopeValue + (1.0f - InState.StartingEnvelopeValue) * EnvValue;
 					}
 
+					InState.CurrentEnvelopeValue = OutEnvPtr[EndAttackFrame - 1];
+					InState.EnvEase.SetValue(InState.CurrentEnvelopeValue, true /* bInit */);
+
 					// We have finished our attack phase in this block
 					if (InState.CurrentSampleIndex == InState.AttackSampleCount)
 					{
@@ -292,6 +292,7 @@ namespace Metasound
 							// We're now done, lets taper it off and finish things
 							if (EndReleaseFrame < EndFrame)
 							{
+								InState.EnvEase.SetValue(0.f);
 								for (int32 i = EndReleaseFrame; i < EndFrame; ++i)
 								{
 									InState.CurrentEnvelopeValue = InState.EnvEase.GetNextValue();
@@ -323,6 +324,102 @@ namespace Metasound
 	template<typename EnvelopeClass, typename ValueType>
 	class TADSREnvelopeNodeOperator : public TExecutableOperator<TADSREnvelopeNodeOperator<EnvelopeClass, ValueType>>
 	{
+		enum class ETriggerType : uint8
+		{
+			Attack,
+			Release,
+			None
+		};
+
+		struct FTriggerInfo
+		{
+			ETriggerType Type;
+			int32 FrameIndex;
+		};
+
+		struct FTriggerIterator
+		{
+			FTriggerIterator(const FTrigger& InAttackTrigger, const FTrigger& InReleaseTrigger, int32 InNumFramesPerBlock)
+			: AttackTrigger(InAttackTrigger)
+			, ReleaseTrigger(InReleaseTrigger)
+			, NumFramesPerBlock(InNumFramesPerBlock)
+			{
+				if (AttackTrigger.NumTriggeredInBlock() > 0)
+				{
+					NextAttackFrame = AttackTrigger[0];
+					AttackTriggerIndex++;
+				}
+				else
+				{
+					NextAttackFrame = NumFramesPerBlock;
+				}
+
+				if (ReleaseTrigger.NumTriggeredInBlock() > 0)
+				{
+					NextReleaseFrame = ReleaseTrigger[0];
+					ReleaseTriggerIndex++;
+				}
+				else
+				{
+					NextReleaseFrame = NumFramesPerBlock;
+				}
+			}
+
+			FTriggerInfo NextTrigger()
+			{
+				FTriggerInfo Info;
+
+				if ((NextAttackFrame <= NextReleaseFrame) && (NextAttackFrame < NumFramesPerBlock))
+				{
+					Info.Type = ETriggerType::Attack;
+					Info.FrameIndex = NextAttackFrame;
+
+					if (AttackTriggerIndex < AttackTrigger.NumTriggeredInBlock())
+					{
+						NextAttackFrame = AttackTrigger[AttackTriggerIndex];
+						AttackTriggerIndex++;
+					}
+					else
+					{
+						NextAttackFrame = NumFramesPerBlock;
+					}
+				}
+				else if (NextReleaseFrame < NumFramesPerBlock)
+				{
+					Info.Type = ETriggerType::Release;
+					Info.FrameIndex = NextReleaseFrame;
+
+					if (ReleaseTriggerIndex < ReleaseTrigger.NumTriggeredInBlock())
+					{
+						NextReleaseFrame = ReleaseTrigger[ReleaseTriggerIndex];
+						ReleaseTriggerIndex++;
+					}
+					else
+					{
+						NextReleaseFrame = NumFramesPerBlock;
+					}
+				}
+				else
+				{
+					Info.Type = ETriggerType::None;
+					Info.FrameIndex = NumFramesPerBlock;
+				}
+
+				return Info;
+			}
+
+		private:
+			const FTrigger& AttackTrigger;
+			const FTrigger& ReleaseTrigger;
+
+			int32 NumFramesPerBlock = 0;
+			int32 AttackTriggerIndex = 0;
+			int32 NextAttackFrame = 0;
+
+			int32 ReleaseTriggerIndex = 0;
+			int32 NextReleaseFrame = 0;
+		};
+
 	public:
 		static const FVertexInterface& GetDefaultInterface()
 		{
@@ -432,10 +529,8 @@ namespace Metasound
 			, AttackCurveFactor(InArgs.AttackCurveFactor)
 			, DecayCurveFactor(InArgs.DecayCurveFactor)
 			, ReleaseCurveFactor(InArgs.ReleaseCurveFactor)
-			, OnAttackTrigger(TDataWriteReferenceFactory<FTrigger>::CreateAny(InArgs.OperatorSettings))
 			, OnDecayTrigger(TDataWriteReferenceFactory<FTrigger>::CreateAny(InArgs.OperatorSettings))
 			, OnSustainTrigger(TDataWriteReferenceFactory<FTrigger>::CreateAny(InArgs.OperatorSettings))
-			, OnReleaseTrigger(TDataWriteReferenceFactory<FTrigger>::CreateAny(InArgs.OperatorSettings))
 			, OnDone(TDataWriteReferenceFactory<FTrigger>::CreateAny(InArgs.OperatorSettings))
 			, OutputEnvelope(TDataWriteReferenceFactory<ValueType>::CreateAny(InArgs.OperatorSettings))
 		{
@@ -470,10 +565,10 @@ namespace Metasound
 			using namespace ADSREnvelopeVertexNames;
 
 			FDataReferenceCollection Outputs;
-			Outputs.AddDataReadReference(METASOUND_GET_PARAM_NAME(OutputOnAttackTrigger), OnAttackTrigger);
+			Outputs.AddDataReadReference(METASOUND_GET_PARAM_NAME(OutputOnAttackTrigger), TriggerAttackIn);
 			Outputs.AddDataReadReference(METASOUND_GET_PARAM_NAME(OutputOnDecayTrigger), OnDecayTrigger);
 			Outputs.AddDataReadReference(METASOUND_GET_PARAM_NAME(OutputOnSustainTrigger), OnSustainTrigger);
-			Outputs.AddDataReadReference(METASOUND_GET_PARAM_NAME(OutputOnReleaseTrigger), OnReleaseTrigger);
+			Outputs.AddDataReadReference(METASOUND_GET_PARAM_NAME(OutputOnReleaseTrigger), TriggerReleaseIn);
 			Outputs.AddDataReadReference(METASOUND_GET_PARAM_NAME(OutputOnDone), OnDone);
 			Outputs.AddDataReadReference(METASOUND_GET_PARAM_NAME(OutputEnvelopeValue), OutputEnvelope);
 
@@ -485,7 +580,7 @@ namespace Metasound
 			float AttackTimeSeconds = AttackTime->GetSeconds();
 			float DecayTimeSeconds = DecayTime->GetSeconds();
 			float ReleaseTimeSeconds = ReleaseTime->GetSeconds();
-			EnvState.AttackSampleCount = SampleRate * FMath::Max(0.0f, AttackTimeSeconds);
+			EnvState.AttackSampleCount = FMath::Max(1, SampleRate * FMath::Max(0.0f, AttackTimeSeconds));
 			EnvState.DecaySampleCount = SampleRate * FMath::Max(0.0f, DecayTimeSeconds);
 			EnvState.SustainLevel = FMath::Max(0.0f, *SustainLevel);
 			EnvState.ReleaseSampleCount = SampleRate * FMath::Max(0.0f, ReleaseTimeSeconds);
@@ -521,51 +616,52 @@ namespace Metasound
 		{
 			using namespace ADSREnvelopeNodePrivate;
 
-			OnAttackTrigger->AdvanceBlock();
 			OnDecayTrigger->AdvanceBlock();
 			OnSustainTrigger->AdvanceBlock();
-			OnReleaseTrigger->AdvanceBlock();
 			OnDone->AdvanceBlock();
  
-			// check for any updates to input params
-			UpdateParams();
- 		
-			TriggerReleaseIn->ExecuteBlock(
-				[&](int32 StartFrame, int32 EndFrame)
-				{
-				},
-				// OnTrigger
-				[&](int32 StartFrame, int32 EndFrame)
-				{
-					// This schedules a release sample count
-					EnvState.bIsInRelease = true;
-					OnReleaseTrigger->TriggerFrame(StartFrame);
-				}
-			);
+			
 
-			TriggerAttackIn->ExecuteBlock(
-				// OnPreTrigger
-				[&](int32 StartFrame, int32 EndFrame)
-				{
-					ProcessEnvelopeOutput(StartFrame, EndFrame);
-				},
-				// OnTrigger
-				[&](int32 StartFrame, int32 EndFrame)
+			FTriggerIterator TriggerIter(*TriggerAttackIn, *TriggerReleaseIn, NumFramesPerBlock);
+
+			FTriggerInfo NextTrigger = TriggerIter.NextTrigger();
+			if (NextTrigger.FrameIndex > 0)
+			{
+				// Process envelope before receiving any triggers on this block
+				ProcessEnvelopeOutput(0, NextTrigger.FrameIndex);
+			}
+
+			while (NextTrigger.Type != ETriggerType::None)
+			{
+				FTriggerInfo CurrentTrigger = NextTrigger;
+				NextTrigger = TriggerIter.NextTrigger();
+
+				switch (CurrentTrigger.Type)
 				{
-					// Get latest params
+					case ETriggerType::Attack:
+					// check for any updates to input params
 					UpdateParams();
-
 					// Set the sample index to the top of the envelope
 					EnvState.CurrentSampleIndex = 0;
 					EnvState.StartingEnvelopeValue = EnvState.CurrentEnvelopeValue;
 					EnvState.bIsInRelease = false;
 
-					ProcessEnvelopeOutput(StartFrame, EndFrame);
+						break;
+
+					case ETriggerType::Release:
+
+						EnvState.bIsInRelease = true;
 
-					// Forward the trigger
-					OnAttackTrigger->TriggerFrame(StartFrame);
+						break;
+
+					default:
+						{
+							checkNoEntry();
+						}
 				}
-			);
+
+				ProcessEnvelopeOutput(CurrentTrigger.FrameIndex, NextTrigger.FrameIndex);
+			}
 		}
 
 	private:
@@ -579,10 +675,8 @@ namespace Metasound
 		FFloatReadRef DecayCurveFactor;
 		FFloatReadRef ReleaseCurveFactor;
 
-		FTriggerWriteRef OnAttackTrigger;
 		FTriggerWriteRef OnDecayTrigger;
 		FTriggerWriteRef OnSustainTrigger;
-		FTriggerWriteRef OnReleaseTrigger;
 		FTriggerWriteRef OnDone;
 		TDataWriteReference<ValueType> OutputEnvelope;
 
diff --git a/Engine/Plugins/Runtime/Metasound/Source/MetasoundStandardNodes/Private/MetasoundTriggerDelayNode.cpp b/Engine/Plugins/Runtime/Metasound/Source/MetasoundStandardNodes/Private/MetasoundTriggerDelayNode.cpp
--- a/Engine/Plugins/Runtime/Metasound/Source/MetasoundStandardNodes/Private/MetasoundTriggerDelayNode.cpp
+++ b/Engine/Plugins/Runtime/Metasound/Source/MetasoundStandardNodes/Private/MetasoundTriggerDelayNode.cpp
@@ -10,7 +10,6 @@
 #include "MetasoundNodeRegistrationMacro.h"
 #include "MetasoundOperatorInterface.h"
 #include "MetasoundPrimitives.h"
-#include "MetasoundSampleCounter.h"
 #include "MetasoundStandardNodesNames.h"
 #include "MetasoundStandardNodesCategories.h"
 #include "MetasoundParamHelper.h"
@@ -42,26 +41,20 @@ namespace Metasound
 			void Execute();
 
 		private:
-			FSampleCounter NextTriggerCounter;
-			FSampleCount LastFrameReset;
-
+			float SampleRate;
 			FTriggerReadRef TriggerIn;
 			FTriggerReadRef TriggerReset;
 			FTriggerWriteRef TriggerOut;
 
 			FTimeReadRef DelayTime;
-
-			FSampleCount FramesPerBlock;
 	};
 
 	FTriggerDelayOperator::FTriggerDelayOperator(const FOperatorSettings& InSettings, const FTriggerReadRef& InTriggerReset, const FTriggerReadRef& InTriggerIn, const FTimeReadRef& InTimeDelay)
-	: NextTriggerCounter(-1, InSettings.GetSampleRate())
-	, LastFrameReset(-1)
+	: SampleRate(InSettings.GetSampleRate())
 	, TriggerIn(InTriggerIn)
 	, TriggerReset(InTriggerReset)
 	, TriggerOut(FTriggerWriteRef::CreateNew(InSettings))
 	, DelayTime(InTimeDelay)
-	, FramesPerBlock(InSettings.GetNumFramesPerBlock())
 	{
 	}
 
@@ -91,49 +84,26 @@ namespace Metasound
 		// Advance internal counter to get rid of old triggers.
 		TriggerOut->AdvanceBlock();
 
-		LastFrameReset = -1;
-
-		TriggerReset->ExecuteBlock(
+		TriggerIn->ExecuteBlock(
 			[&](int32 StartFrame, int32 EndFrame)
 			{
 			},
 			[this](int32 StartFrame, int32 EndFrame)
 			{
-				LastFrameReset = StartFrame;
+				const int32 FrameToTrigger = FMath::Max(0, FMath::RoundToInt(DelayTime->GetSeconds() * SampleRate)) + StartFrame;
+				TriggerOut->TriggerFrame(FrameToTrigger);
 			}
 		);
 
-		TriggerIn->ExecuteBlock(
+		TriggerReset->ExecuteBlock(
 			[&](int32 StartFrame, int32 EndFrame)
 			{
 			},
 			[this](int32 StartFrame, int32 EndFrame)
 			{
-				if (StartFrame > LastFrameReset)
-				{
-					NextTriggerCounter.SetNumSamples(*DelayTime);
-					LastFrameReset = -1;
-				}
-				else
-				{
-					NextTriggerCounter.SetNumSamples(-1);
-				}
+				TriggerOut->RemoveAfter(StartFrame);
 			}
 		);
-
-		if (NextTriggerCounter.GetNumSamples() >= 0)
-		{
-			FSampleCount SamplesRemaining = NextTriggerCounter.GetNumSamples() - FramesPerBlock;
-			if (SamplesRemaining > 0.0f)
-			{
-				NextTriggerCounter -= FramesPerBlock;
-			}
-			else
-			{
-				TriggerOut->TriggerFrame(SamplesRemaining + (int32)FramesPerBlock);
-				NextTriggerCounter.SetNumSamples(-1);
-			}
-		}
 	}
 
 	TUniquePtr<IOperator> FTriggerDelayOperator::CreateOperator(const FCreateOperatorParams& InParams, FBuildErrorArray& OutErrors)
@@ -141,6 +111,7 @@ namespace Metasound
 		using namespace TriggerDelayVertexNames;
 
 		const FInputVertexInterface& InputInterface = GetVertexInterface().GetInputInterface();
+
 		FTimeReadRef Delay = InParams.InputDataReferences.GetDataReadReferenceOrConstructWithVertexDefault<FTime>(InputInterface, METASOUND_GET_PARAM_NAME(InputDelayTime), InParams.OperatorSettings);
 
 		FTriggerReadRef TriggerIn = InParams.InputDataReferences.GetDataReadReferenceOrConstruct<FTrigger>(METASOUND_GET_PARAM_NAME(InputInTriggerDelay), InParams.OperatorSettings);
