Fix naming of metasound array node input

diff --git a/Engine/Plugins/Runtime/Metasound/Source/MetasoundFrontend/Public/MetasoundArrayNodes.h b/Engine/Plugins/Runtime/Metasound/Source/MetasoundFrontend/Public/MetasoundArrayNodes.h
--- a/Engine/Plugins/Runtime/Metasound/Source/MetasoundFrontend/Public/MetasoundArrayNodes.h
+++ b/Engine/Plugins/Runtime/Metasound/Source/MetasoundFrontend/Public/MetasoundArrayNodes.h
@@ -46,6 +46,15 @@ namespace Metasound
 
 	namespace ArrayNodeVertexNames
 	{
+		static const TCHAR* InputInitialArrayName = TEXT("Array");
+#if WITH_EDITOR
+		static const FText InputInitialArrayTooltip = LOCTEXT("InitialArrayTooltip", "Initial Array");
+		static const FText InputInitialArrayDisplayName = LOCTEXT("InitialArrayDisplayName", "Init Array"); 
+#else
+		static const FText InputInitialArrayTooltip = FText::GetEmpty();
+		static const FText InputInitialArrayDisplayName = FText::GetEmpty();
+#endif
+
 		METASOUND_PARAM(InputArray, "Array", "Input Array.")
 		METASOUND_PARAM(InputLeftArray, "Left Array", "Input Left Array.")
 		METASOUND_PARAM(InputRightArray, "Right Array", "Input Right Array.")
@@ -367,7 +376,7 @@ namespace Metasound
 			static const FVertexInterface DefaultInterface(
 				FInputVertexInterface(
 					TInputDataVertex<FTrigger>(METASOUND_GET_PARAM_NAME_AND_METADATA(InputTriggerSet)),
-					TInputDataVertex<ArrayType>(METASOUND_GET_PARAM_NAME_AND_METADATA(InputArray)),
+					TInputDataVertex<ArrayType>(InputInitialArrayName, FDataVertexMetadata { InputInitialArrayTooltip, InputInitialArrayDisplayName }),
 					TInputDataVertex<int32>(METASOUND_GET_PARAM_NAME_AND_METADATA(InputIndex)),
 					TInputDataVertex<ElementType>(METASOUND_GET_PARAM_NAME_AND_METADATA(InputValue))
 				),
@@ -418,7 +427,7 @@ namespace Metasound
 			
 			TDataReadReference<FTrigger> Trigger = InParams.InputDataReferences.GetDataReadReferenceOrConstructWithVertexDefault<FTrigger>(Inputs, METASOUND_GET_PARAM_NAME(InputTriggerSet), InParams.OperatorSettings);
 
-			FArrayDataReadReference InitArray = InParams.InputDataReferences.GetDataReadReferenceOrConstructWithVertexDefault<ArrayType>(Inputs, METASOUND_GET_PARAM_NAME(InputArray), InParams.OperatorSettings);
+			FArrayDataReadReference InitArray = InParams.InputDataReferences.GetDataReadReferenceOrConstructWithVertexDefault<ArrayType>(Inputs, InputInitialArrayName, InParams.OperatorSettings);
 			FArrayDataWriteReference Array = TDataWriteReferenceFactory<ArrayType>::CreateExplicitArgs(InParams.OperatorSettings, *InitArray);
 
 			TDataReadReference<int32> Index = InParams.InputDataReferences.GetDataReadReferenceOrConstructWithVertexDefault<int32>(Inputs, METASOUND_GET_PARAM_NAME(InputIndex), InParams.OperatorSettings);
