MetaSound input widget fixes for range UX issues and undo/redo
- fix undo/redo when changing value on widget
- when changing widget type for non volume widgets, only change range if necessary
- clamp range X/Y in a more intuitive way for changing the min first, then the max
- disable graph refresh on changing default from details panel

diff --git a/Engine/Plugins/Runtime/Metasound/Source/MetasoundEditor/Private/MetasoundEditorGraphMemberDefaults.cpp b/Engine/Plugins/Runtime/Metasound/Source/MetasoundEditor/Private/MetasoundEditorGraphMemberDefaults.cpp
--- a/Engine/Plugins/Runtime/Metasound/Source/MetasoundEditor/Private/MetasoundEditorGraphMemberDefaults.cpp
+++ b/Engine/Plugins/Runtime/Metasound/Source/MetasoundEditor/Private/MetasoundEditorGraphMemberDefaults.cpp
@@ -180,17 +180,10 @@ void UMetasoundEditorGraphMemberDefaultFloat::PostEditChangeChainProperty(FPrope
 	{
 		SetDefault(Default);
 	}
-	else if (PropertyChangedEvent.GetPropertyName().IsEqual(GET_MEMBER_NAME_CHECKED(UMetasoundEditorGraphMemberDefaultFloat, WidgetType)) || PropertyChangedEvent.GetPropertyName().IsEqual(GET_MEMBER_NAME_CHECKED(UMetasoundEditorGraphMemberDefaultFloat, WidgetValueType)))
+	else if (PropertyChangedEvent.GetPropertyName().IsEqual(GET_MEMBER_NAME_CHECKED(UMetasoundEditorGraphMemberDefaultFloat, WidgetType)) ||
+		PropertyChangedEvent.GetPropertyName().IsEqual(GET_MEMBER_NAME_CHECKED(UMetasoundEditorGraphMemberDefaultFloat, WidgetValueType)))
 	{
-		if (WidgetValueType == EMetasoundMemberDefaultWidgetValueType::Linear)
-		{
-			SetRange(FVector2D(0.0f, 1.0f));
-		}
-		else if (WidgetValueType == EMetasoundMemberDefaultWidgetValueType::Frequency)
-		{
-			SetRange(FVector2D(MIN_FILTER_FREQUENCY, MAX_FILTER_FREQUENCY));
-		}
-		else if (WidgetValueType == EMetasoundMemberDefaultWidgetValueType::Volume)
+		if (WidgetValueType == EMetasoundMemberDefaultWidgetValueType::Volume)
 		{
 			SetRange(FVector2D(-100.0f, 0.0f));
 		}
@@ -202,21 +195,25 @@ void UMetasoundEditorGraphMemberDefaultFloat::PostEditChangeChainProperty(FPrope
 	{
 		if (PropertyChangedEvent.ChangeType != EPropertyChangeType::Interactive)
 		{
-			// if Range.Y < Range.X, set Range.X to Range.Y
-			Range.X = FMath::Min(Range.X, Range.Y);
+			// if Range.X > Range.Y, set Range.Y to Range.X
+			Range.Y = FMath::Max(Range.X, Range.Y);
 			ForceRefresh();
 		}
 	}
 	else if (PropertyChangedEvent.GetPropertyName().IsEqual(GET_MEMBER_NAME_CHECKED(UMetasoundEditorGraphMemberDefaultFloat, ClampDefault)))
 	{
-		// set range to reasonable limit given current value
-		if (FMath::IsNearlyEqual(Default, 0.0f))
+		// If value is not within current range, keep it, otherwise set it to a reasonable range 
+		if (!(Default >= Range.X && Default <= Range.Y))
 		{
-			SetRange(FVector2D(0.0f, 1.0f));
-		}
-		else
-		{
-			SetRange(FVector2D(FMath::Min(0.0f, Default), FMath::Max(0.0f, Default)));
+			// set range to reasonable limit given current value
+			if (FMath::IsNearlyEqual(Default, 0.0f))
+			{
+				SetRange(FVector2D(0.0f, 1.0f));
+			}
+			else
+			{
+				SetRange(FVector2D(FMath::Min(0.0f, Default), FMath::Max(0.0f, Default)));
+			}
 		}
 		OnClampChanged.Broadcast(ClampDefault);
 	}
@@ -224,19 +221,18 @@ void UMetasoundEditorGraphMemberDefaultFloat::PostEditChangeChainProperty(FPrope
 	UMetasoundEditorGraphMember* Member = GetParentMember();
 	if (ensure(Member))
 	{
-		FMetasoundFrontendDocumentModifyContext& ModifyContext = Member->GetOwningGraph()->GetModifyContext();
-
 		// Only update member on non-interactive changes to avoid refreshing the details panel mid-update
 		if (PropertyChangedEvent.ChangeType != EPropertyChangeType::Interactive)
 		{
+			FMetasoundFrontendDocumentModifyContext& ModifyContext = Member->GetOwningGraph()->GetModifyContext();
 			ModifyContext.AddMemberIDsModified({ Member->GetMemberID() });
-		}
 
-		// Mark all nodes as modified to refresh them on synchronization.  This ensures all corresponding widgets get updated.
-		const TArray<UMetasoundEditorGraphMemberNode*> MemberNodes = Member->GetNodes();
-		TSet<FGuid> NodesToRefresh;
-		Algo::Transform(MemberNodes, NodesToRefresh, [](const UMetasoundEditorGraphMemberNode* MemberNode) { return MemberNode->GetNodeID(); });
-		ModifyContext.AddNodeIDsModified({ NodesToRefresh });
+			// Mark all nodes as modified to refresh them on synchronization.  This ensures all corresponding widgets get updated.
+			const TArray<UMetasoundEditorGraphMemberNode*> MemberNodes = Member->GetNodes();
+			TSet<FGuid> NodesToRefresh;
+			Algo::Transform(MemberNodes, NodesToRefresh, [](const UMetasoundEditorGraphMemberNode* MemberNode) { return MemberNode->GetNodeID(); });
+			ModifyContext.AddNodeIDsModified({ NodesToRefresh });
+		}
 	}
 
 	Super::PostEditChangeChainProperty(PropertyChangedEvent);
diff --git a/Engine/Plugins/Runtime/Metasound/Source/MetasoundEditor/Private/SMetasoundGraphNode.cpp b/Engine/Plugins/Runtime/Metasound/Source/MetasoundEditor/Private/SMetasoundGraphNode.cpp
--- a/Engine/Plugins/Runtime/Metasound/Source/MetasoundEditor/Private/SMetasoundGraphNode.cpp
+++ b/Engine/Plugins/Runtime/Metasound/Source/MetasoundEditor/Private/SMetasoundGraphNode.cpp
@@ -667,7 +667,15 @@ namespace Metasound
 						{
 							if (InputWidget.IsValid())
 							{
-								constexpr bool bPostTransaction = false;
+								if (!this->bIsInputWidgetTransacting)
+								{
+									GEditor->BeginTransaction(LOCTEXT("MetasoundSetInputDefault", "Set MetaSound Input Default"));
+									this->bIsInputWidgetTransacting = true;
+								}
+								GraphMember->GetOwningGraph()->GetMetasound()->Modify();
+								DefaultFloat->Modify();
+
+								constexpr bool bPostTransaction = true;
 								float Output = InputWidget->GetOutputValue(Value);
 								DefaultFloat->SetDefault(Output);
 								GraphMember->UpdateFrontendDefaultLiteral(bPostTransaction);
@@ -678,9 +686,21 @@ namespace Metasound
 						{
 							if (InputWidget.IsValid())
 							{
-								constexpr bool bPostTransaction = true;
+								bool bPostTransaction = false;
 								float Output = InputWidget->GetOutputValue(Value);
 								DefaultFloat->SetDefault(Output);
+
+								if (bIsInputWidgetTransacting)
+								{
+									GEditor->EndTransaction();
+									bIsInputWidgetTransacting = false;
+								}
+								else
+								{
+									bPostTransaction = true;
+									UE_LOG(LogMetaSound, Warning, TEXT("Unmatched MetaSound editor widget transaction."));
+								}
+
 								GraphMember->UpdateFrontendDefaultLiteral(bPostTransaction);
 
 								if (UMetasoundEditorGraph* Graph = GraphMember->GetOwningGraph())
@@ -754,6 +774,7 @@ namespace Metasound
 									+ SHorizontalBox::Slot()
 									.HAlign(HAlign_Fill)
 									.VAlign(VAlign_Center)
+									.Padding(WidgetPadding, 0.0f, WidgetPadding, 0.0f)
 									.AutoWidth()
 									[
 										Slot1.ToSharedRef()
@@ -761,7 +782,6 @@ namespace Metasound
 								+ SHorizontalBox::Slot()
 									.HAlign(HAlign_Center)
 									.VAlign(VAlign_Fill)
-									.Padding(WidgetPadding, 0.0f, 0.0f, WidgetPadding)
 									.AutoWidth()
 									[
 										Slot2.ToSharedRef()
diff --git a/Engine/Plugins/Runtime/Metasound/Source/MetasoundEditor/Private/SMetasoundGraphNode.h b/Engine/Plugins/Runtime/Metasound/Source/MetasoundEditor/Private/SMetasoundGraphNode.h
--- a/Engine/Plugins/Runtime/Metasound/Source/MetasoundEditor/Private/SMetasoundGraphNode.h
+++ b/Engine/Plugins/Runtime/Metasound/Source/MetasoundEditor/Private/SMetasoundGraphNode.h
@@ -82,6 +82,10 @@ namespace Metasound
 			FDelegateHandle InputSliderOnValueChangedDelegateHandle;
 			// Handle for on input slider range changed  
 			FDelegateHandle InputSliderOnRangeChangedDelegateHandle;
+
+			// Whether the input widget is currently transacting 
+			// for keeping track of transaction state across delegates to only commit transaction on value commit
+			bool bIsInputWidgetTransacting = false;
 		};
 
 		class SMetaSoundGraphNodeKnot : public SGraphNodeKnot
