Initialize MetaSound reroute node id on copy paste

#rnx

diff --git a/Engine/Plugins/Runtime/Metasound/Source/MetasoundEditor/Private/MetasoundEditor.cpp b/Engine/Plugins/Runtime/Metasound/Source/MetasoundEditor/Private/MetasoundEditor.cpp
--- a/Engine/Plugins/Runtime/Metasound/Source/MetasoundEditor/Private/MetasoundEditor.cpp
+++ b/Engine/Plugins/Runtime/Metasound/Source/MetasoundEditor/Private/MetasoundEditor.cpp
@@ -2350,6 +2350,7 @@ namespace Metasound
 									FMetasoundFrontendNodeInterface NodeInterface = FRerouteNodeTemplate::CreateNodeInterfaceFromDataType(DataType);
 
 									FNodeHandle NodeHandle = MetasoundAsset->GetRootGraphHandle()->AddTemplateNode(RerouteTemplateKey, MoveTemp(NodeInterface));
+									FGraphBuilder::InitGraphNode(NodeHandle, ExternalNode, *Metasound);
 									bIsValid = NodeHandle->IsValid();
 								}
 							}
diff --git a/Engine/Plugins/Runtime/Metasound/Source/MetasoundEditor/Private/MetasoundEditorGraph.cpp b/Engine/Plugins/Runtime/Metasound/Source/MetasoundEditor/Private/MetasoundEditorGraph.cpp
--- a/Engine/Plugins/Runtime/Metasound/Source/MetasoundEditor/Private/MetasoundEditorGraph.cpp
+++ b/Engine/Plugins/Runtime/Metasound/Source/MetasoundEditor/Private/MetasoundEditorGraph.cpp
@@ -1894,12 +1894,26 @@ void UMetasoundEditorGraph::ValidateInternal(Metasound::Editor::FGraphValidation
 	using namespace Metasound::Frontend;
 
 	OutResults = FGraphValidationResults();
-
+	TSet<FGuid> NodeGuids; 
 	TArray<UMetasoundEditorGraphNode*> NodesToValidate;
 	GetNodesOfClass<UMetasoundEditorGraphNode>(NodesToValidate);
+	bool bNodeIdFound = false;
 	for (UMetasoundEditorGraphNode* Node : NodesToValidate)
 	{
 		FGraphNodeValidationResult NodeResult(*Node);
+
+		// Validate there is only 1 editor node per guid 
+		// Input nodes are currently not 1:1 with their frontend representation
+		// but when they are, they can be validated here as well 
+		if (!Node->IsA<UMetasoundEditorGraphInputNode>())
+		{
+			NodeGuids.Add(Node->GetNodeID(), &bNodeIdFound);
+			if (bNodeIdFound)
+			{
+				NodeResult.SetMessage(EMessageSeverity::Warning, TEXT("The internal node this represents is referenced multiple times and may have unintended behavior. Please delete and readd this node."));
+			}
+		}
+
 		Node->Validate(NodeResult);
 
 		OutResults.NodeResults.Add(MoveTemp(NodeResult));
diff --git a/Engine/Plugins/Runtime/Metasound/Source/MetasoundEditor/Private/MetasoundEditorGraphNode.cpp b/Engine/Plugins/Runtime/Metasound/Source/MetasoundEditor/Private/MetasoundEditorGraphNode.cpp
--- a/Engine/Plugins/Runtime/Metasound/Source/MetasoundEditor/Private/MetasoundEditorGraphNode.cpp
+++ b/Engine/Plugins/Runtime/Metasound/Source/MetasoundEditor/Private/MetasoundEditorGraphNode.cpp
@@ -211,8 +211,6 @@ void UMetasoundEditorGraphNode::Validate(Metasound::Editor::FGraphNodeValidation
 	using namespace Metasound::Editor;
 	using namespace Metasound::Frontend;
 
-	OutResult = FGraphNodeValidationResult(*this);
-
 #if WITH_EDITOR
 	// Validate that non-reroute inputs are connected to "real" outputs
 	if (GetClassName() != FRerouteNodeTemplate::ClassName)
diff --git a/Engine/Plugins/Runtime/Metasound/Source/MetasoundEditor/Private/SMetasoundGraphNode.cpp b/Engine/Plugins/Runtime/Metasound/Source/MetasoundEditor/Private/SMetasoundGraphNode.cpp
--- a/Engine/Plugins/Runtime/Metasound/Source/MetasoundEditor/Private/SMetasoundGraphNode.cpp
+++ b/Engine/Plugins/Runtime/Metasound/Source/MetasoundEditor/Private/SMetasoundGraphNode.cpp
@@ -974,7 +974,7 @@ namespace Metasound
 			{
 				if (Pin->Direction == EGPD_Output)
 				{
-					if (!HasRequiredConnections())
+					if (!HasRequiredConnections() || Pin->GetOwningNode()->ErrorType <= static_cast<uint32>(EMessageSeverity::Warning))
 					{
 						return FLinearColor::Yellow;
 					}
@@ -992,7 +992,7 @@ namespace Metasound
 
 			if (UEdGraphPin* Pin = GetPinObj())
 			{
-				if (!HasRequiredConnections())
+				if (!HasRequiredConnections() || Pin->GetOwningNode()->ErrorType <= static_cast<uint32>(EMessageSeverity::Warning))
 				{
 					return &Editor::Style::GetSlateBrushSafe("MetasoundEditor.Graph.InvalidReroute");
 				}
