Fix Occasional Unexpected Behavior in Bitcrusher Node

diff --git a/Engine/Plugins/Runtime/Metasound/Source/MetasoundStandardNodes/Private/MetasoundBitcrusherNode.cpp b/Engine/Plugins/Runtime/Metasound/Source/MetasoundStandardNodes/Private/MetasoundBitcrusherNode.cpp
--- a/Engine/Plugins/Runtime/Metasound/Source/MetasoundStandardNodes/Private/MetasoundBitcrusherNode.cpp
+++ b/Engine/Plugins/Runtime/Metasound/Source/MetasoundStandardNodes/Private/MetasoundBitcrusherNode.cpp
@@ -41,11 +41,12 @@ namespace Metasound
 			, AudioOutput(FAudioBufferWriteRef::CreateNew(InSettings))
 			, CrushedSampleRate(InCrushedSampleRate)
 			, CrushedBitDepth(InCrushedBitDepth)
+			, MaxSampleRate(InSettings.GetSampleRate())
 		{
 			// Passing in 1 for NumChannels because the node takes one audio input
 			Bitcrusher.Init(InSettings.GetSampleRate(), 1);
-			Bitcrusher.SetSampleRateCrush(*CrushedSampleRate);
-			Bitcrusher.SetBitDepthCrush(*CrushedBitDepth);
+			Bitcrusher.SetSampleRateCrush(FMath::Clamp(*CrushedSampleRate, 1.0f, MaxSampleRate));
+			Bitcrusher.SetBitDepthCrush(FMath::Clamp(*CrushedBitDepth, 1.0f, Bitcrusher.GetMaxBitDepth()));
 		}
 
 		static const FNodeClassMetadata& GetNodeInfo()
@@ -140,14 +141,14 @@ namespace Metasound
 				return;
 			}
 
-			const float CurrSampleRate = FMath::Clamp(*CrushedSampleRate, 1.0f, 16000.0f);
+			const float CurrSampleRate = FMath::Clamp(*CrushedSampleRate, 1.0f, MaxSampleRate);
 			if (!FMath::IsNearlyEqual(PrevSampleRate, CurrSampleRate))
 			{
 				Bitcrusher.SetSampleRateCrush(CurrSampleRate);
 				PrevSampleRate = CurrSampleRate;
 			}
 
-			const float CurrBitDepth = FMath::Clamp(*CrushedBitDepth, 1.0f, 16.0f);
+			const float CurrBitDepth = FMath::Clamp(*CrushedBitDepth, 1.0f, Bitcrusher.GetMaxBitDepth());
 			if (!FMath::IsNearlyEqual(PrevBitDepth, CurrBitDepth))
 			{
 				Bitcrusher.SetBitDepthCrush(CurrBitDepth);
@@ -175,6 +176,7 @@ namespace Metasound
 		// Cached input parameters
 		float PrevSampleRate;
 		float PrevBitDepth;
+		const float MaxSampleRate;
 	
 	};
 
diff --git a/Engine/Source/Runtime/SignalProcessing/Private/BitCrusher.cpp b/Engine/Source/Runtime/SignalProcessing/Private/BitCrusher.cpp
--- a/Engine/Source/Runtime/SignalProcessing/Private/BitCrusher.cpp
+++ b/Engine/Source/Runtime/SignalProcessing/Private/BitCrusher.cpp
@@ -4,6 +4,8 @@
 
 namespace Audio
 {
+	static constexpr float MaxBitDepth = 32.0f;
+
 	FBitCrusher::FBitCrusher()
 		: SampleRate(0)
 		, BitDepth(16.0f)
@@ -34,8 +36,9 @@ namespace Audio
 
 	void FBitCrusher::SetBitDepthCrush(const float InBitDepth)
 	{
-		BitDepth = FMath::Clamp(InBitDepth, 1.0f, 32.0f);
+		BitDepth = FMath::Clamp(InBitDepth, 1.0f, MaxBitDepth);
 		BitDelta = 1.0f / FMath::Pow(2.0f, BitDepth);
+		ReciprocalBitDelta = 1.0f / BitDelta;
 	}
 
 	void FBitCrusher::ProcessAudioFrame(const float* InFrame, float* OutFrame)
@@ -47,7 +50,7 @@ namespace Audio
 
 			for (int32 ChannelIndex = 0; ChannelIndex < NumChannels; ++ChannelIndex)
 			{
-				LastOutput[ChannelIndex] = BitDelta * (float)FMath::FloorToInt(InFrame[ChannelIndex] / BitDelta + 0.5f);
+				LastOutput[ChannelIndex] = BitDelta * FMath::FloorToFloat(InFrame[ChannelIndex] * ReciprocalBitDelta + 0.5f);
 			}
 		}
 
@@ -64,4 +67,9 @@ namespace Audio
 			ProcessAudioFrame(&InBuffer[SampleIndex], &OutBuffer[SampleIndex]);
 		}
 	}
+
+	float FBitCrusher::GetMaxBitDepth()
+	{
+		return MaxBitDepth;
+	}
 }
diff --git a/Engine/Source/Runtime/SignalProcessing/Public/DSP/BitCrusher.h b/Engine/Source/Runtime/SignalProcessing/Public/DSP/BitCrusher.h
--- a/Engine/Source/Runtime/SignalProcessing/Public/DSP/BitCrusher.h
+++ b/Engine/Source/Runtime/SignalProcessing/Public/DSP/BitCrusher.h
@@ -31,11 +31,15 @@ namespace Audio
 		void ProcessAudioFrame(const float* InFrame, float* OutFrame);
 		void ProcessAudio(const float* InBuffer, const int32 InNumSamples, float* OutBuffer);
 
+		// Returns the maximum value given bit depths will be clamped to.
+		static float GetMaxBitDepth();
+
 	private:
 		// i.e. 8 bit, etc. But can be float!
 		float SampleRate;
 		float BitDepth;
 		float BitDelta;
+		float ReciprocalBitDelta;
 
 		// The current phase of the bit crusher
 		float Phase;
