Remove AnalyzerAddress from MetaSound Editor UX/auto-node registration
#rnx

diff --git a/Engine/Plugins/Runtime/Metasound/Source/MetasoundEditor/Private/MetasoundNodeDetailCustomization.cpp b/Engine/Plugins/Runtime/Metasound/Source/MetasoundEditor/Private/MetasoundNodeDetailCustomization.cpp
--- a/Engine/Plugins/Runtime/Metasound/Source/MetasoundEditor/Private/MetasoundNodeDetailCustomization.cpp
+++ b/Engine/Plugins/Runtime/Metasound/Source/MetasoundEditor/Private/MetasoundNodeDetailCustomization.cpp
@@ -1,6 +1,7 @@
 // Copyright Epic Games, Inc. All Rights Reserved.
 #include "MetasoundNodeDetailCustomization.h"
 
+#include "Analysis/MetasoundFrontendAnalyzerAddress.h"
 #include "Components/AudioComponent.h"
 #include "Containers/Set.h"
 #include "Delegates/Delegate.h"
@@ -51,7 +52,8 @@ namespace Metasound
 			static const TSet<FName> HiddenInputTypeNames =
 			{
 				"Audio:Mono",
-				"Audio:Stereo"
+				"Audio:Stereo",
+				GetMetasoundDataTypeName<Frontend::FAnalyzerAddress>()
 			};
 
 			static const FText OverrideInputDefaultText = LOCTEXT("OverridePresetInputDefault", "Override Inherited Default");
diff --git a/Engine/Plugins/Runtime/Metasound/Source/MetasoundFrontend/Private/Analysis/MetasoundFrontendAnalyzerAddress.cpp b/Engine/Plugins/Runtime/Metasound/Source/MetasoundFrontend/Private/Analysis/MetasoundFrontendAnalyzerAddress.cpp
--- a/Engine/Plugins/Runtime/Metasound/Source/MetasoundFrontend/Private/Analysis/MetasoundFrontendAnalyzerAddress.cpp
+++ b/Engine/Plugins/Runtime/Metasound/Source/MetasoundFrontend/Private/Analysis/MetasoundFrontendAnalyzerAddress.cpp
@@ -1,7 +1,9 @@
 // Copyright Epic Games, Inc. All Rights Reserved.
 
 #include "Analysis/MetasoundFrontendAnalyzerAddress.h"
+#include "MetasoundArrayNodesRegistration.h"
 #include "MetasoundDataTypeRegistrationMacro.h"
+#include "MetasoundFrontendDataTypeTraits.h"
 #include "Misc/AssertionMacros.h"
 
 
@@ -85,5 +87,29 @@ namespace Metasound
 		}
 	} // namespace Frontend
 
+	template<>
+	struct TEnableArrayNodes<Frontend::FAnalyzerAddress>
+	{
+		static constexpr bool Value = false;
+	};
+
+	template<>
+	struct TEnableTransmissionNodeRegistration<Frontend::FAnalyzerAddress>
+	{
+		static constexpr bool Value = false;
+	};
+
+	template<typename FromDataType>
+	struct TEnableAutoConverterNodeRegistration<FromDataType, Frontend::FAnalyzerAddress>
+	{
+		static constexpr bool Value = false;
+	};
+
+	template<>
+	struct TEnableConstructorVertex<Frontend::FAnalyzerAddress>
+	{
+		static constexpr bool Value = false;
+	};
+
 	REGISTER_METASOUND_DATATYPE(Frontend::FAnalyzerAddress, "AnalyzerAddress");
 } // namespace Metasound
