Remove input receive node injection by implementing receive implementation within non-ctor input node operator
#rnx

diff --git a/Engine/Plugins/Runtime/Metasound/Source/MetasoundEngine/Private/MetasoundSource.cpp b/Engine/Plugins/Runtime/Metasound/Source/MetasoundEngine/Private/MetasoundSource.cpp
--- a/Engine/Plugins/Runtime/Metasound/Source/MetasoundEngine/Private/MetasoundSource.cpp
+++ b/Engine/Plugins/Runtime/Metasound/Source/MetasoundEngine/Private/MetasoundSource.cpp
@@ -16,7 +16,6 @@
 #include "MetasoundEnvironment.h"
 #include "MetasoundFrontendController.h"
 #include "MetasoundFrontendDataTypeRegistry.h"
-#include "MetasoundFrontendInjectReceiveNodes.h"
 #include "MetasoundFrontendQuery.h"
 #include "MetasoundFrontendQuerySteps.h"
 #include "MetasoundFrontendTransform.h"
diff --git a/Engine/Plugins/Runtime/Metasound/Source/MetasoundFrontend/Private/MetasoundAssetBase.cpp b/Engine/Plugins/Runtime/Metasound/Source/MetasoundFrontend/Private/MetasoundAssetBase.cpp
--- a/Engine/Plugins/Runtime/Metasound/Source/MetasoundFrontend/Private/MetasoundAssetBase.cpp
+++ b/Engine/Plugins/Runtime/Metasound/Source/MetasoundFrontend/Private/MetasoundAssetBase.cpp
@@ -19,7 +19,6 @@
 #include "MetasoundFrontendDocument.h"
 #include "MetasoundFrontendDocumentVersioning.h"
 #include "MetasoundFrontendGraph.h"
-#include "MetasoundFrontendInjectReceiveNodes.h"
 #include "MetasoundFrontendNodeTemplateRegistry.h"
 #include "MetasoundFrontendRegistries.h"
 #include "MetasoundFrontendSearchEngine.h"
@@ -62,6 +61,40 @@ namespace Metasound
 					}
 				}
 			}
+
+			// Remove all inputs from the provided array of public inputs which are of non-transmittable data types.
+			TArray<FMetasoundFrontendClassInput> GetTransmittableInputsFromPublicInputs(const TArray<FMetasoundFrontendClassInput>& InPublicInputs)
+			{
+				using namespace Metasound::Frontend;
+
+				TArray<FMetasoundFrontendClassInput> TransmittableInputs;
+
+				IDataTypeRegistry& Registry = IDataTypeRegistry::Get();
+				Algo::TransformIf(InPublicInputs, TransmittableInputs,
+					[&Registry](const FMetasoundFrontendClassInput& Input)
+					{
+						if (Input.AccessType != EMetasoundFrontendVertexAccessType::Reference)
+						{
+							return false;
+						}
+
+						FDataTypeRegistryInfo Info;
+						if (!Registry.GetDataTypeInfo(Input.TypeName, Info))
+						{
+							return false;
+						}
+
+						if (!Info.bIsTransmittable)
+						{
+							return false;
+						}
+
+						return true;
+					}, [](const FMetasoundFrontendClassInput& Input) { return Input; }
+					);
+
+				return TransmittableInputs;
+			}
 		} // namespace AssetBasePrivate
 
 		FAutoConsoleVariableRef CVarMetaSoundBlockRate(
@@ -179,7 +212,8 @@ void FMetasoundAssetBase::RegisterGraphWithFrontend(Metasound::Frontend::FMetaSo
 
 		virtual TUniquePtr<INode> CreateNode(const FNodeInitData&) const override
 		{
-			return FFrontendGraphBuilder().CreateGraph(*PreprocessedDoc, Name);
+			static const TSet<FName> ReferencedGraphTransmissibleInputNames; // Empty as referenced graphs do not support transmission
+			return FFrontendGraphBuilder().CreateGraph(*PreprocessedDoc, ReferencedGraphTransmissibleInputNames, Name);
 		}
 
 		virtual TUniquePtr<INode> CreateNode(FDefaultLiteralNodeConstructorParams&&) const override
@@ -234,6 +268,7 @@ void FMetasoundAssetBase::RegisterGraphWithFrontend(Metasound::Frontend::FMetaSo
 	TSharedPtr<FMetasoundFrontendDocument> PreprocessedDocument = MakeShared<FMetasoundFrontendDocument>(GetDocumentChecked());
 	FDocumentTemplatePreprocessTransform TemplatePreprocessor;
 	TemplatePreprocessor.Transform(*PreprocessedDocument);
+	CacheRuntimeData(*PreprocessedDocument);
 
 	TUniquePtr<INodeRegistryEntry> RegistryEntry = MakeUnique<FNodeRegistryEntry>(AssetName, PreprocessedDocument, AssetClassInfo.AssetPath);
 	RegistryKey = FMetasoundFrontendRegistryContainer::Get()->RegisterNode(MoveTemp(RegistryEntry));
@@ -271,10 +306,6 @@ void FMetasoundAssetBase::RegisterGraphWithFrontend(Metasound::Frontend::FMetaSo
 		}
 		UE_LOG(LogMetaSound, Error, TEXT("Registration failed for MetaSound node class '%s' of UObject class '%s'"), *AssetName, *ClassName);
 	}
-
-	// Triggers the existing runtime data to be out-of-date.
-	CurrentCachedRuntimeDataChangeID = FGuid::NewGuid();
-	CacheRuntimeData(*PreprocessedDocument);
 }
 
 void FMetasoundAssetBase::UnregisterGraphWithFrontend()
@@ -545,7 +576,7 @@ const FMetasoundFrontendDocumentModifyContext& FMetasoundAssetBase::GetModifyCon
 }
 #endif // WITH_EDITOR
 
-TSharedPtr<Metasound::IGraph, ESPMode::ThreadSafe> FMetasoundAssetBase::BuildMetasoundDocument(const FMetasoundFrontendDocument& InPreprocessedDoc, const TArray<FMetasoundFrontendClassInput>& InTransmittableInputs) const
+TSharedPtr<Metasound::IGraph, ESPMode::ThreadSafe> FMetasoundAssetBase::BuildMetasoundDocument(const FMetasoundFrontendDocument& InPreprocessedDoc, const TSet<FName>& InTransmittableInputNames) const
 {
 	using namespace Metasound;
 	using namespace Metasound::Frontend;
@@ -553,19 +584,8 @@ TSharedPtr<Metasound::IGraph, ESPMode::ThreadSafe> FMetasoundAssetBase::BuildMet
 	METASOUND_TRACE_CPUPROFILER_EVENT_SCOPE(MetaSoundAssetBase::BuildMetasoundDocument);
 
 	// Create graph which can spawn instances. 
-	TUniquePtr<FFrontendGraph> FrontendGraph = FFrontendGraphBuilder::CreateGraph(InPreprocessedDoc, GetOwningAssetName());
-	if (FrontendGraph.IsValid())
-	{
-		TSet<FVertexName> TransmittableInputNames;
-		Algo::Transform(InTransmittableInputs, TransmittableInputNames, [](const FMetasoundFrontendClassInput& Input) { return Input.Name; });
-
-		bool bSuccessfullyInjectedReceiveNodes = InjectReceiveNodes(*FrontendGraph, FMetaSoundParameterTransmitter::CreateSendAddressFromEnvironment, TransmittableInputNames);
-		if (!bSuccessfullyInjectedReceiveNodes)
-		{
-			UE_LOG(LogMetaSound, Error, TEXT("Error while injecting async communication hooks. Instance communication may not function properly [Name:%s]."), *GetOwningAssetName());
-		}
-	}
-	else
+	TUniquePtr<FFrontendGraph> FrontendGraph = FFrontendGraphBuilder::CreateGraph(InPreprocessedDoc, InTransmittableInputNames, GetOwningAssetName());
+	if (!FrontendGraph.IsValid())
 	{
 		UE_LOG(LogMetaSound, Error, TEXT("Failed to build MetaSound graph in asset '%s'"), *GetOwningAssetName());
 	}
@@ -817,21 +837,6 @@ FString FMetasoundAssetBase::GetOwningAssetName() const
 	return FString();
 }
 
-TArray<FMetasoundFrontendClassInput> FMetasoundAssetBase::GetTransmittableClassInputs() const
-{
-	using namespace Metasound;
-	using namespace Metasound::Frontend;
-
-	METASOUND_TRACE_CPUPROFILER_EVENT_SCOPE(FMetasoundAssetBase::GetTransmittableClassInputs);
-
-	check(IsInGameThread() || IsInAudioThread());
-
-	TArray<FMetasoundFrontendClassInput> TransmittableInputs = GetPublicClassInputs();
-	RemoveNonTransmittableClassInputs(TransmittableInputs);
-
-	return TransmittableInputs;
-}
-
 TArray<FMetasoundFrontendClassInput> FMetasoundAssetBase::GetPublicClassInputs() const
 {
 	using namespace Metasound;
@@ -870,31 +875,6 @@ TArray<FMetasoundFrontendClassInput> FMetasoundAssetBase::GetPublicClassInputs()
 	return PublicInputs;
 }
 
-void FMetasoundAssetBase::RemoveNonTransmittableClassInputs(TArray<FMetasoundFrontendClassInput>& InOutClassInputs) const
-{
-	using namespace Metasound::Frontend;
-
-	// Do not transmit vertices which are not transmittable. Async communication
-	// is not supported without transmission.
-	IDataTypeRegistry& Registry = IDataTypeRegistry::Get();
-	auto IsNotTransmittable = [&Registry](const FMetasoundFrontendClassVertex& InVertex)
-	{
-		// Don't transmit constructor inputs 
-		if (InVertex.AccessType == EMetasoundFrontendVertexAccessType::Reference)
-		{
-			FDataTypeRegistryInfo Info;
-			if (Registry.GetDataTypeInfo(InVertex.TypeName, Info))
-			{
-				return !Info.bIsTransmittable;
-			}
-		}
-
-		return true;
-	};
-
-	InOutClassInputs.RemoveAllSwap(IsNotTransmittable);
-}
-
 void FMetasoundAssetBase::RebuildReferencedAssetClassKeys()
 {
 	using namespace Metasound::Frontend;
@@ -905,33 +885,27 @@ void FMetasoundAssetBase::RebuildReferencedAssetClassKeys()
 
 const FMetasoundAssetBase::FRuntimeData& FMetasoundAssetBase::CacheRuntimeData(const FMetasoundFrontendDocument& InPreprocessedDoc)
 {
-	METASOUND_TRACE_CPUPROFILER_EVENT_SCOPE(MetaSoundAssetBase::CacheRuntimeData);
+	using namespace Metasound::Frontend;
 
-	// Check if a ChangeID has been generated before.
-	if (!CurrentCachedRuntimeDataChangeID.IsValid())
-	{
-		CurrentCachedRuntimeDataChangeID = FGuid::NewGuid();
-	}
+	METASOUND_TRACE_CPUPROFILER_EVENT_SCOPE(MetaSoundAssetBase::CacheRuntimeData);
 
-	// Check if CachedRuntimeData is out-of-date.
-	if (CachedRuntimeData.ChangeID != CurrentCachedRuntimeDataChangeID)
-	{
-		// Update CachedRuntimeData.
+	CurrentCachedRuntimeDataChangeID = FGuid::NewGuid();
+	CachedRuntimeData.ChangeID = CurrentCachedRuntimeDataChangeID;
 
-		TArray<FMetasoundFrontendClassInput> PublicInputs = GetPublicClassInputs();
-		TArray<FMetasoundFrontendClassInput> TransmittableInputs = PublicInputs;
-		RemoveNonTransmittableClassInputs(TransmittableInputs);
+	TArray<FMetasoundFrontendClassInput> PublicInputs = GetPublicClassInputs();
+	TArray<FMetasoundFrontendClassInput> TransmittableInputs = AssetBasePrivate::GetTransmittableInputsFromPublicInputs(PublicInputs);
 
-		TSharedPtr<Metasound::IGraph, ESPMode::ThreadSafe> Graph = BuildMetasoundDocument(InPreprocessedDoc, TransmittableInputs);
+	TSet<FName> TransmittableInputNames;
+	Algo::Transform(TransmittableInputs, TransmittableInputNames, [](const FMetasoundFrontendClassInput& Input) { return Input.Name; });
+	TSharedPtr<Metasound::IGraph, ESPMode::ThreadSafe> Graph = BuildMetasoundDocument(InPreprocessedDoc, TransmittableInputNames);
 
-		CachedRuntimeData =
-		{
-			CurrentCachedRuntimeDataChangeID,
-			MoveTemp(PublicInputs),
-			MoveTemp(TransmittableInputs),
-			MoveTemp(Graph)
-		};
-	}
+	CachedRuntimeData =
+	{
+		CurrentCachedRuntimeDataChangeID,
+		MoveTemp(PublicInputs),
+		MoveTemp(TransmittableInputs),
+		MoveTemp(Graph)
+	};
 
 	return CachedRuntimeData;
 }
diff --git a/Engine/Plugins/Runtime/Metasound/Source/MetasoundFrontend/Private/MetasoundFrontendDataTypeRegistry.cpp b/Engine/Plugins/Runtime/Metasound/Source/MetasoundFrontend/Private/MetasoundFrontendDataTypeRegistry.cpp
--- a/Engine/Plugins/Runtime/Metasound/Source/MetasoundFrontend/Private/MetasoundFrontendDataTypeRegistry.cpp
+++ b/Engine/Plugins/Runtime/Metasound/Source/MetasoundFrontend/Private/MetasoundFrontendDataTypeRegistry.cpp
@@ -178,7 +178,14 @@ namespace Metasound
 
 				virtual TUniquePtr<INode> CreateNode(FDefaultNamedVertexWithLiteralNodeConstructorParams&& InParams) const override
 				{
-					return DataTypeEntry->CreateInputNode(MoveTemp(InParams));
+					FInputNodeConstructorParams InputParams;
+					InputParams.InitParam = MoveTemp(InParams.InitParam);
+					InputParams.InstanceID = InParams.InstanceID;
+					InputParams.NodeName = InParams.NodeName;
+					InputParams.VertexName = InParams.VertexName;
+					InputParams.bEnableTransmission = false;
+
+					return DataTypeEntry->CreateInputNode(MoveTemp(InputParams));
 				}
 
 				virtual TUniquePtr<INodeRegistryEntry> Clone() const override
@@ -212,7 +219,14 @@ namespace Metasound
 
 				virtual TUniquePtr<INode> CreateNode(FDefaultNamedVertexWithLiteralNodeConstructorParams&& InParams) const override
 				{
-					return DataTypeEntry->CreateConstructorInputNode(MoveTemp(InParams));
+					FInputNodeConstructorParams InputParams;
+					InputParams.InitParam = MoveTemp(InParams.InitParam);
+					InputParams.InstanceID = InParams.InstanceID;
+					InputParams.NodeName = InParams.NodeName;
+					InputParams.VertexName = InParams.VertexName;
+					InputParams.bEnableTransmission = false;
+
+					return DataTypeEntry->CreateConstructorInputNode(MoveTemp(InputParams));
 				}
 
 				virtual TUniquePtr<INodeRegistryEntry> Clone() const override
diff --git a/Engine/Plugins/Runtime/Metasound/Source/MetasoundFrontend/Private/MetasoundFrontendGraph.cpp b/Engine/Plugins/Runtime/Metasound/Source/MetasoundFrontend/Private/MetasoundFrontendGraph.cpp
--- a/Engine/Plugins/Runtime/Metasound/Source/MetasoundFrontend/Private/MetasoundFrontendGraph.cpp
+++ b/Engine/Plugins/Runtime/Metasound/Source/MetasoundFrontend/Private/MetasoundFrontendGraph.cpp
@@ -196,7 +196,7 @@ namespace Metasound
 		return TUniquePtr<INode>(nullptr);
 	}
 
-	TUniquePtr<INode> FFrontendGraphBuilder::CreateInputNode(const FMetasoundFrontendNode& InNode, const FMetasoundFrontendClass& InClass, const FMetasoundFrontendClassInput& InOwningGraphClassInput)
+	TUniquePtr<INode> FFrontendGraphBuilder::CreateInputNode(const FMetasoundFrontendNode& InNode, const FMetasoundFrontendClass& InClass, const FMetasoundFrontendClassInput& InOwningGraphClassInput, bool bEnableTransmission)
 	{
 		using namespace Metasound::Frontend;
 
@@ -220,7 +220,8 @@ namespace Metasound
 						InNode.Name,
 						InNode.GetID(),
 						InputVertex.Name,
-						MoveTemp(Literal)
+						MoveTemp(Literal),
+						bEnableTransmission
 					};
 
 					ensureAlwaysMsgf(InOwningGraphClassInput.AccessType != EMetasoundFrontendVertexAccessType::Unset, TEXT("Graph Class Input cannot be set to access type of 'Unset'"));
@@ -413,7 +414,7 @@ namespace Metasound
 		return Literal;
 	}
 
-	bool FFrontendGraphBuilder::AddNodesToGraph(FBuildGraphContext& InGraphContext)
+	bool FFrontendGraphBuilder::AddNodesToGraph(FBuildGraphContext& InGraphContext, const TSet<FName>& InTransmittableInputNames)
 	{
 		TSet<FNodeIDVertexID> GraphEdgeDestinations;
 		const TArray<FMetasoundFrontendEdge>& GraphEdges = InGraphContext.GraphClass.Graph.Edges;
@@ -437,7 +438,8 @@ namespace Metasound
 
 						if ((nullptr != ClassInput) && (INDEX_NONE != InputIndex))
 						{
-							TSharedPtr<const INode> InputNode(CreateInputNode(Node, *NodeClass, *ClassInput).Release());
+							const bool bEnableTransmission = InTransmittableInputNames.Contains(ClassInput->Name);
+							TSharedPtr<const INode> InputNode(CreateInputNode(Node, *NodeClass, *ClassInput, bEnableTransmission).Release());
 							InGraphContext.Graph->AddInputNode(Node.GetID(), InputIndex, ClassInput->Name, InputNode);
 						}
 						else
@@ -801,7 +803,7 @@ namespace Metasound
 		return bSuccess;
 	}
 
-	TUniquePtr<FFrontendGraph> FFrontendGraphBuilder::CreateGraph(FBuildContext& InContext, const FMetasoundFrontendGraphClass& InGraphClass)
+	TUniquePtr<FFrontendGraph> FFrontendGraphBuilder::CreateGraph(FBuildContext& InContext, const FMetasoundFrontendGraphClass& InGraphClass, const TSet<FName>& InTransmittableInputNames)
 	{
 		const FString GraphName = InGraphClass.Metadata.GetClassName().GetFullName().ToString();
 
@@ -812,7 +814,7 @@ namespace Metasound
 			InContext
 		};
 
-		bool bSuccess = AddNodesToGraph(BuildGraphContext);
+		bool bSuccess = AddNodesToGraph(BuildGraphContext, InTransmittableInputNames);
 
 		if (bSuccess)
 		{
@@ -835,7 +837,7 @@ namespace Metasound
 		}
 	}
 
-	TUniquePtr<FFrontendGraph> FFrontendGraphBuilder::CreateGraph(const FMetasoundFrontendGraphClass& InGraph, const TArray<FMetasoundFrontendGraphClass>& InSubgraphs, const TArray<FMetasoundFrontendClass>& InDependencies, const FString& InDebugAssetName)
+	TUniquePtr<FFrontendGraph> FFrontendGraphBuilder::CreateGraph(const FMetasoundFrontendGraphClass& InGraph, const TArray<FMetasoundFrontendGraphClass>& InSubgraphs, const TArray<FMetasoundFrontendClass>& InDependencies, const TSet<FName>& InTransmittableInputNames, const FString& InDebugAssetName)
 	{
 		FBuildContext Context;
 		Context.DebugAssetName = InDebugAssetName;
@@ -864,7 +866,7 @@ namespace Metasound
 		// Create each subgraph.
 		for (const FMetasoundFrontendGraphClass* FrontendSubgraphPtr : FrontendSubgraphPtrs)
 		{
-			TSharedPtr<const INode> Subgraph(CreateGraph(Context, *FrontendSubgraphPtr).Release());
+			TSharedPtr<const INode> Subgraph(CreateGraph(Context, *FrontendSubgraphPtr, InTransmittableInputNames).Release());
 			if (!Subgraph.IsValid())
 			{
 				UE_LOG(LogMetaSound, Warning, TEXT("Failed to create subgraph [SubgraphName: %s] in asset '%s'"), *FrontendSubgraphPtr->Metadata.GetClassName().ToString(), *InDebugAssetName);
@@ -877,12 +879,12 @@ namespace Metasound
 		}
 
 		// Create parent graph.
-		return CreateGraph(Context, InGraph);
+		return CreateGraph(Context, InGraph, InTransmittableInputNames);
 	}
 	
 	/* Metasound document should be inflated by now. */
-	TUniquePtr<FFrontendGraph> FFrontendGraphBuilder::CreateGraph(const FMetasoundFrontendDocument& InDocument, const FString& InDebugAssetName)
+	TUniquePtr<FFrontendGraph> FFrontendGraphBuilder::CreateGraph(const FMetasoundFrontendDocument& InDocument, const TSet<FName>& InTransmittableInputNames, const FString& InDebugAssetName)
 	{
-		return CreateGraph(InDocument.RootGraph, InDocument.Subgraphs, InDocument.Dependencies, InDebugAssetName);
+		return CreateGraph(InDocument.RootGraph, InDocument.Subgraphs, InDocument.Dependencies, InTransmittableInputNames, InDebugAssetName);
 	}
 }
diff --git a/Engine/Plugins/Runtime/Metasound/Source/MetasoundFrontend/Private/MetasoundFrontendGraphController.cpp b/Engine/Plugins/Runtime/Metasound/Source/MetasoundFrontend/Private/MetasoundFrontendGraphController.cpp
--- a/Engine/Plugins/Runtime/Metasound/Source/MetasoundFrontend/Private/MetasoundFrontendGraphController.cpp
+++ b/Engine/Plugins/Runtime/Metasound/Source/MetasoundFrontend/Private/MetasoundFrontendGraphController.cpp
@@ -1600,7 +1600,8 @@ namespace Metasound
 				const TArray<FMetasoundFrontendClass>& Dependencies = OwningDocument->GetDependencies();
 
 				FString UnknownAsset = TEXT("UnknownAsset");
-				TUniquePtr<FFrontendGraph> Graph = FFrontendGraphBuilder::CreateGraph(*GraphClass, Subgraphs, Dependencies, UnknownAsset);
+				TSet<FName> TransmittableInputNames;
+				TUniquePtr<FFrontendGraph> Graph = FFrontendGraphBuilder::CreateGraph(*GraphClass, Subgraphs, Dependencies, TransmittableInputNames, UnknownAsset);
 
 				if (!Graph.IsValid())
 				{
diff --git a/Engine/Plugins/Runtime/Metasound/Source/MetasoundFrontend/Private/MetasoundFrontendInjectReceiveNodes.cpp b/Engine/Plugins/Runtime/Metasound/Source/MetasoundFrontend/Private/MetasoundFrontendInjectReceiveNodes.cpp
deleted file mode 100644
diff --git a/Engine/Plugins/Runtime/Metasound/Source/MetasoundFrontend/Public/MetasoundAssetBase.h b/Engine/Plugins/Runtime/Metasound/Source/MetasoundFrontend/Public/MetasoundAssetBase.h
--- a/Engine/Plugins/Runtime/Metasound/Source/MetasoundFrontend/Public/MetasoundAssetBase.h
+++ b/Engine/Plugins/Runtime/Metasound/Source/MetasoundFrontend/Public/MetasoundAssetBase.h
@@ -194,9 +194,6 @@ protected:
 		TSharedPtr<Metasound::IGraph, ESPMode::ThreadSafe> Graph;
 	};
 
-	// Returns the cached runtime data. Call updates cached data if out-of-date.
-	const FRuntimeData& CacheRuntimeData(const FMetasoundFrontendDocument& InPreprocessedDoc);
-
 	// Returns the cached runtime data.
 	const FRuntimeData& GetRuntimeData() const;
 
@@ -204,13 +201,9 @@ protected:
 	// Prefer accessing public class inputs using CacheRuntimeData.
 	TArray<FMetasoundFrontendClassInput> GetPublicClassInputs() const;
 
-	// Returns all transmissible class inputs.  This is a potentially expensive.
-	// Prefer accessing transmissible class inputs using CacheRuntimeData.
-	TArray<FMetasoundFrontendClassInput> GetTransmittableClassInputs() const;
-
 private:
-	// Remove all inputs from the array which are of non-transmittable data types. 
-	void RemoveNonTransmittableClassInputs(TArray<FMetasoundFrontendClassInput>& InOutClassInputs) const;
+	// Returns the cached runtime data. Call updates cached data if out-of-date.
+	const FRuntimeData& CacheRuntimeData(const FMetasoundFrontendDocument& InPreprocessedDoc);
 
 	Metasound::Frontend::FNodeRegistryKey RegistryKey;
 
@@ -218,7 +211,7 @@ private:
 	FGuid CurrentCachedRuntimeDataChangeID;
 	FRuntimeData CachedRuntimeData;
 
-	TSharedPtr<Metasound::IGraph, ESPMode::ThreadSafe> BuildMetasoundDocument(const FMetasoundFrontendDocument& InPreprocessDoc, const TArray<FMetasoundFrontendClassInput>& InTransmittableInputs) const;
+	TSharedPtr<Metasound::IGraph, ESPMode::ThreadSafe> BuildMetasoundDocument(const FMetasoundFrontendDocument& InPreprocessDoc, const TSet<FName>& InTransmittableInputNames) const;
 	Metasound::FSendAddress CreateSendAddress(uint64 InInstanceID, const Metasound::FVertexName& InVertexName, const FName& InDataTypeName) const;
 	Metasound::Frontend::FNodeHandle AddInputPinForSendAddress(const Metasound::FMetaSoundParameterTransmitter::FSendInfo& InSendInfo, Metasound::Frontend::FGraphHandle InGraph) const;
 };
diff --git a/Engine/Plugins/Runtime/Metasound/Source/MetasoundFrontend/Public/MetasoundDataTypeRegistrationMacro.h b/Engine/Plugins/Runtime/Metasound/Source/MetasoundFrontend/Public/MetasoundDataTypeRegistrationMacro.h
--- a/Engine/Plugins/Runtime/Metasound/Source/MetasoundFrontend/Public/MetasoundDataTypeRegistrationMacro.h
+++ b/Engine/Plugins/Runtime/Metasound/Source/MetasoundFrontend/Public/MetasoundDataTypeRegistrationMacro.h
@@ -296,7 +296,7 @@ namespace Metasound
 
 					if constexpr (bIsParsable)
 					{
-						TInputNode<TDataType, EVertexAccessType::Reference> InputPrototype(TEXT(""), FGuid(), TEXT(""), FLiteral());
+						TInputNode<TDataType, EVertexAccessType::Reference> InputPrototype(FInputNodeConstructorParams { });
 						InputClass = Metasound::Frontend::GenerateClass(InputPrototype.GetMetadata(), EMetasoundFrontendClassType::Input);
 
 						TOutputNode<TDataType, EVertexAccessType::Reference> OutputPrototype(TEXT(""), FGuid(), TEXT(""));
@@ -319,7 +319,7 @@ namespace Metasound
 
 						if constexpr (bIsConstructorType)
 						{
-							TInputNode<TDataType, EVertexAccessType::Value> ConstructorInputPrototype(TEXT(""), FGuid(), TEXT(""), FLiteral());
+							TInputNode<TDataType, EVertexAccessType::Value> ConstructorInputPrototype(FInputNodeConstructorParams { });
 							ConstructorInputClass = Metasound::Frontend::GenerateClass(ConstructorInputPrototype.GetMetadata(), EMetasoundFrontendClassType::Input);
 							TOutputNode<TDataType, EVertexAccessType::Value> ConstructorOutputPrototype(TEXT(""), FGuid(), TEXT(""));
 							ConstructorOutputClass = Metasound::Frontend::GenerateClass(ConstructorOutputPrototype.GetMetadata(), EMetasoundFrontendClassType::Output);
@@ -388,7 +388,7 @@ namespace Metasound
 				{
 					if constexpr (bIsParsable)
 					{
-						return MakeUnique<TInputNode<TDataType, EVertexAccessType::Reference>>(InParams.NodeName, InParams.InstanceID, InParams.VertexName, MoveTemp(InParams.InitParam));
+						return MakeUnique<TInputNode<TDataType, EVertexAccessType::Reference>>(MoveTemp(InParams));
 					}
 					else
 					{
@@ -400,7 +400,8 @@ namespace Metasound
 				{
 					if constexpr (bIsParsable && bIsConstructorType)
 					{
-						return MakeUnique<TInputNode<TDataType, EVertexAccessType::Value>>(InParams.NodeName, InParams.InstanceID, InParams.VertexName, MoveTemp(InParams.InitParam));
+						checkf(!InParams.bEnableTransmission, TEXT("Cannot enable transmission on a constructor input."));
+						return MakeUnique<TInputNode<TDataType, EVertexAccessType::Value>>(MoveTemp(InParams));
 					}
 					else
 					{
diff --git a/Engine/Plugins/Runtime/Metasound/Source/MetasoundFrontend/Public/MetasoundFrontendController.h b/Engine/Plugins/Runtime/Metasound/Source/MetasoundFrontend/Public/MetasoundFrontendController.h
--- a/Engine/Plugins/Runtime/Metasound/Source/MetasoundFrontend/Public/MetasoundFrontendController.h
+++ b/Engine/Plugins/Runtime/Metasound/Source/MetasoundFrontend/Public/MetasoundFrontendController.h
@@ -978,7 +978,7 @@ namespace Metasound
 			 */
 			virtual FNodeHandle CreateEmptySubgraph(const FMetasoundFrontendClassMetadata& InInfo) = 0;
 
-			/** Creates a runtime operator for the given graph.
+			/** Creates a runtime operator for the given graph. Does not support input value manipulation via transmission.
 			 *
 			 * @param InSettings - Settings to use when creating operators.
 			 * @param InEnvironment - Environment variables available during creation.
diff --git a/Engine/Plugins/Runtime/Metasound/Source/MetasoundFrontend/Public/MetasoundFrontendGraph.h b/Engine/Plugins/Runtime/Metasound/Source/MetasoundFrontend/Public/MetasoundFrontendGraph.h
--- a/Engine/Plugins/Runtime/Metasound/Source/MetasoundFrontend/Public/MetasoundFrontendGraph.h
+++ b/Engine/Plugins/Runtime/Metasound/Source/MetasoundFrontend/Public/MetasoundFrontendGraph.h
@@ -109,10 +109,10 @@ namespace Metasound
 		static bool IsFlat(const FMetasoundFrontendGraphClass& InRoot, const TArray<FMetasoundFrontendClass>& InDependencies);
 
 		/* Metasound document should be in order to create this graph. */
-		static TUniquePtr<FFrontendGraph> CreateGraph(const FMetasoundFrontendDocument& InDocument, const FString& InDebugAssetName);
+		static TUniquePtr<FFrontendGraph> CreateGraph(const FMetasoundFrontendDocument& InDocument, const TSet<FName>& TransmittableInputNames, const FString& InDebugAssetName);
 
 		/* Metasound document should be in order to create this graph. */
-		static TUniquePtr<FFrontendGraph> CreateGraph(const FMetasoundFrontendGraphClass& InGraphClass, const TArray<FMetasoundFrontendGraphClass>& InSubgraphs, const TArray<FMetasoundFrontendClass>& InDependencies, const FString& InDebugAssetName);
+		static TUniquePtr<FFrontendGraph> CreateGraph(const FMetasoundFrontendGraphClass& InGraphClass, const TArray<FMetasoundFrontendGraphClass>& InSubgraphs, const TArray<FMetasoundFrontendClass>& InDependencies, const TSet<FName>& TransmittableInputNames, const FString& InDebugAssetName);
 
 	private:
 		struct FDefaultLiteralData
@@ -152,20 +152,20 @@ namespace Metasound
 
 		static bool SortSubgraphDependencies(TArray<const FMetasoundFrontendGraphClass*>& Subgraphs);
 
-		static TUniquePtr<FFrontendGraph> CreateGraph(FBuildContext& InContext, const FMetasoundFrontendGraphClass& InSubgraph);
+		static TUniquePtr<FFrontendGraph> CreateGraph(FBuildContext& InContext, const FMetasoundFrontendGraphClass& InSubgraph, const TSet<FName>& TransmittableInputNames);
 
 		static const FMetasoundFrontendClassInput* FindClassInputForInputNode(const FMetasoundFrontendGraphClass& InOwningGraph, const FMetasoundFrontendNode& InInputNode, int32& OutClassInputIndex);
 		static const FMetasoundFrontendClassOutput* FindClassOutputForOutputNode(const FMetasoundFrontendGraphClass& InOwningGraph, const FMetasoundFrontendNode& InOutputNode, int32& OutClassOutputIndex);
 		static const FMetasoundFrontendLiteral* FindInputLiteralForInputNode(const FMetasoundFrontendNode& InInputNode, const FMetasoundFrontendClass& InInputNodeClass, const FMetasoundFrontendClassInput& InOwningGraphClassInput);
 		static const FMetasoundFrontendVariable* FindVariableForVariableNode(const FMetasoundFrontendNode& InVariableNode, const FMetasoundFrontendGraph& InGraph);
 
-		static TUniquePtr<INode> CreateInputNode(const FMetasoundFrontendNode& InNode, const FMetasoundFrontendClass& InClass, const FMetasoundFrontendClassInput& InOwningGraphClassInput);
+		static TUniquePtr<INode> CreateInputNode(const FMetasoundFrontendNode& InNode, const FMetasoundFrontendClass& InClass, const FMetasoundFrontendClassInput& InOwningGraphClassInput, bool bEnableTransmission);
 		static TUniquePtr<INode> CreateOutputNode(const FMetasoundFrontendNode& InNode, const FMetasoundFrontendClass& InClass, FBuildGraphContext& InGraphContext, const TSet<FNodeIDVertexID>& InEdgeDestinations);
 		static TUniquePtr<INode> CreateVariableNode(const FMetasoundFrontendNode& InNode, const FMetasoundFrontendGraph& InGraph);
 		static TUniquePtr<INode> CreateExternalNode(const FMetasoundFrontendNode& InNode, const FMetasoundFrontendClass& InClass, FBuildGraphContext& InGraphContext, const TSet<FNodeIDVertexID>& InEdgeDestinations);
 
 		// Returns false on error
-		static bool AddNodesToGraph(FBuildGraphContext& InGraphContext);
+		static bool AddNodesToGraph(FBuildGraphContext& InGraphContext, const TSet<FName>& TransmittableInputNames);
 
 		// Returns false on error
 		static bool AddEdgesToGraph(FBuildGraphContext& InGraphContext);
diff --git a/Engine/Plugins/Runtime/Metasound/Source/MetasoundFrontend/Public/MetasoundFrontendInjectReceiveNodes.h b/Engine/Plugins/Runtime/Metasound/Source/MetasoundFrontend/Public/MetasoundFrontendInjectReceiveNodes.h
deleted file mode 100644
diff --git a/Engine/Plugins/Runtime/Metasound/Source/MetasoundFrontend/Public/MetasoundFrontendRegistries.h b/Engine/Plugins/Runtime/Metasound/Source/MetasoundFrontend/Public/MetasoundFrontendRegistries.h
--- a/Engine/Plugins/Runtime/Metasound/Source/MetasoundFrontend/Public/MetasoundFrontendRegistries.h
+++ b/Engine/Plugins/Runtime/Metasound/Source/MetasoundFrontend/Public/MetasoundFrontendRegistries.h
@@ -6,45 +6,15 @@
 #include "MetasoundEnum.h"
 #include "MetasoundFrontendDocument.h"
 #include "MetasoundLiteral.h"
+#include "MetasoundNodeConstructorParams.h"
 #include "MetasoundNodeInterface.h"
 #include "MetasoundOperatorInterface.h"
 #include "MetasoundRouter.h"
 #include "MetasoundVertex.h"
 #include "Templates/Function.h"
 
-
 namespace Metasound
 {
-	// Base implementation for NodeConstructorCallbacks
-	struct FDefaultNamedVertexNodeConstructorParams
-	{
-		// the instance name and name of the specific connection that should be used.
-		FVertexName NodeName;
-		FGuid InstanceID;
-		FVertexName VertexName;
-	};
-
-	struct FDefaultNamedVertexWithLiteralNodeConstructorParams
-	{
-		// the instance name and name of the specific connection that should be used.
-		FVertexName NodeName;
-		FGuid InstanceID;
-		FVertexName VertexName;
-		FLiteral InitParam = FLiteral::CreateInvalid();
-	};
-
-	struct FDefaultLiteralNodeConstructorParams
-	{
-		FVertexName NodeName;
-		FGuid InstanceID;
-		FLiteral Literal = FLiteral::CreateInvalid();
-	};
-
-	using FInputNodeConstructorParams = FDefaultNamedVertexWithLiteralNodeConstructorParams;
-	using FOutputNodeConstructorParams = FDefaultNamedVertexNodeConstructorParams;
-	using FLiteralNodeConstructorParams = FDefaultLiteralNodeConstructorParams;
-	using FVariableNodeConstructorParams = FDefaultLiteralNodeConstructorParams;
-
 	using FIterateMetasoundFrontendClassFunction = TFunctionRef<void(const FMetasoundFrontendClass&)>;
 
 	namespace Frontend
diff --git a/Engine/Plugins/Runtime/Metasound/Source/MetasoundFrontend/Public/MetasoundInputNode.h b/Engine/Plugins/Runtime/Metasound/Source/MetasoundFrontend/Public/MetasoundInputNode.h
--- a/Engine/Plugins/Runtime/Metasound/Source/MetasoundFrontend/Public/MetasoundInputNode.h
+++ b/Engine/Plugins/Runtime/Metasound/Source/MetasoundFrontend/Public/MetasoundInputNode.h
@@ -3,19 +3,24 @@
 
 #include "Internationalization/Text.h"
 #include "MetasoundBuilderInterface.h"
+#include "MetasoundBuildError.h"
 #include "MetasoundDataReference.h"
 #include "MetasoundFrontendDataTypeTraits.h"
+#include "MetasoundNodeConstructorParams.h"
 #include "MetasoundLiteral.h"
 #include "MetasoundNode.h"
 #include "MetasoundNodeInterface.h"
 #include "MetasoundNodeRegistrationMacro.h"
 #include "MetasoundOperatorInterface.h"
+#include "MetasoundParameterTransmitter.h"
+#include "MetasoundRouter.h"
 #include "MetasoundTrigger.h"
 #include "MetasoundVertexData.h"
 #include "UObject/NameTypes.h"
 
 #define LOCTEXT_NAMESPACE "MetasoundFrontend"
 
+
 namespace Metasound
 {
 	/** A writable input and a readable output. */
@@ -78,14 +83,6 @@ namespace Metasound
 				TExecutableDataType<DataType>::Execute(*InputValue, *OutputValue);
 			}
 
-			void ExecutPassThrough()
-			{
-				if (TExecutableDataType<DataType>::bIsExecutable)
-				{
-					*OutputValue = *InputValue;
-				}
-			}
-
 			static void ExecuteFunction(IOperator* InOperator)
 			{
 				static_cast<TInputOperator<DataType>*>(InOperator)->Execute();
@@ -120,21 +117,16 @@ namespace Metasound
 
 		TPassThroughOperator(const FVertexName& InDataReferenceName, FDataReadReference InDataReference)
 		: TInputOperator<DataType>(InDataReferenceName, WriteCast(InDataReference)) // Write cast is safe because `GetExecuteFunction() and Bind() are overridden, ensuring that data is not written.
-		, DataReferenceName(InDataReferenceName)
 		{
 		}
 
 		virtual ~TPassThroughOperator() = default;
 
-		virtual void Bind(FVertexInterfaceData& InOutVertexData) const override
-		{
-			InOutVertexData.GetInputs().BindReadVertex(DataReferenceName, Super::InputValue);
-			InOutVertexData.GetOutputs().BindReadVertex(DataReferenceName, Super::OutputValue);
-		}
-
 		static void ExecuteFunction(IOperator* InOperator)
 		{
-			static_cast<TPassThroughOperator<DataType>*>(InOperator)->ExecutPassThrough();
+			using FPassThroughOperator = TPassThroughOperator<DataType>;
+			FPassThroughOperator* PassThroughOperator = static_cast<FPassThroughOperator*>(InOperator);
+			*(PassThroughOperator->OutputValue) = *(PassThroughOperator->InputValue);
 		}
 
 		virtual IOperator::FExecuteFunction GetExecuteFunction() override
@@ -184,11 +176,73 @@ namespace Metasound
 			{
 				return &TPassThroughOperator<DataType>::ExecuteFunction;
 			}
+
 			return nullptr;
 		}
+	};
+
+	/** FInputValueOperator provides support for transmittable inputs. */
+	template<typename DataType>
+	class TInputReceiverOperator : public TInputOperator<DataType>
+	{
+	public:
+		using FDataReadReference = TDataReadReference<DataType>;
+		using FInputReceiverOperator = TInputReceiverOperator<DataType>;
+		using Super = TInputOperator<DataType>;
+
+		/** Construct an TInputReceiverOperator with the name of the vertex, value reference associated with input, SendAddress, & Receiver
+		 */
+		TInputReceiverOperator(const FVertexName& InDataReferenceName, FDataReadReference InDataReference, FSendAddress&& InSendAddress, TReceiverPtr<DataType>&& InReceiver)
+			: Super(InDataReferenceName, WriteCast(InDataReference)) // Write cast is safe because `GetExecuteFunction() and Bind() are overridden, ensuring that data is not written.
+			, SendAddress(MoveTemp(InSendAddress))
+			, Receiver(MoveTemp(InReceiver))
+		{
+		}
+
+		virtual ~TInputReceiverOperator()
+		{
+			Receiver.Reset();
+			FDataTransmissionCenter::Get().UnregisterDataChannelIfUnconnected(SendAddress);
+		}
+
+		void Execute()
+		{
+			DataType& OutputData = *Super::OutputValue;
+
+			bool bHasNewData = false;
+			bHasNewData = Receiver->CanPop();
+			if (bHasNewData)
+			{
+				Receiver->Pop(OutputData);
+				bHasNotReceivedData = false;
+			}
+
+			if (bHasNotReceivedData)
+			{
+				OutputData = *Super::InputValue;
+				bHasNewData = true;
+			}
+
+			if constexpr (TExecutableDataType<DataType>::bIsExecutable)
+			{
+				TExecutableDataType<DataType>::ExecuteInline(OutputData, bHasNewData);
+			}
+		}
+
+		static void ExecuteFunction(IOperator* InOperator)
+		{
+			static_cast<FInputReceiverOperator*>(InOperator)->Execute();
+		}
+
+		IOperator::FExecuteFunction GetExecuteFunction()
+		{
+			return &FInputReceiverOperator::ExecuteFunction;
+		}
 
 	private:
-		FVertexName DataReferenceName;
+		FSendAddress SendAddress;
+		TReceiverPtr<DataType> Receiver;
+		bool bHasNotReceivedData = true;
 	};
 
 	/** FInputValueOperator provides an input for value references. */
@@ -382,14 +436,9 @@ namespace Metasound
 			};
 
 		public:
-
-			FInputNodeOperatorFactory(const DataType& InValue)
-			: ReferenceCreator(MakeUnique<FCreateWithCopy>(InValue))
-			{
-			}
-
-			explicit FInputNodeOperatorFactory(FLiteral&& InLiteral)
+			explicit FInputNodeOperatorFactory(FLiteral&& InLiteral, bool bInEnableTransmission)
 			: ReferenceCreator(MakeUnique<FCreateWithLiteral>(MoveTemp(InLiteral)))
+			, bEnableTransmission(bInEnableTransmission)
 			{
 			}
 
@@ -399,6 +448,8 @@ namespace Metasound
 				using FOwnedInputOperatorType = std::conditional_t<bIsReferenceVertexAccess, TInputOperator<DataType>, FInputValueOperator>;
 				using FPassThroughInputOperatorType = std::conditional_t<bIsReferenceVertexAccess, TPassThroughOperator<DataType>, FInputValueOperator>;
 
+				checkf(!(bEnableTransmission && bIsValueVertexAccess), TEXT("Input cannot enable transmission for vertex with access 'EVertexAccessType::Reference'"));
+
 				const FInputNodeType& InputNode = static_cast<const FInputNodeType&>(InParams.Node);
 				const FVertexName& VertexKey = InputNode.GetVertexName();
 
@@ -407,15 +458,31 @@ namespace Metasound
 					// Pass through input value
 					return MakeUnique<FPassThroughInputOperatorType>(VertexKey, ReferenceCreator->CreateDataReference(*Ref));
 				}
-				else
+
+				FDataReference DataRef = ReferenceCreator->CreateDataReference(InParams.OperatorSettings);
+
+				if (bEnableTransmission)
 				{
-					// Owned input value
-					return MakeUnique<FOwnedInputOperatorType>(VertexKey, ReferenceCreator->CreateDataReference(InParams.OperatorSettings));
+					const FName DataTypeName = GetMetasoundDataTypeName<DataType>();
+					FSendAddress SendAddress = FMetaSoundParameterTransmitter::CreateSendAddressFromEnvironment(InParams.Environment, VertexKey, DataTypeName);
+					TReceiverPtr<DataType> Receiver = FDataTransmissionCenter::Get().RegisterNewReceiver<DataType>(SendAddress, FReceiverInitParams{ InParams.OperatorSettings });
+					if (Receiver.IsValid())
+					{
+						// Transmittable input value
+						return MakeUnique<TInputReceiverOperator<DataType>>(VertexKey, DataRef, MoveTemp(SendAddress), MoveTemp(Receiver));
+					}
+
+					AddBuildError<FInputReceiverInitializationError>(OutResults.Errors, InParams.Node, VertexKey, DataTypeName);
+					return nullptr;
 				}
+
+				// Owned input value
+				return MakeUnique<FOwnedInputOperatorType>(VertexKey, DataRef);
 			}
 
 		private:
 			TUniquePtr<FDataReferenceCreatorBase> ReferenceCreator;
+			bool bEnableTransmission = false;
 		};
 
 
@@ -449,7 +516,7 @@ namespace Metasound
 
 			return Info;
 		}
-		
+
 		template<typename... ArgTypes>
 		UE_DEPRECATED(5.1, "Moved to internal implementation.")
 		static FOperatorFactorySharedRef CreateOperatorFactoryWithArgs(ArgTypes&&... Args)
@@ -483,11 +550,11 @@ namespace Metasound
 
 		/* Construct a TInputNode using the TInputOperatorLiteralFactory<> and moving
 		 * InParam to the TInputOperatorLiteralFactory constructor.*/
-		explicit TInputNode(const FVertexName& InNodeName, const FGuid& InInstanceID, const FVertexName& InVertexName, FLiteral&& InParam)
-		:	FNode(InNodeName, InInstanceID, GetNodeInfo(InVertexName))
-		,	VertexName(InVertexName)
-		,	Interface(DeclareVertexInterface(InVertexName))
-		, 	Factory(MakeShared<FInputNodeOperatorFactory>(MoveTemp(InParam)))
+		explicit TInputNode(FInputNodeConstructorParams&& InParams)
+		:	FNode(InParams.NodeName, InParams.InstanceID, GetNodeInfo(InParams.VertexName))
+		,	VertexName(InParams.VertexName)
+		,	Interface(DeclareVertexInterface(InParams.VertexName))
+		,	Factory(MakeShared<FInputNodeOperatorFactory>(MoveTemp(InParams.InitParam), InParams.bEnableTransmission))
 		{
 		}
 
diff --git a/Engine/Plugins/Runtime/Metasound/Source/MetasoundFrontend/Public/MetasoundNodeConstructorParams.h b/Engine/Plugins/Runtime/Metasound/Source/MetasoundFrontend/Public/MetasoundNodeConstructorParams.h
new file mode 100644
--- /dev/null
+++ b/Engine/Plugins/Runtime/Metasound/Source/MetasoundFrontend/Public/MetasoundNodeConstructorParams.h
@@ -0,0 +1,48 @@
+// Copyright Epic Games, Inc. All Rights Reserved.
+#pragma once
+
+#include "MetasoundLiteral.h"
+#include "MetasoundVertex.h"
+
+
+namespace Metasound
+{
+	// Base implementation for NodeConstructorCallbacks
+	struct FDefaultNamedVertexNodeConstructorParams
+	{
+		// the instance name and name of the specific connection that should be used.
+		FVertexName NodeName;
+		FGuid InstanceID;
+		FVertexName VertexName;
+	};
+
+	struct FDefaultNamedVertexWithLiteralNodeConstructorParams
+	{
+		// the instance name and name of the specific connection that should be used.
+		FVertexName NodeName;
+		FGuid InstanceID;
+		FVertexName VertexName;
+		FLiteral InitParam = FLiteral::CreateInvalid();
+	};
+
+	struct FDefaultLiteralNodeConstructorParams
+	{
+		FVertexName NodeName;
+		FGuid InstanceID;
+		FLiteral Literal = FLiteral::CreateInvalid();
+	};
+
+	struct FInputNodeConstructorParams
+	{
+		// the instance name and name of the specific connection that should be used.
+		FVertexName NodeName;
+		FGuid InstanceID;
+		FVertexName VertexName;
+		FLiteral InitParam = FLiteral::CreateInvalid();
+		bool bEnableTransmission = false;
+	};
+
+	using FOutputNodeConstructorParams = FDefaultNamedVertexNodeConstructorParams;
+	using FLiteralNodeConstructorParams = FDefaultLiteralNodeConstructorParams;
+	using FVariableNodeConstructorParams = FDefaultLiteralNodeConstructorParams;
+} // namespace Metasound
\ No newline at end of file
diff --git a/Engine/Plugins/Runtime/Metasound/Source/MetasoundGraphCore/Private/MetasoundBuildError.cpp b/Engine/Plugins/Runtime/Metasound/Source/MetasoundGraphCore/Private/MetasoundBuildError.cpp
--- a/Engine/Plugins/Runtime/Metasound/Source/MetasoundGraphCore/Private/MetasoundBuildError.cpp
+++ b/Engine/Plugins/Runtime/Metasound/Source/MetasoundGraphCore/Private/MetasoundBuildError.cpp
@@ -203,6 +203,17 @@ namespace Metasound
 	{
 		AddDataEdge(InEdge);
 	}
+
+	const FName FInputReceiverInitializationError::ErrorType = "InputReceiverInitializationError";
+
+	FInputReceiverInitializationError::FInputReceiverInitializationError(const INode& InInputNode, const FName& InVertexKey, const FName& InDataType)
+		: FBuildErrorBase(ErrorType, FText::Format(LOCTEXT("InputReceiverInitializationError", "Failed to create transmission receiever for input '{0}' of type '{1}'.")
+			, FText::FromName(InVertexKey)
+			, FText::FromName(InDataType)
+		))
+	{
+		AddNode(InInputNode);
+	}
 }
 
 #undef LOCTEXT_NAMESPACE
diff --git a/Engine/Plugins/Runtime/Metasound/Source/MetasoundGraphCore/Public/MetasoundBuildError.h b/Engine/Plugins/Runtime/Metasound/Source/MetasoundGraphCore/Public/MetasoundBuildError.h
--- a/Engine/Plugins/Runtime/Metasound/Source/MetasoundGraphCore/Public/MetasoundBuildError.h
+++ b/Engine/Plugins/Runtime/Metasound/Source/MetasoundGraphCore/Public/MetasoundBuildError.h
@@ -228,4 +228,19 @@ namespace Metasound
 
 			virtual ~FInvalidConnectionDataTypeError() = default;
 	};
-}
+
+	/** FInputReceiverInitializationError
+	 *
+	 * Caused by Inputs that are set to enable transmission fail to create a receiver.
+	 */
+	class METASOUNDGRAPHCORE_API FInputReceiverInitializationError : public FBuildErrorBase
+	{
+		public:
+
+			static const FName ErrorType;
+			
+			FInputReceiverInitializationError(const INode& InInputNode, const FName& InVertexKey, const FName& InDataType);
+
+			virtual ~FInputReceiverInitializationError() = default;
+	};
+} // namespace Metasound
diff --git a/Engine/Plugins/Runtime/Metasound/Source/MetasoundStandardNodes/Private/MetasoundAudioFormats.cpp b/Engine/Plugins/Runtime/Metasound/Source/MetasoundStandardNodes/Private/MetasoundAudioFormats.cpp
--- a/Engine/Plugins/Runtime/Metasound/Source/MetasoundStandardNodes/Private/MetasoundAudioFormats.cpp
+++ b/Engine/Plugins/Runtime/Metasound/Source/MetasoundStandardNodes/Private/MetasoundAudioFormats.cpp
@@ -347,10 +347,10 @@ namespace Metasound
 	public:
 		static constexpr bool bCanRegister = true;
 
-		TInputNode(const FVertexName& InInstanceName, const FGuid& InInstanceID, const FVertexName& InVertexName, FLiteral&& InLiteral)
-		:	FNode(InInstanceName, InInstanceID, GetNodeClassMetadata(InVertexName))
-		,	VertexInterface(GetVertexInterface(InVertexName))
-		,	Factory(MakeShared<FInputOperatorFactory, ESPMode::ThreadSafe>(InVertexName))
+		TInputNode(FInputNodeConstructorParams&& InParams)
+			: FNode(InParams.NodeName, InParams.InstanceID, GetNodeClassMetadata(InParams.VertexName))
+			, VertexInterface(GetVertexInterface(InParams.VertexName))
+			, Factory(MakeShared<FInputOperatorFactory, ESPMode::ThreadSafe>(InParams.VertexName))
 		{
 		}
 
@@ -614,10 +614,10 @@ namespace Metasound
 	public:
 		static constexpr bool bCanRegister = true;
 
-		TInputNode(const FVertexName& InInstanceName, const FGuid& InInstanceID, const FVertexName& InVertexName, FLiteral&& InLiteral)
-		:	FNode(InInstanceName, InInstanceID, GetNodeClassMetadata(InVertexName))
-		,	VertexInterface(GetVertexInterface(InVertexName))
-		,	Factory(MakeShared<FInputOperatorFactory, ESPMode::ThreadSafe>(InVertexName))
+		TInputNode(FInputNodeConstructorParams&& InParams)
+			: FNode(InParams.NodeName, InParams.InstanceID, GetNodeClassMetadata(InParams.VertexName))
+			, VertexInterface(GetVertexInterface(InParams.VertexName))
+			, Factory(MakeShared<FInputOperatorFactory, ESPMode::ThreadSafe>(InParams.VertexName))
 		{
 		}
 
