Remove 5.0 channel format from MetaSoundSource and MetaSound WavePlayer node.
#rnx

diff --git a/Engine/Plugins/Runtime/Metasound/Source/MetasoundEngine/Private/MetasoundEngineNodesNames.cpp b/Engine/Plugins/Runtime/Metasound/Source/MetasoundEngine/Private/MetasoundEngineNodesNames.cpp
--- a/Engine/Plugins/Runtime/Metasound/Source/MetasoundEngine/Private/MetasoundEngineNodesNames.cpp
+++ b/Engine/Plugins/Runtime/Metasound/Source/MetasoundEngine/Private/MetasoundEngineNodesNames.cpp
@@ -11,7 +11,6 @@ namespace Metasound
 		const FName MonoVariant = "Mono";
 		const FName StereoVariant = "Stereo";
 		const FName QuadVariant = "Quad";
-		const FName FiveDotZeroVariant = "5dot0";
 		const FName FiveDotOneVariant = "5dot1";
 		const FName SevenDotOneVariant = "7dot1";
 	}
diff --git a/Engine/Plugins/Runtime/Metasound/Source/MetasoundEngine/Private/MetasoundInterface.cpp b/Engine/Plugins/Runtime/Metasound/Source/MetasoundEngine/Private/MetasoundInterface.cpp
--- a/Engine/Plugins/Runtime/Metasound/Source/MetasoundEngine/Private/MetasoundInterface.cpp
+++ b/Engine/Plugins/Runtime/Metasound/Source/MetasoundEngine/Private/MetasoundInterface.cpp
@@ -256,7 +256,6 @@ namespace Metasound
 
 				RegisterInterface(OutputFormatStereoInterface::CreateInterface(*UMetaSoundSource::StaticClass()), nullptr, bIsDefault, bEditorCanAddOrRemove, IDataReference::RouterName);			
 				RegisterInterface(OutputFormatQuadInterface::CreateInterface(*UMetaSoundSource::StaticClass()), nullptr, bIsDefault, bEditorCanAddOrRemove, IDataReference::RouterName);			
-				RegisterInterface(OutputFormatFiveDotZeroInterface::CreateInterface(*UMetaSoundSource::StaticClass()), nullptr, bIsDefault, bEditorCanAddOrRemove, IDataReference::RouterName);			
 				RegisterInterface(OutputFormatFiveDotOneInterface::CreateInterface(*UMetaSoundSource::StaticClass()), nullptr, bIsDefault, bEditorCanAddOrRemove, IDataReference::RouterName);			
 				RegisterInterface(OutputFormatSevenDotOneInterface::CreateInterface(*UMetaSoundSource::StaticClass()), nullptr, bIsDefault, bEditorCanAddOrRemove, IDataReference::RouterName);			
 			}
diff --git a/Engine/Plugins/Runtime/Metasound/Source/MetasoundEngine/Private/MetasoundSource.cpp b/Engine/Plugins/Runtime/Metasound/Source/MetasoundEngine/Private/MetasoundSource.cpp
--- a/Engine/Plugins/Runtime/Metasound/Source/MetasoundEngine/Private/MetasoundSource.cpp
+++ b/Engine/Plugins/Runtime/Metasound/Source/MetasoundEngine/Private/MetasoundSource.cpp
@@ -95,19 +95,6 @@ namespace Metasound
 							}
 						}
 					},
-					{
-						EMetasoundSourceAudioFormat::FiveDotZero,
-						{
-							OutputFormatFiveDotZeroInterface::GetVersion(),
-							{
-								OutputFormatFiveDotZeroInterface::Outputs::FrontLeftOut,
-								OutputFormatFiveDotZeroInterface::Outputs::FrontRightOut,
-								OutputFormatFiveDotZeroInterface::Outputs::FrontCenterOut,
-								OutputFormatFiveDotZeroInterface::Outputs::SideLeftOut,
-								OutputFormatFiveDotZeroInterface::Outputs::SideRightOut
-							}
-						}
-					},
 					{
 						EMetasoundSourceAudioFormat::FiveDotOne,
 						{
diff --git a/Engine/Plugins/Runtime/Metasound/Source/MetasoundEngine/Private/MetasoundWavePlayerNode.cpp b/Engine/Plugins/Runtime/Metasound/Source/MetasoundEngine/Private/MetasoundWavePlayerNode.cpp
--- a/Engine/Plugins/Runtime/Metasound/Source/MetasoundEngine/Private/MetasoundWavePlayerNode.cpp
+++ b/Engine/Plugins/Runtime/Metasound/Source/MetasoundEngine/Private/MetasoundWavePlayerNode.cpp
@@ -1181,26 +1181,6 @@ namespace Metasound
 	using FQuadWavePlayerNode = TWavePlayerNode<FQuadAudioChannelConfigurationInfo>;
 	METASOUND_REGISTER_NODE(FQuadWavePlayerNode);
 
-	struct FFiveDotZeroAudioChannelConfigurationInfo
-	{
-		static FText GetNodeDisplayName() { return METASOUND_LOCTEXT("Metasound_WavePlayerFiveDotZeroNodeDisplayName", "Wave Player (5.0)"); }
-		static FName GetVariantName() { return Metasound::EngineNodes::FiveDotZeroVariant; }
-
-		static TArray<FOutputDataVertex> GetAudioOutputs()
-		{
-			using namespace WavePlayerVertexNames;
-			return {
-				TOutputDataVertex<FAudioBuffer>(METASOUND_GET_PARAM_NAME_AND_METADATA(OutputAudioFrontLeft)),
-				TOutputDataVertex<FAudioBuffer>(METASOUND_GET_PARAM_NAME_AND_METADATA(OutputAudioFrontRight)),
-				TOutputDataVertex<FAudioBuffer>(METASOUND_GET_PARAM_NAME_AND_METADATA(OutputAudioFrontCenter)),
-				TOutputDataVertex<FAudioBuffer>(METASOUND_GET_PARAM_NAME_AND_METADATA(OutputAudioSideLeft)),
-				TOutputDataVertex<FAudioBuffer>(METASOUND_GET_PARAM_NAME_AND_METADATA(OutputAudioSideRight))
-			};
-		}
-	};
-	using FFiveDotZeroWavePlayerNode = TWavePlayerNode<FFiveDotZeroAudioChannelConfigurationInfo>;
-	METASOUND_REGISTER_NODE(FFiveDotZeroWavePlayerNode);
-
 	struct FFiveDotOneAudioChannelConfigurationInfo
 	{
 		static FText GetNodeDisplayName() { return METASOUND_LOCTEXT("Metasound_WavePlayerFiveDotOneNodeDisplayName", "Wave Player (5.1)"); }
diff --git a/Engine/Plugins/Runtime/Metasound/Source/MetasoundEngine/Public/MetasoundEngineNodesNames.h b/Engine/Plugins/Runtime/Metasound/Source/MetasoundEngine/Public/MetasoundEngineNodesNames.h
--- a/Engine/Plugins/Runtime/Metasound/Source/MetasoundEngine/Public/MetasoundEngineNodesNames.h
+++ b/Engine/Plugins/Runtime/Metasound/Source/MetasoundEngine/Public/MetasoundEngineNodesNames.h
@@ -13,7 +13,6 @@ namespace Metasound
 		extern const FName MonoVariant;
 		extern const FName StereoVariant;
 		extern const FName QuadVariant;
-		extern const FName FiveDotZeroVariant;
 		extern const FName FiveDotOneVariant;
 		extern const FName SevenDotOneVariant;
 	}
diff --git a/Engine/Plugins/Runtime/Metasound/Source/MetasoundEngine/Public/MetasoundSource.h b/Engine/Plugins/Runtime/Metasound/Source/MetasoundEngine/Public/MetasoundSource.h
--- a/Engine/Plugins/Runtime/Metasound/Source/MetasoundEngine/Public/MetasoundSource.h
+++ b/Engine/Plugins/Runtime/Metasound/Source/MetasoundEngine/Public/MetasoundSource.h
@@ -25,7 +25,6 @@ enum class EMetasoundSourceAudioFormat : uint8
 	Mono,
 	Stereo,
 	Quad,
-	FiveDotZero UMETA(DisplayName="5.0"),
 	FiveDotOne UMETA(DisplayName="5.1"),
 	SevenDotOne UMETA(DisplayName="7.1"),
 
diff --git a/Engine/Plugins/Runtime/Metasound/Source/MetasoundFrontend/Private/Interfaces/MetasoundFrontendOutputFormatInterfaces.cpp b/Engine/Plugins/Runtime/Metasound/Source/MetasoundFrontend/Private/Interfaces/MetasoundFrontendOutputFormatInterfaces.cpp
--- a/Engine/Plugins/Runtime/Metasound/Source/MetasoundFrontend/Private/Interfaces/MetasoundFrontendOutputFormatInterfaces.cpp
+++ b/Engine/Plugins/Runtime/Metasound/Source/MetasoundFrontend/Private/Interfaces/MetasoundFrontendOutputFormatInterfaces.cpp
@@ -253,48 +253,6 @@ namespace Metasound
 		} // namespace OutputFormatQuadInterface
 #undef AUDIO_PARAMETER_INTERFACE_NAMESPACE
 
-#define AUDIO_PARAMETER_INTERFACE_NAMESPACE "UE.OutputFormat.5dot0"
-		namespace OutputFormatFiveDotZeroInterface
-		{
-			const FMetasoundFrontendVersion& GetVersion()
-			{
-				static const FMetasoundFrontendVersion Version = { AUDIO_PARAMETER_INTERFACE_NAMESPACE, { 1, 0 } };
-				return Version;
-			}
-
-			namespace Outputs
-			{
-				const FName FrontLeftOut = AUDIO_PARAMETER_INTERFACE_MEMBER_DEFINE("Audio:0");
-				const FName FrontRightOut = AUDIO_PARAMETER_INTERFACE_MEMBER_DEFINE("Audio:1");
-				const FName FrontCenterOut = AUDIO_PARAMETER_INTERFACE_MEMBER_DEFINE("Audio:2");
-				const FName SideLeftOut = AUDIO_PARAMETER_INTERFACE_MEMBER_DEFINE("Audio:4");
-				const FName SideRightOut = AUDIO_PARAMETER_INTERFACE_MEMBER_DEFINE("Audio:5");
-			}
-
-			Audio::FParameterInterfacePtr CreateInterface(const UClass& InUClass)
-			{
-				struct FInterface : public Audio::FParameterInterface
-				{
-					FInterface(const UClass& InAssetClass)
-						: FParameterInterface(OutputFormatFiveDotZeroInterface::GetVersion().Name, OutputFormatFiveDotZeroInterface::GetVersion().Number.ToInterfaceVersion(), InAssetClass)
-					{
-						using namespace MetasoundOutputFormatInterfacesPrivate;
-						Outputs =
-						{
-							GetFrontLeftOutput(Outputs::FrontLeftOut),
-							GetFrontRightOutput(Outputs::FrontRightOut),
-							GetFrontCenterOutput(Outputs::FrontCenterOut),
-							GetSideLeftOutput(Outputs::SideLeftOut),
-							GetSideRightOutput(Outputs::SideRightOut)
-						};
-					}
-				};
-
-				return MakeShared<FInterface>(InUClass);
-			}
-		} // namespace OutputFormatFiveDotZeroInterface
-#undef AUDIO_PARAMETER_INTERFACE_NAMESPACE
-
 #define AUDIO_PARAMETER_INTERFACE_NAMESPACE "UE.OutputFormat.5dot1"
 		namespace OutputFormatFiveDotOneInterface
 		{
diff --git a/Engine/Plugins/Runtime/Metasound/Source/MetasoundFrontend/Public/Interfaces/MetasoundFrontendOutputFormatInterfaces.h b/Engine/Plugins/Runtime/Metasound/Source/MetasoundFrontend/Public/Interfaces/MetasoundFrontendOutputFormatInterfaces.h
--- a/Engine/Plugins/Runtime/Metasound/Source/MetasoundFrontend/Public/Interfaces/MetasoundFrontendOutputFormatInterfaces.h
+++ b/Engine/Plugins/Runtime/Metasound/Source/MetasoundFrontend/Public/Interfaces/MetasoundFrontendOutputFormatInterfaces.h
@@ -49,21 +49,6 @@ namespace Metasound
 			METASOUNDFRONTEND_API Audio::FParameterInterfacePtr CreateInterface(const UClass& InClass);
 		} // namespace OutputFormatQuadInterface
 
-		namespace OutputFormatFiveDotZeroInterface
-		{
-			namespace Outputs
-			{
-				METASOUNDFRONTEND_API const extern FName FrontLeftOut;
-				METASOUNDFRONTEND_API const extern FName FrontRightOut;
-				METASOUNDFRONTEND_API const extern FName FrontCenterOut;
-				METASOUNDFRONTEND_API const extern FName SideLeftOut;
-				METASOUNDFRONTEND_API const extern FName SideRightOut;
-			} // namespace Outputs
-
-			METASOUNDFRONTEND_API const FMetasoundFrontendVersion& GetVersion();
-			METASOUNDFRONTEND_API Audio::FParameterInterfacePtr CreateInterface(const UClass& InClass);
-		} // namespace OutputFormatFiveDotZeroInterface
-
 		namespace OutputFormatFiveDotOneInterface
 		{
 			namespace Outputs
