- Fix for off by one on looping Sound Waves in MetaSounds
- Fixed typo in Nearly Finished tooltip

diff --git a/Engine/Plugins/Runtime/Metasound/Source/MetasoundEngine/Private/MetasoundWavePlayerNode.cpp b/Engine/Plugins/Runtime/Metasound/Source/MetasoundEngine/Private/MetasoundWavePlayerNode.cpp
--- a/Engine/Plugins/Runtime/Metasound/Source/MetasoundEngine/Private/MetasoundWavePlayerNode.cpp
+++ b/Engine/Plugins/Runtime/Metasound/Source/MetasoundEngine/Private/MetasoundWavePlayerNode.cpp
@@ -36,8 +36,8 @@ namespace Metasound
 		METASOUND_PARAM(InputLoopDuration, "Loop Duration", "The duration of the loop when wave player is enabled for looping. A negative value will loop the whole wave asset.")
 		
 		METASOUND_PARAM(OutputTriggerOnPlay, "On Play", "Triggers when Play is triggered.")
-		METASOUND_PARAM(OutputTriggerOnDone, "On Finished", "Triggers when the wave played has finished playing.")
-		METASOUND_PARAM(OutputTriggerOnNearlyDone, "On Nearly Finished", "Triggers when the wave played has almost finished playing (the block before it finishes). Allows time for logic to trigger different variations to play seamlessly. playing.")
+		METASOUND_PARAM(OutputTriggerOnDone, "On Finished", "Triggers when the wave played has finished playing or Stop is triggered.")
+		METASOUND_PARAM(OutputTriggerOnNearlyDone, "On Nearly Finished", "Triggers when the wave played has almost finished playing (the block before it finishes). Allows time for logic to trigger different variations to play seamlessly.")
 		METASOUND_PARAM(OutputTriggerOnLooped, "On Looped", "Triggers when the wave player has looped.")
 		METASOUND_PARAM(OutputTriggerOnCuePoint, "On Cue Point", "Triggers when a wave cue point was hit during playback.")
 		METASOUND_PARAM(OutputCuePointID, "Cue Point ID", "The cue point ID that was triggered.")
diff --git a/Engine/Plugins/Runtime/Metasound/Source/MetasoundEngine/Private/MetasoundWaveProxyReader.cpp b/Engine/Plugins/Runtime/Metasound/Source/MetasoundEngine/Private/MetasoundWaveProxyReader.cpp
--- a/Engine/Plugins/Runtime/Metasound/Source/MetasoundEngine/Private/MetasoundWaveProxyReader.cpp
+++ b/Engine/Plugins/Runtime/Metasound/Source/MetasoundEngine/Private/MetasoundWaveProxyReader.cpp
@@ -521,7 +521,7 @@ namespace Metasound
 			const int32 MinLoopDurationInFrames = FMath::CeilToInt(MinLoopDurationInSeconds * SampleRate);
 			const float LoopEndTime = Settings.LoopStartTimeInSeconds + Settings.LoopDurationInSeconds;
 			const int32 MinLoopEndFrameIndex = LoopStartFrameIndex + MinLoopDurationInFrames;
-			const int32 MaxLoopEndFrameIndex = NumFramesInWave + 1;
+			const int32 MaxLoopEndFrameIndex = NumFramesInWave;
 
 			LoopStartFrameIndex = FMath::FloorToInt(Settings.LoopStartTimeInSeconds * SampleRate);
 			LoopEndFrameIndex = FMath::Clamp(LoopEndTime * SampleRate, MinLoopEndFrameIndex, MaxLoopEndFrameIndex);
@@ -529,7 +529,7 @@ namespace Metasound
 		else
 		{
 			LoopStartFrameIndex = 0;
-			LoopEndFrameIndex = NumFramesInWave + 1;
+			LoopEndFrameIndex = NumFramesInWave;
 		}
 	}
 }
