[Metasounds] Inputs not constructing TInputReceiverOperators intermittently causing parameter transmission to fail
#rnx

diff --git a/Engine/Plugins/Runtime/Metasound/Source/MetasoundFrontend/Public/MetasoundInputNode.h b/Engine/Plugins/Runtime/Metasound/Source/MetasoundFrontend/Public/MetasoundInputNode.h
--- a/Engine/Plugins/Runtime/Metasound/Source/MetasoundFrontend/Public/MetasoundInputNode.h
+++ b/Engine/Plugins/Runtime/Metasound/Source/MetasoundFrontend/Public/MetasoundInputNode.h
@@ -209,8 +209,7 @@ namespace Metasound
 		{
 			DataType& OutputData = *Super::OutputValue;
 
-			bool bHasNewData = false;
-			bHasNewData = Receiver->CanPop();
+			bool bHasNewData = Receiver->CanPop();
 			if (bHasNewData)
 			{
 				Receiver->Pop(OutputData);
@@ -453,14 +452,6 @@ namespace Metasound
 				const FInputNodeType& InputNode = static_cast<const FInputNodeType&>(InParams.Node);
 				const FVertexName& VertexKey = InputNode.GetVertexName();
 
-				if (const FAnyDataReference* Ref = InParams.InputData.FindDataReference(VertexKey))
-				{
-					// Pass through input value
-					return MakeUnique<FPassThroughInputOperatorType>(VertexKey, ReferenceCreator->CreateDataReference(*Ref));
-				}
-
-				FDataReference DataRef = ReferenceCreator->CreateDataReference(InParams.OperatorSettings);
-
 				if (bEnableTransmission)
 				{
 					const FName DataTypeName = GetMetasoundDataTypeName<DataType>();
@@ -469,14 +460,30 @@ namespace Metasound
 					if (Receiver.IsValid())
 					{
 						// Transmittable input value
-						return MakeUnique<TInputReceiverOperator<DataType>>(VertexKey, DataRef, MoveTemp(SendAddress), MoveTemp(Receiver));
+						if (const FAnyDataReference* Ref = InParams.InputData.FindDataReference(VertexKey))
+						{
+							return MakeUnique<TInputReceiverOperator<DataType>>(VertexKey, ReferenceCreator->CreateDataReference(*Ref), MoveTemp(SendAddress), MoveTemp(Receiver));
+						}
+						else
+						{
+							FDataReference DataRef = ReferenceCreator->CreateDataReference(InParams.OperatorSettings);
+							return MakeUnique<TInputReceiverOperator<DataType>>(VertexKey, DataRef, MoveTemp(SendAddress), MoveTemp(Receiver));
+						}
 					}
 
 					AddBuildError<FInputReceiverInitializationError>(OutResults.Errors, InParams.Node, VertexKey, DataTypeName);
 					return nullptr;
 				}
 
+				if (const FAnyDataReference* Ref = InParams.InputData.FindDataReference(VertexKey))
+				{
+					// Pass through input value
+					return MakeUnique<FPassThroughInputOperatorType>(VertexKey, ReferenceCreator->CreateDataReference(*Ref));
+				}
+
+
 				// Owned input value
+				FDataReference DataRef = ReferenceCreator->CreateDataReference(InParams.OperatorSettings);
 				return MakeUnique<FOwnedInputOperatorType>(VertexKey, DataRef);
 			}
 
