Replace use of FSoftObjectPath with UPROPERTY for dependent metasound assets.

diff --git a/Engine/Plugins/Runtime/Metasound/Source/MetasoundEditor/Private/MetasoundEditorGraphBuilder.cpp b/Engine/Plugins/Runtime/Metasound/Source/MetasoundEditor/Private/MetasoundEditorGraphBuilder.cpp
--- a/Engine/Plugins/Runtime/Metasound/Source/MetasoundEditor/Private/MetasoundEditorGraphBuilder.cpp
+++ b/Engine/Plugins/Runtime/Metasound/Source/MetasoundEditor/Private/MetasoundEditorGraphBuilder.cpp
@@ -454,7 +454,7 @@ namespace Metasound
 				NewGraphNode->SetNodeLocation(InLocation);
 
 				// Adding external node may introduce referenced asset so rebuild referenced keys.
-				MetaSoundAsset->RebuildReferencedAssetClassKeys();
+				MetaSoundAsset->RebuildReferencedAssetClasses();
 			}
 
 			return NewGraphNode;
@@ -1769,7 +1769,7 @@ namespace Metasound
 					{
 						if (FMetasoundAssetBase* EditedMetaSound = IMetasoundUObjectRegistry::Get().GetObjectAsAssetBase(Asset))
 						{
-							EditedMetaSound->RebuildReferencedAssetClassKeys();
+							EditedMetaSound->RebuildReferencedAssetClasses();
 							if (EditedMetaSound->IsReferencedAsset(*MetaSoundAsset))
 							{
 								EditedReferencingMetaSounds.Add(EditedMetaSound);
@@ -1816,7 +1816,7 @@ namespace Metasound
 					{
 						if (FMetasoundAssetBase* EditedMetaSound = IMetasoundUObjectRegistry::Get().GetObjectAsAssetBase(Asset))
 						{
-							EditedMetaSound->RebuildReferencedAssetClassKeys();
+							EditedMetaSound->RebuildReferencedAssetClasses();
 							if (EditedMetaSound->IsReferencedAsset(*MetaSoundAsset))
 							{
 								EditedMetaSound->GetModifyContext().SetDocumentModified();
diff --git a/Engine/Plugins/Runtime/Metasound/Source/MetasoundEngine/Private/Metasound.cpp b/Engine/Plugins/Runtime/Metasound/Source/MetasoundEngine/Private/Metasound.cpp
--- a/Engine/Plugins/Runtime/Metasound/Source/MetasoundEngine/Private/Metasound.cpp
+++ b/Engine/Plugins/Runtime/Metasound/Source/MetasoundEngine/Private/Metasound.cpp
@@ -7,6 +7,7 @@
 #include "MetasoundAssetBase.h"
 #include "MetasoundAudioFormats.h"
 #include "MetasoundEngineArchetypes.h"
+#include "MetasoundEngineAsset.h"
 #include "MetasoundEngineEnvironment.h"
 #include "MetasoundEnvironment.h"
 #include "MetasoundFrontendController.h"
@@ -73,7 +74,7 @@ void UMetaSoundPatch::PostDuplicate(EDuplicateMode::Type InDuplicateMode)
 void UMetaSoundPatch::PostEditUndo()
 {
 	Super::PostEditUndo();
-	Metasound::PostEditUndo(*this);
+	Metasound::FMetaSoundEngineAssetHelper::PostEditUndo(*this);
 }
 #endif // WITHEDITOR
 
@@ -86,13 +87,13 @@ void UMetaSoundPatch::BeginDestroy()
 void UMetaSoundPatch::PreSave(FObjectPreSaveContext InSaveContext)
 {
 	Super::PreSave(InSaveContext);
-	Metasound::PreSaveAsset(*this, InSaveContext);
+	Metasound::FMetaSoundEngineAssetHelper::PreSaveAsset(*this, InSaveContext);
 }
 
 void UMetaSoundPatch::Serialize(FArchive& InArchive)
 {
 	Super::Serialize(InArchive);
-	Metasound::SerializeToArchive(*this, InArchive);
+	Metasound::FMetaSoundEngineAssetHelper::SerializeToArchive(*this, InArchive);
 }
 
 #if WITH_EDITORONLY_DATA
@@ -124,30 +125,45 @@ FText UMetaSoundPatch::GetDisplayName() const
 	return FMetasoundAssetBase::GetDisplayName(MoveTemp(TypeName));
 }
 
+
 void UMetaSoundPatch::SetRegistryAssetClassInfo(const Metasound::Frontend::FNodeClassInfo& InNodeInfo)
 {
-	Metasound::SetMetaSoundRegistryAssetClassInfo(*this, InNodeInfo);
+	Metasound::FMetaSoundEngineAssetHelper::SetMetaSoundRegistryAssetClassInfo(*this, InNodeInfo);
 }
 #endif // WITH_EDITORONLY_DATA
 
+void UMetaSoundPatch::PostLoad() 
+{
+	Super::PostLoad();
+	Metasound::FMetaSoundEngineAssetHelper::PostLoad(*this);
+}
+
 Metasound::Frontend::FNodeClassInfo UMetaSoundPatch::GetAssetClassInfo() const
 {
 	return { GetDocumentChecked().RootGraph, FSoftObjectPath(this) };
 }
 
-void UMetaSoundPatch::SetReferencedAssetClassKeys(TSet<Metasound::Frontend::FNodeRegistryKey>&& InKeys)
+#if WITH_EDITOR
+void UMetaSoundPatch::SetReferencedAssetClasses(TSet<Metasound::Frontend::IMetaSoundAssetManager::FAssetInfo>&& InAssetClasses)
+{
+	Metasound::FMetaSoundEngineAssetHelper::SetReferencedAssetClasses(*this, MoveTemp(InAssetClasses));
+}
+#endif
+
+TArray<FMetasoundAssetBase*> UMetaSoundPatch::GetReferencedAssets()
 {
-	ReferencedAssetClassKeys = MoveTemp(InKeys);
+	return Metasound::FMetaSoundEngineAssetHelper::GetReferencedAssets(*this);
 }
 
-TSet<FSoftObjectPath>& UMetaSoundPatch::GetReferencedAssetClassCache()
+const TSet<FSoftObjectPath>& UMetaSoundPatch::GetAsyncReferencedAssetClassPaths() const 
 {
 	return ReferenceAssetClassCache;
 }
 
-const TSet<FSoftObjectPath>& UMetaSoundPatch::GetReferencedAssetClassCache() const
+void UMetaSoundPatch::OnAsyncReferencedAssetsLoaded(const TArray<FMetasoundAssetBase*>& InAsyncReferences)
 {
-	return ReferenceAssetClassCache;
+	Metasound::FMetaSoundEngineAssetHelper::OnAsyncReferencedAssetsLoaded(*this, InAsyncReferences);
 }
+
 #undef LOCTEXT_NAMESPACE // MetaSound
 
diff --git a/Engine/Plugins/Runtime/Metasound/Source/MetasoundEngine/Private/MetasoundAssetSubsystem.cpp b/Engine/Plugins/Runtime/Metasound/Source/MetasoundEngine/Private/MetasoundAssetSubsystem.cpp
--- a/Engine/Plugins/Runtime/Metasound/Source/MetasoundEngine/Private/MetasoundAssetSubsystem.cpp
+++ b/Engine/Plugins/Runtime/Metasound/Source/MetasoundEngine/Private/MetasoundAssetSubsystem.cpp
@@ -110,8 +110,11 @@ void UMetaSoundAssetSubsystem::PostInitAssetScan()
 			AddOrUpdateAsset(AssetData);
 		});
 	}
+
+	bIsInitialAssetScanComplete = true;
 }
 
+#if WITH_EDITORONLY_DATA
 void UMetaSoundAssetSubsystem::AddAssetReferences(FMetasoundAssetBase& InAssetBase)
 {
 	using namespace Metasound;
@@ -126,52 +129,53 @@ void UMetaSoundAssetSubsystem::AddAssetReferences(FMetasoundAssetBase& InAssetBa
 		UE_LOG(LogMetaSound, Verbose, TEXT("Adding asset '%s' to MetaSoundAsset registry."), *InAssetBase.GetOwningAssetName());
 	}
 
-	bool bLoadFromPathCache = false;
+	bool bAddFromReferencedAssets = false;
 	const TSet<FString>& ReferencedAssetClassKeys = InAssetBase.GetReferencedAssetClassKeys();
 	for (const FString& ReferencedAssetClassKey : ReferencedAssetClassKeys)
 	{
 		if (!ContainsKey(ReferencedAssetClassKey))
 		{
 			UE_LOG(LogMetaSound, Verbose, TEXT("Missing referenced class '%s' asset entry."), *ReferencedAssetClassKey);
-			bLoadFromPathCache = true;
+			bAddFromReferencedAssets = true;
 		}
 	}
 
 	// All keys are loaded
-	if (!bLoadFromPathCache)
+	if (!bAddFromReferencedAssets)
 	{
 		return;
 	}
 
 	UE_LOG(LogMetaSound, Verbose, TEXT("Attempting preemptive reference load..."));
 
-	// If keys are not loaded, iterate class cache paths to prime asset manager with
-	// hint paths as the registration is getting called before asset manager has 
-	// completed initial scan.
-	const TSet<FSoftObjectPath>& AssetClassCache = InAssetBase.GetReferencedAssetClassCache();
-	for (const FSoftObjectPath& AssetClassPath : AssetClassCache)
+	TArray<FMetasoundAssetBase*> ReferencedAssets = InAssetBase.GetReferencedAssets();
+	for (FMetasoundAssetBase* Asset : ReferencedAssets)
 	{
-		if (FMetasoundAssetBase* MetaSoundAsset = TryLoadAsset(AssetClassPath))
+		if (Asset)
 		{
-			FNodeClassInfo ClassInfo = MetaSoundAsset->GetAssetClassInfo();
+			FNodeClassInfo ClassInfo = Asset->GetAssetClassInfo();
 			const FNodeRegistryKey ClassKey = NodeRegistryKey::CreateKey(ClassInfo);
 			if (!ContainsKey(ClassKey))
 			{
 				UE_LOG(LogMetaSound, Verbose,
-					TEXT("Preemptive load of class '%s' from hint path '%s' due to early "
+					TEXT("Preemptive load of class '%s' due to early "
 						"registration request (asset scan likely not complete)."),
-					*ClassKey,
-					*AssetClassPath.ToString());
+					*ClassKey);
 
-				UObject* MetaSoundObject = MetaSoundAsset->GetOwningAsset();
+				UObject* MetaSoundObject = Asset->GetOwningAsset();
 				if (ensureAlways(MetaSoundObject))
 				{
 					AddOrUpdateAsset(*MetaSoundObject);
 				}
 			}
 		}
+		else
+		{
+			UE_LOG(LogMetaSound, Warning, TEXT("Null referenced dependent asset in %s. Resaving asset in editor may fix the issue"), *InAssetBase.GetOwningAssetName());
+		}
 	}
 }
+#endif
 
 Metasound::Frontend::FNodeRegistryKey UMetaSoundAssetSubsystem::AddOrUpdateAsset(const UObject& InObject)
 {
@@ -299,23 +303,30 @@ void UMetaSoundAssetSubsystem::RebuildDenyListCache(const UAssetManager& InAsset
 	AutoUpdateDenyListChangeID = Settings->DenyListCacheChangeID;
 }
 
-TSet<Metasound::Frontend::FNodeRegistryKey> UMetaSoundAssetSubsystem::GetReferencedKeys(const FMetasoundAssetBase& InAssetBase) const
+#if WITH_EDITOR
+TSet<UMetaSoundAssetSubsystem::FAssetInfo> UMetaSoundAssetSubsystem::GetReferencedAssetClasses(const FMetasoundAssetBase& InAssetBase) const
 {
-	METASOUND_TRACE_CPUPROFILER_EVENT_SCOPE(UMetaSoundAssetSubsystem::GetReferencedKeys);
+	METASOUND_TRACE_CPUPROFILER_EVENT_SCOPE(UMetaSoundAssetSubsystem::GetReferencedAssetClasses);
 	using namespace Metasound::Frontend;
 
-	TSet<FNodeRegistryKey> OutKeys;
+	if (!bIsInitialAssetScanComplete)
+	{
+		UE_LOG(LogMetaSound, Warning, TEXT("Attempt to get registered dependent assets for %s before asset scan is complete may result in missed dependencies"), *InAssetBase.GetOwningAssetName());
+	}
+
+	TSet<FAssetInfo> OutAssetInfos;
 	const FMetasoundFrontendDocument& Document = InAssetBase.GetDocumentChecked();
 	for (const FMetasoundFrontendClass& Class : Document.Dependencies)
 	{
 		const FNodeRegistryKey Key = NodeRegistryKey::CreateKey(Class.Metadata);
-		if (ContainsKey(Key))
+		if (const FSoftObjectPath* ObjectPath = PathMap.Find(Key))
 		{
-			OutKeys.Add(Key);
+			OutAssetInfos.Add(FAssetInfo{Key, *ObjectPath});
 		}
 	}
-	return OutKeys;
+	return MoveTemp(OutAssetInfos);
 }
+#endif
 
 void UMetaSoundAssetSubsystem::RescanAutoUpdateDenyList()
 {
@@ -360,6 +371,152 @@ bool UMetaSoundAssetSubsystem::TryLoadReferencedAssets(const FMetasoundAssetBase
 	return bSucceeded;
 }
 
+void UMetaSoundAssetSubsystem::RequestAsyncLoadReferencedAssets(FMetasoundAssetBase& InAssetBase)
+{
+	const TSet<FSoftObjectPath>& AsyncReferences = InAssetBase.GetAsyncReferencedAssetClassPaths();
+	if (AsyncReferences.Num() > 0)
+	{
+		if (UObject* OwningAsset = InAssetBase.GetOwningAsset())
+		{
+			TArray<FSoftObjectPath> PathsToLoad = AsyncReferences.Array();
+
+			// Protect against duplicate calls to async load assets. 
+			if (FMetaSoundAsyncAssetDependencies* ExistingAsyncLoad = FindLoadingDependencies(OwningAsset))
+			{
+				if (ExistingAsyncLoad->Dependencies == PathsToLoad)
+				{
+					// early out since these are already actively being loaded.
+					return;
+				}
+			}
+
+			int32 AsyncLoadID = AsyncLoadIDCounter++;
+
+			auto AssetsLoadedDelegate = [this, AsyncLoadID]()
+			{
+				this->OnAssetsLoaded(AsyncLoadID);
+			};
+
+			// Store async loading data for use when async load is complete. 
+			FMetaSoundAsyncAssetDependencies& AsyncDependencies = LoadingDependencies.AddDefaulted_GetRef();
+
+			AsyncDependencies.LoadID = AsyncLoadID;
+			AsyncDependencies.MetaSound = OwningAsset;
+			AsyncDependencies.Dependencies = PathsToLoad;
+			AsyncDependencies.StreamableHandle = StreamableManager.RequestAsyncLoad(PathsToLoad, AssetsLoadedDelegate);
+		}
+		else
+		{
+			UE_LOG(LogMetaSound, Error, TEXT("Cannot load async asset as FMetasoundAssetBase null owning UObject"), *InAssetBase.GetOwningAssetName());
+		}
+	}
+
+}
+
+void UMetaSoundAssetSubsystem::WaitUntilAsyncLoadReferencedAssetsComplete(FMetasoundAssetBase& InAssetBase)
+{
+	UObject* OwningAsset = InAssetBase.GetOwningAsset();
+	if (OwningAsset)
+	{
+		while (FMetaSoundAsyncAssetDependencies* LoadingDependency = FindLoadingDependencies(OwningAsset))
+		{
+			// Grab shared ptr to handle as LoadingDependencies may be deleted and have it's shared pointer removed. 
+			TSharedPtr<FStreamableHandle> StreamableHandle = LoadingDependency->StreamableHandle;
+			if (StreamableHandle.IsValid())
+			{
+				UE_LOG(LogMetaSound, Verbose, TEXT("Waiting on async load (id: %d) from asset %s"), LoadingDependency->LoadID, *InAssetBase.GetOwningAssetName());
+
+				EAsyncPackageState::Type LoadState = StreamableHandle->WaitUntilComplete();
+				if (EAsyncPackageState::Complete != LoadState)
+				{
+					UE_LOG(LogMetaSound, Error, TEXT("Failed to complete loading of async dependent assets from parent asset %s"), *InAssetBase.GetOwningAssetName());
+					RemoveLoadingDependencies(LoadingDependency->LoadID);
+				}
+				else
+				{
+					// This will remove the loading dependencies from internal storage
+					OnAssetsLoaded(LoadingDependency->LoadID);
+				}
+
+				// This will prevent OnAssetsLoaded from being called via the streamables
+				// internal delegate complete callback.
+				StreamableHandle->CancelHandle();
+			}
+		}
+	}
+}
+
+FMetaSoundAsyncAssetDependencies* UMetaSoundAssetSubsystem::FindLoadingDependencies(const UObject* InParentAsset)
+{
+	auto IsEqualMetaSoundUObject = [InParentAsset](const FMetaSoundAsyncAssetDependencies& InDependencies) -> bool
+	{
+		return (InDependencies.MetaSound == InParentAsset);
+	}; 
+
+	return LoadingDependencies.FindByPredicate(IsEqualMetaSoundUObject);
+}
+
+FMetaSoundAsyncAssetDependencies* UMetaSoundAssetSubsystem::FindLoadingDependencies(int32 InLoadID)
+{
+	auto IsEqualID = [InLoadID](const FMetaSoundAsyncAssetDependencies& InDependencies) -> bool
+	{
+		return (InDependencies.LoadID == InLoadID);
+	};
+	
+	return LoadingDependencies.FindByPredicate(IsEqualID);
+}
+
+void UMetaSoundAssetSubsystem::RemoveLoadingDependencies(int32 InLoadID)
+{
+	auto IsEqualID = [InLoadID](const FMetaSoundAsyncAssetDependencies& InDependencies) -> bool
+	{
+		return (InDependencies.LoadID == InLoadID);
+	};
+	LoadingDependencies.RemoveAllSwap(IsEqualID);
+}
+
+void UMetaSoundAssetSubsystem::OnAssetsLoaded(int32 InLoadID)
+{
+	FMetaSoundAsyncAssetDependencies* LoadedDependencies = FindLoadingDependencies(InLoadID);
+	if (ensureMsgf(LoadedDependencies, TEXT("Call to async asset load complete with invalid IDs %d"), InLoadID))
+	{
+		if (LoadedDependencies->StreamableHandle.IsValid())
+		{
+			if (LoadedDependencies->MetaSound)
+			{
+				Metasound::IMetasoundUObjectRegistry& UObjectRegistry = Metasound::IMetasoundUObjectRegistry::Get();
+				FMetasoundAssetBase* ParentAssetBase = UObjectRegistry.GetObjectAsAssetBase(LoadedDependencies->MetaSound);
+				if (ensureMsgf(ParentAssetBase, TEXT("UClass of Parent MetaSound asset %s is not registered in metasound UObject Registery"), *LoadedDependencies->MetaSound->GetPathName()))
+				{
+					// Get all async loaded assets
+					TArray<UObject*> LoadedAssets;
+					LoadedDependencies->StreamableHandle->GetLoadedAssets(LoadedAssets);
+
+					// Cast UObjects to FMetaSoundAssetBase
+					TArray<FMetasoundAssetBase*> LoadedAssetBases;
+					for (UObject* AssetDependency : LoadedAssets)
+					{
+						if (AssetDependency)
+						{
+							FMetasoundAssetBase* AssetDependencyBase = UObjectRegistry.GetObjectAsAssetBase(AssetDependency);
+							if (ensure(AssetDependencyBase))
+							{
+								LoadedAssetBases.Add(AssetDependencyBase);
+							}
+						}
+					}
+
+					// Update parent asset with loaded assets. 
+					ParentAssetBase->OnAsyncReferencedAssetsLoaded(LoadedAssetBases);
+				}
+			}
+		}
+
+		// Remove from active array of loading dependencies.
+		RemoveLoadingDependencies(InLoadID);
+	}
+}
+
 const FSoftObjectPath* UMetaSoundAssetSubsystem::FindObjectPathFromKey(const Metasound::Frontend::FNodeRegistryKey& InRegistryKey) const
 {
 	return PathMap.Find(InRegistryKey);
diff --git a/Engine/Plugins/Runtime/Metasound/Source/MetasoundEngine/Private/MetasoundEngineAsset.h b/Engine/Plugins/Runtime/Metasound/Source/MetasoundEngine/Private/MetasoundEngineAsset.h
new file mode 100644
--- /dev/null
+++ b/Engine/Plugins/Runtime/Metasound/Source/MetasoundEngine/Private/MetasoundEngineAsset.h
@@ -0,0 +1,183 @@
+// Copyright Epic Games, Inc. All Rights Reserved.
+
+#pragma once
+
+#include "Metasound.h"
+#include "MetasoundAssetManager.h"
+#include "MetasoundUObjectRegistry.h"
+#include "Serialization/Archive.h"
+
+#if WITH_EDITORONLY_DATA
+#include "MetasoundFrontendRegistries.h"
+#include "Algo/Transform.h"
+#endif // WITH_EDITORONLY_DATA
+
+
+namespace Metasound
+{
+	/** MetaSound Engine Asset helper provides routines for UObject based MetaSound assets. 
+	 * Any UObject deriving from FMetaSoundAssetBase should use these helper functions
+	 * in their UObject overrides. 
+	 */
+	struct FMetaSoundEngineAssetHelper
+	{
+#if WITH_EDITOR
+		template <typename TMetaSoundObject>
+		static void PostEditUndo(TMetaSoundObject& InMetaSound)
+		{
+			InMetaSound.GetModifyContext().SetForceRefreshViews();
+			if (UMetasoundEditorGraphBase* Graph = Cast<UMetasoundEditorGraphBase>(InMetaSound.GetGraph()))
+			{
+				Graph->RegisterGraphWithFrontend();
+			}
+		}
+
+		template<typename TMetaSoundObject>
+		static void SetReferencedAssetClasses(TMetaSoundObject& InMetaSound, TSet<Metasound::Frontend::IMetaSoundAssetManager::FAssetInfo>&& InAssetClasses)
+		{
+			using namespace Metasound::Frontend;
+			
+			InMetaSound.ReferencedAssetClassKeys.Reset();
+			InMetaSound.ReferencedAssetClassObjects.Reset();
+
+			for (const IMetaSoundAssetManager::FAssetInfo& AssetClass : InAssetClasses)
+			{
+				InMetaSound.ReferencedAssetClassKeys.Add(AssetClass.RegistryKey);
+				if (UObject* Object = AssetClass.AssetPath.TryLoad())
+				{
+					InMetaSound.ReferencedAssetClassObjects.Add(Object);
+				}
+				else
+				{
+					UE_LOG(LogMetaSound, Error, TEXT("Failed to load referenced asset %s from asset %s"), *AssetClass.AssetPath.ToString(), *InMetaSound.GetPathName());
+				}
+			}
+		}
+#endif // WITH_EDITOR
+
+		template <typename TMetaSoundObject>
+		static TArray<FMetasoundAssetBase*> GetReferencedAssets(TMetaSoundObject& InMetaSound)
+		{
+			TArray<FMetasoundAssetBase*> ReferencedAssets;
+
+			IMetasoundUObjectRegistry& UObjectRegistry = IMetasoundUObjectRegistry::Get();
+
+			for (TObjectPtr<UObject>& Object : InMetaSound.ReferencedAssetClassObjects)
+			{
+				if (FMetasoundAssetBase* Asset = UObjectRegistry.GetObjectAsAssetBase(Object))
+				{
+					ReferencedAssets.Add(Asset);
+				}
+				else
+				{
+					UE_LOG(LogMetaSound, Error, TEXT("Referenced asset \"%s\", referenced from \"%s\", is not convertible to FMetasoundAssetBase"), *Object->GetPathName(), *InMetaSound.GetPathName());
+				}
+			}
+
+			return ReferencedAssets;
+		}
+
+		template <typename TMetaSoundObject>
+		static void PreSaveAsset(TMetaSoundObject& InMetaSound, FObjectPreSaveContext InSaveContext)
+		{
+#if WITH_EDITORONLY_DATA
+			using namespace Metasound::Frontend;
+
+			if (UMetasoundEditorGraphBase* MetaSoundGraph = Cast<UMetasoundEditorGraphBase>(InMetaSound.GetGraph()))
+			{
+				// Cooked data must be deterministic, so do not call register graph as this can
+				// initiate an auto-update and/or local registry data cache and modify serialized data.
+				if (!InSaveContext.IsCooking())
+				{
+					MetaSoundGraph->RegisterGraphWithFrontend();
+					MetaSoundGraph->GetModifyContext().SetForceRefreshViews();
+				}
+			}
+
+			IMetaSoundAssetManager::GetChecked().WaitUntilAsyncLoadReferencedAssetsComplete(InMetaSound);
+#endif // WITH_EDITORONLY_DATA
+		}
+
+		template <typename TMetaSoundObject>
+		static void SerializeToArchive(TMetaSoundObject& InMetaSound, FArchive& InArchive)
+		{
+			if (InArchive.IsLoading())
+			{
+				if (InMetaSound.VersionAsset())
+				{
+#if WITH_EDITORONLY_DATA
+					if (UMetasoundEditorGraphBase* MetaSoundGraph = Cast<UMetasoundEditorGraphBase>(InMetaSound.GetGraph()))
+					{
+						MetaSoundGraph->SetVersionedOnLoad();
+					}
+#endif // WITH_EDITORONLY_DATA
+				}
+
+			}
+		}
+
+		template<typename TMetaSoundObject>
+		static void PostLoad(TMetaSoundObject& InMetaSound)
+		{
+			Frontend::IMetaSoundAssetManager::GetChecked().RequestAsyncLoadReferencedAssets(InMetaSound);
+		}
+
+		template<typename TMetaSoundObject>
+		static void OnAsyncReferencedAssetsLoaded(TMetaSoundObject& InMetaSound, const TArray<FMetasoundAssetBase*>& InAsyncReferences)
+		{
+			for (FMetasoundAssetBase* AssetBase : InAsyncReferences)
+			{
+				if (AssetBase)
+				{
+					if (UObject* OwningAsset = AssetBase->GetOwningAsset())
+					{
+						InMetaSound.ReferencedAssetClassObjects.Add(OwningAsset);
+						InMetaSound.ReferenceAssetClassCache.Remove(FSoftObjectPath(OwningAsset));
+					}
+				}
+			}
+		}
+
+#if WITH_EDITORONLY_DATA
+		template <typename TMetaSoundObject>
+		static void SetMetaSoundRegistryAssetClassInfo(TMetaSoundObject& InMetaSound, const Metasound::Frontend::FNodeClassInfo& InClassInfo)
+		{
+			using namespace Metasound;
+			using namespace Metasound::Frontend;
+
+			check(AssetTags::AssetClassID == GET_MEMBER_NAME_CHECKED(TMetaSoundObject, AssetClassID));
+			check(AssetTags::RegistryInputTypes == GET_MEMBER_NAME_CHECKED(TMetaSoundObject, RegistryInputTypes));
+			check(AssetTags::RegistryOutputTypes == GET_MEMBER_NAME_CHECKED(TMetaSoundObject, RegistryOutputTypes));
+			check(AssetTags::RegistryVersionMajor == GET_MEMBER_NAME_CHECKED(TMetaSoundObject, RegistryVersionMajor));
+			check(AssetTags::RegistryVersionMinor == GET_MEMBER_NAME_CHECKED(TMetaSoundObject, RegistryVersionMinor));
+
+			bool bMarkDirty = InMetaSound.AssetClassID != InClassInfo.AssetClassID;
+			bMarkDirty |= InMetaSound.RegistryVersionMajor != InClassInfo.Version.Major;
+			bMarkDirty |= InMetaSound.RegistryVersionMinor != InClassInfo.Version.Minor;
+
+			InMetaSound.AssetClassID = InClassInfo.AssetClassID;
+			InMetaSound.RegistryVersionMajor = InClassInfo.Version.Major;
+			InMetaSound.RegistryVersionMinor = InClassInfo.Version.Minor;
+
+			{
+				TArray<FString> InputTypes;
+				Algo::Transform(InClassInfo.InputTypes, InputTypes, [](const FName& Name) { return Name.ToString(); });
+
+				const FString TypeString = FString::Join(InputTypes, *AssetTags::ArrayDelim);
+				bMarkDirty |= InMetaSound.RegistryInputTypes != TypeString;
+				InMetaSound.RegistryInputTypes = TypeString;
+			}
+
+			{
+				TArray<FString> OutputTypes;
+				Algo::Transform(InClassInfo.OutputTypes, OutputTypes, [](const FName& Name) { return Name.ToString(); });
+
+				const FString TypeString = FString::Join(OutputTypes, *AssetTags::ArrayDelim);
+				bMarkDirty |= InMetaSound.RegistryOutputTypes != TypeString;
+				InMetaSound.RegistryOutputTypes = TypeString;
+			}
+		}
+#endif // WITH_EDITORONLY_DATA
+	};
+} // namespace Metasound
+
diff --git a/Engine/Plugins/Runtime/Metasound/Source/MetasoundEngine/Private/MetasoundSource.cpp b/Engine/Plugins/Runtime/Metasound/Source/MetasoundEngine/Private/MetasoundSource.cpp
--- a/Engine/Plugins/Runtime/Metasound/Source/MetasoundEngine/Private/MetasoundSource.cpp
+++ b/Engine/Plugins/Runtime/Metasound/Source/MetasoundEngine/Private/MetasoundSource.cpp
@@ -12,6 +12,7 @@
 #include "MetasoundAssetManager.h"
 #include "MetasoundAudioFormats.h"
 #include "MetasoundEngineArchetypes.h"
+#include "MetasoundEngineAsset.h"
 #include "MetasoundEngineEnvironment.h"
 #include "MetasoundEnvironment.h"
 #include "MetasoundFrontendController.h"
@@ -181,7 +182,7 @@ UMetaSoundSource::UMetaSoundSource(const FObjectInitializer& ObjectInitializer)
 void UMetaSoundSource::PostEditUndo()
 {
 	Super::PostEditUndo();
-	Metasound::PostEditUndo(*this);
+	Metasound::FMetaSoundEngineAssetHelper::PostEditUndo(*this);
 }
 
 void UMetaSoundSource::PostDuplicate(EDuplicateMode::Type InDuplicateMode)
@@ -292,31 +293,39 @@ void UMetaSoundSource::BeginDestroy()
 void UMetaSoundSource::PreSave(FObjectPreSaveContext InSaveContext)
 {
 	Super::PreSave(InSaveContext);
-	Metasound::PreSaveAsset(*this, InSaveContext);
+	Metasound::FMetaSoundEngineAssetHelper::PreSaveAsset(*this, InSaveContext);
 }
 
 void UMetaSoundSource::Serialize(FArchive& InArchive)
 {
 	Super::Serialize(InArchive);
-	Metasound::SerializeToArchive(*this, InArchive);
+	Metasound::FMetaSoundEngineAssetHelper::SerializeToArchive(*this, InArchive);
 }
 
-void UMetaSoundSource::SetReferencedAssetClassKeys(TSet<Metasound::Frontend::FNodeRegistryKey>&& InKeys)
+#if WITH_EDITOR
+void UMetaSoundSource::SetReferencedAssetClasses(TSet<Metasound::Frontend::IMetaSoundAssetManager::FAssetInfo>&& InAssetClasses)
 {
-	ReferencedAssetClassKeys = MoveTemp(InKeys);
+	Metasound::FMetaSoundEngineAssetHelper::SetReferencedAssetClasses(*this, MoveTemp(InAssetClasses));
 }
+#endif // WITH_EDITOR
 
-TSet<FSoftObjectPath>& UMetaSoundSource::GetReferencedAssetClassCache()
+TArray<FMetasoundAssetBase*> UMetaSoundSource::GetReferencedAssets()
 {
-	return ReferenceAssetClassCache;
+	return Metasound::FMetaSoundEngineAssetHelper::GetReferencedAssets(*this);
 }
 
-const TSet<FSoftObjectPath>& UMetaSoundSource::GetReferencedAssetClassCache() const
+const TSet<FSoftObjectPath>& UMetaSoundSource::GetAsyncReferencedAssetClassPaths() const 
 {
 	return ReferenceAssetClassCache;
 }
 
+void UMetaSoundSource::OnAsyncReferencedAssetsLoaded(const TArray<FMetasoundAssetBase*>& InAsyncReferences)
+{
+	Metasound::FMetaSoundEngineAssetHelper::OnAsyncReferencedAssetsLoaded(*this, InAsyncReferences);
+}
+
 #if WITH_EDITORONLY_DATA
+
 UEdGraph* UMetaSoundSource::GetGraph()
 {
 	return Graph;
@@ -345,20 +354,22 @@ FText UMetaSoundSource::GetDisplayName() const
 	return FMetasoundAssetBase::GetDisplayName(MoveTemp(TypeName));
 }
 
+
+void UMetaSoundSource::SetRegistryAssetClassInfo(const Metasound::Frontend::FNodeClassInfo& InNodeInfo)
+{
+	Metasound::FMetaSoundEngineAssetHelper::SetMetaSoundRegistryAssetClassInfo(*this, InNodeInfo);
+}
+#endif // WITH_EDITORONLY_DATA
+
 void UMetaSoundSource::PostLoad()
 {
 	Super::PostLoad();
+	Metasound::FMetaSoundEngineAssetHelper::PostLoad(*this);
 
 	Duration = GetDuration();
 	bLooping = IsLooping();
 }
 
-void UMetaSoundSource::SetRegistryAssetClassInfo(const Metasound::Frontend::FNodeClassInfo& InNodeInfo)
-{
-	Metasound::SetMetaSoundRegistryAssetClassInfo(*this, InNodeInfo);
-}
-#endif // WITH_EDITORONLY_DATA
-
 void UMetaSoundSource::InitParameters(TArray<FAudioParameter>& ParametersToInit, FName InFeatureName)
 {
 	using namespace Metasound::SourcePrivate;
diff --git a/Engine/Plugins/Runtime/Metasound/Source/MetasoundEngine/Public/Metasound.h b/Engine/Plugins/Runtime/Metasound/Source/MetasoundEngine/Public/Metasound.h
--- a/Engine/Plugins/Runtime/Metasound/Source/MetasoundEngine/Public/Metasound.h
+++ b/Engine/Plugins/Runtime/Metasound/Source/MetasoundEngine/Public/Metasound.h
@@ -13,10 +13,7 @@
 #include "MetasoundUObjectRegistry.h"
 #include "Serialization/Archive.h"
 #include "UObject/ObjectSaveContext.h"
-
-#if WITH_EDITORONLY_DATA
-#include "Algo/Transform.h"
-#endif // WITH_EDITORONLY_DATA
+#include "UObject/SoftObjectPath.h"
 
 #include "Metasound.generated.h"
 
@@ -46,96 +43,9 @@ public:
 
 namespace Metasound
 {
-#if WITH_EDITOR
-	template <typename TMetaSoundObject>
-	void PostEditUndo(TMetaSoundObject& InMetaSound)
-	{
-		InMetaSound.GetModifyContext().SetForceRefreshViews();
-		if (UMetasoundEditorGraphBase* Graph = Cast<UMetasoundEditorGraphBase>(InMetaSound.GetGraph()))
-		{
-			Graph->RegisterGraphWithFrontend();
-		}
-	}
-#endif // WITH_EDITOR
-
-	template <typename TMetaSoundObject>
-	void PreSaveAsset(TMetaSoundObject& InMetaSound, FObjectPreSaveContext InSaveContext)
-	{
-		using namespace Metasound::Frontend;
-
-#if WITH_EDITORONLY_DATA
-		if (UMetasoundEditorGraphBase* MetaSoundGraph = Cast<UMetasoundEditorGraphBase>(InMetaSound.GetGraph()))
-		{
-			// Cooked data must be deterministic, so do not call register graph as this can
-			// initiate an auto-update and/or local registry data cache and modify serialized data.
-			if (!InSaveContext.IsCooking())
-			{
-				MetaSoundGraph->RegisterGraphWithFrontend();
-				MetaSoundGraph->GetModifyContext().SetForceRefreshViews();
-			}
-		}
-#endif // WITH_EDITORONLY_DATA
-	}
-
-	template <typename TMetaSoundObject>
-	void SerializeToArchive(TMetaSoundObject& InMetaSound, FArchive& InArchive)
-	{
-		if (InArchive.IsLoading())
-		{
-			if (InMetaSound.VersionAsset())
-			{
-#if WITH_EDITORONLY_DATA
-				if (UMetasoundEditorGraphBase* MetaSoundGraph = Cast<UMetasoundEditorGraphBase>(InMetaSound.GetGraph()))
-				{
-					MetaSoundGraph->SetVersionedOnLoad();
-				}
-#endif // WITH_EDITORONLY_DATA
-			}
-		}
-	}
-
-#if WITH_EDITORONLY_DATA
-	template <typename TMetaSoundObject>
-	void SetMetaSoundRegistryAssetClassInfo(TMetaSoundObject& InMetaSound, const Metasound::Frontend::FNodeClassInfo& InClassInfo)
-	{
-		using namespace Metasound;
-		using namespace Metasound::Frontend;
-
-		check(AssetTags::AssetClassID == GET_MEMBER_NAME_CHECKED(TMetaSoundObject, AssetClassID));
-		check(AssetTags::RegistryInputTypes == GET_MEMBER_NAME_CHECKED(TMetaSoundObject, RegistryInputTypes));
-		check(AssetTags::RegistryOutputTypes == GET_MEMBER_NAME_CHECKED(TMetaSoundObject, RegistryOutputTypes));
-		check(AssetTags::RegistryVersionMajor == GET_MEMBER_NAME_CHECKED(TMetaSoundObject, RegistryVersionMajor));
-		check(AssetTags::RegistryVersionMinor == GET_MEMBER_NAME_CHECKED(TMetaSoundObject, RegistryVersionMinor));
-
-		bool bMarkDirty = InMetaSound.AssetClassID != InClassInfo.AssetClassID;
-		bMarkDirty |= InMetaSound.RegistryVersionMajor != InClassInfo.Version.Major;
-		bMarkDirty |= InMetaSound.RegistryVersionMinor != InClassInfo.Version.Minor;
-
-		InMetaSound.AssetClassID = InClassInfo.AssetClassID;
-		InMetaSound.RegistryVersionMajor = InClassInfo.Version.Major;
-		InMetaSound.RegistryVersionMinor = InClassInfo.Version.Minor;
-
-		{
-			TArray<FString> InputTypes;
-			Algo::Transform(InClassInfo.InputTypes, InputTypes, [](const FName& Name) { return Name.ToString(); });
-
-			const FString TypeString = FString::Join(InputTypes, *AssetTags::ArrayDelim);
-			bMarkDirty |= InMetaSound.RegistryInputTypes != TypeString;
-			InMetaSound.RegistryInputTypes = TypeString;
-		}
-
-		{
-			TArray<FString> OutputTypes;
-			Algo::Transform(InClassInfo.OutputTypes, OutputTypes, [](const FName& Name) { return Name.ToString(); });
-
-			const FString TypeString = FString::Join(OutputTypes, *AssetTags::ArrayDelim);
-			bMarkDirty |= InMetaSound.RegistryOutputTypes != TypeString;
-			InMetaSound.RegistryOutputTypes = TypeString;
-		}
-	}
-#endif // WITH_EDITORONLY_DATA
-} // namespace Metasound
-
+	// Forward declare
+	struct FMetaSoundEngineAssetHelper;
+}
 
 /**
  * This asset type is used for Metasound assets that can only be used as nodes in other Metasound graphs.
@@ -146,10 +56,20 @@ class METASOUNDENGINE_API UMetaSoundPatch : public UObject, public FMetasoundAss
 {
 	GENERATED_BODY()
 
+	friend struct Metasound::FMetaSoundEngineAssetHelper;
 protected:
 	UPROPERTY(EditAnywhere, Category = CustomView)
 	FMetasoundFrontendDocument RootMetaSoundDocument;
 
+	UPROPERTY()
+	TSet<FString> ReferencedAssetClassKeys;
+
+	UPROPERTY()
+	TSet<TObjectPtr<UObject>> ReferencedAssetClassObjects;
+
+	UPROPERTY()
+	TSet<FSoftObjectPath> ReferenceAssetClassCache;
+
 #if WITH_EDITORONLY_DATA
 	UPROPERTY()
 	TObjectPtr<UMetasoundEditorGraphBase> Graph;
@@ -158,12 +78,6 @@ protected:
 public:
 	UMetaSoundPatch(const FObjectInitializer& ObjectInitializer);
 
-	UPROPERTY()
-	TSet<FString> ReferencedAssetClassKeys;
-
-	UPROPERTY()
-	TSet<FSoftObjectPath> ReferenceAssetClassCache;
-
 	UPROPERTY(AssetRegistrySearchable)
 	FGuid AssetClassID;
 
@@ -208,8 +122,10 @@ public:
 	{
 		Graph = CastChecked<UMetasoundEditorGraphBase>(InGraph);
 	}
+
 #endif // #if WITH_EDITORONLY_DATA
 
+
 #if WITH_EDITOR
 	virtual void PostDuplicate(EDuplicateMode::Type DuplicateMode) override;
 	virtual void PostEditUndo() override;
@@ -218,20 +134,20 @@ public:
 	virtual void BeginDestroy() override;
 	virtual void PreSave(FObjectPreSaveContext InSaveContext) override;
 	virtual void Serialize(FArchive& InArchive) override;
-
-	virtual TSet<FSoftObjectPath>& GetReferencedAssetClassCache() override;
-	virtual const TSet<FSoftObjectPath>& GetReferencedAssetClassCache() const override;
+	virtual void PostLoad() override;
 
 	// Returns Asset Metadata associated with this MetaSound
 	virtual Metasound::Frontend::FNodeClassInfo GetAssetClassInfo() const override;
 
-
 	virtual bool ConformObjectDataToInterfaces() override { return false; }
 
 	virtual const TSet<FString>& GetReferencedAssetClassKeys() const override
 	{
 		return ReferencedAssetClassKeys;
 	}
+	virtual TArray<FMetasoundAssetBase*> GetReferencedAssets() override;
+	virtual const TSet<FSoftObjectPath>& GetAsyncReferencedAssetClassPaths() const override;
+	virtual void OnAsyncReferencedAssetsLoaded(const TArray<FMetasoundAssetBase*>& InAsyncReferences) override;
 
 	UObject* GetOwningAsset() override
 	{
@@ -244,7 +160,9 @@ public:
 	}
 
 protected:
-	virtual void SetReferencedAssetClassKeys(TSet<Metasound::Frontend::FNodeRegistryKey>&& InKeys) override;
+#if WITH_EDITOR
+	virtual void SetReferencedAssetClasses(TSet<Metasound::Frontend::IMetaSoundAssetManager::FAssetInfo>&& InAssetClasses) override;
+#endif // #if WITH_EDITOR
 
 	Metasound::Frontend::FDocumentAccessPtr GetDocument() override
 	{
diff --git a/Engine/Plugins/Runtime/Metasound/Source/MetasoundEngine/Public/MetasoundAssetSubsystem.h b/Engine/Plugins/Runtime/Metasound/Source/MetasoundEngine/Public/MetasoundAssetSubsystem.h
--- a/Engine/Plugins/Runtime/Metasound/Source/MetasoundEngine/Public/MetasoundAssetSubsystem.h
+++ b/Engine/Plugins/Runtime/Metasound/Source/MetasoundEngine/Public/MetasoundAssetSubsystem.h
@@ -2,6 +2,7 @@
 #pragma once
 
 #include "Engine/Engine.h"
+#include "Engine/StreamableManager.h"
 #include "MetasoundAssetBase.h"
 #include "MetasoundAssetManager.h"
 #include "Subsystems/EngineSubsystem.h"
@@ -25,6 +26,26 @@ struct METASOUNDENGINE_API FMetaSoundAssetDirectory
 	FDirectoryPath Directory;
 };
 
+/** Contains info of assets which are currently async loading. */
+USTRUCT()
+struct FMetaSoundAsyncAssetDependencies
+{
+	GENERATED_BODY()
+
+	// ID of the async load
+	int32 LoadID = 0;
+
+	// Parent MetaSound 
+	UPROPERTY(Transient)
+	TObjectPtr<UObject> MetaSound;
+
+	// Dependencies of parent MetaSound
+	TArray<FSoftObjectPath> Dependencies;
+
+	// Handle to in-flight streaming request
+	TSharedPtr<FStreamableHandle> StreamableHandle;
+};
+
 /** The subsystem in charge of the MetaSound asset registry */
 UCLASS()
 class METASOUNDENGINE_API UMetaSoundAssetSubsystem : public UEngineSubsystem, public Metasound::Frontend::IMetaSoundAssetManager
@@ -32,6 +53,8 @@ class METASOUNDENGINE_API UMetaSoundAssetSubsystem : public UEngineSubsystem, pu
 	GENERATED_BODY()
 
 public:
+	using FAssetInfo = Metasound::Frontend::IMetaSoundAssetManager::FAssetInfo;
+
 	virtual void Initialize(FSubsystemCollectionBase& InCollection) override;
 
 	Metasound::Frontend::FNodeRegistryKey AddOrUpdateAsset(const FAssetData& InAssetData);
@@ -39,16 +62,22 @@ public:
 	void RemoveAsset(const FAssetData& InAssetData);
 	void RenameAsset(const FAssetData& InAssetData, bool bInReregisterWithFrontend = true);
 
+#if WITH_EDITORONLY_DATA
 	virtual void AddAssetReferences(FMetasoundAssetBase& InAssetBase) override;
+#endif
 	virtual Metasound::Frontend::FNodeRegistryKey  AddOrUpdateAsset(const UObject& InObject) override;
 	virtual bool CanAutoUpdate(const FMetasoundFrontendClassName& InClassName) const override;
 	virtual bool ContainsKey(const Metasound::Frontend::FNodeRegistryKey& InRegistryKey) const override;
 	virtual const FSoftObjectPath* FindObjectPathFromKey(const Metasound::Frontend::FNodeRegistryKey& RegistryKey) const override;
-	virtual TSet<Metasound::Frontend::FNodeRegistryKey> GetReferencedKeys(const FMetasoundAssetBase& InAssetBase) const override;
+#if WITH_EDITOR
+	virtual TSet<FAssetInfo> GetReferencedAssetClasses(const FMetasoundAssetBase& InAssetBase) const override;
+#endif
 	virtual void RescanAutoUpdateDenyList() override;
 	virtual FMetasoundAssetBase* TryLoadAsset(const FSoftObjectPath& InObjectPath) const override;
 	virtual FMetasoundAssetBase* TryLoadAssetFromKey(const Metasound::Frontend::FNodeRegistryKey& RegistryKey) const override;
 	virtual bool TryLoadReferencedAssets(const FMetasoundAssetBase& InAssetBase, TArray<FMetasoundAssetBase*>& OutReferencedAssets) const override;
+	virtual void RequestAsyncLoadReferencedAssets(FMetasoundAssetBase& InAssetBase) override;
+	virtual void WaitUntilAsyncLoadReferencedAssetsComplete(FMetasoundAssetBase& InAssetBase) override;
 
 	UFUNCTION(BlueprintCallable, Category = "MetaSounds|Registration")
 	void RegisterAssetClassesInDirectories(const TArray<FMetaSoundAssetDirectory>& Directories);
@@ -57,6 +86,7 @@ public:
 	void UnregisterAssetClassesInDirectories(const TArray<FMetaSoundAssetDirectory>& Directories);
 
 protected:
+
 	void PostEngineInit();
 	void PostInitAssetScan();
 	void RebuildDenyListCache(const UAssetManager& InAssetManager);
@@ -64,7 +94,19 @@ protected:
 	void SearchAndIterateDirectoryAssets(const TArray<FDirectoryPath>& InDirectories, TFunctionRef<void(const FAssetData&)> InFunction);
 
 private:
+
+	UPROPERTY(Transient)
+	TArray<FMetaSoundAsyncAssetDependencies> LoadingDependencies;
+
+	FMetaSoundAsyncAssetDependencies* FindLoadingDependencies(const UObject* InParentAsset);
+	FMetaSoundAsyncAssetDependencies* FindLoadingDependencies(int32 InLoadID);
+	void RemoveLoadingDependencies(int32 InLoadID);
+	void OnAssetsLoaded(int32 InLoadID);
+
+	FStreamableManager StreamableManager;
+	int32 AsyncLoadIDCounter = 0;
 	int32 AutoUpdateDenyListChangeID = INDEX_NONE;
 	TSet<FName> AutoUpdateDenyListCache;
 	TMap<Metasound::Frontend::FNodeRegistryKey, FSoftObjectPath> PathMap;
+	std::atomic<bool> bIsInitialAssetScanComplete = false;
 };
diff --git a/Engine/Plugins/Runtime/Metasound/Source/MetasoundEngine/Public/MetasoundSource.h b/Engine/Plugins/Runtime/Metasound/Source/MetasoundEngine/Public/MetasoundSource.h
--- a/Engine/Plugins/Runtime/Metasound/Source/MetasoundEngine/Public/MetasoundSource.h
+++ b/Engine/Plugins/Runtime/Metasound/Source/MetasoundEngine/Public/MetasoundSource.h
@@ -17,6 +17,11 @@
 
 #include "MetasoundSource.generated.h"
 
+namespace Metasound
+{
+	// Forward declare
+	struct FMetaSoundEngineAssetHelper;
+}
 
 /** Declares the output audio format of the UMetaSoundSource */
 UENUM()
@@ -31,7 +36,6 @@ enum class EMetasoundSourceAudioFormat : uint8
 	COUNT UMETA(Hidden)
 };
 
-
 /**
  * This Metasound type can be played as an audio source.
  */
@@ -40,6 +44,7 @@ class METASOUNDENGINE_API UMetaSoundSource : public USoundWaveProcedural, public
 {
 	GENERATED_BODY()
 
+	friend struct Metasound::FMetaSoundEngineAssetHelper;
 protected:
 	UPROPERTY(EditAnywhere, Category = CustomView)
 	FMetasoundFrontendDocument RootMetasoundDocument;
@@ -47,6 +52,9 @@ protected:
 	UPROPERTY()
 	TSet<FString> ReferencedAssetClassKeys;
 
+	UPROPERTY()
+	TSet<TObjectPtr<UObject>> ReferencedAssetClassObjects;
+
 	UPROPERTY()
 	TSet<FSoftObjectPath> ReferenceAssetClassCache;
 
@@ -78,9 +86,6 @@ public:
 	UPROPERTY(AssetRegistrySearchable)
 	int32 RegistryVersionMinor = 0;
 
-	//~ Begin UObject Interface.
-	virtual void PostLoad() override;
-	//~ End UObject Interface.
 
 	// Sets Asset Registry Metadata associated with this MetaSoundSource
 	virtual void SetRegistryAssetClassInfo(const Metasound::Frontend::FNodeClassInfo& InNodeInfo) override;
@@ -112,6 +117,7 @@ public:
 	}
 #endif // #if WITH_EDITORONLY_DATA
 
+
 #if WITH_EDITOR
 	virtual void PostEditUndo() override;
 
@@ -143,13 +149,14 @@ public:
 	{
 		return ReferencedAssetClassKeys;
 	}
+	virtual TArray<FMetasoundAssetBase*> GetReferencedAssets() override;
+	virtual const TSet<FSoftObjectPath>& GetAsyncReferencedAssetClassPaths() const override;
+	virtual void OnAsyncReferencedAssetsLoaded(const TArray<FMetasoundAssetBase*>& InAsyncReferences) override;
 
 	virtual void BeginDestroy() override;
 	virtual void PreSave(FObjectPreSaveContext InSaveContext) override;
 	virtual void Serialize(FArchive& Ar) override;
-
-	virtual TSet<FSoftObjectPath>& GetReferencedAssetClassCache() override;
-	virtual const TSet<FSoftObjectPath>& GetReferencedAssetClassCache() const override;
+	virtual void PostLoad() override;
 
 	// Returns Asset Metadata associated with this MetaSoundSource
 	virtual Metasound::Frontend::FNodeClassInfo GetAssetClassInfo() const override;
@@ -182,7 +189,6 @@ public:
 	virtual bool EnableSubmixSendsOnPreview() const override { return true; }
 protected:
 
-
 	Metasound::Frontend::FDocumentAccessPtr GetDocument() override
 	{
 		using namespace Metasound::Frontend;
@@ -202,7 +208,9 @@ protected:
 	/** Gets all the default parameters for this Asset.  */
 	virtual bool GetAllDefaultParameters(TArray<FAudioParameter>& OutParameters) const override;
 
-	virtual void SetReferencedAssetClassKeys(TSet<Metasound::Frontend::FNodeRegistryKey>&& InKeys) override;
+#if WITH_EDITOR
+	virtual void SetReferencedAssetClasses(TSet<Metasound::Frontend::IMetaSoundAssetManager::FAssetInfo>&& InAssetClasses) override;
+#endif // #if WITH_EDITOR
 
 private:
 
diff --git a/Engine/Plugins/Runtime/Metasound/Source/MetasoundFrontend/Private/MetasoundAssetBase.cpp b/Engine/Plugins/Runtime/Metasound/Source/MetasoundFrontend/Private/MetasoundAssetBase.cpp
--- a/Engine/Plugins/Runtime/Metasound/Source/MetasoundFrontend/Private/MetasoundAssetBase.cpp
+++ b/Engine/Plugins/Runtime/Metasound/Source/MetasoundFrontend/Private/MetasoundAssetBase.cpp
@@ -95,6 +95,61 @@ namespace Metasound
 
 				return TransmittableInputs;
 			}
+			
+
+			// Registers node by copying document. Updates to document require re-registration.
+			class FNodeRegistryEntry : public INodeRegistryEntry
+			{
+			public:
+				FNodeRegistryEntry(const FString& InName, TSharedPtr<FMetasoundFrontendDocument> InPreprocessedDoc, const FSoftObjectPath& InAssetPath)
+				: Name(InName)
+				, PreprocessedDoc(InPreprocessedDoc)
+				{
+					// Copy FrontendClass to preserve original document.
+					FrontendClass = PreprocessedDoc->RootGraph;
+					FrontendClass.Metadata.SetType(EMetasoundFrontendClassType::External);
+
+					ClassInfo = FNodeClassInfo(PreprocessedDoc->RootGraph, InAssetPath);
+				}
+
+				virtual ~FNodeRegistryEntry() = default;
+
+				virtual const FNodeClassInfo& GetClassInfo() const override
+				{
+					return ClassInfo;
+				}
+
+				virtual TUniquePtr<INode> CreateNode(const FNodeInitData&) const override
+				{
+					static const TSet<FName> ReferencedGraphTransmissibleInputNames; // Empty as referenced graphs do not support transmission
+					return FFrontendGraphBuilder().CreateGraph(*PreprocessedDoc, ReferencedGraphTransmissibleInputNames, Name);
+				}
+
+				virtual TUniquePtr<INode> CreateNode(FDefaultLiteralNodeConstructorParams&&) const override { return nullptr; }
+				virtual TUniquePtr<INode> CreateNode(FDefaultNamedVertexNodeConstructorParams&&) const override { return nullptr; }
+				virtual TUniquePtr<INode> CreateNode(FDefaultNamedVertexWithLiteralNodeConstructorParams&&) const override { return nullptr; }
+
+				virtual const FMetasoundFrontendClass& GetFrontendClass() const override
+				{
+					return FrontendClass;
+				}
+
+				virtual TUniquePtr<INodeRegistryEntry> Clone() const override
+				{
+					return MakeUnique<FNodeRegistryEntry>(Name, PreprocessedDoc, ClassInfo.AssetPath);
+				}
+
+				virtual bool IsNative() const override
+				{
+					return false;
+				}
+
+			private:
+				FString Name;
+				TSharedPtr<FMetasoundFrontendDocument> PreprocessedDoc;
+				FMetasoundFrontendClass FrontendClass;
+				FNodeClassInfo ClassInfo;
+			};
 		} // namespace AssetBasePrivate
 
 		FAutoConsoleVariableRef CVarMetaSoundBlockRate(
@@ -127,46 +182,26 @@ void FMetasoundAssetBase::RegisterGraphWithFrontend(Metasound::Frontend::FMetaSo
 		}
 	}
 
-	if (InRegistrationOptions.bRegisterDependencies)
-	{
-		// Must be called in case register is called prior to asset scan being completed
-		// or in certain cases where asset is being loaded for first time.
-		IMetaSoundAssetManager::GetChecked().AddAssetReferences(*this);
-	}
 
-	if (InRegistrationOptions.bRebuildReferencedAssetClassKeys)
+#if WITH_EDITOR
+	if (InRegistrationOptions.bRebuildReferencedAssetClasses)
 	{
-		RebuildReferencedAssetClassKeys();
+		RebuildReferencedAssetClasses();
 	}
+#endif
 
 	if (InRegistrationOptions.bRegisterDependencies)
 	{
-		TArray<FMetasoundAssetBase*> References;
-		ensureAlways(IMetaSoundAssetManager::GetChecked().TryLoadReferencedAssets(*this, References));
-
-		GetReferencedAssetClassCache().Reset();
-		for (FMetasoundAssetBase* Reference : References)
-		{
-			if (InRegistrationOptions.bForceReregister || !Reference->IsRegistered())
-			{
-				// TODO: Check for infinite recursion and error if so
-				Reference->RegisterGraphWithFrontend(InRegistrationOptions);
-			}
-
-			if (UObject* RefAsset = Reference->GetOwningAsset())
-			{
-				GetReferencedAssetClassCache().Add(RefAsset);
-			}
-		}
+		RegisterAssetDependencies(InRegistrationOptions);
 	}
+	IMetaSoundAssetManager::GetChecked().AddOrUpdateAsset(*GetOwningAsset());
 
 	// Auto update must be done after all referenced asset classes are registered
 	if (InRegistrationOptions.bAutoUpdate)
 	{
-		FString OwningAssetName = GetOwningAssetName();
-		const bool bAutoUpdated = FAutoUpdateRootGraph(MoveTemp(OwningAssetName), InRegistrationOptions.bAutoUpdateLogWarningOnDroppedConnection).Transform(GetDocumentHandle());
+		const bool bDidUpdate = AutoUpdate(InRegistrationOptions.bAutoUpdateLogWarningOnDroppedConnection);
 #if WITH_EDITOR
-		if (bAutoUpdated || InRegistrationOptions.bForceViewSynchronization)
+		if (bDidUpdate || InRegistrationOptions.bForceViewSynchronization)
 		{
 			GetModifyContext().SetForceRefreshViews();
 		}
@@ -188,110 +223,26 @@ void FMetasoundAssetBase::RegisterGraphWithFrontend(Metasound::Frontend::FMetaSo
 	CacheRegistryMetadata();
 #endif // WITH_EDITOR
 
-	// Registers node by copying document. Updates to document require re-registration.
-	class FNodeRegistryEntry : public INodeRegistryEntry
-	{
-	public:
-		FNodeRegistryEntry(const FString& InName, TSharedPtr<FMetasoundFrontendDocument> InPreprocessedDoc, const FSoftObjectPath& InAssetPath)
-		: Name(InName)
-		, PreprocessedDoc(InPreprocessedDoc)
-		{
-			// Copy FrontendClass to preserve original document.
-			FrontendClass = PreprocessedDoc->RootGraph;
-			FrontendClass.Metadata.SetType(EMetasoundFrontendClassType::External);
-
-			ClassInfo = FNodeClassInfo(PreprocessedDoc->RootGraph, InAssetPath);
-		}
-
-		virtual ~FNodeRegistryEntry() = default;
-
-		virtual const FNodeClassInfo& GetClassInfo() const override
-		{
-			return ClassInfo;
-		}
-
-		virtual TUniquePtr<INode> CreateNode(const FNodeInitData&) const override
-		{
-			static const TSet<FName> ReferencedGraphTransmissibleInputNames; // Empty as referenced graphs do not support transmission
-			return FFrontendGraphBuilder().CreateGraph(*PreprocessedDoc, ReferencedGraphTransmissibleInputNames, Name);
-		}
-
-		virtual TUniquePtr<INode> CreateNode(FDefaultLiteralNodeConstructorParams&&) const override
-		{
-			return nullptr;
-		}
-
-		virtual TUniquePtr<INode> CreateNode(FDefaultNamedVertexNodeConstructorParams&&) const override
-		{
-			return nullptr;
-		}
-
-		virtual TUniquePtr<INode> CreateNode(FDefaultNamedVertexWithLiteralNodeConstructorParams&&) const override
-		{
-			return nullptr;
-		}
-
-		virtual const FMetasoundFrontendClass& GetFrontendClass() const override
-		{
-			return FrontendClass;
-		}
-
-		virtual TUniquePtr<INodeRegistryEntry> Clone() const override
-		{
-			return MakeUnique<FNodeRegistryEntry>(Name, PreprocessedDoc, ClassInfo.AssetPath);
-		}
-
-		virtual bool IsNative() const override
-		{
-			return false;
-		}
-
-	private:
-		FString Name;
-		TSharedPtr<FMetasoundFrontendDocument> PreprocessedDoc;
-		FMetasoundFrontendClass FrontendClass;
-		FNodeClassInfo ClassInfo;
-	};
-
-	UnregisterGraphWithFrontend();
-
 	FString AssetName;
-	FString AssetPath;
 	const UObject* OwningAsset = GetOwningAsset();
 	if (ensure(OwningAsset))
 	{
 		AssetName = OwningAsset->GetName();
-		AssetPath = OwningAsset->GetPathName();
 	}
 
-	FNodeClassInfo AssetClassInfo = GetAssetClassInfo();
-	TSharedPtr<FMetasoundFrontendDocument> PreprocessedDocument = MakeShared<FMetasoundFrontendDocument>(GetDocumentChecked());
-	FDocumentTemplatePreprocessTransform TemplatePreprocessor;
-	TemplatePreprocessor.Transform(*PreprocessedDocument);
+	TSharedPtr<FMetasoundFrontendDocument> PreprocessedDocument = PreprocessDocument();
 	CacheRuntimeData(*PreprocessedDocument);
+	FNodeClassInfo AssetClassInfo = GetAssetClassInfo();
+
+	TUniquePtr<INodeRegistryEntry> RegistryEntry = MakeUnique<AssetBasePrivate::FNodeRegistryEntry>(AssetName, PreprocessedDocument, AssetClassInfo.AssetPath);
 
-	TUniquePtr<INodeRegistryEntry> RegistryEntry = MakeUnique<FNodeRegistryEntry>(AssetName, PreprocessedDocument, AssetClassInfo.AssetPath);
+	UnregisterGraphWithFrontend();
 	RegistryKey = FMetasoundFrontendRegistryContainer::Get()->RegisterNode(MoveTemp(RegistryEntry));
 
 	if (NodeRegistryKey::IsValid(RegistryKey) && FMetasoundFrontendRegistryContainer::Get()->IsNodeRegistered(RegistryKey))
 	{
 #if WITH_EDITORONLY_DATA
-		// Refresh Asset Registry Info if successfully registered with Frontend
-		const FMetasoundFrontendGraphClass& DocumentClassGraph = GetDocumentHandle()->GetRootGraphClass();
-		const FMetasoundFrontendClassMetadata& DocumentClassMetadata = DocumentClassGraph.Metadata;
-		AssetClassInfo.AssetClassID = FGuid(DocumentClassMetadata.GetClassName().Name.ToString());
-		FNodeClassName ClassName = DocumentClassMetadata.GetClassName().ToNodeClassName();
-		FMetasoundFrontendClass GraphClass;
-
-		AssetClassInfo.Version = DocumentClassMetadata.GetVersion();
-
-		AssetClassInfo.InputTypes.Reset();
-		Algo::Transform(GraphClass.Interface.Inputs, AssetClassInfo.InputTypes, [] (const FMetasoundFrontendClassInput& Input) { return Input.TypeName; });
-
-		AssetClassInfo.OutputTypes.Reset();
-		Algo::Transform(GraphClass.Interface.Outputs, AssetClassInfo.OutputTypes, [](const FMetasoundFrontendClassOutput& Output) { return Output.TypeName; });
-
-		SetRegistryAssetClassInfo(MoveTemp(AssetClassInfo));
+		UpdateAssetRegistry();
 #endif // WITH_EDITORONLY_DATA
 	}
 	else
@@ -875,12 +826,80 @@ TArray<FMetasoundFrontendClassInput> FMetasoundAssetBase::GetPublicClassInputs()
 	return PublicInputs;
 }
 
-void FMetasoundAssetBase::RebuildReferencedAssetClassKeys()
+
+#if WITH_EDITOR
+void FMetasoundAssetBase::RebuildReferencedAssetClasses()
+{
+	using namespace Metasound::Frontend;
+
+	IMetaSoundAssetManager& AssetManager = IMetaSoundAssetManager::GetChecked();
+
+	AssetManager.AddAssetReferences(*this);
+	TSet<IMetaSoundAssetManager::FAssetInfo> ReferencedAssetClasses = AssetManager.GetReferencedAssetClasses(*this);
+	SetReferencedAssetClasses(MoveTemp(ReferencedAssetClasses));
+}
+#endif // WITH_EDITOR
+
+void FMetasoundAssetBase::RegisterAssetDependencies(const Metasound::Frontend::FMetaSoundAssetRegistrationOptions& InRegistrationOptions)
+{
+	using namespace Metasound::Frontend;
+
+	IMetaSoundAssetManager& AssetManager = IMetaSoundAssetManager::GetChecked();
+	TArray<FMetasoundAssetBase*> References = GetReferencedAssets();
+	for (FMetasoundAssetBase* Reference : References)
+	{
+		if (InRegistrationOptions.bForceReregister || !Reference->IsRegistered())
+		{
+			// TODO: Check for infinite recursion and error if so
+			AssetManager.AddOrUpdateAsset(*(Reference->GetOwningAsset()));
+			Reference->RegisterGraphWithFrontend(InRegistrationOptions);
+		}
+	}
+}
+
+bool FMetasoundAssetBase::AutoUpdate(bool bInLogWarningsOnDroppedConnection)
+{
+	using namespace Metasound::Frontend;
+
+	FString OwningAssetName = GetOwningAssetName();
+	const bool bAutoUpdated = FAutoUpdateRootGraph(MoveTemp(OwningAssetName), bInLogWarningsOnDroppedConnection).Transform(GetDocumentHandle());
+	return bAutoUpdated;
+}
+
+#if WITH_EDITORONLY_DATA
+void FMetasoundAssetBase::UpdateAssetRegistry() 
 {
+	using namespace Metasound;
 	using namespace Metasound::Frontend;
 
-	TSet<FNodeRegistryKey> ReferencedKeys = IMetaSoundAssetManager::GetChecked().GetReferencedKeys(*this);
-	SetReferencedAssetClassKeys(MoveTemp(ReferencedKeys));
+	FNodeClassInfo AssetClassInfo = GetAssetClassInfo();
+	// Refresh Asset Registry Info if successfully registered with Frontend
+	const FMetasoundFrontendGraphClass& DocumentClassGraph = GetDocumentHandle()->GetRootGraphClass();
+	const FMetasoundFrontendClassMetadata& DocumentClassMetadata = DocumentClassGraph.Metadata;
+	AssetClassInfo.AssetClassID = FGuid(DocumentClassMetadata.GetClassName().Name.ToString());
+	FNodeClassName ClassName = DocumentClassMetadata.GetClassName().ToNodeClassName();
+	FMetasoundFrontendClass GraphClass;
+
+	AssetClassInfo.Version = DocumentClassMetadata.GetVersion();
+
+	AssetClassInfo.InputTypes.Reset();
+	Algo::Transform(GraphClass.Interface.Inputs, AssetClassInfo.InputTypes, [] (const FMetasoundFrontendClassInput& Input) { return Input.TypeName; });
+
+	AssetClassInfo.OutputTypes.Reset();
+	Algo::Transform(GraphClass.Interface.Outputs, AssetClassInfo.OutputTypes, [](const FMetasoundFrontendClassOutput& Output) { return Output.TypeName; });
+
+	SetRegistryAssetClassInfo(MoveTemp(AssetClassInfo));
+}
+#endif
+
+TSharedPtr<FMetasoundFrontendDocument> FMetasoundAssetBase::PreprocessDocument()
+{
+	using namespace Metasound::Frontend;
+
+	TSharedPtr<FMetasoundFrontendDocument> PreprocessedDocument = MakeShared<FMetasoundFrontendDocument>(GetDocumentChecked());
+	FDocumentTemplatePreprocessTransform TemplatePreprocessor;
+	TemplatePreprocessor.Transform(*PreprocessedDocument);
+	return PreprocessedDocument;
 }
 
 const FMetasoundAssetBase::FRuntimeData& FMetasoundAssetBase::CacheRuntimeData(const FMetasoundFrontendDocument& InPreprocessedDoc)
diff --git a/Engine/Plugins/Runtime/Metasound/Source/MetasoundFrontend/Public/MetasoundAssetBase.h b/Engine/Plugins/Runtime/Metasound/Source/MetasoundFrontend/Public/MetasoundAssetBase.h
--- a/Engine/Plugins/Runtime/Metasound/Source/MetasoundFrontend/Public/MetasoundAssetBase.h
+++ b/Engine/Plugins/Runtime/Metasound/Source/MetasoundFrontend/Public/MetasoundAssetBase.h
@@ -78,8 +78,10 @@ public:
 	// Sets/overwrites the root class metadata
 	virtual void SetMetadata(FMetasoundFrontendClassMetadata& InMetadata);
 
-	// Rebuild the asset class dependency key array.
-	void RebuildReferencedAssetClassKeys();
+#if WITH_EDITOR
+	// Rebuild dependent asset classes
+	void RebuildReferencedAssetClasses();
+#endif // WITH_EDITOR
 
 	// Returns the interface entries declared by the given asset's document from the InterfaceRegistry.
 	bool GetDeclaredInterfaces(TArray<const Metasound::Frontend::IInterfaceRegistryEntry*>& OutInterfaces) const;
@@ -93,13 +95,16 @@ public:
 	// Returns all the class keys of this asset's referenced assets
 	virtual const TSet<FString>& GetReferencedAssetClassKeys() const = 0;
 
-	// Returns set of cached class references set on last registration
-	// prior to serialize. Used at runtime to hint where to load referenced
-	// class if sound loads before AssetManager scan is completed.  When registered
-	// hint paths to classes here can be superseded by another asset class if it shares
-	// the same key and has already been registered in the MetaSoundAssetManager.
-	virtual TSet<FSoftObjectPath>& GetReferencedAssetClassCache() = 0;
-	virtual const TSet<FSoftObjectPath>& GetReferencedAssetClassCache() const = 0;
+	// Returns set of class references set call to serialize in the editor
+	// Used at runtime load register referenced classes.
+	virtual TArray<FMetasoundAssetBase*> GetReferencedAssets() = 0;
+
+	// Return all dependent asset paths to load asynchronously
+	virtual const TSet<FSoftObjectPath>& GetAsyncReferencedAssetClassPaths() const = 0;
+
+	// Called when async assets have finished loading.
+	virtual void OnAsyncReferencedAssetsLoaded(const TArray<FMetasoundAssetBase*>& InAsyncReferences) = 0;
+
 
 	bool AddingReferenceCausesLoop(const FSoftObjectPath& InReferencePath) const;
 	bool IsReferencedAsset(const FMetasoundAssetBase& InAssetToCheck) const;
@@ -162,7 +167,9 @@ public:
 	FString GetOwningAssetName() const;
 
 protected:
-	virtual void SetReferencedAssetClassKeys(TSet<Metasound::Frontend::FNodeRegistryKey>&& InKeys) = 0;
+#if WITH_EDITOR
+	virtual void SetReferencedAssetClasses(TSet<Metasound::Frontend::IMetaSoundAssetManager::FAssetInfo>&& InAssetClasses) = 0;
+#endif
 
 	// Get information for communicating asynchronously with MetaSound running instance.
 	TArray<FSendInfoAndVertexName> GetSendInfos(uint64 InInstanceID) const;
@@ -201,7 +208,14 @@ protected:
 	// Prefer accessing public class inputs using CacheRuntimeData.
 	TArray<FMetasoundFrontendClassInput> GetPublicClassInputs() const;
 
+
+	bool AutoUpdate(bool bInLogWarningsOnDroppedConnection);
+	void RegisterAssetDependencies(const Metasound::Frontend::FMetaSoundAssetRegistrationOptions& InRegistrationOptions);
+	TSharedPtr<FMetasoundFrontendDocument> PreprocessDocument();
 private:
+#if WITH_EDITORONLY_DATA
+	void UpdateAssetRegistry();
+#endif
 	// Returns the cached runtime data. Call updates cached data if out-of-date.
 	const FRuntimeData& CacheRuntimeData(const FMetasoundFrontendDocument& InPreprocessedDoc);
 
diff --git a/Engine/Plugins/Runtime/Metasound/Source/MetasoundFrontend/Public/MetasoundAssetManager.h b/Engine/Plugins/Runtime/Metasound/Source/MetasoundFrontend/Public/MetasoundAssetManager.h
--- a/Engine/Plugins/Runtime/Metasound/Source/MetasoundFrontend/Public/MetasoundAssetManager.h
+++ b/Engine/Plugins/Runtime/Metasound/Source/MetasoundFrontend/Public/MetasoundAssetManager.h
@@ -51,8 +51,10 @@ namespace Metasound
 			// If true, warnings will be logged if updating a node results in existing connections being discarded.
 			bool bAutoUpdateLogWarningOnDroppedConnection = false;
 
-			// Attempt to rebuild referenced class keys (only run if class not registered or set to force re-register)
-			bool bRebuildReferencedAssetClassKeys = true;
+#if WITH_EDITOR
+			// Attempt to rebuild referenced classes (only run if class not registered or set to force re-register)
+			bool bRebuildReferencedAssetClasses = true;
+#endif
 		};
 
 		class METASOUNDFRONTEND_API IMetaSoundAssetManager
@@ -80,12 +82,30 @@ namespace Metasound
 				return *Instance;
 			}
 
+			struct FAssetInfo
+			{
+				const Metasound::Frontend::FNodeRegistryKey RegistryKey;
+				FSoftObjectPath AssetPath;
+
+				FORCEINLINE friend bool operator==(const FAssetInfo& InLHS, const FAssetInfo& InRHS)
+				{
+					return (InLHS.RegistryKey == InRHS.RegistryKey) && (InLHS.AssetPath == InRHS.AssetPath);
+				}
+
+				FORCEINLINE friend uint32 GetTypeHash(const IMetaSoundAssetManager::FAssetInfo& InInfo)
+				{
+					return HashCombineFast(::GetTypeHash(InInfo.RegistryKey), GetTypeHash(InInfo.AssetPath));
+				}
+			};
+
 			// Whether or not manager is being used to run tests or not (enabling instances to be reset without asserting.)
 			virtual bool IsTesting() const { return false; }
 
+#if WITH_EDITORONLY_DATA
 			// Adds missing assets using the provided asset's local reference class cache. Used
 			// to prime system from asset attempting to register prior to asset scan being complete.
 			virtual void AddAssetReferences(FMetasoundAssetBase& InAssetBase) = 0;
+#endif
 
 			// Add or Update a MetaSound Asset's entry data
 			virtual Metasound::Frontend::FNodeRegistryKey AddOrUpdateAsset(const UObject& InObject) = 0;
@@ -99,8 +119,10 @@ namespace Metasound
 			// Returns path associated with the given key (null if key is not registered with the AssetManager or was not loaded from asset)
 			virtual const FSoftObjectPath* FindObjectPathFromKey(const Metasound::Frontend::FNodeRegistryKey& InRegistryKey) const = 0;
 
-			// Generates all asset keys associated with registered assets that are referenced by the provided asset's graph.
-			virtual TSet<Metasound::Frontend::FNodeRegistryKey> GetReferencedKeys(const FMetasoundAssetBase& InAssetBase) const = 0;
+#if WITH_EDITOR
+			// Generates all asset info associated with registered assets that are referenced by the provided asset's graph.
+			virtual TSet<FAssetInfo> GetReferencedAssetClasses(const FMetasoundAssetBase& InAssetBase) const = 0;
+#endif // WITH_EDITOR
 
 			// Rescans settings for denied assets not to run reference auto-update against.
 			virtual void RescanAutoUpdateDenyList() = 0;
@@ -114,6 +136,12 @@ namespace Metasound
 			// Try to load referenced assets of the given asset or return them if they are already loaded (non-recursive).
 			// @return - True if all referenced assets successfully loaded, false if not.
 			virtual bool TryLoadReferencedAssets(const FMetasoundAssetBase& InAssetBase, TArray<FMetasoundAssetBase*>& OutReferencedAssets) const = 0;
+			
+			// Requests an async load of all async referenced assets of the input asset.
+			virtual void RequestAsyncLoadReferencedAssets(FMetasoundAssetBase& InAssetBase) = 0;
+
+			// Waits until all async load requests related to this asset are complete. 
+ 			virtual void WaitUntilAsyncLoadReferencedAssetsComplete(FMetasoundAssetBase& InAssetBase) = 0;
 		};
 	} // namespace Frontend
 };
