Fix for CDO attempting to access MetaSound asset manager
#rnx

diff --git a/Engine/Plugins/Runtime/Metasound/Source/MetasoundEngine/Private/MetasoundEngineAsset.h b/Engine/Plugins/Runtime/Metasound/Source/MetasoundEngine/Private/MetasoundEngineAsset.h
--- a/Engine/Plugins/Runtime/Metasound/Source/MetasoundEngine/Private/MetasoundEngineAsset.h
+++ b/Engine/Plugins/Runtime/Metasound/Source/MetasoundEngine/Private/MetasoundEngineAsset.h
@@ -94,7 +94,12 @@ namespace Metasound
 				}
 			}
 
-			IMetaSoundAssetManager::GetChecked().WaitUntilAsyncLoadReferencedAssetsComplete(InMetaSound);
+			// Do not call asset manager on CDO objects which may be loaded before asset 
+			// manager is set.
+			if (IMetaSoundAssetManager* AssetManager = IMetaSoundAssetManager::Get())
+			{
+				AssetManager->WaitUntilAsyncLoadReferencedAssetsComplete(InMetaSound);
+			}
 #endif // WITH_EDITORONLY_DATA
 		}
 
@@ -112,14 +117,27 @@ namespace Metasound
 					}
 #endif // WITH_EDITORONLY_DATA
 				}
-
 			}
 		}
 
 		template<typename TMetaSoundObject>
 		static void PostLoad(TMetaSoundObject& InMetaSound)
 		{
-			Frontend::IMetaSoundAssetManager::GetChecked().RequestAsyncLoadReferencedAssets(InMetaSound);
+			using namespace Frontend;
+			// Do not call asset manager on CDO objects which may be loaded before asset 
+			// manager is set.
+			const bool bIsCDO = InMetaSound.HasAnyFlags(RF_ClassDefaultObject);
+			if (!bIsCDO)
+			{
+				if (IMetaSoundAssetManager* AssetManager = IMetaSoundAssetManager::Get())
+				{
+					AssetManager->RequestAsyncLoadReferencedAssets(InMetaSound);
+				}
+				else
+				{
+					UE_LOG(LogMetaSound, Warning, TEXT("Request for to load async references ignored. Likely due loading before the MetaSoundEngine module."));
+				}
+			}
 		}
 
 		template<typename TMetaSoundObject>
