A few little tweaks to Metasound "core" to help with Harmonix metasound work

diff --git a/Engine/Plugins/Runtime/Metasound/Source/MetasoundFrontend/Private/MetasoundTrigger.cpp b/Engine/Plugins/Runtime/Metasound/Source/MetasoundFrontend/Private/MetasoundTrigger.cpp
--- a/Engine/Plugins/Runtime/Metasound/Source/MetasoundFrontend/Private/MetasoundTrigger.cpp
+++ b/Engine/Plugins/Runtime/Metasound/Source/MetasoundFrontend/Private/MetasoundTrigger.cpp
@@ -91,6 +91,16 @@ namespace Metasound
 		return TriggeredFrames[InTriggerIndex];
 	}
 
+	int32 FTrigger::First() const
+	{
+		return TriggeredFrames.IsEmpty() ? INDEX_NONE : TriggeredFrames[0];
+	}
+
+	int32 FTrigger::Last() const
+	{
+		return TriggeredFrames.IsEmpty() ? INDEX_NONE : TriggeredFrames.Last();
+	}
+
 	bool FTrigger::IsTriggered() const
 	{
 		return bHasTrigger;
diff --git a/Engine/Plugins/Runtime/Metasound/Source/MetasoundFrontend/Public/MetasoundTrigger.h b/Engine/Plugins/Runtime/Metasound/Source/MetasoundFrontend/Public/MetasoundTrigger.h
--- a/Engine/Plugins/Runtime/Metasound/Source/MetasoundFrontend/Public/MetasoundTrigger.h
+++ b/Engine/Plugins/Runtime/Metasound/Source/MetasoundFrontend/Public/MetasoundTrigger.h
@@ -67,6 +67,18 @@ namespace Metasound
 			 */
 			int32 operator[](int32 InTriggerIndex) const;
 
+			/** Returns frame index for the first trigger in the block.
+			 *
+			 * @return The frame of the first trigger in the block or -1 if there is no trigger in the block.
+			 */
+			int32 First() const;
+
+			/** Returns frame index for the last trigger in the block.
+			 *
+			 * @return The frame of the last trigger in the block or -1 if there is no trigger in the block.
+			 */
+			int32 Last() const;
+
 			/** Returns true if there are any triggered frames. */
 			bool IsTriggered() const;
 
@@ -125,6 +137,11 @@ namespace Metasound
 				}
 			}
 
+			const TArray<int32>& GetTriggeredFrames() const
+			{
+				return TriggeredFrames;
+			}
+
 		private:
 			void UpdateLastTriggerIndexInBlock();
 
diff --git a/Engine/Plugins/Runtime/Metasound/Source/MetasoundGraphCore/Public/MetasoundDataReferenceMacro.h b/Engine/Plugins/Runtime/Metasound/Source/MetasoundGraphCore/Public/MetasoundDataReferenceMacro.h
--- a/Engine/Plugins/Runtime/Metasound/Source/MetasoundGraphCore/Public/MetasoundDataReferenceMacro.h
+++ b/Engine/Plugins/Runtime/Metasound/Source/MetasoundGraphCore/Public/MetasoundDataReferenceMacro.h
@@ -24,17 +24,19 @@
 	typedef ::Metasound::TDataReadReference<DataType> DataReadReferenceTypeName; \
 	typedef ::Metasound::TDataWriteReference<DataType> DataWriteReferenceTypeName;
 
+#define DECLARE_METASOUND_DATA_REFERENCE_TYPES_NO_ALIASES(DataType, ModuleApi) \
+	DECLARE_METASOUND_DATA_REFERENCE_CORE_TYPE(DataType, ModuleApi) \
+	DECLARE_METASOUND_DATA_REFERENCE_CORE_TYPE(::Metasound::TVariable<DataType>, ModuleApi) \
+	DECLARE_METASOUND_DATA_REFERENCE_CORE_TYPE(TArray<DataType>, ModuleApi) \
+	DECLARE_METASOUND_DATA_REFERENCE_CORE_TYPE(::Metasound::TVariable<TArray<DataType>>, ModuleApi)
+
 /** Macro to make declaring a metasound parameter simple.  */
 // Declares a metasound parameter type by
 // - Adding typedefs for commonly used template types.
 // - Defining parameter type traits.
 #define DECLARE_METASOUND_DATA_REFERENCE_TYPES(DataType, ModuleApi, DataTypeInfoTypeName, DataReadReferenceTypeName, DataWriteReferenceTypeName) \
-	DECLARE_METASOUND_DATA_REFERENCE_CORE_TYPE(DataType, ModuleApi); \
-	DECLARE_METASOUND_DATA_REFERENCE_ALIAS_TYPES(DataType, DataTypeInfoTypeName, DataReadReferenceTypeName, DataWriteReferenceTypeName); \
-	DECLARE_METASOUND_DATA_REFERENCE_CORE_TYPE(::Metasound::TVariable<DataType>, ModuleApi); \
-	DECLARE_METASOUND_DATA_REFERENCE_CORE_TYPE(TArray<DataType>, ModuleApi); \
-	DECLARE_METASOUND_DATA_REFERENCE_CORE_TYPE(::Metasound::TVariable<TArray<DataType>>, ModuleApi);
-
+	DECLARE_METASOUND_DATA_REFERENCE_TYPES_NO_ALIASES(DataType, ModuleApi) \
+	DECLARE_METASOUND_DATA_REFERENCE_ALIAS_TYPES(DataType, DataTypeInfoTypeName, DataReadReferenceTypeName, DataWriteReferenceTypeName)
 
 #define DEFINE_METASOUND_DATA_REFERENCE_CORE_TYPE(DataType, DataTypeName, DataTypeLoctextKey) \
 	const TCHAR* ::Metasound::TDataReferenceTypeInfo<DataType>::TypeName = TEXT(DataTypeName); \
@@ -51,8 +53,8 @@
 
 // This only needs to be called if you don't plan on calling REGISTER_METASOUND_DATATYPE.
 #define DEFINE_METASOUND_DATA_TYPE(DataType, DataTypeName) \
-	DEFINE_METASOUND_DATA_REFERENCE_CORE_TYPE(DataType, DataTypeName, DataTypeName); \
-	DEFINE_METASOUND_DATA_REFERENCE_CORE_TYPE(::Metasound::TVariable<DataType>, DataTypeName METASOUND_DATA_TYPE_NAME_VARIABLE_TYPE_SPECIFIER, DataTypeName"_Variable"); \
-	DEFINE_METASOUND_DATA_REFERENCE_CORE_TYPE(TArray<DataType>, DataTypeName METASOUND_DATA_TYPE_NAME_ARRAY_TYPE_SPECIFIER, DataTypeName"_Array"); \
-	DEFINE_METASOUND_DATA_REFERENCE_CORE_TYPE(::Metasound::TVariable<TArray<DataType>>, DataTypeName METASOUND_DATA_TYPE_NAME_ARRAY_TYPE_SPECIFIER METASOUND_DATA_TYPE_NAME_VARIABLE_TYPE_SPECIFIER, DataTypeName"_Array_Variable"); 
+	DEFINE_METASOUND_DATA_REFERENCE_CORE_TYPE(DataType, DataTypeName, DataTypeName) \
+	DEFINE_METASOUND_DATA_REFERENCE_CORE_TYPE(::Metasound::TVariable<DataType>, DataTypeName METASOUND_DATA_TYPE_NAME_VARIABLE_TYPE_SPECIFIER, DataTypeName"_Variable") \
+	DEFINE_METASOUND_DATA_REFERENCE_CORE_TYPE(TArray<DataType>, DataTypeName METASOUND_DATA_TYPE_NAME_ARRAY_TYPE_SPECIFIER, DataTypeName"_Array") \
+	DEFINE_METASOUND_DATA_REFERENCE_CORE_TYPE(::Metasound::TVariable<TArray<DataType>>, DataTypeName METASOUND_DATA_TYPE_NAME_ARRAY_TYPE_SPECIFIER METASOUND_DATA_TYPE_NAME_VARIABLE_TYPE_SPECIFIER, DataTypeName"_Array_Variable")
 
