Replace use of FSoftObjectPath with UPROPERTY for dependent metasound assets.

diff --git a/Engine/Plugins/Runtime/Metasound/Source/MetasoundEngine/Private/MetasoundEngineAsset.h b/Engine/Plugins/Runtime/Metasound/Source/MetasoundEngine/Private/MetasoundEngineAsset.h
--- a/Engine/Plugins/Runtime/Metasound/Source/MetasoundEngine/Private/MetasoundEngineAsset.h
+++ b/Engine/Plugins/Runtime/Metasound/Source/MetasoundEngine/Private/MetasoundEngineAsset.h
@@ -129,13 +129,16 @@ namespace Metasound
 			const bool bIsCDO = InMetaSound.HasAnyFlags(RF_ClassDefaultObject);
 			if (!bIsCDO)
 			{
-				if (IMetaSoundAssetManager* AssetManager = IMetaSoundAssetManager::Get())
+				if (InMetaSound.GetAsyncReferencedAssetClassPaths().Num() > 0)
 				{
-					AssetManager->RequestAsyncLoadReferencedAssets(InMetaSound);
-				}
-				else
-				{
-					UE_LOG(LogMetaSound, Warning, TEXT("Request for to load async references ignored. Likely due loading before the MetaSoundEngine module."));
+					if (IMetaSoundAssetManager* AssetManager = IMetaSoundAssetManager::Get())
+					{
+						AssetManager->RequestAsyncLoadReferencedAssets(InMetaSound);
+					}
+					else
+					{
+						UE_LOG(LogMetaSound, Warning, TEXT("Request for to load async references ignored from asset %s. Likely due loading before the MetaSoundEngine module."), *InMetaSound.GetPathName());
+					}
 				}
 			}
 		}
