Fix warning when pasting MetaSound variable nodes caused by frontend nodes not being added

#rnx

diff --git a/Engine/Plugins/Runtime/Metasound/Source/MetasoundEditor/Private/MetasoundEditor.cpp b/Engine/Plugins/Runtime/Metasound/Source/MetasoundEditor/Private/MetasoundEditor.cpp
--- a/Engine/Plugins/Runtime/Metasound/Source/MetasoundEditor/Private/MetasoundEditor.cpp
+++ b/Engine/Plugins/Runtime/Metasound/Source/MetasoundEditor/Private/MetasoundEditor.cpp
@@ -2419,12 +2419,12 @@ namespace Metasound
 								}
 								else
 								{
-									// Fix-up if variable getter node does not exist but variable does
-									FConstNodeHandle NodeHandle = VariableNode->GetConstNodeHandle();
+									// Add new variable node
+									const FNodeClassName NodeClassName = VariableNode->GetClassName().ToNodeClassName();
+									FNodeHandle NodeHandle = FGraphBuilder::AddVariableNodeHandle(*Metasound, Variable->GetVariableID(), NodeClassName, VariableNode);
 									if (!NodeHandle->IsValid())
 									{
-										const FNodeClassName NodeClassName = VariableNode->GetClassName().ToNodeClassName();
-										NodeHandle = FGraphBuilder::AddVariableNodeHandle(*Metasound, Variable->GetVariableID(), NodeClassName, VariableNode);
+										NodesToRemove.Add(GraphNode);
 									}
 								}
 							}
