[MetaSound Editor]
- Pin inspector/inspector & reroute interaction fixes
- Trivial fix for dirtying assets on MetaSound sync
#rnx

diff --git a/Engine/Plugins/Runtime/Metasound/Source/MetasoundEditor/Private/MetasoundEditor.cpp b/Engine/Plugins/Runtime/Metasound/Source/MetasoundEditor/Private/MetasoundEditor.cpp
--- a/Engine/Plugins/Runtime/Metasound/Source/MetasoundEditor/Private/MetasoundEditor.cpp
+++ b/Engine/Plugins/Runtime/Metasound/Source/MetasoundEditor/Private/MetasoundEditor.cpp
@@ -722,6 +722,7 @@ namespace Metasound
 				Stop();
 			}
 
+			GraphConnectionManager.Reset();
 			InterfacesView.Reset();
 			DestroyAnalyzers();
 			check(GEditor);
@@ -1516,7 +1517,7 @@ namespace Metasound
 
 					const FName& AudioBufferTypeName = GetMetasoundDataTypeName<FAudioBuffer>();
 					const FSampleRate SampleRate = static_cast<FSampleRate>(AudioDevice->GetSampleRate());
-					GraphConnectionManager = FGraphConnectionManager(*MetasoundAsset, *PreviewComp, SampleRate);
+					GraphConnectionManager = MakeUnique<FGraphConnectionManager>(*MetasoundAsset, *PreviewComp, SampleRate);
 				}
 
 				MetasoundGraphEditor->RegisterActiveTimer(0.0f,
@@ -1540,7 +1541,7 @@ namespace Metasound
 							SetPreviewID(INDEX_NONE);
 							PlayTime = 0.0;
 							PlayTimeWidget->SetText(FText::GetEmpty());
-							GraphConnectionManager = { };
+							GraphConnectionManager = MakeUnique<FGraphConnectionManager>();
 
 							return EActiveTimerReturnType::Stop;
 						}
@@ -2613,12 +2614,12 @@ namespace Metasound
 
 		FGraphConnectionManager& FEditor::GetConnectionManager()
 		{
-			return GraphConnectionManager;
+			return *GraphConnectionManager.Get();
 		}
 
 		const FGraphConnectionManager& FEditor::GetConnectionManager() const
 		{
-			return GraphConnectionManager;
+			return *GraphConnectionManager.Get();
 		}
 
 		UAudioComponent* FEditor::GetAudioComponent() const
@@ -3139,7 +3140,7 @@ namespace Metasound
 
 			RefreshEditorContext();
 
-			GraphConnectionManager.Update(DeltaTime);
+			GraphConnectionManager->Update(DeltaTime);
 		}
 
 		void FEditor::RefreshEditorContext()
diff --git a/Engine/Plugins/Runtime/Metasound/Source/MetasoundEditor/Private/MetasoundEditor.h b/Engine/Plugins/Runtime/Metasound/Source/MetasoundEditor/Private/MetasoundEditor.h
--- a/Engine/Plugins/Runtime/Metasound/Source/MetasoundEditor/Private/MetasoundEditor.h
+++ b/Engine/Plugins/Runtime/Metasound/Source/MetasoundEditor/Private/MetasoundEditor.h
@@ -92,6 +92,11 @@ namespace Metasound
 		public:
 			static const FName EditorName;
 
+			FEditor()
+				: GraphConnectionManager(MakeUnique<FGraphConnectionManager>())
+			{
+			}
+
 			virtual ~FEditor();
 
 			virtual void RegisterTabSpawners(const TSharedRef<FTabManager>& TabManager) override;
@@ -422,7 +427,7 @@ namespace Metasound
 			TSharedPtr<STextBlock> PlayTimeWidget;
 			double PlayTime = 0.0;
 
-			FGraphConnectionManager GraphConnectionManager;
+			TUniquePtr<FGraphConnectionManager> GraphConnectionManager;
 
 			/** Command list for this editor */
 			TSharedPtr<FUICommandList> GraphEditorCommands;
diff --git a/Engine/Plugins/Runtime/Metasound/Source/MetasoundEditor/Private/MetasoundEditorGraphBuilder.cpp b/Engine/Plugins/Runtime/Metasound/Source/MetasoundEditor/Private/MetasoundEditorGraphBuilder.cpp
--- a/Engine/Plugins/Runtime/Metasound/Source/MetasoundEditor/Private/MetasoundEditorGraphBuilder.cpp
+++ b/Engine/Plugins/Runtime/Metasound/Source/MetasoundEditor/Private/MetasoundEditorGraphBuilder.cpp
@@ -146,11 +146,6 @@ namespace Metasound
 				bEditorGraphModified |= FGraphBuilder::SynchronizeNodes(InMetaSound);
 				bEditorGraphModified |= FGraphBuilder::SynchronizeConnections(InMetaSound);
 
-				if (bEditorGraphModified)
-				{
-					InMetaSound.MarkPackageDirty();
-				}
-
 				return bEditorGraphModified;
 			}
 		} // namespace GraphBuilderPrivate
@@ -238,6 +233,12 @@ namespace Metasound
 				return false;
 			}
 
+			const UEdGraphPin* ReroutedPin = FindReroutedOutputPin(InPin);
+			if (ReroutedPin != InPin)
+			{
+				return false;
+			}
+
 			return true;
 		}
 
@@ -856,6 +857,33 @@ namespace Metasound
 			return OutputPin;
 		}
 
+		const UEdGraphPin* FGraphBuilder::FindReroutedOutputPin(const UEdGraphPin* OutputPin)
+		{
+			using namespace Frontend;
+
+			if (OutputPin)
+			{
+				if (UMetasoundEditorGraphExternalNode* ExternalNode = Cast<UMetasoundEditorGraphExternalNode>(OutputPin->GetOwningNode()))
+				{
+					if (ExternalNode->GetClassName() == FRerouteNodeTemplate::ClassName)
+					{
+						auto IsInput = [](const UEdGraphPin* Pin) { check(Pin); return Pin->Direction == EGPD_Input; };
+						if (const UEdGraphPin* RerouteInput = *ExternalNode->Pins.FindByPredicate(IsInput))
+						{
+							const TArray<UEdGraphPin*>& LinkedTo = RerouteInput->LinkedTo;
+							if (!LinkedTo.IsEmpty())
+							{
+								const UEdGraphPin* ReroutedOutput = LinkedTo.Last();
+								return FindReroutedOutputPin(ReroutedOutput);
+							}
+						}
+					}
+				}
+			}
+
+			return OutputPin;
+		}
+
 		Frontend::FOutputHandle FGraphBuilder::FindReroutedOutputHandleFromPin(const UEdGraphPin* OutputPin)
 		{
 			using namespace Frontend;
diff --git a/Engine/Plugins/Runtime/Metasound/Source/MetasoundEditor/Private/MetasoundEditorGraphBuilder.h b/Engine/Plugins/Runtime/Metasound/Source/MetasoundEditor/Private/MetasoundEditorGraphBuilder.h
--- a/Engine/Plugins/Runtime/Metasound/Source/MetasoundEditor/Private/MetasoundEditorGraphBuilder.h
+++ b/Engine/Plugins/Runtime/Metasound/Source/MetasoundEditor/Private/MetasoundEditorGraphBuilder.h
@@ -215,6 +215,7 @@ namespace Metasound
 			static Frontend::FConstOutputHandle GetConstOutputHandleFromPin(const UEdGraphPin* InPin);
 
 			static UEdGraphPin* FindReroutedOutputPin(UEdGraphPin* InPin);
+			static const UEdGraphPin* FindReroutedOutputPin(const UEdGraphPin* InPin);
 
 			// Find the "concrete" output handle associated with an output pin.  If the given output pin is on
 			// a reroute node, will recursively search for the non-rerouted output its representing.
diff --git a/Engine/Plugins/Runtime/Metasound/Source/MetasoundEditor/Private/MetasoundEditorGraphConnectionDrawingPolicy.cpp b/Engine/Plugins/Runtime/Metasound/Source/MetasoundEditor/Private/MetasoundEditorGraphConnectionDrawingPolicy.cpp
--- a/Engine/Plugins/Runtime/Metasound/Source/MetasoundEditor/Private/MetasoundEditorGraphConnectionDrawingPolicy.cpp
+++ b/Engine/Plugins/Runtime/Metasound/Source/MetasoundEditor/Private/MetasoundEditorGraphConnectionDrawingPolicy.cpp
@@ -12,6 +12,7 @@
 #include "MetasoundTrace.h"
 #include "Misc/App.h"
 #include "NodeTemplates/MetasoundFrontendNodeTemplateReroute.h"
+#include "SMetasoundPinValueInspector.h"
 #include "Templates/Function.h"
 #include "WaveTableSampler.h"
 
@@ -131,10 +132,10 @@ namespace Metasound
 					return;
 				}
 
-				const FGuid NodeID = Node->GetNodeID();
+				Frontend::FConstOutputHandle OutputHandle = Frontend::FindReroutedOutput(InData.OutputHandle);
+				const FGuid NodeID = OutputHandle->GetOwningNodeID();
 				TNumericType Value = InDefaultValue;
-				FName OutputName = InData.OutputHandle->GetName();
-				
+				FName OutputName = OutputHandle->GetName();
 				Editor->GetConnectionManager().GetValue(NodeID, OutputName, Value);
 
 				FLinearColor InnerColor = Frontend::DisplayStyle::EdgeAnimation::DefaultColor;
diff --git a/Engine/Plugins/Runtime/Metasound/Source/MetasoundEditor/Private/MetasoundEditorGraphConnectionManager.cpp b/Engine/Plugins/Runtime/Metasound/Source/MetasoundEditor/Private/MetasoundEditorGraphConnectionManager.cpp
--- a/Engine/Plugins/Runtime/Metasound/Source/MetasoundEditor/Private/MetasoundEditorGraphConnectionManager.cpp
+++ b/Engine/Plugins/Runtime/Metasound/Source/MetasoundEditor/Private/MetasoundEditorGraphConnectionManager.cpp
@@ -14,17 +14,17 @@ namespace Metasound
 	{
 		FGraphConnectionManager::FGraphConnectionManager(const FMetasoundAssetBase& InAssetBase, const UAudioComponent& InAudioComponent, FSampleRate InSampleRate)
 			: AudioComponent(&InAudioComponent)
-			, GraphAnalyzerView({ InAssetBase, InAudioComponent.GetInstanceOwnerID(), InSampleRate })
+			, GraphAnalyzerView(MakeUnique<Frontend::FMetasoundGraphAnalyzerView>(InAssetBase, InAudioComponent.GetInstanceOwnerID(), InSampleRate))
 		{
 			using namespace Frontend;
 
-			GraphAnalyzerView.AddAnalyzerForAllSupportedOutputs(FVertexAnalyzerForwardBool::GetAnalyzerName());
-			GraphAnalyzerView.AddAnalyzerForAllSupportedOutputs(FVertexAnalyzerForwardFloat::GetAnalyzerName());
-			GraphAnalyzerView.AddAnalyzerForAllSupportedOutputs(FVertexAnalyzerForwardInt::GetAnalyzerName());
-			GraphAnalyzerView.AddAnalyzerForAllSupportedOutputs(FVertexAnalyzerForwardString::GetAnalyzerName());
+			GraphAnalyzerView->AddAnalyzerForAllSupportedOutputs(FVertexAnalyzerForwardBool::GetAnalyzerName());
+			GraphAnalyzerView->AddAnalyzerForAllSupportedOutputs(FVertexAnalyzerForwardFloat::GetAnalyzerName());
+			GraphAnalyzerView->AddAnalyzerForAllSupportedOutputs(FVertexAnalyzerForwardInt::GetAnalyzerName());
+			GraphAnalyzerView->AddAnalyzerForAllSupportedOutputs(FVertexAnalyzerForwardString::GetAnalyzerName());
 
-			GraphAnalyzerView.AddAnalyzerForAllSupportedOutputs(FVertexAnalyzerEnvelopeFollower::GetAnalyzerName());
-			GraphAnalyzerView.AddAnalyzerForAllSupportedOutputs(FVertexAnalyzerTriggerDensity::GetAnalyzerName());
+			GraphAnalyzerView->AddAnalyzerForAllSupportedOutputs(FVertexAnalyzerEnvelopeFollower::GetAnalyzerName());
+			GraphAnalyzerView->AddAnalyzerForAllSupportedOutputs(FVertexAnalyzerTriggerDensity::GetAnalyzerName());
 		}
 
 		bool FGraphConnectionManager::GetValue(const FGuid& InNodeID, FVertexName InOutputName, float& OutValue) const
@@ -127,7 +127,7 @@ namespace Metasound
 
 			float CurrentValue = 0.0f;
 
-			TArray<FMetasoundAnalyzerView*> Views = GraphAnalyzerView.GetAnalyzerViews(InAnalyzerName);
+			TArray<FMetasoundAnalyzerView*> Views = GraphAnalyzerView->GetAnalyzerViews(InAnalyzerName);
 			for (const FMetasoundAnalyzerView* View : Views)
 			{
 				check(View);
diff --git a/Engine/Plugins/Runtime/Metasound/Source/MetasoundEditor/Private/MetasoundEditorGraphConnectionManager.h b/Engine/Plugins/Runtime/Metasound/Source/MetasoundEditor/Private/MetasoundEditorGraphConnectionManager.h
--- a/Engine/Plugins/Runtime/Metasound/Source/MetasoundEditor/Private/MetasoundEditorGraphConnectionManager.h
+++ b/Engine/Plugins/Runtime/Metasound/Source/MetasoundEditor/Private/MetasoundEditorGraphConnectionManager.h
@@ -79,9 +79,12 @@ namespace Metasound
 
 		class FGraphConnectionManager
 		{
+			FGraphConnectionManager(const FGraphConnectionManager&) = delete;
+
 		public:
 			FGraphConnectionManager() = default;
 			FGraphConnectionManager(const FMetasoundAssetBase& InAssetBase, const UAudioComponent& InAudioComponent, FSampleRate InSampleRate);
+			~FGraphConnectionManager() = default;
 
 			bool GetValue(const FGuid& InNodeID, FVertexName InOutputName, float& OutValue) const;
 			bool GetValue(const FGuid& InNodeID, FVertexName InOutputName, bool& OutValue) const;
@@ -111,14 +114,15 @@ namespace Metasound
 			{
 				using namespace Frontend;
 
-				TArray<const FMetasoundAnalyzerView*> Views = GraphAnalyzerView.GetAnalyzerViewsForOutput(InNodeID, InOutputName, ForwardValueAnalyzerClass::GetAnalyzerName());
+				const Frontend::FMetasoundGraphAnalyzerView* GraphView = GraphAnalyzerView.Get();
+				TArray<const FMetasoundAnalyzerView*> Views = GraphView->GetAnalyzerViewsForOutput(InNodeID, InOutputName, ForwardValueAnalyzerClass::GetAnalyzerName());
 				if (!Views.IsEmpty())
 				{
 					const FMetasoundAnalyzerView** View = Views.FindByPredicate([](const FMetasoundAnalyzerView* Candidate)
-						{
-							check(Candidate);
-							return Candidate->AnalyzerAddress.AnalyzerMemberName == ForwardValueAnalyzerClass::FOutputs::GetValue().Name;
-						});
+					{
+						check(Candidate);
+						return Candidate->AnalyzerAddress.AnalyzerMemberName == ForwardValueAnalyzerClass::FOutputs::GetValue().Name;
+					});
 
 					if (View)
 					{
@@ -156,7 +160,7 @@ namespace Metasound
 			{
 				using namespace Frontend;
 
-				TArray<FMetasoundAnalyzerView*> Views = GraphAnalyzerView.GetAnalyzerViews(ForwardValueAnalyzerClass::GetAnalyzerName());
+				TArray<FMetasoundAnalyzerView*> Views = GraphAnalyzerView->GetAnalyzerViews(ForwardValueAnalyzerClass::GetAnalyzerName());
 				for (FMetasoundAnalyzerView* View : Views)
 				{
 					check(View);
@@ -217,7 +221,7 @@ namespace Metasound
 			TMap<FString, bool> BoolConnectionValues;
 			TMap<FString, FString> StringConnectionValues;
 
-			Frontend::FMetasoundGraphAnalyzerView GraphAnalyzerView;
+			TUniquePtr<Frontend::FMetasoundGraphAnalyzerView> GraphAnalyzerView;
 		};
 	} // namespace Editor
 } // Metasound
diff --git a/Engine/Plugins/Runtime/Metasound/Source/MetasoundEditor/Private/SMetasoundGraphNode.cpp b/Engine/Plugins/Runtime/Metasound/Source/MetasoundEditor/Private/SMetasoundGraphNode.cpp
--- a/Engine/Plugins/Runtime/Metasound/Source/MetasoundEditor/Private/SMetasoundGraphNode.cpp
+++ b/Engine/Plugins/Runtime/Metasound/Source/MetasoundEditor/Private/SMetasoundGraphNode.cpp
@@ -36,6 +36,7 @@
 #include "SLevelOfDetailBranchNode.h"
 #include "SMetasoundGraphEnumPin.h"
 #include "SMetasoundGraphPin.h"
+#include "SMetasoundPinValueInspector.h"
 #include "SPinTypeSelector.h"
 #include "Styling/AppStyle.h"
 #include "Styling/SlateColor.h"
@@ -51,29 +52,13 @@
 #include "Widgets/SWidget.h"
 #include "Widgets/Text/SInlineEditableTextBlock.h"
 
-#define LOCTEXT_NAMESPACE "MetasoundGraphNode"
+#define LOCTEXT_NAMESPACE "MetasoundEditor"
 
 
 namespace Metasound
 {
 	namespace Editor
 	{
-		class SMetaSoundGraphPinKnot : public SGraphPinKnot
-		{
-		public:
-			SLATE_BEGIN_ARGS(SMetaSoundGraphPinKnot) {}
-			SLATE_END_ARGS()
-
-			void Construct(const FArguments& InArgs, UEdGraphPin* InPin);
-
-			virtual const FSlateBrush* GetPinIcon() const override;
-			virtual FSlateColor GetPinColor() const override;
-
-		protected:
-			bool HasRequiredConnections() const;
-		};
-
-
 		SMetaSoundGraphNode::~SMetaSoundGraphNode()
 		{
 			UMetasoundEditorGraphNode& Node = GetMetaSoundNode();
@@ -191,7 +176,7 @@ namespace Metasound
 		{
 			const FText ToolTip = InToolTip
 				? *InToolTip
-				: LOCTEXT("TriggerTestToolTip", "Executes trigger if currently previewing MetaSound.");
+				: LOCTEXT("MetasoundGraphNode_TriggerTestToolTip", "Executes trigger if currently previewing MetaSound.");
 
 			TSharedPtr<SButton> SimulationButton;
 			TSharedRef<SWidget> SimulationWidget = SNew(SHorizontalBox)
@@ -669,7 +654,7 @@ namespace Metasound
 							{
 								if (!this->bIsInputWidgetTransacting)
 								{
-									GEditor->BeginTransaction(LOCTEXT("MetasoundSetInputDefault", "Set MetaSound Input Default"));
+									GEditor->BeginTransaction(LOCTEXT("MetasoundGraphNode_MetasoundSetInputDefault", "Set MetaSound Input Default"));
 									this->bIsInputWidgetTransacting = true;
 								}
 								GraphMember->GetOwningGraph()->GetMetasound()->Modify();
@@ -955,101 +940,6 @@ namespace Metasound
 			check(GraphNode);
 			return *Cast<UMetasoundEditorGraphNode>(GraphNode);
 		}
-
-		void SMetaSoundGraphPinKnot::Construct(const FArguments& InArgs, UEdGraphPin* InPin)
-		{
-			SGraphPinKnot::Construct(SGraphPinKnot::FArguments(), InPin);
-
-			TSharedRef<SWidget> PinWidgetRef = SPinTypeSelector::ConstructPinTypeImage(
-				MakeAttributeSP(this, &SMetaSoundGraphPinKnot::GetPinIcon),
-				MakeAttributeSP(this, &SMetaSoundGraphPinKnot::GetPinColor),
-				MakeAttributeSP(this, &SMetaSoundGraphPinKnot::GetSecondaryPinIcon),
-				MakeAttributeSP(this, &SMetaSoundGraphPinKnot::GetSecondaryPinColor));
-			PinImage = PinWidgetRef;
-		}
-
-		FSlateColor SMetaSoundGraphPinKnot::GetPinColor() const
-		{
-			if (UEdGraphPin* Pin = GetPinObj())
-			{
-				if (Pin->Direction == EGPD_Output)
-				{
-					if (!HasRequiredConnections() || Pin->GetOwningNode()->ErrorType <= static_cast<uint32>(EMessageSeverity::Warning))
-					{
-						return FLinearColor::Yellow;
-					}
-				}
-			}
-
-			return SGraphPin::GetPinColor();
-		}
-
-		const FSlateBrush* SMetaSoundGraphPinKnot::GetPinIcon() const
-		{
-			using namespace Metasound::Frontend;
-
-			bool bIsConstructorPin = false;
-
-			if (UEdGraphPin* Pin = GetPinObj())
-			{
-				if (!HasRequiredConnections() || Pin->GetOwningNode()->ErrorType <= static_cast<uint32>(EMessageSeverity::Warning))
-				{
-					return &Editor::Style::GetSlateBrushSafe("MetasoundEditor.Graph.InvalidReroute");
-				}
-
-				Pin = FGraphBuilder::FindReroutedOutputPin(Pin);
-				if (const UMetasoundEditorGraphNode* Node = Cast<UMetasoundEditorGraphNode>(Pin->GetOwningNode()))
-				{
-					if (const UMetasoundEditorGraphMemberNode* MemberNode = Cast<UMetasoundEditorGraphMemberNode>(Node))
-					{
-						if (const UMetasoundEditorGraphInput* Input = Cast<UMetasoundEditorGraphInput>(MemberNode->GetMember()))
-						{
-							FConstNodeHandle NodeHandle = Input->GetConstNodeHandle();
-							TArray<FConstOutputHandle> OutputHandles = NodeHandle->GetConstOutputs();
-
-							check(!OutputHandles.IsEmpty());
-							bIsConstructorPin = OutputHandles.Last()->GetVertexAccessType() == EMetasoundFrontendVertexAccessType::Value;
-						}
-					}
-				}
-			}
-
-			if (bIsConstructorPin)
-			{
-				const bool bIsConnected = IsConnected();
-				if (IsArray())
-				{
-					const FName BrushName = bIsConnected ? "MetasoundEditor.Graph.ConstructorPinArray" : "MetasoundEditor.Graph.ConstructorPinArrayDisconnected";
-					return &Editor::Style::GetSlateBrushSafe(BrushName);
-				}
-				else
-				{
-					const FName BrushName = bIsConnected ? "MetasoundEditor.Graph.ConstructorPin" : "MetasoundEditor.Graph.ConstructorPinDisconnected";
-					return &Editor::Style::GetSlateBrushSafe(BrushName);
-				}
-			}
-
-			return SGraphPin::GetPinIcon();
-		}
-
-		bool SMetaSoundGraphPinKnot::HasRequiredConnections() const
-		{
-			using namespace Frontend;
-
-			if (UEdGraphPin* Pin = GetPinObj())
-			{
-				if (const UMetasoundEditorGraphExternalNode* OwningNode = Cast<UMetasoundEditorGraphExternalNode>(Pin->GetOwningNode()))
-				{
-					if (const INodeTemplate* Template = INodeTemplateRegistry::Get().FindTemplate(FRerouteNodeTemplate::GetRegistryKey()))
-					{
-						FConstNodeHandle NodeHandle = OwningNode->GetConstNodeHandle();
-						return Template->HasRequiredConnections(NodeHandle);
-					}
-				}
-			}
-
-			return false;
-		}
 	} // namespace Editor
 } // namespace Metasound
-#undef LOCTEXT_NAMESPACE // MetasoundGraphNode
+#undef LOCTEXT_NAMESPACE // MetasoundEditor
diff --git a/Engine/Plugins/Runtime/Metasound/Source/MetasoundEditor/Private/SMetasoundGraphPin.cpp b/Engine/Plugins/Runtime/Metasound/Source/MetasoundEditor/Private/SMetasoundGraphPin.cpp
new file mode 100644
--- /dev/null
+++ b/Engine/Plugins/Runtime/Metasound/Source/MetasoundEditor/Private/SMetasoundGraphPin.cpp
@@ -0,0 +1,110 @@
+// Copyright Epic Games, Inc. All Rights Reserved.
+#include "SMetasoundGraphPin.h"
+
+#include "Logging/TokenizedMessage.h"
+#include "MetasoundEditorModule.h"
+#include "MetasoundFrontendNodeTemplateRegistry.h"
+#include "NodeTemplates/MetasoundFrontendNodeTemplateReroute.h"
+#include "SPinTypeSelector.h"
+
+
+namespace Metasound
+{
+	namespace Editor
+	{
+		void SMetaSoundGraphPinKnot::Construct(const FArguments& InArgs, UEdGraphPin* InPin)
+		{
+			SGraphPinKnot::Construct(SGraphPinKnot::FArguments(), InPin);
+
+			TSharedRef<SWidget> PinWidgetRef = SPinTypeSelector::ConstructPinTypeImage(
+				MakeAttributeSP(this, &SMetaSoundGraphPinKnot::GetPinIcon),
+				MakeAttributeSP(this, &SMetaSoundGraphPinKnot::GetPinColor),
+				MakeAttributeSP(this, &SMetaSoundGraphPinKnot::GetSecondaryPinIcon),
+				MakeAttributeSP(this, &SMetaSoundGraphPinKnot::GetSecondaryPinColor));
+			PinImage = PinWidgetRef;
+		}
+
+		FSlateColor SMetaSoundGraphPinKnot::GetPinColor() const
+		{
+			if (const UEdGraphPin* Pin = SGraphPinKnot::GetPinObj())
+			{
+				if (Pin->Direction == EGPD_Output)
+				{
+					if (!HasRequiredConnections() || Pin->GetOwningNode()->ErrorType <= static_cast<uint32>(EMessageSeverity::Warning))
+					{
+						return FLinearColor::Yellow;
+					}
+				}
+			}
+
+			return SGraphPinKnot::GetPinColor();
+		}
+
+		const FSlateBrush* SMetaSoundGraphPinKnot::GetPinIcon() const
+		{
+			using namespace Metasound::Frontend;
+
+			bool bIsConstructorPin = false;
+
+			if (const UEdGraphPin* Pin = SGraphPinKnot::GetPinObj())
+			{
+				if (!HasRequiredConnections() || Pin->GetOwningNode()->ErrorType <= static_cast<uint32>(EMessageSeverity::Warning))
+				{
+					return &Editor::Style::GetSlateBrushSafe("MetasoundEditor.Graph.InvalidReroute");
+				}
+
+				Pin = FGraphBuilder::FindReroutedOutputPin(Pin);
+				if (const UMetasoundEditorGraphNode* Node = Cast<UMetasoundEditorGraphNode>(Pin->GetOwningNode()))
+				{
+					if (const UMetasoundEditorGraphMemberNode* MemberNode = Cast<UMetasoundEditorGraphMemberNode>(Node))
+					{
+						if (const UMetasoundEditorGraphInput* Input = Cast<UMetasoundEditorGraphInput>(MemberNode->GetMember()))
+						{
+							FConstNodeHandle NodeHandle = Input->GetConstNodeHandle();
+							TArray<FConstOutputHandle> OutputHandles = NodeHandle->GetConstOutputs();
+
+							check(!OutputHandles.IsEmpty());
+							bIsConstructorPin = OutputHandles.Last()->GetVertexAccessType() == EMetasoundFrontendVertexAccessType::Value;
+						}
+					}
+				}
+			}
+
+			if (bIsConstructorPin)
+			{
+				const bool bIsConnected = IsConnected();
+				if (IsArray())
+				{
+					const FName BrushName = bIsConnected ? "MetasoundEditor.Graph.ConstructorPinArray" : "MetasoundEditor.Graph.ConstructorPinArrayDisconnected";
+					return &Editor::Style::GetSlateBrushSafe(BrushName);
+				}
+				else
+				{
+					const FName BrushName = bIsConnected ? "MetasoundEditor.Graph.ConstructorPin" : "MetasoundEditor.Graph.ConstructorPinDisconnected";
+					return &Editor::Style::GetSlateBrushSafe(BrushName);
+				}
+			}
+
+			return SGraphPinKnot::GetPinIcon();
+		}
+
+		bool SMetaSoundGraphPinKnot::HasRequiredConnections() const
+		{
+			using namespace Frontend;
+
+			if (UEdGraphPin* Pin = SGraphPinKnot::GetPinObj())
+			{
+				if (const UMetasoundEditorGraphExternalNode* OwningNode = Cast<UMetasoundEditorGraphExternalNode>(Pin->GetOwningNode()))
+				{
+					if (const INodeTemplate* Template = INodeTemplateRegistry::Get().FindTemplate(FRerouteNodeTemplate::GetRegistryKey()))
+					{
+						FConstNodeHandle NodeHandle = OwningNode->GetConstNodeHandle();
+						return Template->HasRequiredConnections(NodeHandle);
+					}
+				}
+			}
+
+			return false;
+		}
+	} // namespace Editor
+} // namespace Metasound
diff --git a/Engine/Plugins/Runtime/Metasound/Source/MetasoundEditor/Private/SMetasoundGraphPin.h b/Engine/Plugins/Runtime/Metasound/Source/MetasoundEditor/Private/SMetasoundGraphPin.h
--- a/Engine/Plugins/Runtime/Metasound/Source/MetasoundEditor/Private/SMetasoundGraphPin.h
+++ b/Engine/Plugins/Runtime/Metasound/Source/MetasoundEditor/Private/SMetasoundGraphPin.h
@@ -1,6 +1,7 @@
 // Copyright Epic Games, Inc. All Rights Reserved.
 #pragma once
 
+#include "EdGraph/EdGraphPin.h"
 #include "Styling/AppStyle.h"
 #include "KismetPins/SGraphPinBool.h"
 #include "KismetPins/SGraphPinNum.h"
@@ -13,12 +14,16 @@
 #include "MetasoundEditorGraphNode.h"
 #include "MetasoundFrontendController.h"
 #include "SGraphPin.h"
+#include "SGraphNodeKnot.h"
 #include "SMetasoundPinValueInspector.h"
+#include "SPinTypeSelector.h"
 #include "Styling/SlateStyleRegistry.h"
 #include "Widgets/Input/SButton.h"
 #include "Widgets/SBoxPanel.h"
 #include "Widgets/SNullWidget.h"
 #include "Widgets/SWidget.h"
+#include "Styling/SlateColor.h"
+#include "Styling/SlateBrush.h"
 
 #define LOCTEXT_NAMESPACE "MetasoundEditor"
 
@@ -28,16 +33,89 @@ namespace Metasound
 	namespace Editor
 	{
 		template <typename ParentPinType>
-		class METASOUNDEDITOR_API TMetasoundGraphPin : public ParentPinType
+		class TMetasoundGraphPin : public ParentPinType
 		{
 			TSharedPtr<SMetasoundPinValueInspector> PinInspector;
 
+		protected:
+			static TWeakPtr<FPinValueInspectorTooltip> OpenPinInspector(UEdGraphPin& InPin, TSharedPtr<SMetasoundPinValueInspector>& OutPinInspector)
+			{
+				TSharedPtr<SMetasoundPinValueInspector> NewPinInspector = SNew(SMetasoundPinValueInspector);
+				TWeakPtr<FPinValueInspectorTooltip> NewTooltip = FPinValueInspectorTooltip::SummonTooltip(&InPin, NewPinInspector);
+				if (NewTooltip.IsValid())
+				{
+					OutPinInspector = NewPinInspector;
+					return NewTooltip;
+				}
+
+				return nullptr;
+			}
+
+			static void UpdatePinInspector(
+				UEdGraphPin& InPin,
+				const bool bIsHoveringPin,
+				TSharedPtr<SMetasoundPinValueInspector>& OutPinInspector,
+				TWeakPtr<FPinValueInspectorTooltip>& OutInspectorTooltip,
+				TFunctionRef<void(FVector2D&)> InGetTooltipLocation)
+			{
+				const bool bCanInspectPin = FGraphBuilder::CanInspectPin(&InPin);
+				if (bIsHoveringPin && bCanInspectPin)
+				{
+					if (OutPinInspector.IsValid())
+					{
+						const UEdGraphPin* InspectedPin = OutPinInspector->GetPinRef().Get();
+						if (InspectedPin == &InPin)
+						{
+							OutPinInspector->UpdateMessage();
+						}
+					}
+					else
+					{
+						OutInspectorTooltip = OpenPinInspector(InPin, OutPinInspector);
+						TSharedPtr<FPinValueInspectorTooltip> NewTooltip = OutInspectorTooltip.Pin();
+						if (NewTooltip.IsValid())
+						{
+							FVector2D TooltipLocation;
+							InGetTooltipLocation(TooltipLocation);
+							NewTooltip->MoveTooltip(TooltipLocation);
+						}
+					}
+				}
+				else if (OutPinInspector.IsValid())
+				{
+					TSharedPtr<FPinValueInspectorTooltip> InspectorTooltip = OutInspectorTooltip.Pin();
+					if (InspectorTooltip.IsValid())
+					{
+						if (InspectorTooltip->TooltipCanClose())
+						{
+							constexpr bool bForceDismiss = true;
+							InspectorTooltip->TryDismissTooltip(bForceDismiss);
+							OutInspectorTooltip.Reset();
+							OutPinInspector.Reset();
+						}
+					}
+					else
+					{
+						OutPinInspector.Reset();
+					}
+				}
+			}
+
 		public:
 			SLATE_BEGIN_ARGS(TMetasoundGraphPin<ParentPinType>)
 			{
 			}
 			SLATE_END_ARGS()
 
+			virtual ~TMetasoundGraphPin() = default;
+
+			void CacheNodeOffset(const FGeometry& AllottedGeometry)
+			{
+				const FVector2D UnscaledPosition = ParentPinType::OwnerNodePtr.Pin()->GetUnscaledPosition();
+				ParentPinType::CachedNodeOffset = FVector2D(AllottedGeometry.AbsolutePosition) / AllottedGeometry.Scale - UnscaledPosition;
+				ParentPinType::CachedNodeOffset.Y += AllottedGeometry.Size.Y * 0.5f;
+			}
+
 			Frontend::FConstInputHandle GetConstInputHandle() const
 			{
 				using namespace Frontend;
@@ -264,61 +342,16 @@ namespace Metasound
 
 			virtual void Tick(const FGeometry& AllottedGeometry, const double InCurrentTime, const float InDeltaTime) override
 			{
-				ParentPinType::CachedNodeOffset = FVector2D(AllottedGeometry.AbsolutePosition) / AllottedGeometry.Scale - ParentPinType::OwnerNodePtr.Pin()->GetUnscaledPosition();
-				ParentPinType::CachedNodeOffset.Y += AllottedGeometry.Size.Y * 0.5f;
+				CacheNodeOffset(AllottedGeometry);
 
-				UEdGraphPin* GraphPin = ParentPinType::GetPinObj();
-
-				// Pause updates if menu is hosted (so user can capture state for ex. using the copy value action)
-				TSharedPtr<FPinValueInspectorTooltip> InspectorTooltip = ParentPinType::ValueInspectorTooltip.Pin();
-				if (InspectorTooltip.IsValid())
-				{
-					const bool bIsHoveringPin = ParentPinType::IsHovered();
-					const bool bIsInspectingPin = bIsHoveringPin && FGraphBuilder::CanInspectPin(GraphPin);
-					if (bIsInspectingPin)
-					{
-						if (PinInspector.IsValid())
-						{
-							PinInspector->UpdateMessage();
-						}
-					}
-					else
-					{
-						if (InspectorTooltip->TooltipCanClose())
-						{
-							PinInspector.Reset();
-							InspectorTooltip->TryDismissTooltip();
-						}
-					}
-				}
-				else
+				if (UEdGraphPin* GraphPin = ParentPinType::GetPinObj())
 				{
 					const bool bIsHoveringPin = ParentPinType::IsHovered();
-					const bool bCanInspectPin = FGraphBuilder::CanInspectPin(GraphPin);
-					if (bIsHoveringPin && bCanInspectPin)
-					{
-						const FEdGraphPinReference* CurrentRef = nullptr;
-						if (FPinValueInspectorTooltip::ValueInspectorWidget.IsValid())
-						{
-							CurrentRef = &FPinValueInspectorTooltip::ValueInspectorWidget->GetPinRef();
-						}
-
-						// Only update if reference is not already set.  This avoids ping-pong between pins which can happen
-						// when hovering connections as this state causes IsHovered to return true for all the associated pins
-						// for the given connection.
-						if (!CurrentRef || !CurrentRef->Get() || CurrentRef->Get() == GraphPin)
+					UpdatePinInspector(*GraphPin, bIsHoveringPin, PinInspector, ParentPinType::ValueInspectorTooltip,
+						[this](FVector2D& OutTooltipLocation)
 						{
-							PinInspector = SNew(SMetasoundPinValueInspector);
-							ParentPinType::ValueInspectorTooltip = FPinValueInspectorTooltip::SummonTooltip(GraphPin, PinInspector);
-							InspectorTooltip = ParentPinType::ValueInspectorTooltip.Pin();
-							if (InspectorTooltip.IsValid())
-							{
-								FVector2D TooltipLocation;
-								ParentPinType::GetInteractiveTooltipLocation(TooltipLocation);
-								InspectorTooltip->MoveTooltip(TooltipLocation);
-							}
-						}
-					}
+							ParentPinType::GetInteractiveTooltipLocation(OutTooltipLocation);
+						});
 				}
 			}
 		};
@@ -326,6 +359,8 @@ namespace Metasound
 		class SMetasoundGraphPin : public TMetasoundGraphPin<SGraphPin>
 		{
 		public:
+			virtual ~SMetasoundGraphPin() = default;
+
 			void Construct(const FArguments& InArgs, UEdGraphPin* InGraphPinObj)
 			{
 				SGraphPin::Construct(SGraphPin::FArguments(), InGraphPinObj);
@@ -335,6 +370,8 @@ namespace Metasound
 		class SMetasoundGraphPinBool : public TMetasoundGraphPin<SGraphPinBool>
 		{
 		public:
+			virtual ~SMetasoundGraphPinBool() = default;
+
 			void Construct(const FArguments& InArgs, UEdGraphPin* InGraphPinObj)
 			{
 				SGraphPinBool::Construct(SGraphPinBool::FArguments(), InGraphPinObj);
@@ -344,6 +381,8 @@ namespace Metasound
 		class SMetasoundGraphPinFloat : public TMetasoundGraphPin<SGraphPinNum<float>>
 		{
 		public:
+			virtual ~SMetasoundGraphPinFloat() = default;
+
 			void Construct(const FArguments& InArgs, UEdGraphPin* InGraphPinObj)
 			{
 				SGraphPinNum<float>::Construct(SGraphPinNum<float>::FArguments(), InGraphPinObj);
@@ -353,6 +392,8 @@ namespace Metasound
 		class SMetasoundGraphPinInteger : public TMetasoundGraphPin<SGraphPinInteger>
 		{
 		public:
+			virtual ~SMetasoundGraphPinInteger() = default;
+
 			void Construct(const FArguments& InArgs, UEdGraphPin* InGraphPinObj)
 			{
 				SGraphPinInteger::Construct(SGraphPinInteger::FArguments(), InGraphPinObj);
@@ -362,6 +403,8 @@ namespace Metasound
 		class SMetasoundGraphPinObject : public TMetasoundGraphPin<SGraphPinObject>
 		{
 		public:
+			virtual ~SMetasoundGraphPinObject() = default;
+
 			void Construct(const FArguments& InArgs, UEdGraphPin* InGraphPinObj)
 			{
 				SGraphPinObject::Construct(SGraphPinObject::FArguments(), InGraphPinObj);
@@ -371,11 +414,28 @@ namespace Metasound
 		class SMetasoundGraphPinString : public TMetasoundGraphPin<SGraphPinString>
 		{
 		public:
+			virtual ~SMetasoundGraphPinString() = default;
+
 			void Construct(const FArguments& InArgs, UEdGraphPin* InGraphPinObj)
 			{
 				SGraphPinString::Construct(SGraphPinString::FArguments(), InGraphPinObj);
 			}
 		};
+
+		class SMetaSoundGraphPinKnot : public TMetasoundGraphPin<SGraphPinKnot>
+		{
+		public:
+			virtual ~SMetaSoundGraphPinKnot() = default;
+
+			void Construct(const FArguments& InArgs, UEdGraphPin* InPin);
+
+			virtual FSlateColor GetPinColor() const override;
+
+			virtual const FSlateBrush* GetPinIcon() const override;
+
+		protected:
+			bool HasRequiredConnections() const;
+		};
 	} // namespace Editor
 } // namespace Metasound
 #undef LOCTEXT_NAMESPACE
diff --git a/Engine/Plugins/Runtime/Metasound/Source/MetasoundEditor/Private/SMetasoundPinValueInspector.cpp b/Engine/Plugins/Runtime/Metasound/Source/MetasoundEditor/Private/SMetasoundPinValueInspector.cpp
--- a/Engine/Plugins/Runtime/Metasound/Source/MetasoundEditor/Private/SMetasoundPinValueInspector.cpp
+++ b/Engine/Plugins/Runtime/Metasound/Source/MetasoundEditor/Private/SMetasoundPinValueInspector.cpp
@@ -77,6 +77,21 @@ namespace Metasound
 				return Literals;
 			}
 
+			UEdGraphPin* ResolvePinObjectAsOutput(UEdGraphPin* InPin)
+			{
+				if (InPin && !InPin->LinkedTo.IsEmpty())
+				{
+					// Swap to show connected output if input (Only ever one)
+					if (InPin->Direction == EGPD_Input)
+					{
+						InPin = InPin->LinkedTo.Last();
+						check(InPin->Direction == EGPD_Output);
+					}
+				}
+
+				return InPin;
+			}
+
 			void SetLiteralFromText(const FName InDataType, const FText& InText, FMetasoundFrontendLiteral& OutLiteral)
 			{
 				const FString TextString = InText.ToString();
@@ -104,27 +119,15 @@ namespace Metasound
 
 		FMetasoundNumericDebugLineItem::FMetasoundNumericDebugLineItem(UEdGraphPin* InGraphPinObj, FGetValueStringFunction&& InGetValueStringFunction)
 			: FDebugLineItem(DLT_Message)
-			, GraphPinObj(InGraphPinObj)
+			, GraphPinObj(PinValueInspectorPrivate::ResolvePinObjectAsOutput(InGraphPinObj))
 			, GetValueStringFunction(MoveTemp(InGetValueStringFunction))
 		{
 			using namespace Frontend;
 
-			if (GraphPinObj && !GraphPinObj->LinkedTo.IsEmpty())
-			{
-				// Swap to show connected output if input (Only ever one)
-				if (GraphPinObj->Direction == EGPD_Input)
-				{
-					GraphPinObj = GraphPinObj->LinkedTo.Last();
-					check(GraphPinObj->Direction == EGPD_Output);
-				}
-
-				FConstOutputHandle OutputHandle = GetOutputHandle();
-				ColorLiterals = PinValueInspectorPrivate::InitializeLiterals(OutputHandle->GetDataType());
-
-				Update();
-			}
-
-			DisplayName = FGraphBuilder::GetDisplayName(*GetOutputHandle());
+			FConstOutputHandle OutputHandle = GetReroutedOutputHandle();
+			ColorLiterals = PinValueInspectorPrivate::InitializeLiterals(GetReroutedOutputHandle()->GetDataType());
+			Update();
+			DisplayName = FGraphBuilder::GetDisplayName(*GetReroutedOutputHandle());
 		}
 
 		FColorPickerArgs FMetasoundNumericDebugLineItem::InitPickerArgs()
@@ -155,12 +158,25 @@ namespace Metasound
 			return *Node;
 		}
 
-		Frontend::FConstOutputHandle FMetasoundNumericDebugLineItem::GetOutputHandle() const
+		UObject* FMetasoundNumericDebugLineItem::GetOutermostObject()
+		{
+			if (GraphPinObj)
+			{
+				if (UEdGraphNode* OwningNode = GraphPinObj->GetOwningNode())
+				{
+					return OwningNode->GetOutermostObject();
+				}
+			}
+
+			return nullptr;
+		}
+
+		Frontend::FConstOutputHandle FMetasoundNumericDebugLineItem::GetReroutedOutputHandle() const
 		{
 			return FGraphBuilder::FindReroutedConstOutputHandleFromPin(GraphPinObj);
 		}
 
-		Frontend::FOutputHandle FMetasoundNumericDebugLineItem::GetOutputHandle()
+		Frontend::FOutputHandle FMetasoundNumericDebugLineItem::GetReroutedOutputHandle()
 		{
 			return FGraphBuilder::FindReroutedOutputHandleFromPin(GraphPinObj);
 		}
@@ -183,20 +199,20 @@ namespace Metasound
 
 			FScopedTransaction Transaction(PinValueInspectorPrivate::DisableValueColorizationText);
 
-			const UMetasoundEditorGraphNode& Node = GetReroutedNodeChecked();
-			UObject* ParentObject = Node.GetOutermostObject();
+			UObject* ParentObject = GetOutermostObject();
 			check(ParentObject);
 			ParentObject->Modify();
 
 			FGraphHandle GraphHandle = GetGraphHandle();
 			FMetasoundFrontendGraphStyle Style = GraphHandle->GetGraphStyle();
 
-			const FGuid NodeID = Node.GetNodeID();
-			const FName OutputName = GetOutputHandle()->GetName();
+			const FConstOutputHandle OutputHandle = GetReroutedOutputHandle();
+			const FGuid NodeID = OutputHandle->GetOwningNodeID();
+			const FName OutputName = OutputHandle->GetName();
 			Style.EdgeStyles.RemoveAllSwap([&NodeID, &OutputName](FMetasoundFrontendEdgeStyle& EdgeStyle)
-				{
-					return EdgeStyle.NodeID == NodeID && EdgeStyle.OutputName == OutputName;
-				});
+			{
+				return EdgeStyle.NodeID == NodeID && EdgeStyle.OutputName == OutputName;
+			});
 
 			GraphHandle->SetGraphStyle(Style);
 			bIsValueColorizationEnabled = false;
@@ -208,8 +224,7 @@ namespace Metasound
 
 			FScopedTransaction Transaction(PinValueInspectorPrivate::EnableValueColorizationText);
 
-			const UMetasoundEditorGraphNode& Node = GetReroutedNodeChecked();
-			UObject* ParentObject = Node.GetOutermostObject();
+			UObject* ParentObject = GetOutermostObject();
 			check(ParentObject);
 			ParentObject->Modify();
 
@@ -221,9 +236,9 @@ namespace Metasound
 
 			FGraphHandle GraphHandle = GetGraphHandle();
 			FMetasoundFrontendGraphStyle Style = GraphHandle->GetGraphStyle();
-			FConstOutputHandle OutputHandle = GetOutputHandle();
+			FConstOutputHandle OutputHandle = GetReroutedOutputHandle();
 
-			const FGuid NodeID = Node.GetNodeID();
+			const FGuid NodeID = OutputHandle->GetOwningNodeID();
 			const FName OutputName = OutputHandle->GetName();
 			FMetasoundFrontendEdgeStyle* EdgeStyle = Style.EdgeStyles.FindByPredicate([&NodeID, &OutputName](FMetasoundFrontendEdgeStyle& EdgeStyle)
 			{
@@ -248,19 +263,20 @@ namespace Metasound
 			using namespace Frontend;
 
 			static const FText SetValueFormat = LOCTEXT("PinValueInspector_SetMinColorFormat", "Set '{0}[{1}]' Color to '{2}'");
-			const FName Name = GetOutputHandle()->GetName();
+
+			const FConstOutputHandle OutputHandle = GetReroutedOutputHandle();
+			const FName Name = OutputHandle->GetName();
 			FScopedTransaction Transaction(FText::Format(SetValueFormat, FText::FromName(Name), FText::AsNumber(InIndex), FText::FromString(InColor.ToString())));
 
-			const UMetasoundEditorGraphNode& Node = GetReroutedNodeChecked();
-			UObject* ParentObject = Node.GetOutermostObject();
+			UObject* ParentObject = GetOutermostObject();
 			check(ParentObject);
 			ParentObject->Modify();
 
 			FGraphHandle GraphHandle = GetGraphHandle();
 			FMetasoundFrontendGraphStyle Style = GraphHandle->GetGraphStyle();
 
-			const FGuid NodeID = Node.GetNodeID();
-			const FName OutputName = GetOutputHandle()->GetName();
+			const FGuid NodeID = OutputHandle->GetOwningNodeID();
+			const FName OutputName = GetReroutedOutputHandle()->GetName();
 			FMetasoundFrontendEdgeStyle* EdgeStyle = Style.EdgeStyles.FindByPredicate([&NodeID, &OutputName](FMetasoundFrontendEdgeStyle& EdgeStyle)
 			{
 				return EdgeStyle.NodeID == NodeID && EdgeStyle.OutputName == OutputName;
@@ -281,40 +297,37 @@ namespace Metasound
 		{
 			using namespace Frontend;
 
-			if (const UMetasoundEditorGraphNode* Node = GetReroutedNode())
-			{
-				FOutputHandle OutputHandle = GetOutputHandle();
+			FOutputHandle OutputHandle = GetReroutedOutputHandle();
 
-				static const FText SetValueFormat = LOCTEXT("PinValueInspector_SetColorLiteralValueFormat", "Set '{0}[{1}]' EdgeStyle Color Value to '{2}'");
-				const FName Name = OutputHandle->GetName();
-				FScopedTransaction Transaction(FText::Format(SetValueFormat, FText::FromName(Name), FText::AsNumber(InIndex), InValueText));
+			static const FText SetValueFormat = LOCTEXT("PinValueInspector_SetColorLiteralValueFormat", "Set '{0}[{1}]' EdgeStyle Color Value to '{2}'");
+			const FName Name = OutputHandle->GetName();
+			FScopedTransaction Transaction(FText::Format(SetValueFormat, FText::FromName(Name), FText::AsNumber(InIndex), InValueText));
 
-				UObject* ParentObject = Node->GetOutermostObject();
-				check(ParentObject);
-				ParentObject->Modify();
+			UObject* ParentObject = GetOutermostObject();
+			check(ParentObject);
+			ParentObject->Modify();
 
-				const FGuid NodeID = Node->GetNodeID();
-				const FName OutputName = OutputHandle->GetName();
+			const FGuid NodeID = OutputHandle->GetOwningNodeID();
+			const FName OutputName = OutputHandle->GetName();
 
-				FGraphHandle OwningGraph = GetGraphHandle();
-				FMetasoundFrontendGraphStyle Style = OwningGraph->GetGraphStyle();
+			FGraphHandle OwningGraph = GetGraphHandle();
+			FMetasoundFrontendGraphStyle Style = OwningGraph->GetGraphStyle();
 
-				bool bUpdated = false;
-				for (FMetasoundFrontendEdgeStyle& EdgeStyle : Style.EdgeStyles)
+			bool bUpdated = false;
+			for (FMetasoundFrontendEdgeStyle& EdgeStyle : Style.EdgeStyles)
+			{
+				if (EdgeStyle.NodeID == NodeID && EdgeStyle.OutputName == OutputName)
 				{
-					if (EdgeStyle.NodeID == NodeID && EdgeStyle.OutputName == OutputName)
+					if (ensure(EdgeStyle.LiteralColorPairs.Num() > InIndex))
 					{
-						if (ensure(EdgeStyle.LiteralColorPairs.Num() > InIndex))
-						{
-							const FName DataType = OutputHandle->GetDataType();
-							PinValueInspectorPrivate::SetLiteralFromText(DataType, InValueText, OutNewLiteral);
-							EdgeStyle.LiteralColorPairs[InIndex].Value = OutNewLiteral;
-						}
+						const FName DataType = OutputHandle->GetDataType();
+						PinValueInspectorPrivate::SetLiteralFromText(DataType, InValueText, OutNewLiteral);
+						EdgeStyle.LiteralColorPairs[InIndex].Value = OutNewLiteral;
 					}
 				}
-
-				OwningGraph->SetGraphStyle(Style);
 			}
+
+			OwningGraph->SetGraphStyle(Style);
 		}
 
 		Frontend::FGraphHandle FMetasoundNumericDebugLineItem::GetGraphHandle()
@@ -353,9 +366,9 @@ namespace Metasound
 				return nullptr;
 			}
 
-			const UMetasoundEditorGraphNode& Node = GetReroutedNodeChecked();
-			const FGuid NodeID = Node.GetNodeID();
-			const FName OutputName = GetOutputHandle()->GetName();
+			const FConstOutputHandle OutputHandle = GetReroutedOutputHandle();
+			const FGuid NodeID = OutputHandle->GetOwningNodeID();
+			const FName OutputName = OutputHandle->GetName();
 			const FMetasoundFrontendGraphStyle& Style = OwningGraph->GetGraphStyle();
 
 			return Style.EdgeStyles.FindByPredicate([&NodeID, &OutputName](FMetasoundFrontendEdgeStyle& EdgeStyle)
@@ -370,12 +383,10 @@ namespace Metasound
 
 			bIsValueColorizationEnabled = false;
 
-			const FMetasoundFrontendEdgeStyle* EdgeStyle = GetEdgeStyle();
-
- 			const UMetasoundEditorGraphNode& Node = GetReroutedNodeChecked();
-			FGuid NodeID = Node.GetNodeID();
-			FName OutputName = GetOutputHandle()->GetName();
-			if (EdgeStyle)
+			const FConstOutputHandle OutputHandle = GetReroutedOutputHandle();
+			const FGuid NodeID = OutputHandle->GetOwningNodeID();
+			const FName OutputName = OutputHandle->GetName();
+			if (const FMetasoundFrontendEdgeStyle* EdgeStyle = GetEdgeStyle())
 			{
 				if (ensure(EdgeStyle->LiteralColorPairs.Num() == 2))
 				{
@@ -461,7 +472,7 @@ namespace Metasound
 					.Font(FontInfo)
 				];
 
-			const FName DataType = GetOutputHandle()->GetDataType();
+			const FName DataType = GetReroutedOutputHandle()->GetDataType();
 			const bool bConnectionAnimationsEnabled = GetDefault<UMetasoundEditorSettings>()->AnalyzerAnimationSettings.bAnimateConnections;
 			const bool bIsTypeSupported = PinValueInspectorPrivate::IsMetasoundPrimitiveValueColorizationSupported(DataType);
 			if (bConnectionAnimationsEnabled && bIsTypeSupported)
@@ -640,7 +651,7 @@ namespace Metasound
 		void SMetasoundPinValueInspector::PopulateTreeView()
 		{
 			// Locate the class property associated with the source pin and set it to the root node.
-			GraphPinObj = GetPinRef().Get();
+			UEdGraphPin* GraphPinObj = PinValueInspectorPrivate::ResolvePinObjectAsOutput(GetPinRef().Get());
 
 			// Don't populate if no pin found or pin is not connected (analyzers are not enabled for disconnected node outputs)
 			if (!GraphPinObj || GraphPinObj->LinkedTo.IsEmpty())
diff --git a/Engine/Plugins/Runtime/Metasound/Source/MetasoundEditor/Private/SMetasoundPinValueInspector.h b/Engine/Plugins/Runtime/Metasound/Source/MetasoundEditor/Private/SMetasoundPinValueInspector.h
--- a/Engine/Plugins/Runtime/Metasound/Source/MetasoundEditor/Private/SMetasoundPinValueInspector.h
+++ b/Engine/Plugins/Runtime/Metasound/Source/MetasoundEditor/Private/SMetasoundPinValueInspector.h
@@ -6,9 +6,11 @@
 #include "Engine/Engine.h"
 #include "InputCoreTypes.h"
 #include "Input/Reply.h"
+#include "MetasoundEditorGraphBuilder.h"
 #include "MetasoundEditorGraphConnectionManager.h"
 #include "Misc/Attribute.h"
 #include "ScopedTransaction.h"
+#include "SGraphPin.h"
 #include "SPinValueInspector.h"
 #include "Templates/Function.h"
 #include "Templates/SharedPointer.h"
@@ -57,9 +59,11 @@ namespace Metasound
 
 			const UMetasoundEditorGraphNode& GetReroutedNodeChecked() const;
 
-			Frontend::FConstOutputHandle GetOutputHandle() const;
+			UObject* GetOutermostObject();
 
-			Frontend::FOutputHandle GetOutputHandle();
+			Frontend::FConstOutputHandle GetReroutedOutputHandle() const;
+
+			Frontend::FOutputHandle GetReroutedOutputHandle();
 
 			Frontend::FGraphHandle GetGraphHandle();
 
@@ -99,9 +103,6 @@ namespace Metasound
 
 		class METASOUNDEDITOR_API SMetasoundPinValueInspector : public SPinValueInspector
 		{
-			/** The GraphPin that this widget is inspecting. */
-			UEdGraphPin* GraphPinObj;
-
 			TSharedPtr<FMetasoundNumericDebugLineItem> LineItem;
 
 		public:
diff --git a/Engine/Plugins/Runtime/Metasound/Source/MetasoundFrontend/Private/Analysis/MetasoundFrontendGraphAnalyzerView.cpp b/Engine/Plugins/Runtime/Metasound/Source/MetasoundFrontend/Private/Analysis/MetasoundFrontendGraphAnalyzerView.cpp
--- a/Engine/Plugins/Runtime/Metasound/Source/MetasoundFrontend/Private/Analysis/MetasoundFrontendGraphAnalyzerView.cpp
+++ b/Engine/Plugins/Runtime/Metasound/Source/MetasoundFrontend/Private/Analysis/MetasoundFrontendGraphAnalyzerView.cpp
@@ -66,6 +66,21 @@ namespace Metasound
 			ActiveAnalyzerSender = FDataTransmissionCenter::Get().RegisterNewSender(AnalyzerAddress, SenderParams);
 		}
 
+		FMetasoundGraphAnalyzerView::~FMetasoundGraphAnalyzerView()
+		{
+			// Only unregister the data channel if we had a sender using that 
+			// data channel. This protects against removing the data channel 
+			// multiple times. Multiple removals of data channels has caused
+			// race conditions between newly created transmitters and transmitters
+			// being cleaned up.
+			if (ActiveAnalyzerSender.IsValid())
+			{
+				const FGraphAnalyzerAddress AnalyzerAddress(InstanceID);
+				ActiveAnalyzerSender.Reset();
+				FDataTransmissionCenter::Get().UnregisterDataChannel(AnalyzerAddress);
+			}
+		}
+
 		void FMetasoundGraphAnalyzerView::AddAnalyzerForAllSupportedOutputs(FName InAnalyzerName, bool bInRequiresConnection)
 		{
 			using namespace GraphAnalyzerViewPrivate;
diff --git a/Engine/Plugins/Runtime/Metasound/Source/MetasoundFrontend/Public/Analysis/MetasoundFrontendGraphAnalyzerView.h b/Engine/Plugins/Runtime/Metasound/Source/MetasoundFrontend/Public/Analysis/MetasoundFrontendGraphAnalyzerView.h
--- a/Engine/Plugins/Runtime/Metasound/Source/MetasoundFrontend/Public/Analysis/MetasoundFrontendGraphAnalyzerView.h
+++ b/Engine/Plugins/Runtime/Metasound/Source/MetasoundFrontend/Public/Analysis/MetasoundFrontendGraphAnalyzerView.h
@@ -36,9 +36,16 @@ namespace Metasound
 			using FMetasoundGraphAnalyzerOutputKey = TTuple<FGuid /* NodeID */, FName /* OutputName */>;
 			TMap<FMetasoundGraphAnalyzerOutputKey, TArray<FMetasoundAnalyzerView>> AnalyzerViews;
 
+			FMetasoundGraphAnalyzerView(const FMetasoundGraphAnalyzerView&) = delete;
+			FMetasoundGraphAnalyzerView(FMetasoundGraphAnalyzerView&&) = delete;
+			FMetasoundGraphAnalyzerView& operator=(const FMetasoundGraphAnalyzerView&) = delete;
+			FMetasoundGraphAnalyzerView& operator=(FMetasoundGraphAnalyzerView&&) = delete;
+
 		public:
 			FMetasoundGraphAnalyzerView() = default;
+
 			FMetasoundGraphAnalyzerView(const FMetasoundAssetBase& InAssetBase, uint64 InInstanceID, FSampleRate InSampleRate);
+			~FMetasoundGraphAnalyzerView();
 
 			void AddAnalyzerForAllSupportedOutputs(FName InAnalyzerName, bool bInRequiresConnection = true);
 			void RemoveAnalyzerForAllSupportedOutputs(FName InAnalyzerName);
diff --git a/Engine/Plugins/Runtime/Metasound/Source/MetasoundFrontend/Public/Analysis/MetasoundFrontendVertexAnalyzer.h b/Engine/Plugins/Runtime/Metasound/Source/MetasoundFrontend/Public/Analysis/MetasoundFrontendVertexAnalyzer.h
--- a/Engine/Plugins/Runtime/Metasound/Source/MetasoundFrontend/Public/Analysis/MetasoundFrontendVertexAnalyzer.h
+++ b/Engine/Plugins/Runtime/Metasound/Source/MetasoundFrontend/Public/Analysis/MetasoundFrontendVertexAnalyzer.h
@@ -73,13 +73,26 @@ namespace Metasound
 		{
 		public:
 			TBoundAnalyzerOutput(const FAnalyzerAddress& InAnalyzerOutputAddress, const FOperatorSettings& InOperatorSettings, TDataReadReference<DataType>&& InData)
-				: DataRef(MoveTemp(InData))
+				: Address(InAnalyzerOutputAddress)
+				, DataRef(MoveTemp(InData))
 			{
 				Sender = FDataTransmissionCenter::Get().RegisterNewSender<DataType>(InAnalyzerOutputAddress, FSenderInitParams { InOperatorSettings, 0 });
 				ensure(Sender);
 			}
 
-			virtual ~TBoundAnalyzerOutput() = default;
+			virtual ~TBoundAnalyzerOutput()
+			{
+				// Only unregister the data channel if we had a sender using that 
+				// data channel. This protects against removing the data channel 
+				// multiple times. Multiple removals of data channels has caused
+				// race conditions between newly created transmitters and transmitters
+				// being cleaned up.
+				if (Sender.IsValid())
+				{
+					Sender.Reset();
+					FDataTransmissionCenter::Get().UnregisterDataChannel(Address);
+				}
+			}
 
 			FAnyDataReference GetDataReference() const override
 			{
@@ -95,6 +108,7 @@ namespace Metasound
 			}
 
 		private:
+			FAnalyzerAddress Address;
 			TDataReadReference<DataType> DataRef;
 			TUniquePtr<TSender<DataType>> Sender;
 		};
