MetaSound preset fixes: Update node descriptions for referencing graphs, remove variable header, prevent preset vertices from being renamed

diff --git a/Engine/Plugins/Runtime/Metasound/Source/MetasoundEditor/Private/MetasoundEditor.cpp b/Engine/Plugins/Runtime/Metasound/Source/MetasoundEditor/Private/MetasoundEditor.cpp
--- a/Engine/Plugins/Runtime/Metasound/Source/MetasoundEditor/Private/MetasoundEditor.cpp
+++ b/Engine/Plugins/Runtime/Metasound/Source/MetasoundEditor/Private/MetasoundEditor.cpp
@@ -2892,10 +2892,19 @@ namespace Metasound
 
 		void FEditor::CollectStaticSections(TArray<int32>& StaticSectionIDs)
 		{
+			Frontend::FConstDocumentHandle DocumentHandle = GetMetaSoundGraphChecked().GetDocumentHandle();
+			const FMetasoundFrontendGraphClass& RootGraphClass = DocumentHandle->GetRootGraphClass();
+			const bool bIsPreset = RootGraphClass.PresetOptions.bIsPreset;
+
 			for (int32 i = 0; i < static_cast<int32>(ENodeSection::COUNT); ++i)
 			{
 				if (static_cast<ENodeSection>(i) != ENodeSection::None)
 				{
+					// Presets do not have variables
+					if (bIsPreset && static_cast<ENodeSection>(i) == ENodeSection::Variables)
+					{
+						continue;
+					}
 					StaticSectionIDs.Add(i);
 				}
 			}
diff --git a/Engine/Plugins/Runtime/Metasound/Source/MetasoundEditor/Private/MetasoundEditorGraph.cpp b/Engine/Plugins/Runtime/Metasound/Source/MetasoundEditor/Private/MetasoundEditorGraph.cpp
--- a/Engine/Plugins/Runtime/Metasound/Source/MetasoundEditor/Private/MetasoundEditorGraph.cpp
+++ b/Engine/Plugins/Runtime/Metasound/Source/MetasoundEditor/Private/MetasoundEditorGraph.cpp
@@ -217,6 +217,9 @@ void UMetasoundEditorGraphVertex::SetDescription(const FText& InDescription, boo
 
 			FNodeHandle NodeHandle = GetNodeHandle();
 			NodeHandle->SetDescription(InDescription);
+
+			Graph->RegisterGraphWithFrontend();
+			Graph->GetModifyContext().AddMemberIDsModified({ GetMemberID() });
 		}
 	}
 }
@@ -538,7 +541,10 @@ bool UMetasoundEditorGraphVertex::IsInterfaceMember() const
 
 bool UMetasoundEditorGraphVertex::CanRename() const
 {
-	return !IsInterfaceMember();
+	Metasound::Frontend::FConstDocumentHandle DocumentHandle = GetOwningGraph()->GetDocumentHandle();
+	const FMetasoundFrontendGraphClass& RootGraphClass = DocumentHandle->GetRootGraphClass();
+
+	return !RootGraphClass.PresetOptions.bIsPreset && !IsInterfaceMember();
 }
 
 bool UMetasoundEditorGraphVertex::CanRename(const FText& InNewName, FText& OutError) const
@@ -557,6 +563,15 @@ bool UMetasoundEditorGraphVertex::CanRename(const FText& InNewName, FText& OutEr
 		return false;
 	}
 
+	Metasound::Frontend::FConstDocumentHandle DocumentHandle = GetOwningGraph()->GetDocumentHandle();
+	const FMetasoundFrontendGraphClass& RootGraphClass = DocumentHandle->GetRootGraphClass();
+
+	if (RootGraphClass.PresetOptions.bIsPreset)
+	{
+		OutError = FText::Format(LOCTEXT("GraphVertexRenameInvalid_Preset", "{0} is a vertex in a preset graph and cannot be renamed."), InNewName);
+		return false;
+	}
+
 	bool bIsNameValid = true;
 	const FString NewName = InNewName.ToString();
 	FConstNodeHandle NodeHandle = GetConstNodeHandle();
