Enable metasound plugin by default.

diff --git a/Engine/Plugins/Runtime/Metasound/Metasound.uplugin b/Engine/Plugins/Runtime/Metasound/Metasound.uplugin
--- a/Engine/Plugins/Runtime/Metasound/Metasound.uplugin
+++ b/Engine/Plugins/Runtime/Metasound/Metasound.uplugin
@@ -10,7 +10,7 @@
 	"DocsURL": "",
 	"MarketplaceURL": "",
 	"SupportURL": "",
-	"EnabledByDefault": false,
+	"EnabledByDefault": true,
 	"CanContainContent": true,
 	"IsBetaVersion": true,
 	"Installed": false,
