Fix crashes in Audio Bus Reader node when no Audio Bus is given

diff --git a/Engine/Plugins/Runtime/Metasound/Source/MetasoundEngine/Private/MetasoundAudioBusReaderNode.cpp b/Engine/Plugins/Runtime/Metasound/Source/MetasoundEngine/Private/MetasoundAudioBusReaderNode.cpp
--- a/Engine/Plugins/Runtime/Metasound/Source/MetasoundEngine/Private/MetasoundAudioBusReaderNode.cpp
+++ b/Engine/Plugins/Runtime/Metasound/Source/MetasoundEngine/Private/MetasoundAudioBusReaderNode.cpp
@@ -103,16 +103,18 @@ namespace Metasound
 			}
 			
 			const FAudioBusProxyPtr& AudioBusProxy = AudioBusAsset->GetAudioBusProxy();
-
-			if (FAudioDeviceManager* ADM = FAudioDeviceManager::Get())
+			if (AudioBusProxy.IsValid())
 			{
-				if (FAudioDevice* AudioDevice = ADM->GetAudioDeviceRaw(AudioDeviceId))
+				if (FAudioDeviceManager* ADM = FAudioDeviceManager::Get())
 				{
-					// Start the audio bus in case it's not already started
-					AudioBusChannels = AudioBusProxy->NumChannels;
-					AudioDevice->StartAudioBus(AudioBusProxy->AudioBusId, AudioBusChannels, false);
-					
-					AudioBusPatchOutput = AudioDevice->AddPatchForAudioBus(AudioBusProxy->AudioBusId);
+					if (FAudioDevice* AudioDevice = ADM->GetAudioDeviceRaw(AudioDeviceId))
+					{
+						// Start the audio bus in case it's not already started
+						AudioBusChannels = AudioBusProxy->NumChannels;
+						AudioDevice->StartAudioBus(AudioBusProxy->AudioBusId, AudioBusChannels, false);
+
+						AudioBusPatchOutput = AudioDevice->AddPatchForAudioBus(AudioBusProxy->AudioBusId);
+					}
 				}
 			}
 		}
@@ -149,6 +151,11 @@ namespace Metasound
 			const int32 BlockSizeFrames = AudioOutputs[0]->Num();
 			const int32 NumSamplesToPop = BlockSizeFrames * AudioBusChannels;
 
+			if (!AudioBusPatchOutput.IsValid())
+			{
+				return;
+			}
+
 			// We want to write out zeroed audio until there is enough audio queued up in the patch output to avoid any underruns due to thread timing.
 			// From testing, waiting for the first buffer for 3 times the size of a single metasound block was sufficient to avoid pops. 
 			// Once we have started popping audio off the patch output, we want to keep it popping. Unless there are CPU underruns, there should be enough
