Fix for Limiter Distortion with non-zero attack times

diff --git a/Engine/Plugins/Runtime/Metasound/Source/MetasoundStandardNodes/Private/MetasoundLimiterNode.cpp b/Engine/Plugins/Runtime/Metasound/Source/MetasoundStandardNodes/Private/MetasoundLimiterNode.cpp
--- a/Engine/Plugins/Runtime/Metasound/Source/MetasoundStandardNodes/Private/MetasoundLimiterNode.cpp
+++ b/Engine/Plugins/Runtime/Metasound/Source/MetasoundStandardNodes/Private/MetasoundLimiterNode.cpp
@@ -21,7 +21,7 @@ namespace Metasound
 	namespace LimiterVertexNames
 	{
 		METASOUND_PARAM(InputAudio, "Audio", "Incoming audio signal to compress.");
-		METASOUND_PARAM(InputInGainDb, "Input Gain dB", "Gain to apply to the input before limiting.");
+		METASOUND_PARAM(InputInGainDb, "Input Gain dB", "Gain to apply to the input before limiting, in decibels. Maximum 100 dB. ");
 		METASOUND_PARAM(InputThresholdDb, "Threshold dB", "Amplitude threshold above which gain will be reduced.");
 		METASOUND_PARAM(InputReleaseTime, "Release Time", "How long it takes for audio below the threshold to return to its original volume level.");
 		METASOUND_PARAM(InputKneeMode, "Knee", "Whether the limiter uses a hard or soft knee.");
@@ -50,6 +50,7 @@ namespace Metasound
 
 		static constexpr float HardKneeBandwitdh = 0.0f;
 		static constexpr float SoftKneeBandwitdh = 10.0f;
+		static constexpr float MaxInputGain = 100.0f;
 
 		FLimiterOperator(const FOperatorSettings& InSettings,
 			const FAudioBufferReadRef& InAudio,
@@ -70,7 +71,7 @@ namespace Metasound
 		{
 			Limiter.Init(InSettings.GetSampleRate(), 1);
 			Limiter.SetProcessingMode(Audio::EDynamicsProcessingMode::Limiter);
-			Limiter.SetInputGain(*InGainDb);
+			Limiter.SetInputGain(FMath::Min(*InGainDbInput, MaxInputGain));
 			Limiter.SetThreshold(*ThresholdDbInput);
 			Limiter.SetAttackTime(0.0f);
 			Limiter.SetReleaseTime(PrevReleaseTime);
@@ -182,7 +183,7 @@ namespace Metasound
 			/* Update parameters */
 			if (!FMath::IsNearlyEqual(*InGainDbInput, PrevInGainDb))
 			{
-				Limiter.SetInputGain(*InGainDbInput);
+				Limiter.SetInputGain(FMath::Min(*InGainDbInput, MaxInputGain));
 			}
 			if (!FMath::IsNearlyEqual(*ThresholdDbInput, PrevThresholdDb))
 			{
diff --git a/Engine/Source/Runtime/SignalProcessing/Private/DynamicsProcessor.cpp b/Engine/Source/Runtime/SignalProcessing/Private/DynamicsProcessor.cpp
--- a/Engine/Source/Runtime/SignalProcessing/Private/DynamicsProcessor.cpp
+++ b/Engine/Source/Runtime/SignalProcessing/Private/DynamicsProcessor.cpp
@@ -260,11 +260,6 @@ namespace Audio
 
 				// Write into the output with the computed gain value
 				OutFrame[Channel] = Gain[Channel] * LookaheadOutput * OutputGain * InputGain;
-
-				if (ProcessingMode == EDynamicsProcessingMode::Limiter)
-				{
-					OutFrame[Channel] = FMath::Clamp(OutFrame[Channel], -ConvertToLinear(ThresholdDb), ConvertToLinear(ThresholdDb));
-				}
 			}
 		}
 	}
