Static analysis fix for MetaSound node loctext being compiled out for non editor builds

#rnx

diff --git a/Engine/Plugins/Runtime/Metasound/Source/MetasoundStandardNodes/Private/MetasoundBandsplitterNode.cpp b/Engine/Plugins/Runtime/Metasound/Source/MetasoundStandardNodes/Private/MetasoundBandsplitterNode.cpp
--- a/Engine/Plugins/Runtime/Metasound/Source/MetasoundStandardNodes/Private/MetasoundBandsplitterNode.cpp
+++ b/Engine/Plugins/Runtime/Metasound/Source/MetasoundStandardNodes/Private/MetasoundBandsplitterNode.cpp
@@ -173,11 +173,15 @@ namespace Metasound
 				// audio channels
 				for (uint32 ChannelIndex = 0; ChannelIndex < NumChannels; ChannelIndex++)
 				{
+#if WITH_EDITOR
 					const FDataVertexMetadata AudioInputMetadata
 					{
 						GetAudioInputDescription(ChannelIndex),
 						GetAudioInputDisplayName(ChannelIndex)
 					};
+#else 
+					const FDataVertexMetadata AudioInputMetadata;
+#endif // WITH_EDITOR
 
 					InputInterface.Add(TInputDataVertex<FAudioBuffer>(GetAudioInputName(ChannelIndex), AudioInputMetadata));
 				}
@@ -188,11 +192,15 @@ namespace Metasound
 				// Crossover frequencies
 				for (uint32 InputIndex = 0; InputIndex < NumBands - 1; InputIndex++)
 				{
+#if WITH_EDITOR
 					const FDataVertexMetadata CrossoverInputMetadata
 					{
 						GetCrossoverInputDescription(InputIndex),
 						GetCrossoverInputDisplayName(InputIndex)
 					};
+#else
+					const FDataVertexMetadata CrossoverInputMetadata;
+#endif // WITH_EDITOR
 
 					InputInterface.Add(TInputDataVertex<float>(GetCrossoverInputName(InputIndex), CrossoverInputMetadata, (1 + InputIndex) * 500.0f));
 				}
@@ -203,11 +211,15 @@ namespace Metasound
 				{
 					for (uint32 ChannelIndex = 0; ChannelIndex < NumChannels; ChannelIndex++)
 					{
+#if WITH_EDITOR
 						const FDataVertexMetadata AudioOutputMetadata
 						{
 							GetAudioOutputDescription(OutputIndex, ChannelIndex),
 							GetAudioOutputDisplayName(OutputIndex, ChannelIndex)
 						};
+#else 
+						const FDataVertexMetadata AudioOutputMetadata;
+#endif // WITH_EDITOR 
 
 						OutputInterface.Add(TOutputDataVertex<FAudioBuffer>(GetAudioOutputName(OutputIndex, ChannelIndex), AudioOutputMetadata));
 					}
@@ -360,6 +372,26 @@ namespace Metasound
 			return *FString::Printf(TEXT("In %i"), ChannelIndex);
 		}
 
+		static const FVertexName GetCrossoverInputName(uint32 InputIndex)
+		{
+			return *FString::Printf(TEXT("Crossover %i"), InputIndex);
+		}
+
+		static const FVertexName GetAudioOutputName(uint32 OutputIndex, uint32 ChannelIndex)
+		{
+			if (NumChannels == 1)
+			{
+				return *FString::Printf(TEXT("Band %i Out"), OutputIndex);
+			}
+			else if (NumChannels == 2)
+			{
+				return *FString::Printf(TEXT("Band %i %s"), OutputIndex, (ChannelIndex == 0) ? TEXT("L") : TEXT("R"));
+			}
+
+			return *FString::Printf(TEXT("Band %i Out %i"), OutputIndex, ChannelIndex);
+		}
+
+#if WITH_EDITOR
 		static const FText GetAudioInputDescription(uint32 ChannelIndex)
 		{
 			return METASOUND_LOCTEXT_FORMAT("BandSplitterAudioInputDescription", "Audio Channel: {0}", ChannelIndex);
@@ -386,11 +418,6 @@ namespace Metasound
 			return METASOUND_LOCTEXT_FORMAT("BandsplitterAudioInputIn", "In {0}", ChannelIndex);
 		}
 
-		static const FVertexName GetCrossoverInputName(uint32 InputIndex)
-		{
-			return *FString::Printf(TEXT("Crossover %i"), InputIndex);
-		}
-
 		static const FText GetCrossoverInputDisplayName(uint32 InputIndex)
 		{
 			return METASOUND_LOCTEXT_FORMAT("BandsplitterCrossoverInputDisplayName", "Crossover {0}", InputIndex);
@@ -401,20 +428,6 @@ namespace Metasound
 			return METASOUND_LOCTEXT_FORMAT("BandSplitterCrossoverInputDescription", "Crossover Input #: {0}", InputIndex);
 		}
 
-		static const FVertexName GetAudioOutputName(uint32 OutputIndex, uint32 ChannelIndex)
-		{
-			if (NumChannels == 1)
-			{
-				return *FString::Printf(TEXT("Band %i Out"), OutputIndex);
-			}
-			else if (NumChannels == 2)
-			{
-				return *FString::Printf(TEXT("Band %i %s"), OutputIndex, (ChannelIndex == 0) ? TEXT("L") : TEXT("R"));
-			}
-
-			return *FString::Printf(TEXT("Band %i Out %i"), OutputIndex, ChannelIndex);
-		}
-
 		static const FText GetAudioOutputDescription(uint32 OutputIndex, uint32 ChannelIndex)
 		{
 			return METASOUND_LOCTEXT_FORMAT("BandSplitterAudioOutputDescription", "Channel {0} Output for Band {1}", ChannelIndex, OutputIndex);
@@ -440,6 +453,7 @@ namespace Metasound
 
 			return METASOUND_LOCTEXT_FORMAT("BandsplitterAudioOutputOut", "Band {0} Out {1}", OutputIndex, ChannelIndex);
 		}
+#endif // WITH_EDITOR
 
 		FORCEINLINE Audio::EFilterOrder GetFilterOrder(EBandSplitterFilterOrder InFilterOrder) const
 		{
diff --git a/Engine/Plugins/Runtime/Metasound/Source/MetasoundStandardNodes/Private/MetasoundMixerNode.cpp b/Engine/Plugins/Runtime/Metasound/Source/MetasoundStandardNodes/Private/MetasoundMixerNode.cpp
--- a/Engine/Plugins/Runtime/Metasound/Source/MetasoundStandardNodes/Private/MetasoundMixerNode.cpp
+++ b/Engine/Plugins/Runtime/Metasound/Source/MetasoundStandardNodes/Private/MetasoundMixerNode.cpp
@@ -59,20 +59,28 @@ namespace Metasound
 					// audio channels
 					for (uint32 ChanIndex = 0; ChanIndex < NumChannels; ++ChanIndex)
 					{
+#if WITH_EDITOR
 						const FDataVertexMetadata AudioInputMetadata
 						{
 							GetAudioInputDescription(InputIndex, ChanIndex),
 							GetAudioInputDisplayName(InputIndex, ChanIndex)
 						};
+#else 
+						const FDataVertexMetadata AudioInputMetadata;
+#endif // WITH_EDITOR
 						InputInterface.Add(TInputDataVertex<FAudioBuffer>(GetAudioInputName(InputIndex, ChanIndex), AudioInputMetadata));
 					}
 
 					// gain scalar
+#if WITH_EDITOR
 					FDataVertexMetadata GainPinMetaData
 					{
 						GetGainInputDescription(InputIndex),
 						GetGainInputDisplayName(InputIndex)
 					};
+#else 
+					FDataVertexMetadata GainPinMetaData;
+#endif // WITH_EDITOR
 					TInputDataVertex<float>GainVertexModel(GetGainInputName(InputIndex), GainPinMetaData, 1.0f);
 
 					InputInterface.Add(GainVertexModel);
@@ -82,12 +90,15 @@ namespace Metasound
 				FOutputVertexInterface OutputInterface;
 				for (uint32 i = 0; i < NumChannels; ++i)
 				{
+#if WITH_EDITOR
 					const FDataVertexMetadata AudioOutputMetadata
 					{
 						GetAudioOutputDescription(i),
 						GetAudioOutputDisplayName(i)
 					};
-
+#else 
+					const FDataVertexMetadata AudioOutputMetadata;
+#endif // WITH_EDITOR
 					OutputInterface.Add(TOutputDataVertex<FAudioBuffer>(GetAudioOutputName(i), AudioOutputMetadata));
 				}
 
@@ -261,6 +272,26 @@ namespace Metasound
 			return *FString::Printf(TEXT("In %i, %i"), InputIndex, ChannelIndex);
 		}
 
+		static const FVertexName GetGainInputName(uint32 InputIndex)
+		{
+			return *FString::Printf(TEXT("Gain %i"), InputIndex);
+		}
+
+		static const FVertexName GetAudioOutputName(uint32 ChannelIndex)
+		{
+			if (NumChannels == 1)
+			{
+				return TEXT("Out");
+			}
+			else if (NumChannels == 2)
+			{
+				return *FString::Printf(TEXT("Out %s"), (ChannelIndex == 0) ? TEXT("L") : TEXT("R"));
+			}
+
+			return *FString::Printf(TEXT("Out %i"), ChannelIndex);
+		}
+
+#if WITH_EDITOR
 		static const FText GetAudioInputDescription(uint32 InputIndex, uint32 ChannelIndex)
 		{
 			return METASOUND_LOCTEXT_FORMAT("AudioMixerAudioInputDescription", "Audio Input #: {0}, Channel: {1}", InputIndex, ChannelIndex);
@@ -286,11 +317,6 @@ namespace Metasound
 			return METASOUND_LOCTEXT_FORMAT("AudioMixerAudioInputIn", "In {0}, {0}", InputIndex, ChannelIndex);
 		}
 
-		static const FVertexName GetGainInputName(uint32 InputIndex)
-		{
-			return *FString::Printf(TEXT("Gain %i"), InputIndex);
-		}
-
 		static const FText GetGainInputDisplayName(uint32 InputIndex)
 		{
 			return METASOUND_LOCTEXT_FORMAT("AudioMixerGainInputDisplayName", "Gain {0} (Lin)", InputIndex);
@@ -301,20 +327,6 @@ namespace Metasound
 			return METASOUND_LOCTEXT_FORMAT("AudioMixerGainInputDescription", "Gain Input #: {0}", InputIndex);
 		}
 
-		static const FVertexName GetAudioOutputName(uint32 ChannelIndex)
-		{
-			if (NumChannels == 1)
-			{
-				return TEXT("Out");
-			}
-			else if (NumChannels == 2)
-			{
-				return *FString::Printf(TEXT("Out %s"), (ChannelIndex == 0) ? TEXT("L") : TEXT("R"));
-			}
-
-			return *FString::Printf(TEXT("Out %i"), ChannelIndex);
-		}
-
 		static const FText GetAudioOutputDisplayName(uint32 ChannelIndex)
 		{
 			if (NumChannels == 1)
@@ -339,6 +351,7 @@ namespace Metasound
 		{
 			return METASOUND_LOCTEXT_FORMAT("AudioMixerAudioOutputDescription", "Summed output for channel: {0}", ChannelIndex);
 		}
+#endif // WITH_EDITOR
 #pragma endregion
 	}; // class TAudioMixerNodeOperator
 #pragma endregion
diff --git a/Engine/Plugins/Runtime/Metasound/Source/MetasoundStandardNodes/Private/MetasoundWaveWriterNode.cpp b/Engine/Plugins/Runtime/Metasound/Source/MetasoundStandardNodes/Private/MetasoundWaveWriterNode.cpp
--- a/Engine/Plugins/Runtime/Metasound/Source/MetasoundStandardNodes/Private/MetasoundWaveWriterNode.cpp
+++ b/Engine/Plugins/Runtime/Metasound/Source/MetasoundStandardNodes/Private/MetasoundWaveWriterNode.cpp
@@ -279,12 +279,15 @@ namespace Metasound
 				// For backwards compatibility with previous (mono) node, in the case of 1 channels, just provide the old interface.
 				for (int32 InputIndex = 0; InputIndex < NumInputChannels; ++InputIndex)
 				{
+#if WITH_EDITOR
 					const FDataVertexMetadata AudioInputMetadata
 					{
 						  GetAudioInputDescription(InputIndex) // description
 						, GetAudioInputDisplayName(InputIndex) // display name
 					};
-
+#else
+					const FDataVertexMetadata AudioInputMetadata;
+#endif //WITH_EDITOR
 					InputInterface.Add(TInputDataVertex<FAudioBuffer>(GetAudioInputName(InputIndex), AudioInputMetadata));
 				}
 				FOutputVertexInterface OutputInterface;
@@ -385,6 +388,7 @@ namespace Metasound
 			return *FString::Printf(TEXT("In %d"), InInputIndex);
 		}
 
+#if WITH_EDITOR 
 		static const FText GetAudioInputDisplayName(int32 InInputIndex)
 		{
 			if (NumInputChannels == 1)
@@ -397,11 +401,11 @@ namespace Metasound
 			{
 				if (InInputIndex == 0)
 				{
-					return METASOUND_LOCTEXT_FORMAT("AudioInputIn2ChannelName", "In {0} L", InInputIndex);
+					return METASOUND_LOCTEXT_FORMAT("AudioInputIn2ChannelNameL", "In {0} L", InInputIndex);
 				}
 				else
 				{
-					return METASOUND_LOCTEXT_FORMAT("AudioInputIn2ChannelName", "In {0} R", InInputIndex);
+					return METASOUND_LOCTEXT_FORMAT("AudioInputIn2ChannelNameR", "In {0} R", InInputIndex);
 				}
 			}
 			return METASOUND_LOCTEXT_FORMAT("AudioInputInChannelName", "In {0}", InInputIndex);
@@ -411,6 +415,7 @@ namespace Metasound
 		{
 			return METASOUND_LOCTEXT_FORMAT("WaveWriterAudioInputDescription", "Audio Input #: {0}", InputIndex);
 		}
+#endif // WITH_EDITOR
 
 		static FNodeClassMetadata CreateNodeClassMetadata(const FName& InOperatorName, const FText& InDisplayName, const FText& InDescription, const FVertexInterface& InDefaultInterface)
 		{
