Optimization pass on MetaSound SlatePanel/Node/Pins & add profiler events
#rnx

diff --git a/Engine/Plugins/Runtime/Metasound/Source/MetasoundEditor/Private/MetasoundEditor.cpp b/Engine/Plugins/Runtime/Metasound/Source/MetasoundEditor/Private/MetasoundEditor.cpp
--- a/Engine/Plugins/Runtime/Metasound/Source/MetasoundEditor/Private/MetasoundEditor.cpp
+++ b/Engine/Plugins/Runtime/Metasound/Source/MetasoundEditor/Private/MetasoundEditor.cpp
@@ -2698,6 +2698,8 @@ namespace Metasound
 
 		void FEditor::RefreshDetails()
 		{
+			TRACE_CPUPROFILER_EVENT_SCOPE(Metasound::Editor::FEditor::RefreshDetails);
+
 			using namespace Frontend;
 
 			if (MetasoundDetails.IsValid())
@@ -2708,6 +2710,8 @@ namespace Metasound
 
 		void FEditor::RefreshInterfaces()
 		{
+			TRACE_CPUPROFILER_EVENT_SCOPE(Metasound::Editor::FEditor::RefreshInterfaces);
+
 			if (InterfacesDetails.IsValid())
 			{
 				InterfacesDetails->ForceRefresh();
@@ -2716,6 +2720,8 @@ namespace Metasound
 
 		void FEditor::RefreshGraphMemberMenu()
 		{
+			TRACE_CPUPROFILER_EVENT_SCOPE(Metasound::Editor::FEditor::RefreshGraphMemberMenu);
+
 			if (GraphMembersMenu.IsValid())
 			{
 				TArray<TSharedPtr<FEdGraphSchemaAction>> SelectedActions;
@@ -3154,6 +3160,7 @@ namespace Metasound
 
 		void FEditor::RefreshEditorContext()
 		{
+			TRACE_CPUPROFILER_EVENT_SCOPE(Metasound::Editor::FEditor::RefreshEditorContext);
 			const bool bSynchronizedGraph = FGraphBuilder::SynchronizeGraph(*Metasound);
 
 			FMetasoundAssetBase* MetasoundAsset = IMetasoundUObjectRegistry::Get().GetObjectAsAssetBase(Metasound);
diff --git a/Engine/Plugins/Runtime/Metasound/Source/MetasoundEditor/Private/MetasoundEditorGraphBuilder.cpp b/Engine/Plugins/Runtime/Metasound/Source/MetasoundEditor/Private/MetasoundEditorGraphBuilder.cpp
--- a/Engine/Plugins/Runtime/Metasound/Source/MetasoundEditor/Private/MetasoundEditorGraphBuilder.cpp
+++ b/Engine/Plugins/Runtime/Metasound/Source/MetasoundEditor/Private/MetasoundEditorGraphBuilder.cpp
@@ -574,6 +574,8 @@ namespace Metasound
 
 		FGraphValidationResults FGraphBuilder::ValidateGraph(UObject& InMetaSound)
 		{
+			TRACE_CPUPROFILER_EVENT_SCOPE(Metasound::Editor::FGraphBuilder::ValidateGraph);
+
 			using namespace Frontend;
 
 			FMetasoundAssetBase* MetaSoundAsset = IMetasoundUObjectRegistry::Get().GetObjectAsAssetBase(&InMetaSound);
@@ -962,13 +964,13 @@ namespace Metasound
 			return Algo::AnyOf(EditorNodes, [](const UMetasoundEditorGraphNode* Node) { return Node->ErrorType == EMessageSeverity::Error; });
 		}
 
-		bool FGraphBuilder::SynchronizeNodeLocation(UMetasoundEditorGraphNode& InNode)
+		bool FGraphBuilder::SynchronizeNodeLocation(const Frontend::FConstNodeHandle& InNode, UMetasoundEditorGraphNode& OutGraphNode)
 		{
 			bool bModified = false;
 
-			const FMetasoundFrontendNodeStyle& Style = InNode.GetConstNodeHandle()->GetNodeStyle();
+			const FMetasoundFrontendNodeStyle& Style = InNode->GetNodeStyle();
 
-			const FVector2D* Location = Style.Display.Locations.Find(InNode.NodeGuid);
+			const FVector2D* Location = Style.Display.Locations.Find(OutGraphNode.NodeGuid);
 			if (!Location)
 			{
 				// If no specific location found, use default location if provided (zero guid
@@ -980,12 +982,12 @@ namespace Metasound
 			{
 				const int32 LocX = FMath::TruncToInt(Location->X);
 				const int32 LocY = FMath::TruncToInt(Location->Y);
-				const bool bXChanged = static_cast<bool>(LocX - InNode.NodePosX);
-				const bool bYChanged = static_cast<bool>(LocY - InNode.NodePosY);
+				const bool bXChanged = static_cast<bool>(LocX - OutGraphNode.NodePosX);
+				const bool bYChanged = static_cast<bool>(LocY - OutGraphNode.NodePosY);
 				if (bXChanged || bYChanged)
 				{
-					InNode.NodePosX = LocX;
-					InNode.NodePosY = LocY;
+					OutGraphNode.NodePosX = LocX;
+					OutGraphNode.NodePosY = LocY;
 					bModified = true;
 				}
 			}
@@ -2121,6 +2123,8 @@ namespace Metasound
 
 		bool FGraphBuilder::SynchronizeGraph(UObject& InMetaSound)
 		{
+			TRACE_CPUPROFILER_EVENT_SCOPE(Metasound::Editor::FGraphBuilder::SynchronizeGraph);
+
 			using namespace Frontend;
 
 			FMetasoundAssetBase* MetaSoundAsset = IMetasoundUObjectRegistry::Get().GetObjectAsAssetBase(&InMetaSound);
@@ -2172,6 +2176,8 @@ namespace Metasound
 
 		bool FGraphBuilder::SynchronizeNodeMembers(UObject& InMetaSound)
 		{
+			TRACE_CPUPROFILER_EVENT_SCOPE(Metasound::Editor::FGraphBuilder::SynchronizeNodeMembers);
+
 			using namespace Frontend;
 
 			bool bEditorGraphModified = false;
@@ -2248,6 +2254,8 @@ namespace Metasound
 
 		bool FGraphBuilder::SynchronizeNodes(UObject& InMetaSound)
 		{
+			TRACE_CPUPROFILER_EVENT_SCOPE(Metasound::Editor::FGraphBuilder::SynchronizeNodes);
+
 			using namespace Frontend;
 
 			bool bEditorGraphModified = false;
@@ -2272,14 +2280,14 @@ namespace Metasound
 			struct FAssociatedNodes
 			{
 				TArray<UMetasoundEditorGraphNode*> EditorNodes;
-				FNodeHandle Node = Metasound::Frontend::INodeController::GetInvalidHandle();
+				FConstNodeHandle Node = Metasound::Frontend::INodeController::GetInvalidHandle();
 			};
 			TMap<FGuid, FAssociatedNodes> AssociatedNodes;
 
 			// Reverse iterate so paired nodes can safely be removed from the array.
 			for (int32 i = FrontendNodes.Num() - 1; i >= 0; i--)
 			{
-				FNodeHandle Node = FrontendNodes[i];
+				const FConstNodeHandle& Node = FrontendNodes[i];
 				bool bFoundEditorNode = false;
 				for (int32 j = EditorNodes.Num() - 1; j >= 0; --j)
 				{
@@ -2297,7 +2305,7 @@ namespace Metasound
 							AssociatedNodeData.Node = Node;
 						}
 
-						bEditorGraphModified |= SynchronizeNodeLocation(*EditorNode);
+						bEditorGraphModified |= SynchronizeNodeLocation(Node, *EditorNode);
 						AssociatedNodeData.EditorNodes.Add(EditorNode);
 						EditorNodes.RemoveAtSwap(j, 1, false /* bAllowShrinking */);
 					}
@@ -2321,7 +2329,7 @@ namespace Metasound
 			}
 
 			// Add missing editor nodes marked as visible.
-			for (FNodeHandle Node : FrontendNodes)
+			for (const FNodeHandle& Node : FrontendNodes)
 			{
 				const FMetasoundFrontendNodeStyle& CurrentStyle = Node->GetNodeStyle();
 				if (CurrentStyle.Display.Locations.IsEmpty())
@@ -2576,6 +2584,8 @@ namespace Metasound
 
 		bool FGraphBuilder::SynchronizeGraphMembers(UObject& InMetaSound)
 		{
+			TRACE_CPUPROFILER_EVENT_SCOPE(Metasound::Editor::FGraphBuilder::SynchronizeGraphMembers);
+
 			using namespace Frontend;
 
 			bool bEditorGraphModified = false;
diff --git a/Engine/Plugins/Runtime/Metasound/Source/MetasoundEditor/Private/MetasoundEditorGraphBuilder.h b/Engine/Plugins/Runtime/Metasound/Source/MetasoundEditor/Private/MetasoundEditorGraphBuilder.h
--- a/Engine/Plugins/Runtime/Metasound/Source/MetasoundEditor/Private/MetasoundEditorGraphBuilder.h
+++ b/Engine/Plugins/Runtime/Metasound/Source/MetasoundEditor/Private/MetasoundEditorGraphBuilder.h
@@ -274,7 +274,7 @@ namespace Metasound
 			// Synchronizes node location data using handle's data
 			//
 			// @return True if the UMetasoundEditorGraphNode was altered. False otherwise.
-			static bool SynchronizeNodeLocation(UMetasoundEditorGraphNode& InNode);
+			static bool SynchronizeNodeLocation(const Frontend::FConstNodeHandle& InNode, UMetasoundEditorGraphNode& OutGraphNode);
 
 			// Adds and removes pins so that the UMetasoundEditorGraphNode matches the InNode.
 			//
diff --git a/Engine/Plugins/Runtime/Metasound/Source/MetasoundEditor/Private/SMetasoundGraphNode.cpp b/Engine/Plugins/Runtime/Metasound/Source/MetasoundEditor/Private/SMetasoundGraphNode.cpp
--- a/Engine/Plugins/Runtime/Metasound/Source/MetasoundEditor/Private/SMetasoundGraphNode.cpp
+++ b/Engine/Plugins/Runtime/Metasound/Source/MetasoundEditor/Private/SMetasoundGraphNode.cpp
@@ -23,6 +23,7 @@
 #include "MetasoundFrontendArchetypeRegistry.h"
 #include "MetasoundFrontendNodeTemplateRegistry.h"
 #include "MetasoundFrontendRegistries.h"
+#include "MetasoundTrace.h"
 #include "MetasoundTrigger.h"
 #include "NodeFactory.h"
 #include "NodeTemplates/MetasoundFrontendNodeTemplateReroute.h"
@@ -83,14 +84,12 @@ namespace Metasound
 
 		bool SMetaSoundGraphNode::IsVariableAccessor() const
 		{
-			const EMetasoundFrontendClassType ClassType = GetMetaSoundNode().GetNodeHandle()->GetClassMetadata().GetType();
 			return ClassType == EMetasoundFrontendClassType::VariableAccessor
 				|| ClassType == EMetasoundFrontendClassType::VariableDeferredAccessor;
 		}
 
 		bool SMetaSoundGraphNode::IsVariableMutator() const
 		{
-			const EMetasoundFrontendClassType ClassType = GetMetaSoundNode().GetNodeHandle()->GetClassMetadata().GetType();
 			return ClassType == EMetasoundFrontendClassType::VariableMutator;
 		}
 
@@ -107,6 +106,9 @@ namespace Metasound
 		void SMetaSoundGraphNode::Construct(const FArguments& InArgs, class UEdGraphNode* InNode)
 		{
 			GraphNode = InNode;
+			Frontend::FConstNodeHandle NodeHandle = GetMetaSoundNode().GetConstNodeHandle();
+			ClassType = NodeHandle->GetClassMetadata().GetType();
+
 			SetCursor(EMouseCursor::CardinalCross);
 			UpdateGraphNode();
 		}
@@ -529,7 +531,6 @@ namespace Metasound
 			// TODO: Add tweak & add custom bodies
 			if (GraphNode)
 			{
-				const EMetasoundFrontendClassType ClassType = GetMetaSoundNode().GetNodeHandle()->GetClassMetadata().GetType();
 				switch (ClassType)
 				{
 					case EMetasoundFrontendClassType::Variable:
diff --git a/Engine/Plugins/Runtime/Metasound/Source/MetasoundEditor/Private/SMetasoundGraphNode.h b/Engine/Plugins/Runtime/Metasound/Source/MetasoundEditor/Private/SMetasoundGraphNode.h
--- a/Engine/Plugins/Runtime/Metasound/Source/MetasoundEditor/Private/SMetasoundGraphNode.h
+++ b/Engine/Plugins/Runtime/Metasound/Source/MetasoundEditor/Private/SMetasoundGraphNode.h
@@ -4,6 +4,7 @@
 
 #include "Input/Reply.h"
 #include "Layout/Visibility.h"
+#include "MetasoundFrontendDocument.h"
 #include "Misc/Attribute.h"
 #include "SGraphNode.h"
 #include "SGraphNodeKnot.h"
@@ -86,6 +87,8 @@ namespace Metasound
 			// Whether the input widget is currently transacting 
 			// for keeping track of transaction state across delegates to only commit transaction on value commit
 			bool bIsInputWidgetTransacting = false;
+
+			EMetasoundFrontendClassType ClassType;
 		};
 
 		class SMetaSoundGraphNodeKnot : public SGraphNodeKnot
diff --git a/Engine/Plugins/Runtime/Metasound/Source/MetasoundEditor/Private/SMetasoundGraphPin.cpp b/Engine/Plugins/Runtime/Metasound/Source/MetasoundEditor/Private/SMetasoundGraphPin.cpp
--- a/Engine/Plugins/Runtime/Metasound/Source/MetasoundEditor/Private/SMetasoundGraphPin.cpp
+++ b/Engine/Plugins/Runtime/Metasound/Source/MetasoundEditor/Private/SMetasoundGraphPin.cpp
@@ -22,6 +22,27 @@ namespace Metasound
 				MakeAttributeSP(this, &SMetaSoundGraphPinKnot::GetSecondaryPinIcon),
 				MakeAttributeSP(this, &SMetaSoundGraphPinKnot::GetSecondaryPinColor));
 			PinImage = PinWidgetRef;
+
+			CacheAccessType();
+			CacheHasRequiredConnections();
+		}
+
+		void SMetaSoundGraphPinKnot::CacheHasRequiredConnections()
+		{
+			using namespace Frontend;
+
+			bHasRequiredConnections = false;
+			if (UEdGraphPin* Pin = SGraphPinKnot::GetPinObj())
+			{
+				if (const UMetasoundEditorGraphExternalNode* OwningNode = Cast<UMetasoundEditorGraphExternalNode>(Pin->GetOwningNode()))
+				{
+					if (const INodeTemplate* Template = INodeTemplateRegistry::Get().FindTemplate(FRerouteNodeTemplate::GetRegistryKey()))
+					{
+						FConstNodeHandle NodeHandle = OwningNode->GetConstNodeHandle();
+						bHasRequiredConnections = Template->HasRequiredConnections(NodeHandle);
+					}
+				}
+			}
 		}
 
 		FSlateColor SMetaSoundGraphPinKnot::GetPinColor() const
@@ -30,7 +51,7 @@ namespace Metasound
 			{
 				if (Pin->Direction == EGPD_Output)
 				{
-					if (!HasRequiredConnections() || Pin->GetOwningNode()->ErrorType <= static_cast<uint32>(EMessageSeverity::Warning))
+					if (!bHasRequiredConnections || Pin->GetOwningNode()->ErrorType <= static_cast<uint32>(EMessageSeverity::Warning))
 					{
 						return FLinearColor::Yellow;
 					}
@@ -44,32 +65,15 @@ namespace Metasound
 		{
 			using namespace Metasound::Frontend;
 
-			bool bIsConstructorPin = false;
-
 			if (const UEdGraphPin* Pin = SGraphPinKnot::GetPinObj())
 			{
-				if (!HasRequiredConnections() || Pin->GetOwningNode()->ErrorType <= static_cast<uint32>(EMessageSeverity::Warning))
+				if (!bHasRequiredConnections || Pin->GetOwningNode()->ErrorType <= static_cast<uint32>(EMessageSeverity::Warning))
 				{
 					return &Editor::Style::GetSlateBrushSafe("MetasoundEditor.Graph.InvalidReroute");
 				}
-
-				Pin = FGraphBuilder::FindReroutedOutputPin(Pin);
-				if (const UMetasoundEditorGraphNode* Node = Cast<UMetasoundEditorGraphNode>(Pin->GetOwningNode()))
-				{
-					if (const UMetasoundEditorGraphMemberNode* MemberNode = Cast<UMetasoundEditorGraphMemberNode>(Node))
-					{
-						if (const UMetasoundEditorGraphInput* Input = Cast<UMetasoundEditorGraphInput>(MemberNode->GetMember()))
-						{
-							FConstNodeHandle NodeHandle = Input->GetConstNodeHandle();
-							TArray<FConstOutputHandle> OutputHandles = NodeHandle->GetConstOutputs();
-
-							check(!OutputHandles.IsEmpty());
-							bIsConstructorPin = OutputHandles.Last()->GetVertexAccessType() == EMetasoundFrontendVertexAccessType::Value;
-						}
-					}
-				}
 			}
 
+			const bool bIsConstructorPin = TMetasoundGraphPin<SGraphPinKnot>::AccessType == EMetasoundFrontendVertexAccessType::Value;
 			if (bIsConstructorPin)
 			{
 				const bool bIsConnected = IsConnected();
@@ -87,24 +91,5 @@ namespace Metasound
 
 			return SGraphPinKnot::GetPinIcon();
 		}
-
-		bool SMetaSoundGraphPinKnot::HasRequiredConnections() const
-		{
-			using namespace Frontend;
-
-			if (UEdGraphPin* Pin = SGraphPinKnot::GetPinObj())
-			{
-				if (const UMetasoundEditorGraphExternalNode* OwningNode = Cast<UMetasoundEditorGraphExternalNode>(Pin->GetOwningNode()))
-				{
-					if (const INodeTemplate* Template = INodeTemplateRegistry::Get().FindTemplate(FRerouteNodeTemplate::GetRegistryKey()))
-					{
-						FConstNodeHandle NodeHandle = OwningNode->GetConstNodeHandle();
-						return Template->HasRequiredConnections(NodeHandle);
-					}
-				}
-			}
-
-			return false;
-		}
 	} // namespace Editor
 } // namespace Metasound
diff --git a/Engine/Plugins/Runtime/Metasound/Source/MetasoundEditor/Private/SMetasoundGraphPin.h b/Engine/Plugins/Runtime/Metasound/Source/MetasoundEditor/Private/SMetasoundGraphPin.h
--- a/Engine/Plugins/Runtime/Metasound/Source/MetasoundEditor/Private/SMetasoundGraphPin.h
+++ b/Engine/Plugins/Runtime/Metasound/Source/MetasoundEditor/Private/SMetasoundGraphPin.h
@@ -101,13 +101,37 @@ namespace Metasound
 				}
 			}
 
-		public:
-			SLATE_BEGIN_ARGS(TMetasoundGraphPin<ParentPinType>)
+			void CacheAccessType()
 			{
-			}
-			SLATE_END_ARGS()
+				AccessType = EMetasoundFrontendVertexAccessType::Unset;
 
-			virtual ~TMetasoundGraphPin() = default;
+				if (const UEdGraphPin* Pin = ParentPinType::GetPinObj())
+				{
+					if (const UMetasoundEditorGraphNode* Node = Cast<UMetasoundEditorGraphNode>(Pin->GetOwningNode()))
+					{
+						if (const UMetasoundEditorGraphMemberNode* MemberNode = Cast<UMetasoundEditorGraphMemberNode>(Node))
+						{
+							if (const UMetasoundEditorGraphVertex* Vertex = Cast<UMetasoundEditorGraphVertex>(MemberNode->GetMember()))
+							{
+								AccessType = Vertex->GetVertexAccessType();
+							}
+						}
+						else if (const UMetasoundEditorGraphExternalNode* ExternalNode = Cast<UMetasoundEditorGraphExternalNode>(Node))
+						{
+							if (Pin->Direction == EGPD_Input)
+							{
+								Frontend::FInputHandle InputHandle = FGraphBuilder::GetInputHandleFromPin(Pin);
+								AccessType = InputHandle->GetVertexAccessType();
+							}
+							else if (Pin->Direction == EGPD_Output)
+							{
+								Frontend::FOutputHandle OutputHandle = FGraphBuilder::GetOutputHandleFromPin(Pin);
+								AccessType = OutputHandle->GetVertexAccessType();
+							}
+						}
+					}
+				}
+			}
 
 			void CacheNodeOffset(const FGeometry& AllottedGeometry)
 			{
@@ -116,6 +140,16 @@ namespace Metasound
 				ParentPinType::CachedNodeOffset.Y += AllottedGeometry.Size.Y * 0.5f;
 			}
 
+			EMetasoundFrontendVertexAccessType AccessType = EMetasoundFrontendVertexAccessType::Unset;
+
+		public:
+			SLATE_BEGIN_ARGS(TMetasoundGraphPin<ParentPinType>)
+			{
+			}
+			SLATE_END_ARGS()
+
+			virtual ~TMetasoundGraphPin() = default;
+
 			Frontend::FConstInputHandle GetConstInputHandle() const
 			{
 				using namespace Frontend;
@@ -291,34 +325,6 @@ namespace Metasound
 			virtual const FSlateBrush* GetPinIcon() const override
 			{
 				const bool bIsConnected = ParentPinType::IsConnected();
-				EMetasoundFrontendVertexAccessType AccessType = EMetasoundFrontendVertexAccessType::Unset;
-
-				if (const UEdGraphPin* Pin = ParentPinType::GetPinObj())
-				{
-					if (const UMetasoundEditorGraphNode* Node = Cast<UMetasoundEditorGraphNode>(Pin->GetOwningNode()))
-					{
-						if (const UMetasoundEditorGraphMemberNode* MemberNode = Cast<UMetasoundEditorGraphMemberNode>(Node))
-						{
-							if (const UMetasoundEditorGraphVertex* Vertex = Cast<UMetasoundEditorGraphVertex>(MemberNode->GetMember()))
-							{
-								AccessType = Vertex->GetVertexAccessType();
-							}
-						}
-						else if (const UMetasoundEditorGraphExternalNode* ExternalNode = Cast<UMetasoundEditorGraphExternalNode>(Node))
-						{
-							if (Pin->Direction == EGPD_Input)
-							{
-								Frontend::FInputHandle InputHandle = FGraphBuilder::GetInputHandleFromPin(Pin);
-								AccessType = InputHandle->GetVertexAccessType();
-							}
-							else if (Pin->Direction == EGPD_Output)
-							{
-								Frontend::FOutputHandle OutputHandle = FGraphBuilder::GetOutputHandleFromPin(Pin);
-								AccessType = OutputHandle->GetVertexAccessType();
-							}
-						}
-					}
-				}
 
 				// Is constructor pin 
 				if (AccessType == EMetasoundFrontendVertexAccessType::Value)
@@ -363,6 +369,7 @@ namespace Metasound
 
 			void Construct(const FArguments& InArgs, UEdGraphPin* InGraphPinObj)
 			{
+				CacheAccessType();
 				SGraphPin::Construct(SGraphPin::FArguments(), InGraphPinObj);
 			}
 		};
@@ -374,6 +381,7 @@ namespace Metasound
 
 			void Construct(const FArguments& InArgs, UEdGraphPin* InGraphPinObj)
 			{
+				CacheAccessType();
 				SGraphPinBool::Construct(SGraphPinBool::FArguments(), InGraphPinObj);
 			}
 		};
@@ -385,6 +393,7 @@ namespace Metasound
 
 			void Construct(const FArguments& InArgs, UEdGraphPin* InGraphPinObj)
 			{
+				CacheAccessType();
 				SGraphPinNum<float>::Construct(SGraphPinNum<float>::FArguments(), InGraphPinObj);
 			}
 		};
@@ -396,6 +405,7 @@ namespace Metasound
 
 			void Construct(const FArguments& InArgs, UEdGraphPin* InGraphPinObj)
 			{
+				CacheAccessType();
 				SGraphPinInteger::Construct(SGraphPinInteger::FArguments(), InGraphPinObj);
 			}
 		};
@@ -407,6 +417,7 @@ namespace Metasound
 
 			void Construct(const FArguments& InArgs, UEdGraphPin* InGraphPinObj)
 			{
+				CacheAccessType();
 				SGraphPinObject::Construct(SGraphPinObject::FArguments(), InGraphPinObj);
 			}
 		};
@@ -418,6 +429,7 @@ namespace Metasound
 
 			void Construct(const FArguments& InArgs, UEdGraphPin* InGraphPinObj)
 			{
+				CacheAccessType();
 				SGraphPinString::Construct(SGraphPinString::FArguments(), InGraphPinObj);
 			}
 		};
@@ -434,7 +446,9 @@ namespace Metasound
 			virtual const FSlateBrush* GetPinIcon() const override;
 
 		protected:
-			bool HasRequiredConnections() const;
+			void CacheHasRequiredConnections();
+
+			bool bHasRequiredConnections = false;
 		};
 	} // namespace Editor
 } // namespace Metasound
