Fix MetaSound preset inputs being incorrectly selected after edit, save, or play

#rnx

diff --git a/Engine/Plugins/Runtime/Metasound/Source/MetasoundEditor/Private/MetasoundEditor.cpp b/Engine/Plugins/Runtime/Metasound/Source/MetasoundEditor/Private/MetasoundEditor.cpp
--- a/Engine/Plugins/Runtime/Metasound/Source/MetasoundEditor/Private/MetasoundEditor.cpp
+++ b/Engine/Plugins/Runtime/Metasound/Source/MetasoundEditor/Private/MetasoundEditor.cpp
@@ -3199,17 +3199,6 @@ namespace Metasound
 				if (!MembersModified.IsEmpty() || bForceRefreshViews)
 				{
 					RefreshGraphMemberMenu();
-
-					UMetasoundEditorGraph& Graph = GetMetaSoundGraphChecked();
-					for (const FGuid& MemberGuid : MembersModified)
-					{
-						if (UObject* Member = Graph.FindMember(MemberGuid))
-						{
-							// Currently only one member can be selected at a time, so only first found is added
-							Selection.Add(Member);
-							break;
-						}
-					}
 				}
 
 				// Only refresh details panel if
