Fix MetaSound edited graph member not selected after undo/redo

#rnx

diff --git a/Engine/Plugins/Runtime/Metasound/Source/MetasoundEditor/Private/MetasoundEditor.cpp b/Engine/Plugins/Runtime/Metasound/Source/MetasoundEditor/Private/MetasoundEditor.cpp
--- a/Engine/Plugins/Runtime/Metasound/Source/MetasoundEditor/Private/MetasoundEditor.cpp
+++ b/Engine/Plugins/Runtime/Metasound/Source/MetasoundEditor/Private/MetasoundEditor.cpp
@@ -2718,7 +2718,7 @@ namespace Metasound
 			}
 		}
 
-		void FEditor::RefreshGraphMemberMenu()
+		UMetasoundEditorGraphMember* FEditor::RefreshGraphMemberMenu()
 		{
 			TRACE_CPUPROFILER_EVENT_SCOPE(Metasound::Editor::FEditor::RefreshGraphMemberMenu);
 
@@ -2738,10 +2738,12 @@ namespace Metasound
 						{
 							const FName ActionName = Member->GetMemberName();
 							GraphMembersMenu->SelectItemByName(ActionName, ESelectInfo::Direct, Action->GetSectionID());
+							return Member;
 						}
 					}
 				}
 			}
+			return nullptr;
 		}
 
 		void FEditor::UpdateSelectedNodeClasses()
@@ -3198,7 +3200,22 @@ namespace Metasound
 
 				if (!MembersModified.IsEmpty() || bForceRefreshViews)
 				{
-					RefreshGraphMemberMenu();
+					UMetasoundEditorGraphMember* SelectedMember = RefreshGraphMemberMenu();
+
+					// If no member was selected by an action (ex. undo/redo), select a modified member 
+					if (!SelectedMember)
+					{
+						UMetasoundEditorGraph& Graph = GetMetaSoundGraphChecked();
+						for (const FGuid& MemberGuid : MembersModified)
+						{
+							if (UObject* Member = Graph.FindMember(MemberGuid))
+							{
+								// Currently only one member can be selected at a time, so only first found is added
+								Selection.Add(Member);
+								break;
+							}
+						}
+					}
 				}
 
 				// Only refresh details panel if
diff --git a/Engine/Plugins/Runtime/Metasound/Source/MetasoundEditor/Private/MetasoundEditor.h b/Engine/Plugins/Runtime/Metasound/Source/MetasoundEditor/Private/MetasoundEditor.h
--- a/Engine/Plugins/Runtime/Metasound/Source/MetasoundEditor/Private/MetasoundEditor.h
+++ b/Engine/Plugins/Runtime/Metasound/Source/MetasoundEditor/Private/MetasoundEditor.h
@@ -169,8 +169,8 @@ namespace Metasound
 			FGraphConnectionManager& GetConnectionManager();
 			const FGraphConnectionManager& GetConnectionManager() const;
 
-			/** Forces all UX pertaining to the root graph's interface to be refreshed. */
-			void RefreshGraphMemberMenu();
+			/** Forces all UX pertaining to the root graph's interface to be refreshed, returning the first selected member. */
+			UMetasoundEditorGraphMember* RefreshGraphMemberMenu();
 
 			/** Forces refresh of interfaces list. */
 			void RefreshInterfaces();
