Fix MetaSound input widget value not matching actual default value

#rnx

diff --git a/Engine/Plugins/Runtime/Metasound/Source/MetasoundEditor/Private/MetasoundEditorGraphMemberDefaults.cpp b/Engine/Plugins/Runtime/Metasound/Source/MetasoundEditor/Private/MetasoundEditorGraphMemberDefaults.cpp
--- a/Engine/Plugins/Runtime/Metasound/Source/MetasoundEditor/Private/MetasoundEditorGraphMemberDefaults.cpp
+++ b/Engine/Plugins/Runtime/Metasound/Source/MetasoundEditor/Private/MetasoundEditorGraphMemberDefaults.cpp
@@ -167,6 +167,9 @@ void UMetasoundEditorGraphMemberDefaultFloat::ForceRefresh()
 void UMetasoundEditorGraphMemberDefaultFloat::SetFromLiteral(const FMetasoundFrontendLiteral& InLiteral)
 {
 	Metasound::Editor::MemberDefaultsPrivate::ConvertLiteral<float>(InLiteral, Default);
+
+	// If set from literal, we force the default value to be the literal's value which may require the range to be fixed up 
+	SetInitialRange();
 }
 
 void UMetasoundEditorGraphMemberDefaultFloat::UpdatePreviewInstance(const Metasound::FVertexName& InParameterName, TScriptInterface<IAudioParameterControllerInterface>& InParameterInterface) const
@@ -188,6 +191,10 @@ void UMetasoundEditorGraphMemberDefaultFloat::PostEditChangeChainProperty(FPrope
 		{
 			VolumeWidgetDecibelRange = FVector2D(Audio::ConvertToDecibels(Range.X), Audio::ConvertToDecibels(Range.Y));
 		}
+		else
+		{
+			SetInitialRange();
+		}
 
 		// If the widget type is changed to none, we need to refresh clamping the value or not, since if the widget was a slider before, the value was clamped
 		OnClampChanged.Broadcast(ClampDefault);
@@ -228,19 +235,7 @@ void UMetasoundEditorGraphMemberDefaultFloat::PostEditChangeChainProperty(FPrope
 	}
 	else if (PropertyChangedEvent.GetPropertyName().IsEqual(GET_MEMBER_NAME_CHECKED(UMetasoundEditorGraphMemberDefaultFloat, ClampDefault)))
 	{
-		// If value is not within current range, keep it, otherwise set it to a reasonable range 
-		if (!(Default >= Range.X && Default <= Range.Y))
-		{
-			// set range to reasonable limit given current value
-			if (FMath::IsNearlyEqual(Default, 0.0f))
-			{
-				SetRange(FVector2D(0.0f, 1.0f));
-			}
-			else
-			{
-				SetRange(FVector2D(FMath::Min(0.0f, Default), FMath::Max(0.0f, Default)));
-			}
-		}
+		SetInitialRange();
 		OnClampChanged.Broadcast(ClampDefault);
 	}
 	else if (PropertyChangedEvent.PropertyChain.GetActiveMemberNode()->GetValue()->GetFName().IsEqual(GET_MEMBER_NAME_CHECKED(UMetasoundEditorGraphMemberDefaultFloat, VolumeWidgetDecibelRange)))
@@ -299,6 +294,22 @@ void UMetasoundEditorGraphMemberDefaultFloat::SetDefault(const float InDefault)
 	}
 }
 
+void UMetasoundEditorGraphMemberDefaultFloat::SetInitialRange()
+{
+	// If value is within current range, keep it, otherwise set range to something reasonable 
+	if (!(Default >= Range.X && Default <= Range.Y))
+	{
+		if (FMath::IsNearlyEqual(Default, 0.0f))
+		{
+			SetRange(FVector2D(0.0f, 1.0f));
+		}
+		else
+		{
+			SetRange(FVector2D(FMath::Min(0.0f, Default), FMath::Max(0.0f, Default)));
+		}
+	}
+}
+
 float UMetasoundEditorGraphMemberDefaultFloat::GetDefault()
 {
 	return Default;
diff --git a/Engine/Plugins/Runtime/Metasound/Source/MetasoundEditor/Private/MetasoundNodeDetailCustomization.cpp b/Engine/Plugins/Runtime/Metasound/Source/MetasoundEditor/Private/MetasoundNodeDetailCustomization.cpp
--- a/Engine/Plugins/Runtime/Metasound/Source/MetasoundEditor/Private/MetasoundNodeDetailCustomization.cpp
+++ b/Engine/Plugins/Runtime/Metasound/Source/MetasoundEditor/Private/MetasoundNodeDetailCustomization.cpp
@@ -200,7 +200,7 @@ namespace Metasound
 				DefaultRows.Add(WidgetCategoryBuilder.AddExternalObjectProperty(TArray<UObject*>({ DefaultFloat }), GET_MEMBER_NAME_CHECKED(UMetasoundEditorGraphMemberDefaultFloat, WidgetType)));
 				DefaultRows.Add(WidgetCategoryBuilder.AddExternalObjectProperty(TArray<UObject*>({ DefaultFloat }), GET_MEMBER_NAME_CHECKED(UMetasoundEditorGraphMemberDefaultFloat, WidgetOrientation)));
 				DefaultRows.Add(WidgetCategoryBuilder.AddExternalObjectProperty(TArray<UObject*>({ DefaultFloat }), GET_MEMBER_NAME_CHECKED(UMetasoundEditorGraphMemberDefaultFloat, WidgetValueType)));
-				if (DefaultFloat->WidgetValueType == EMetasoundMemberDefaultWidgetValueType::Volume)
+				if (DefaultFloat->WidgetType != EMetasoundMemberDefaultWidget::None && DefaultFloat->WidgetValueType == EMetasoundMemberDefaultWidgetValueType::Volume)
 				{
 					DefaultRows.Add(WidgetCategoryBuilder.AddExternalObjectProperty(TArray<UObject*>({ DefaultFloat }), GET_MEMBER_NAME_CHECKED(UMetasoundEditorGraphMemberDefaultFloat, VolumeWidgetUseLinearOutput)));
 					if (DefaultFloat->VolumeWidgetUseLinearOutput)
diff --git a/Engine/Plugins/Runtime/Metasound/Source/MetasoundEditor/Public/MetasoundEditorGraphMemberDefaults.h b/Engine/Plugins/Runtime/Metasound/Source/MetasoundEditor/Public/MetasoundEditorGraphMemberDefaults.h
--- a/Engine/Plugins/Runtime/Metasound/Source/MetasoundEditor/Public/MetasoundEditorGraphMemberDefaults.h
+++ b/Engine/Plugins/Runtime/Metasound/Source/MetasoundEditor/Public/MetasoundEditorGraphMemberDefaults.h
@@ -182,6 +182,8 @@ public:
 	virtual void PostEditChangeChainProperty(FPropertyChangedChainEvent& InPropertyChangedEvent) override;
 
 	void SetDefault(const float InDefault);
+	// Set Range to reasonable limit given current Default value
+	void SetInitialRange();
 	float GetDefault();
 	FVector2D GetRange();
 	void SetRange(const FVector2D InRange);
