[MetaSounds] Fix/add missing new MetaSound type icons/thumbnails as well as play/stop overlays.  Support for showing/filtering presets in ContentBrowser

diff --git a/Engine/Plugins/Enterprise/GLTFExporter/Source/GLTFExporter/Private/Actions/GLTFProxyAssetActions.h b/Engine/Plugins/Enterprise/GLTFExporter/Source/GLTFExporter/Private/Actions/GLTFProxyAssetActions.h
--- a/Engine/Plugins/Enterprise/GLTFExporter/Source/GLTFExporter/Private/Actions/GLTFProxyAssetActions.h
+++ b/Engine/Plugins/Enterprise/GLTFExporter/Source/GLTFExporter/Private/Actions/GLTFProxyAssetActions.h
@@ -46,7 +46,8 @@ public:
 	virtual void SetSupported(bool bInSupported) override { OriginalActions->SetSupported(bInSupported); }
 	virtual bool IsSupported() const override { return OriginalActions->IsSupported(); }
 	virtual FTopLevelAssetPath GetClassPathName() const override { return OriginalActions->GetClassPathName(); };
-
+	virtual const FSlateBrush* GetIconBrush(const FAssetData& InAssetData, const FName InClassName) const override { return OriginalActions->GetIconBrush(InAssetData, InClassName); }
+	virtual const FSlateBrush* GetThumbnailBrush(const FAssetData& InAssetData, const FName InClassName) const override { return OriginalActions->GetThumbnailBrush(InAssetData, InClassName); }
 private:
 
 	void GetProxyActions(const TArray<UObject*>& InObjects, FToolMenuSection& Section);
diff --git a/Engine/Plugins/Runtime/Metasound/Source/MetasoundEditor/Private/MetasoundAssetTypeActions.cpp b/Engine/Plugins/Runtime/Metasound/Source/MetasoundEditor/Private/MetasoundAssetTypeActions.cpp
--- a/Engine/Plugins/Runtime/Metasound/Source/MetasoundEditor/Private/MetasoundAssetTypeActions.cpp
+++ b/Engine/Plugins/Runtime/Metasound/Source/MetasoundEditor/Private/MetasoundAssetTypeActions.cpp
@@ -12,16 +12,25 @@
 #include "MetasoundSource.h"
 #include "MetasoundEditor.h"
 #include "MetasoundEditorGraphBuilder.h"
+#include "MetasoundEditorModule.h"
 #include "MetasoundEditorSettings.h"
 #include "MetasoundFactory.h"
 #include "MetasoundUObjectRegistry.h"
 #include "ObjectEditorUtils.h"
+#include "Styling/AppStyle.h"
+#include "Styling/CoreStyle.h"
 #include "Styling/ISlateStyle.h"
 #include "Styling/SlateStyleRegistry.h"
+#include "Styling/SlateTypes.h"
 #include "ToolMenus.h"
 #include "ToolMenuSection.h"
 #include "Widgets/Input/SButton.h"
-#include "Styling/AppStyle.h"
+
+#include "AssetRegistry/AssetRegistryModule.h"
+#include "MetasoundAssetManager.h"
+#include "Misc/AssetRegistryInterface.h"
+#include "Modules/ModuleManager.h"
+#include "Widgets/SOverlay.h"
 
 
 #define LOCTEXT_NAMESPACE "MetaSoundEditor"
@@ -70,19 +79,16 @@ namespace Metasound
 					return;
 				}
 
-				FString PresetClassName = TPresetClass::StaticClass()->GetName();
-				const FString EntryName = FString::Printf(TEXT("%sTo%s_Preset"), *PresetClassName, *ClassName);
+				FName PresetClassName = TPresetClass::StaticClass()->GetFName();
+				const FString EntryName = FString::Printf(TEXT("%sTo%s_Preset"), *PresetClassName.ToString(), *ClassName);
 				FToolMenuSection& Section = Menu->FindOrAddSection("GetAssetActions");
-				Section.AddDynamicEntry(*EntryName, FNewToolMenuSectionDelegate::CreateLambda([InPresetClassName = MoveTemp(PresetClassName), Label = InLabel, ToolTip = InToolTip](FToolMenuSection& InSection)
+				Section.AddDynamicEntry(*EntryName, FNewToolMenuSectionDelegate::CreateLambda([PresetClassName, Label = InLabel, ToolTip = InToolTip](FToolMenuSection& InSection)
 				{
-					// TODO: Make class icons for MetaSound types. For now just use SoundCue.
-					//const FString IconName = TEXT("ClassIcon.") + InPresetClassName;
-					const FString IconName = TEXT("ClassIcon.SoundCue");
-
-					const FSlateIcon Icon = FSlateIcon(FAppStyle::GetAppStyleSetName(), *IconName);
+					const FName IconName = *FString::Printf(TEXT("ClassIcon.%s"), *PresetClassName.ToString());
+					const FSlateIcon Icon = Style::CreateSlateIcon(IconName);
 					const FToolMenuExecuteAction UIExecuteAction = FToolMenuExecuteAction::CreateStatic(&ExecuteCreatePreset<TPresetClass, TFactory>);
 
-					const FString PresetEntryName = InPresetClassName + TEXT("_CreatePreset");
+					const FString PresetEntryName = FString::Printf(TEXT("%s_CreatePreset"), *PresetClassName.ToString());
 					InSection.AddMenuEntry(*PresetEntryName, Label, ToolTip, Icon, UIExecuteAction);
 				}));
 			}
@@ -139,6 +145,50 @@ namespace Metasound
 					PreviewComponent->Stop();
 				}
 			}
+
+			bool GetIsPreset(const FSoftObjectPath& InSourcePath)
+			{
+				using namespace Metasound::Editor;
+				using namespace Metasound::Frontend;
+
+				int32 bIsPreset = 0;
+				// If object in memory, try to resolve but not load
+				if (UObject* Object = InSourcePath.ResolveObject())
+				{
+					const FMetasoundAssetBase* MetaSoundAsset = IMetasoundUObjectRegistry::Get().GetObjectAsAssetBase(Object);
+					if (MetaSoundAsset)
+					{
+						bIsPreset = MetaSoundAsset->GetDocumentChecked().RootGraph.PresetOptions.bIsPreset;
+					}
+				}
+				// Otherwise, try to pull from asset registry, but avoid load as this call
+				// would then be slow and is called from the ContentBrowser many times
+				else
+				{
+					FAssetData IconAssetData;
+					const FAssetRegistryModule& AssetRegistryModule = FModuleManager::GetModuleChecked<FAssetRegistryModule>("AssetRegistry");
+					UE::AssetRegistry::EExists Exists = AssetRegistryModule.TryGetAssetByObjectPath(InSourcePath, IconAssetData);
+					if (Exists == UE::AssetRegistry::EExists::Exists)
+					{
+						IconAssetData.GetTagValue(AssetTags::IsPreset, bIsPreset);
+					}
+				}
+
+				return static_cast<bool>(bIsPreset);
+			}
+
+			const FSlateBrush* GetClassBrush(const FAssetData& InAssetData, FName InClassName, bool bIsThumbnail = false)
+			{
+				const bool bIsPreset = AssetTypeActionsPrivate::GetIsPreset(InAssetData.ToSoftObjectPath());
+				FString BrushName = FString::Printf(TEXT("MetasoundEditor.%s"), *InClassName.ToString());
+				if (bIsPreset)
+				{
+					BrushName += TEXT(".Preset");
+				}
+				BrushName += bIsThumbnail ? TEXT(".Thumbnail") : TEXT(".Icon");
+
+				return &Style::GetSlateBrushSafe(FName(*BrushName));
+			}
 		} // namespace AssetTypeActionsPrivate
 
 		UClass* FAssetTypeActions_MetaSoundPatch::GetSupportedClass() const
@@ -192,6 +242,17 @@ namespace Metasound
 			return SubMenus;
 		}
 
+		const FSlateBrush* FAssetTypeActions_MetaSoundPatch::GetThumbnailBrush(const FAssetData& InAssetData, const FName InClassName) const
+		{
+			constexpr bool bIsThumbnail = true;
+			return AssetTypeActionsPrivate::GetClassBrush(InAssetData, InClassName, bIsThumbnail);
+		}
+
+		const FSlateBrush* FAssetTypeActions_MetaSoundPatch::GetIconBrush(const FAssetData& InAssetData, const FName InClassName) const
+		{
+			return AssetTypeActionsPrivate::GetClassBrush(InAssetData, InClassName);
+		}
+
 		UClass* FAssetTypeActions_MetaSoundSource::GetSupportedClass() const
 		{
 			return UMetaSoundSource::StaticClass();
@@ -225,13 +286,6 @@ namespace Metasound
 		{
 			using namespace AssetTypeActionsPrivate;
 			RegisterPresetAction<UMetaSoundSource, UMetaSoundSourceFactory>(SourcePresetLabel, SourcePresetToolTip);
-
-			// This is currently disabled because of the requirement that interfaces of a preset and its referenced
-			// asset must match.  TODO: A feature of interfaces/preset transform is required for this action to work
-			// that introduces the concept of an interface "interface" (i.e. if the interface one graph subscribes to
-			// or contains all of the necessary inputs/outputs of another), so that the preset can "implement" the
-			// referenced graphs "interface."
-// 			RegisterPresetAction<UMetaSoundPatch, UMetaSoundFactory, UMetaSoundSource>(PresetLabel, PresetToolTip);
 		}
 
 		const TArray<FText>& FAssetTypeActions_MetaSoundSource::GetSubMenus() const
@@ -250,21 +304,41 @@ namespace Metasound
 			return SubMenus;
 		}
 
-		TSharedPtr<SWidget> FAssetTypeActions_MetaSoundSource::GetThumbnailOverlay(const FAssetData& AssetData) const
+		TSharedPtr<SWidget> FAssetTypeActions_MetaSoundSource::GetThumbnailOverlay(const FAssetData& InAssetData) const
 		{
-			auto OnGetDisplayBrushLambda = [Path = AssetData.ToSoftObjectPath()]()
+			TSharedPtr<SBox> Box = SNew(SBox)
+				.HAlign(HAlign_Fill)
+				.VAlign(VAlign_Fill)
+				.Padding(2.0f);
+
+			auto OnGetDisplayPlaybackLambda = [Box, Path = InAssetData.ToSoftObjectPath()]()
 			{
 				using namespace AssetTypeActionsPrivate;
 
-				if (IsPlaying(Path))
+				const bool bIsValid = Box.IsValid();
+				if (!bIsValid)
+				{
+					return FCoreStyle::Get().GetBrush(TEXT("NoBrush"));
+				}
+
+				const bool bIsPlaying = IsPlaying(Path);
+				const bool bIsHovered = Box->IsHovered();
+				if (!bIsPlaying && !bIsHovered)
 				{
-					return FAppStyle::GetBrush("MediaAsset.AssetActions.Stop.Large");
+					return FCoreStyle::Get().GetBrush(TEXT("NoBrush"));
 				}
 
-				return FAppStyle::GetBrush("MediaAsset.AssetActions.Play.Large");
+				FString IconName = TEXT("MetasoundEditor");
+				IconName += bIsPlaying ? TEXT(".Stop.Thumbnail") : TEXT(".Play.Thumbnail");
+				if (bIsHovered)
+				{
+					IconName += TEXT(".Hovered");
+				}
+
+				return &Style::GetSlateBrushSafe(FName(*IconName));
 			};
 
-			auto OnClickedLambda = [Path = AssetData.ToSoftObjectPath()]()
+			auto OnClickedLambda = [Path = InAssetData.ToSoftObjectPath()]()
 			{
 				using namespace AssetTypeActionsPrivate;
 				// Load and play sound
@@ -283,7 +357,7 @@ namespace Metasound
 				return FReply::Handled();
 			};
 
-			auto OnToolTipTextLambda = [Path = AssetData.ToSoftObjectPath()]()
+			auto OnToolTipTextLambda = [ClassName = GetSupportedClass()->GetFName(), Path = InAssetData.ToSoftObjectPath()]()
 			{
 				using namespace AssetTypeActionsPrivate;
 
@@ -297,46 +371,36 @@ namespace Metasound
 					Format = LOCTEXT("PreviewMetaSoundFromIconToolTip_Editor", "Preview {0}");
 				}
 
-				FName TypeName = UMetaSoundSource::StaticClass()->GetFName();
-				return FText::Format(Format, FText::FromName(TypeName));
+				return FText::Format(Format, FText::FromName(ClassName));
 			};
 
-			TSharedPtr<SBox> Box;
-			SAssignNew(Box, SBox)
+			Box->SetContent(
+				SNew(SButton)
+				.ButtonStyle(&FCoreStyle::Get().GetWidgetStyle<FButtonStyle>("NoBorder"))
 				.HAlign(HAlign_Center)
 				.VAlign(VAlign_Center)
-				.Padding(FMargin(2));
-
-			auto OnGetVisibilityLambda = [Box, Path = AssetData.ToSoftObjectPath()]()
-			{
-				using namespace AssetTypeActionsPrivate;
-
-				if (Box.IsValid() && (Box->IsHovered() || IsPlaying(Path)))
-				{
-					return EVisibility::Visible;
-				}
-
-				return EVisibility::Hidden;
-			};
-
-			TSharedPtr<SButton> Widget;
-			SAssignNew(Widget, SButton)
-				.ButtonStyle(FAppStyle::Get(), "HoverHintOnly")
 				.ToolTipText_Lambda(OnToolTipTextLambda)
-				.Cursor(EMouseCursor::Default) // The outer widget can specify a DragHand cursor, so overriden here
+				.Cursor(EMouseCursor::Default) // The outer widget can specify a DragHand cursor, so overridden here
 				.ForegroundColor(FSlateColor::UseForeground())
 				.IsFocusable(false)
 				.OnClicked_Lambda(OnClickedLambda)
-				.Visibility_Lambda(OnGetVisibilityLambda)
 				[
 					SNew(SImage)
-					.Image_Lambda(OnGetDisplayBrushLambda)
-				];
+					.Image_Lambda(OnGetDisplayPlaybackLambda)
+				]
+			);
+			return Box;
+		}
 
-			Box->SetContent(Widget.ToSharedRef());
-			Box->SetVisibility(EVisibility::Visible);
+		const FSlateBrush* FAssetTypeActions_MetaSoundSource::GetThumbnailBrush(const FAssetData& InAssetData, const FName InClassName) const
+		{
+			constexpr bool bIsThumbnail = true;
+			return AssetTypeActionsPrivate::GetClassBrush(InAssetData, InClassName, bIsThumbnail);
+		}
 
-			return Box;
+		const FSlateBrush* FAssetTypeActions_MetaSoundSource::GetIconBrush(const FAssetData& InAssetData, const FName InClassName) const
+		{
+			return AssetTypeActionsPrivate::GetClassBrush(InAssetData, InClassName);
 		}
 	} // namespace Editor
 } // namespace Metasound
diff --git a/Engine/Plugins/Runtime/Metasound/Source/MetasoundEditor/Private/MetasoundAssetTypeActions.h b/Engine/Plugins/Runtime/Metasound/Source/MetasoundEditor/Private/MetasoundAssetTypeActions.h
--- a/Engine/Plugins/Runtime/Metasound/Source/MetasoundEditor/Private/MetasoundAssetTypeActions.h
+++ b/Engine/Plugins/Runtime/Metasound/Source/MetasoundEditor/Private/MetasoundAssetTypeActions.h
@@ -22,6 +22,9 @@ namespace Metasound
 			virtual const TArray<FText>& GetSubMenus() const override;
 			virtual void OpenAssetEditor(const TArray<UObject*>& InObjects, TSharedPtr<IToolkitHost> ToolkitHost) override;
 
+			virtual const FSlateBrush* GetThumbnailBrush(const FAssetData& InAssetData, const FName InClassName) const override;
+			virtual const FSlateBrush* GetIconBrush(const FAssetData& InAssetData, const FName InClassName) const override;
+
 			static void RegisterMenuActions();
 		};
 
@@ -36,7 +39,9 @@ namespace Metasound
 			virtual const TArray<FText>& GetSubMenus() const override;
 			virtual void OpenAssetEditor(const TArray<UObject*>& InObjects, TSharedPtr<IToolkitHost> ToolkitHost) override;
 
+			virtual const FSlateBrush* GetThumbnailBrush(const FAssetData& InAssetData, const FName InClassName) const override;
 			virtual TSharedPtr<SWidget> GetThumbnailOverlay(const FAssetData& AssetData) const override;
+			virtual const FSlateBrush* GetIconBrush(const FAssetData& InAssetData, const FName InClassName) const override;
 
 			static void RegisterMenuActions();
 		};
diff --git a/Engine/Plugins/Runtime/Metasound/Source/MetasoundEditor/Private/MetasoundEditor.cpp b/Engine/Plugins/Runtime/Metasound/Source/MetasoundEditor/Private/MetasoundEditor.cpp
--- a/Engine/Plugins/Runtime/Metasound/Source/MetasoundEditor/Private/MetasoundEditor.cpp
+++ b/Engine/Plugins/Runtime/Metasound/Source/MetasoundEditor/Private/MetasoundEditor.cpp
@@ -908,12 +908,33 @@ namespace Metasound
 
 		const FSlateBrush* FEditor::GetDefaultTabIcon() const
 		{
+			FString IconName = TEXT("MetasoundEditor");
 			if (IsPlaying())
 			{
-				return &Style::GetSlateBrushSafe("MetasoundEditor.Play");
+				IconName += TEXT(".Play");
+			}
+			else
+			{
+				if (UMetaSoundSource* MetaSoundSource = Cast<UMetaSoundSource>(Metasound))
+				{
+					IconName += TEXT(".MetasoundSource");
+				}
+				else if (UMetaSoundPatch* MetaSoundPatch = Cast<UMetaSoundPatch>(Metasound))
+				{
+					IconName += TEXT(".MetasoundPatch");
+				}
+
+				const FMetasoundAssetBase* MetaSoundAsset = IMetasoundUObjectRegistry::Get().GetObjectAsAssetBase(Metasound);
+				check(MetaSoundAsset);
+				if (MetaSoundAsset->GetDocumentChecked().RootGraph.PresetOptions.bIsPreset)
+				{
+					IconName += TEXT(".Preset");
+				}
+
+				IconName += TEXT(".Icon");
 			}
 
-			return FAssetEditorToolkit::GetDefaultTabIcon();
+			return &Style::GetSlateBrushSafe(FName(*IconName));
 		}
 
 		FLinearColor FEditor::GetDefaultTabColor() const
diff --git a/Engine/Plugins/Runtime/Metasound/Source/MetasoundEditor/Private/MetasoundEditorModule.cpp b/Engine/Plugins/Runtime/Metasound/Source/MetasoundEditor/Private/MetasoundEditorModule.cpp
--- a/Engine/Plugins/Runtime/Metasound/Source/MetasoundEditor/Private/MetasoundEditorModule.cpp
+++ b/Engine/Plugins/Runtime/Metasound/Source/MetasoundEditor/Private/MetasoundEditorModule.cpp
@@ -101,12 +101,14 @@ namespace Metasound
 
 				// Metasound Editor
 				{
-					Set("MetaSoundPatch.Color", FColor(13, 55, 13));
+					Set("MetaSoundPatch.Color", FColor(31, 133, 31));
 					Set("MetaSoundSource.Color", FColor(103, 214, 66));
 
 					// Actions
 					Set("MetasoundEditor.Play", new IMAGE_BRUSH_SVG(TEXT("Icons/play"), Icon40x40));
 					Set("MetasoundEditor.Play.Small", new IMAGE_BRUSH_SVG(TEXT("Icons/play"), Icon20x20));
+					Set("MetasoundEditor.Play.Thumbnail", new IMAGE_BRUSH_SVG(TEXT("Icons/play_thumbnail"), Icon64));
+					Set("MetasoundEditor.Play.Thumbnail.Hovered", new IMAGE_BRUSH_SVG(TEXT("Icons/play_thumbnail_hover"), Icon64));
 
 					Set("MetasoundEditor.Play.Active.Valid", new IMAGE_BRUSH_SVG(TEXT("Icons/play_active_valid"), Icon40x40));
 					Set("MetasoundEditor.Play.Active.Warning", new IMAGE_BRUSH_SVG(TEXT("Icons/play_active_warning"), Icon40x40));
@@ -119,7 +121,8 @@ namespace Metasound
 					Set("MetasoundEditor.Stop.Disabled", new IMAGE_BRUSH_SVG(TEXT("Icons/stop_disabled"), Icon40x40));
 					Set("MetasoundEditor.Stop.Active", new IMAGE_BRUSH_SVG(TEXT("Icons/stop_active"), Icon40x40));
 					Set("MetasoundEditor.Stop.Inactive", new IMAGE_BRUSH_SVG(TEXT("Icons/stop_inactive"), Icon40x40));
-
+					Set("MetasoundEditor.Stop.Thumbnail", new IMAGE_BRUSH_SVG(TEXT("Icons/stop_thumbnail"), Icon64));
+					Set("MetasoundEditor.Stop.Thumbnail.Hovered", new IMAGE_BRUSH_SVG(TEXT("Icons/stop_thumbnail_hover"), Icon64));
 
 					Set("MetasoundEditor.Import", new FSlateImageBrush(RootToContentDir(TEXT("/Icons/build_40x.png")), Icon40x40));
 					Set("MetasoundEditor.Import.Small", new FSlateImageBrush(RootToContentDir(TEXT("/Icons/build_40x.png")), Icon20x20));
@@ -178,6 +181,15 @@ namespace Metasound
 
 					SetClassIcon(TEXT("MetasoundPatch"));
 					SetClassIcon(TEXT("MetasoundSource"));
+
+					Set("MetasoundEditor.MetasoundPatch.Icon", new IMAGE_BRUSH_SVG(TEXT("Icons/metasoundpatch_icon"), Icon20x20));
+					Set("MetasoundEditor.MetasoundPatch.Preset.Icon", new IMAGE_BRUSH_SVG(TEXT("Icons/metasoundpatchpreset_icon"), Icon20x20));
+					Set("MetasoundEditor.MetasoundSource.Icon", new IMAGE_BRUSH_SVG(TEXT("Icons/metasoundsource_icon"), Icon20x20));
+					Set("MetasoundEditor.MetasoundSource.Preset.Icon", new IMAGE_BRUSH_SVG(TEXT("Icons/metasoundsourcepreset_icon"), Icon20x20));
+					Set("MetasoundEditor.MetasoundPatch.Thumbnail", new IMAGE_BRUSH_SVG(TEXT("Icons/metasoundpatch_thumbnail"), Icon20x20));
+					Set("MetasoundEditor.MetasoundPatch.Preset.Thumbnail", new IMAGE_BRUSH_SVG(TEXT("Icons/metasoundpatchpreset_thumbnail"), Icon20x20));
+					Set("MetasoundEditor.MetasoundSource.Thumbnail", new IMAGE_BRUSH_SVG(TEXT("Icons/metasoundsource_thumbnail"), Icon20x20));
+					Set("MetasoundEditor.MetasoundSource.Preset.Thumbnail", new IMAGE_BRUSH_SVG(TEXT("Icons/metasoundsourcepreset_thumbnail"), Icon20x20));
 				}
 
 				FSlateStyleRegistry::RegisterSlateStyle(*this);
diff --git a/Engine/Plugins/Runtime/Metasound/Source/MetasoundEngine/Private/MetasoundAssetSubsystem.cpp b/Engine/Plugins/Runtime/Metasound/Source/MetasoundEngine/Private/MetasoundAssetSubsystem.cpp
--- a/Engine/Plugins/Runtime/Metasound/Source/MetasoundEngine/Private/MetasoundAssetSubsystem.cpp
+++ b/Engine/Plugins/Runtime/Metasound/Source/MetasoundEngine/Private/MetasoundAssetSubsystem.cpp
@@ -34,12 +34,15 @@ namespace Metasound
 
 			OutInfo.Type = EMetasoundFrontendClassType::External;
 			OutInfo.AssetPath = InAssetData.GetSoftObjectPath();
-
 			FString AssetClassID;
 			bSuccess &= InAssetData.GetTagValue(AssetTags::AssetClassID, AssetClassID);
 			OutInfo.AssetClassID = FGuid(AssetClassID);
 			OutInfo.ClassName = FMetasoundFrontendClassName(FName(), *AssetClassID, FName());
 
+#if WITH_EDITORONLY_DATA
+			InAssetData.GetTagValue(AssetTags::IsPreset, OutInfo.bIsPreset);
+#endif // WITH_EDITORONLY_DATA
+
 			int32 RegistryVersionMajor = 0;
 			bSuccess &= InAssetData.GetTagValue(AssetTags::RegistryVersionMajor, RegistryVersionMajor);
 			OutInfo.Version.Major = RegistryVersionMajor;
diff --git a/Engine/Plugins/Runtime/Metasound/Source/MetasoundEngine/Private/MetasoundEngineAsset.h b/Engine/Plugins/Runtime/Metasound/Source/MetasoundEngine/Private/MetasoundEngineAsset.h
--- a/Engine/Plugins/Runtime/Metasound/Source/MetasoundEngine/Private/MetasoundEngineAsset.h
+++ b/Engine/Plugins/Runtime/Metasound/Source/MetasoundEngine/Private/MetasoundEngineAsset.h
@@ -167,6 +167,7 @@ namespace Metasound
 			using namespace Metasound::Frontend;
 
 			check(AssetTags::AssetClassID == GET_MEMBER_NAME_CHECKED(TMetaSoundObject, AssetClassID));
+			check(AssetTags::IsPreset == GET_MEMBER_NAME_CHECKED(TMetaSoundObject, bIsPreset));
 			check(AssetTags::RegistryInputTypes == GET_MEMBER_NAME_CHECKED(TMetaSoundObject, RegistryInputTypes));
 			check(AssetTags::RegistryOutputTypes == GET_MEMBER_NAME_CHECKED(TMetaSoundObject, RegistryOutputTypes));
 			check(AssetTags::RegistryVersionMajor == GET_MEMBER_NAME_CHECKED(TMetaSoundObject, RegistryVersionMajor));
@@ -175,10 +176,12 @@ namespace Metasound
 			bool bMarkDirty = InMetaSound.AssetClassID != InClassInfo.AssetClassID;
 			bMarkDirty |= InMetaSound.RegistryVersionMajor != InClassInfo.Version.Major;
 			bMarkDirty |= InMetaSound.RegistryVersionMinor != InClassInfo.Version.Minor;
+			bMarkDirty |= InMetaSound.bIsPreset != InClassInfo.bIsPreset;
 
 			InMetaSound.AssetClassID = InClassInfo.AssetClassID;
 			InMetaSound.RegistryVersionMajor = InClassInfo.Version.Major;
 			InMetaSound.RegistryVersionMinor = InClassInfo.Version.Minor;
+			InMetaSound.bIsPreset = InClassInfo.bIsPreset;
 
 			{
 				TArray<FString> InputTypes;
diff --git a/Engine/Plugins/Runtime/Metasound/Source/MetasoundEngine/Public/Metasound.h b/Engine/Plugins/Runtime/Metasound/Source/MetasoundEngine/Public/Metasound.h
--- a/Engine/Plugins/Runtime/Metasound/Source/MetasoundEngine/Public/Metasound.h
+++ b/Engine/Plugins/Runtime/Metasound/Source/MetasoundEngine/Public/Metasound.h
@@ -94,6 +94,9 @@ public:
 	UPROPERTY(AssetRegistrySearchable)
 	int32 RegistryVersionMinor = 0;
 
+	UPROPERTY(AssetRegistrySearchable)
+	bool bIsPreset = false;
+
 	// Sets Asset Registry Metadata associated with this MetaSound
 	virtual void SetRegistryAssetClassInfo(const Metasound::Frontend::FNodeClassInfo& InClassInfo) override;
 
diff --git a/Engine/Plugins/Runtime/Metasound/Source/MetasoundEngine/Public/MetasoundSource.h b/Engine/Plugins/Runtime/Metasound/Source/MetasoundEngine/Public/MetasoundSource.h
--- a/Engine/Plugins/Runtime/Metasound/Source/MetasoundEngine/Public/MetasoundSource.h
+++ b/Engine/Plugins/Runtime/Metasound/Source/MetasoundEngine/Public/MetasoundSource.h
@@ -86,6 +86,8 @@ public:
 	UPROPERTY(AssetRegistrySearchable)
 	int32 RegistryVersionMinor = 0;
 
+	UPROPERTY(AssetRegistrySearchable)
+	bool bIsPreset = false;
 
 	// Sets Asset Registry Metadata associated with this MetaSoundSource
 	virtual void SetRegistryAssetClassInfo(const Metasound::Frontend::FNodeClassInfo& InNodeInfo) override;
diff --git a/Engine/Plugins/Runtime/Metasound/Source/MetasoundFrontend/Private/MetasoundAssetManager.cpp b/Engine/Plugins/Runtime/Metasound/Source/MetasoundFrontend/Private/MetasoundAssetManager.cpp
--- a/Engine/Plugins/Runtime/Metasound/Source/MetasoundFrontend/Private/MetasoundAssetManager.cpp
+++ b/Engine/Plugins/Runtime/Metasound/Source/MetasoundFrontend/Private/MetasoundAssetManager.cpp
@@ -11,6 +11,11 @@ namespace Metasound
 			const FString ArrayDelim = TEXT(",");
 
 			const FName AssetClassID = "AssetClassID";
+
+#if WITH_EDITORONLY_DATA
+			const FName IsPreset = "bIsPreset";
+#endif // WITH_EDITORONLY_DATA
+
 			const FName RegistryVersionMajor = "RegistryVersionMajor";
 			const FName RegistryVersionMinor = "RegistryVersionMinor";
 
diff --git a/Engine/Plugins/Runtime/Metasound/Source/MetasoundFrontend/Private/MetasoundFrontendRegistries.cpp b/Engine/Plugins/Runtime/Metasound/Source/MetasoundFrontend/Private/MetasoundFrontendRegistries.cpp
--- a/Engine/Plugins/Runtime/Metasound/Source/MetasoundFrontend/Private/MetasoundFrontendRegistries.cpp
+++ b/Engine/Plugins/Runtime/Metasound/Source/MetasoundFrontend/Private/MetasoundFrontendRegistries.cpp
@@ -646,6 +646,8 @@ namespace Metasound
 			{
 				OutputTypes.Add(Output.TypeName);
 			}
+
+			bIsPreset = InClass.PresetOptions.bIsPreset;
 #endif // WITH_EDITORONLY_DATA
 		}
 	} // namespace Frontend
diff --git a/Engine/Plugins/Runtime/Metasound/Source/MetasoundFrontend/Public/MetasoundAssetManager.h b/Engine/Plugins/Runtime/Metasound/Source/MetasoundFrontend/Public/MetasoundAssetManager.h
--- a/Engine/Plugins/Runtime/Metasound/Source/MetasoundFrontend/Public/MetasoundAssetManager.h
+++ b/Engine/Plugins/Runtime/Metasound/Source/MetasoundFrontend/Public/MetasoundAssetManager.h
@@ -20,6 +20,10 @@ namespace Metasound
 		{
 			extern const FString METASOUNDFRONTEND_API ArrayDelim;
 
+#if WITH_EDITORONLY_DATA
+			extern const FName METASOUNDFRONTEND_API IsPreset;
+#endif // WITH_EDITORONLY_DATA
+
 			extern const FName METASOUNDFRONTEND_API AssetClassID;
 			extern const FName METASOUNDFRONTEND_API RegistryVersionMajor;
 			extern const FName METASOUNDFRONTEND_API RegistryVersionMinor;
diff --git a/Engine/Plugins/Runtime/Metasound/Source/MetasoundFrontend/Public/MetasoundFrontendRegistries.h b/Engine/Plugins/Runtime/Metasound/Source/MetasoundFrontend/Public/MetasoundFrontendRegistries.h
--- a/Engine/Plugins/Runtime/Metasound/Source/MetasoundFrontend/Public/MetasoundFrontendRegistries.h
+++ b/Engine/Plugins/Runtime/Metasound/Source/MetasoundFrontend/Public/MetasoundFrontendRegistries.h
@@ -47,6 +47,9 @@ namespace Metasound
 
 			// Types of class outputs
 			TSet<FName> OutputTypes;
+
+			// Whether or not class is preset
+			bool bIsPreset = false;
 #endif // WITH_EDITORONLY_DATA
 
 			FNodeClassInfo() = default;
diff --git a/Engine/Plugins/Runtime/Synthesis/Source/SynthesisEditor/Private/AudioImpulseResponseAsset.cpp b/Engine/Plugins/Runtime/Synthesis/Source/SynthesisEditor/Private/AudioImpulseResponseAsset.cpp
--- a/Engine/Plugins/Runtime/Synthesis/Source/SynthesisEditor/Private/AudioImpulseResponseAsset.cpp
+++ b/Engine/Plugins/Runtime/Synthesis/Source/SynthesisEditor/Private/AudioImpulseResponseAsset.cpp
@@ -1,5 +1,8 @@
 // Copyright Epic Games, Inc. All Rights Reserved.
 #include "AudioImpulseResponseAsset.h"
+
+#include "Algo/AnyOf.h"
+#include "AssetRegistry/AssetData.h"
 #include "UObject/ObjectMacros.h"
 #include "UObject/Object.h"
 #include "ToolMenus.h"
@@ -24,12 +27,18 @@ void FAudioImpulseResponseExtension::RegisterMenus()
 
 	Section.AddDynamicEntry("SoundWaveAssetConversion_CreateImpulseResponse", FNewToolMenuSectionDelegate::CreateLambda([](FToolMenuSection& InSection)
 	{
-		const TAttribute<FText> Label = LOCTEXT("SoundWave_CreateImpulseResponse", "Create Impulse Response");
-		const TAttribute<FText> ToolTip = LOCTEXT("SoundWave_CreateImpulseResponseTooltip", "Creates an impulse response asset using the selected sound wave.");
-		const FSlateIcon Icon = FSlateIcon(FAppStyle::GetAppStyleSetName(), "ClassIcon.ImpulseResponse");
-		const FToolMenuExecuteAction UIAction = FToolMenuExecuteAction::CreateStatic(&FAudioImpulseResponseExtension::ExecuteCreateImpulseResponse);
+		if (const UContentBrowserAssetContextMenuContext* Context = InSection.FindContext<UContentBrowserAssetContextMenuContext>())
+		{
+			if (!Algo::AnyOf(Context->SelectedAssets, [](const FAssetData& AssetData){ return AssetData.IsInstanceOf<USoundWaveProcedural>(); }))
+			{
+				const TAttribute<FText> Label = LOCTEXT("SoundWave_CreateImpulseResponse", "Create Impulse Response");
+				const TAttribute<FText> ToolTip = LOCTEXT("SoundWave_CreateImpulseResponseTooltip", "Creates an impulse response asset using the selected sound wave.");
+				const FSlateIcon Icon = FSlateIcon(FAppStyle::GetAppStyleSetName(), "ClassIcon.ImpulseResponse");
+				const FToolMenuExecuteAction UIAction = FToolMenuExecuteAction::CreateStatic(&FAudioImpulseResponseExtension::ExecuteCreateImpulseResponse);
 
-		InSection.AddMenuEntry("SoundWave_CreateImpulseResponse", Label, ToolTip, Icon, UIAction);
+				InSection.AddMenuEntry("SoundWave_CreateImpulseResponse", Label, ToolTip, Icon, UIAction);
+			}
+		}
 	}));
 }
 
diff --git a/Engine/Source/Developer/AssetTools/Public/AssetTypeActions_Base.h b/Engine/Source/Developer/AssetTools/Public/AssetTypeActions_Base.h
--- a/Engine/Source/Developer/AssetTools/Public/AssetTypeActions_Base.h
+++ b/Engine/Source/Developer/AssetTools/Public/AssetTypeActions_Base.h
@@ -172,6 +172,10 @@ public:
 		return GetSupportedClass()->GetClassPathName();
 	}
 
+	virtual const FSlateBrush* GetThumbnailBrush(const FAssetData& InAssetData, const FName InClassName) const override { return nullptr; }
+
+	virtual const FSlateBrush* GetIconBrush(const FAssetData& InAssetData, const FName InClassName) const override { return nullptr; }
+
 protected:
 
 	// Here are some convenience functions for common asset type actions logic
diff --git a/Engine/Source/Developer/AssetTools/Public/IAssetTypeActions.h b/Engine/Source/Developer/AssetTools/Public/IAssetTypeActions.h
--- a/Engine/Source/Developer/AssetTools/Public/IAssetTypeActions.h
+++ b/Engine/Source/Developer/AssetTools/Public/IAssetTypeActions.h
@@ -174,10 +174,14 @@ public:
 	/** Is this asset type supported in the current session? */
 	virtual bool IsSupported() const = 0;
 
-	/**
-	 * Returns class path name as a package + class FName pair
-	 */
+	/** Returns class path name as a package + class FName pair */
 	virtual FTopLevelAssetPath GetClassPathName() const = 0;
+
+	/** Returns thumbnail brush unique for the given asset data.  Returning null falls back to class thumbnail brush. */
+	virtual const FSlateBrush* GetThumbnailBrush(const FAssetData& InAssetData, const FName InClassName) const = 0;
+
+	/** Returns icon brush unique for the given asset data.  Returning null falls back to class icon brush. */
+	virtual const FSlateBrush* GetIconBrush(const FAssetData& InAssetData, const FName InClassName) const = 0;
 };
 
 /**
diff --git a/Engine/Source/Editor/UnrealEd/Private/AssetThumbnail.cpp b/Engine/Source/Editor/UnrealEd/Private/AssetThumbnail.cpp
--- a/Engine/Source/Editor/UnrealEd/Private/AssetThumbnail.cpp
+++ b/Engine/Source/Editor/UnrealEd/Private/AssetThumbnail.cpp
@@ -126,7 +126,7 @@ public:
 
 		TSharedRef<SOverlay> OverlayWidget = SNew(SOverlay);
 
-		UpdateThumbnailClass();
+		UpdateThumbnailClass(AssetTypeActions.Get());
 
 		ClassThumbnailBrushOverride = InArgs._ClassThumbnailBrushOverride;
 
@@ -277,18 +277,38 @@ public:
 
 	}
 
-	void UpdateThumbnailClass()
+	void UpdateThumbnailClass(const IAssetTypeActions* AssetTypeActions)
 	{
 		const FAssetData& AssetData = AssetThumbnail->GetAssetData();
 		ThumbnailClass = MakeWeakObjectPtr(const_cast<UClass*>(FClassIconFinder::GetIconClassForAssetData(AssetData, &bIsClassType)));
+		const FName AssetClassName = AssetThumbnail->GetAssetData().AssetClassPath.GetAssetName();
 
-		// For non-class types, use the default based upon the actual asset class
-		// This has the side effect of not showing a class icon for assets that don't have a proper thumbnail image available
-		const FName DefaultThumbnail = (bIsClassType) ? NAME_None : FName(*FString::Printf(TEXT("ClassThumbnail.%s"), *AssetThumbnail->GetAssetData().AssetClassPath.GetAssetName().ToString()));
-		ThumbnailBrush = FClassIconFinder::FindThumbnailForClass(ThumbnailClass.Get(), DefaultThumbnail);
+		ClassIconBrush = nullptr;
+		ThumbnailBrush = nullptr;
+		if (AssetTypeActions)
+		{
+			if (const FSlateBrush* AssetTypeThumbnail = AssetTypeActions->GetThumbnailBrush(AssetData, AssetClassName))
+			{
+				ThumbnailBrush = AssetTypeThumbnail;
+			}
 
-		ClassIconBrush = FSlateIconFinder::FindIconBrushForClass(ThumbnailClass.Get());
+			if (const FSlateBrush* AssetTypeIcon = AssetTypeActions->GetIconBrush(AssetData, AssetClassName))
+			{
+				ClassIconBrush = AssetTypeIcon;
+			}
+		}
 
+		if (!ThumbnailBrush)
+		{
+			// For non-class types, use the default based upon the actual asset class
+			// This has the side effect of not showing a class icon for assets that don't have a proper thumbnail image available
+			const FName DefaultThumbnail = (bIsClassType) ? NAME_None : FName(*FString::Printf(TEXT("ClassThumbnail.%s"), *AssetClassName.ToString()));
+			ThumbnailBrush = FClassIconFinder::FindThumbnailForClass(ThumbnailClass.Get(), DefaultThumbnail);
+		}
+		if (!ClassIconBrush)
+		{
+			ClassIconBrush = FSlateIconFinder::FindIconBrushForClass(ThumbnailClass.Get());
+		}
 	}
 
 	FSlateColor GetHintBackgroundColor() const
@@ -371,13 +391,17 @@ private:
 
 		UClass* Class = FindObject<UClass>(AssetData.AssetClassPath);
 		FAssetToolsModule& AssetToolsModule = FModuleManager::LoadModuleChecked<FAssetToolsModule>(TEXT("AssetTools"));
-		TWeakPtr<IAssetTypeActions> AssetTypeActions;
+		TSharedPtr<IAssetTypeActions> AssetTypeActions;
 		if ( Class != NULL )
 		{
-			AssetTypeActions = AssetToolsModule.Get().GetAssetTypeActionsForClass(Class);
+			TWeakPtr<IAssetTypeActions> TypeActions = AssetToolsModule.Get().GetAssetTypeActionsForClass(Class);
+			if (TypeActions.IsValid())
+			{
+				AssetTypeActions = TypeActions.Pin();
+			}
 		}
 
-		UpdateThumbnailClass();
+		UpdateThumbnailClass(AssetTypeActions.Get());
 
 		AssetColor = FLinearColor::White;
 		if( AssetTypeColorOverride.IsSet() )
@@ -386,7 +410,7 @@ private:
 		}
 		else if ( AssetTypeActions.IsValid() )
 		{
-			AssetColor = AssetTypeActions.Pin()->GetTypeColor();
+			AssetColor = AssetTypeActions->GetTypeColor();
 		}
 
 		//AssetBackgroundWidget->SetBorderBackgroundColor(AssetColor.CopyWithNewOpacity(0.3f));
