[MetaSounds] Fix for ensure when naming an input with "/".  Forbid use of analyzer '/' in MetaSound member names & don't bloat FName table with attempted renamed members

diff --git a/Engine/Plugins/Runtime/Metasound/Source/MetasoundEditor/Private/MetasoundEditorGraph.cpp b/Engine/Plugins/Runtime/Metasound/Source/MetasoundEditor/Private/MetasoundEditorGraph.cpp
--- a/Engine/Plugins/Runtime/Metasound/Source/MetasoundEditor/Private/MetasoundEditorGraph.cpp
+++ b/Engine/Plugins/Runtime/Metasound/Source/MetasoundEditor/Private/MetasoundEditorGraph.cpp
@@ -2,6 +2,7 @@
 #include "MetasoundEditorGraph.h"
 
 #include "Algo/Transform.h"
+#include "Analysis/MetasoundFrontendAnalyzerAddress.h"
 #include "AudioParameterControllerInterface.h"
 #include "Components/AudioComponent.h"
 #include "EdGraph/EdGraphNode.h"
@@ -16,12 +17,16 @@
 #include "MetasoundVariableNodes.h"
 #include "MetasoundVertex.h"
 #include "ScopedTransaction.h"
+#include "UObject/NameTypes.h"
 #include "UObject/UnrealType.h"
 
 #include UE_INLINE_GENERATED_CPP_BY_NAME(MetasoundEditorGraph)
 
 #define LOCTEXT_NAMESPACE "MetaSoundEditor"
 
+// Parameter names do not support '/' in addition to standard
+// name characters As analyzers use this to pack data into address.
+#define INVALID_PARAMETER_NAME_CHARACTERS INVALID_NAME_CHARACTERS METASOUND_ANALYZER_PATH_SEPARATOR
 
 namespace Metasound
 {
@@ -547,19 +552,25 @@ bool UMetasoundEditorGraphVertex::CanRename() const
 	return !RootGraphClass.PresetOptions.bIsPreset && !IsInterfaceMember();
 }
 
-bool UMetasoundEditorGraphVertex::CanRename(const FText& InNewName, FText& OutError) const
+bool UMetasoundEditorGraphVertex::CanRename(const FText& InNewText, FText& OutError) const
 {
 	using namespace Metasound::Frontend;
 
-	if (InNewName.IsEmptyOrWhitespace())
+	if (InNewText.IsEmptyOrWhitespace())
+	{
+		OutError = FText::Format(LOCTEXT("GraphVertexRenameInvalid_NameEmpty", "{0} cannot be empty string."), InNewText);
+		return false;
+	}
+
+	const FString NewNameString = InNewText.ToString();
+	if (!FName::IsValidXName(NewNameString, INVALID_PARAMETER_NAME_CHARACTERS, &OutError))
 	{
-		OutError = FText::Format(LOCTEXT("GraphVertexRenameInvalid_NameEmpty", "{0} cannot be empty string."), InNewName);
 		return false;
 	}
 
 	if (IsInterfaceMember())
 	{
-		OutError = FText::Format(LOCTEXT("GraphVertexRenameInvalid_GraphVertexRequired", "{0} is interface member and cannot be renamed."), InNewName);
+		OutError = FText::Format(LOCTEXT("GraphVertexRenameInvalid_GraphVertexRequired", "{0} is interface member and cannot be renamed."), InNewText);
 		return false;
 	}
 
@@ -568,22 +579,22 @@ bool UMetasoundEditorGraphVertex::CanRename(const FText& InNewName, FText& OutEr
 
 	if (RootGraphClass.PresetOptions.bIsPreset)
 	{
-		OutError = FText::Format(LOCTEXT("GraphVertexRenameInvalid_Preset", "{0} is a vertex in a preset graph and cannot be renamed."), InNewName);
+		OutError = FText::Format(LOCTEXT("GraphVertexRenameInvalid_Preset", "{0} is a vertex in a preset graph and cannot be renamed."), InNewText);
 		return false;
 	}
 
+	const FName NewName(*NewNameString);
 	bool bIsNameValid = true;
-	const FString NewName = InNewName.ToString();
 	FConstNodeHandle NodeHandle = GetConstNodeHandle();
 	FConstGraphHandle GraphHandle = NodeHandle->GetOwningGraph();
 	GraphHandle->IterateConstNodes([&](FConstNodeHandle NodeToCompare)
 	{
 		if (NodeID != NodeToCompare->GetID())
 		{
-			if (NewName == NodeToCompare->GetNodeName().ToString())
+			if (NewName == NodeToCompare->GetNodeName())
 			{
 				bIsNameValid = false;
-				OutError = FText::Format(LOCTEXT("GraphVertexRenameInvalid_NameTaken", "{0} is already in use"), InNewName);
+				OutError = FText::Format(LOCTEXT("GraphVertexRenameInvalid_NameTaken", "{0} is already in use"), InNewText);
 			}
 		}
 	}, GetClassType());
@@ -1070,20 +1081,21 @@ bool UMetasoundEditorGraphVariable::CanRename(const FText& InNewText, FText& Out
 		return false;
 	}
 	
-	const FName InNewName = FName(*InNewText.ToString());
-	if (!InNewName.IsValid())
+	const FString NewNameString = InNewText.ToString();
+	if (!FName::IsValidXName(NewNameString, INVALID_PARAMETER_NAME_CHARACTERS, &OutError))
 	{
-		OutError = FText::Format(LOCTEXT("GraphVariableRenameInvalid_InvalidName", "{0} is an invalid name."), InNewText);
 		return false;
 	}
 
+	const FName NewName(*NewNameString);
+
 	FConstVariableHandle VariableHandle = GetConstVariableHandle();
 	TArray<FConstVariableHandle> Variables = VariableHandle->GetOwningGraph()->GetVariables();
 	for (const FConstVariableHandle& OtherVariable : Variables)
 	{
 		if (VariableID != OtherVariable->GetID())
 		{
-			if (InNewName == OtherVariable->GetName())
+			if (NewName == OtherVariable->GetName())
 			{
 				OutError = FText::Format(LOCTEXT("GraphVariableRenameInvalid_NameTaken", "{0} is already in use"), InNewText);
 				return false;
diff --git a/Engine/Plugins/Runtime/Metasound/Source/MetasoundFrontend/Private/Analysis/MetasoundFrontendAnalyzerAddress.cpp b/Engine/Plugins/Runtime/Metasound/Source/MetasoundFrontend/Private/Analysis/MetasoundFrontendAnalyzerAddress.cpp
--- a/Engine/Plugins/Runtime/Metasound/Source/MetasoundFrontend/Private/Analysis/MetasoundFrontendAnalyzerAddress.cpp
+++ b/Engine/Plugins/Runtime/Metasound/Source/MetasoundFrontend/Private/Analysis/MetasoundFrontendAnalyzerAddress.cpp
@@ -11,12 +11,10 @@ namespace Metasound
 {
 	namespace Frontend
 	{
-		const FString FAnalyzerAddress::PathSeparator = TEXT("/");
-
 		FAnalyzerAddress::FAnalyzerAddress(const FString& InAddressString)
 		{
 			TArray<FString> Tokens;
-			if (ensureAlwaysMsgf(InAddressString.ParseIntoArray(Tokens, *PathSeparator) == 7, TEXT("Invalid Analyzer Address String Format")))
+			if (ensureAlwaysMsgf(InAddressString.ParseIntoArray(Tokens, METASOUND_ANALYZER_PATH_SEPARATOR) == 7, TEXT("Invalid Analyzer Address String Format")))
 			{
 				InstanceID = static_cast<uint64>(FCString::Atoi64(*Tokens[0]));
 				NodeID = FGuid(Tokens[1]);
@@ -83,7 +81,7 @@ namespace Metasound
 				*AnalyzerName.ToString(),
 				*AnalyzerInstanceID.ToString(),
 				*AnalyzerMemberName.ToString()
-			}, *PathSeparator);
+			}, METASOUND_ANALYZER_PATH_SEPARATOR);
 		}
 	} // namespace Frontend
 
diff --git a/Engine/Plugins/Runtime/Metasound/Source/MetasoundFrontend/Public/Analysis/MetasoundFrontendAnalyzerAddress.h b/Engine/Plugins/Runtime/Metasound/Source/MetasoundFrontend/Public/Analysis/MetasoundFrontendAnalyzerAddress.h
--- a/Engine/Plugins/Runtime/Metasound/Source/MetasoundFrontend/Public/Analysis/MetasoundFrontendAnalyzerAddress.h
+++ b/Engine/Plugins/Runtime/Metasound/Source/MetasoundFrontend/Public/Analysis/MetasoundFrontendAnalyzerAddress.h
@@ -10,6 +10,8 @@
 #include "UObject/NameTypes.h"
 
 
+#define METASOUND_ANALYZER_PATH_SEPARATOR TEXT("/")
+
 namespace Metasound
 {
 	namespace Frontend
@@ -19,8 +21,6 @@ namespace Metasound
 		class METASOUNDFRONTEND_API FAnalyzerAddress : public FTransmissionAddress
 		{
 		public:
-			static const FString PathSeparator;
-
 			FAnalyzerAddress(const FString& InAddressString);
 			FAnalyzerAddress() = default;
 
diff --git a/Engine/Plugins/Runtime/Metasound/Source/MetasoundFrontend/Public/Analysis/MetasoundFrontendGraphAnalyzer.h b/Engine/Plugins/Runtime/Metasound/Source/MetasoundFrontend/Public/Analysis/MetasoundFrontendGraphAnalyzer.h
--- a/Engine/Plugins/Runtime/Metasound/Source/MetasoundFrontend/Public/Analysis/MetasoundFrontendGraphAnalyzer.h
+++ b/Engine/Plugins/Runtime/Metasound/Source/MetasoundFrontend/Public/Analysis/MetasoundFrontendGraphAnalyzer.h
@@ -44,7 +44,7 @@ namespace Metasound
 
 			virtual FString ToString() const override
 			{
-				return FString::Printf(TEXT("%s%s%lld"), *GetAddressType().ToString(), * FAnalyzerAddress::PathSeparator, InstanceID);
+				return FString::Printf(TEXT("%s%s%lld"), *GetAddressType().ToString(), METASOUND_ANALYZER_PATH_SEPARATOR, InstanceID);
 			}
 
 			virtual uint32 GetHash() const override
