Change result of applying an Attack and Release trigger on the same block sample index.

Prior to this change, attack triggers are ALWAYS processed before release triggers.

This change ONLY affects functionality when an attack trigger and release arrive at the same sample AND only if the ADSR is already tiggered. In this case the release trigger will now be processed before the attack trigger... stopping the currently triggered ADSR and then restarting it.

diff --git a/Engine/Plugins/Runtime/Metasound/Source/MetasoundStandardNodes/Private/MetasoundADSREnvelopeNode.cpp b/Engine/Plugins/Runtime/Metasound/Source/MetasoundStandardNodes/Private/MetasoundADSREnvelopeNode.cpp
--- a/Engine/Plugins/Runtime/Metasound/Source/MetasoundStandardNodes/Private/MetasoundADSREnvelopeNode.cpp
+++ b/Engine/Plugins/Runtime/Metasound/Source/MetasoundStandardNodes/Private/MetasoundADSREnvelopeNode.cpp
@@ -75,6 +75,11 @@ namespace Metasound
 
 			// If this is set, we are in release mode
 			bool bIsInRelease = false;
+
+			bool IsTriggered() const 
+			{
+				return CurrentSampleIndex != INDEX_NONE; 
+			}
 		};
 
 
@@ -365,11 +370,32 @@ namespace Metasound
 				}
 			}
 
-			FTriggerInfo NextTrigger()
+			FTriggerInfo NextTrigger(bool EnvelopeIsTriggered)
 			{
 				FTriggerInfo Info;
 
-				if ((NextAttackFrame <= NextReleaseFrame) && (NextAttackFrame < NumFramesPerBlock))
+				// If attack and release triggers arrive on the same frame we will process the 
+				// release first if the envelope is currently triggered, or process the attack first if 
+				// the envelope is not alreay triggered. In fact, in this second case we could do NOTHING and 
+				// "swallow" them both. Since we know that the attack will be followed immediately by a release.
+				// But we don't that as there may be nodes down the line watching the ADSR's output triggers
+				// and they may want to know about the attack and release even if they do cancel out. 
+				if (NextAttackFrame == NextReleaseFrame && EnvelopeIsTriggered && NextReleaseFrame < NumFramesPerBlock)
+				{
+					Info.Type = ETriggerType::Release;
+					Info.FrameIndex = NextReleaseFrame;
+
+					if (ReleaseTriggerIndex < ReleaseTrigger.NumTriggeredInBlock())
+					{
+						NextReleaseFrame = ReleaseTrigger[ReleaseTriggerIndex];
+						ReleaseTriggerIndex++;
+					}
+					else
+					{
+						NextReleaseFrame = NumFramesPerBlock;
+					}
+				}
+				else if ((NextAttackFrame <= NextReleaseFrame) && (NextAttackFrame < NumFramesPerBlock))
 				{
 					Info.Type = ETriggerType::Attack;
 					Info.FrameIndex = NextAttackFrame;
@@ -624,7 +650,7 @@ namespace Metasound
 
 			FTriggerIterator TriggerIter(*TriggerAttackIn, *TriggerReleaseIn, NumFramesPerBlock);
 
-			FTriggerInfo NextTrigger = TriggerIter.NextTrigger();
+			FTriggerInfo NextTrigger = TriggerIter.NextTrigger(EnvState.IsTriggered());
 			if (NextTrigger.FrameIndex > 0)
 			{
 				// Process envelope before receiving any triggers on this block
@@ -634,7 +660,7 @@ namespace Metasound
 			while (NextTrigger.Type != ETriggerType::None)
 			{
 				FTriggerInfo CurrentTrigger = NextTrigger;
-				NextTrigger = TriggerIter.NextTrigger();
+				NextTrigger = TriggerIter.NextTrigger(EnvState.IsTriggered());
 
 				switch (CurrentTrigger.Type)
 				{
