metasounds switch-and-ramp node for smoothing out pops, and the ability for the AD envelope to hard reset (i.e., start from zero rather than the current value)

diff --git a/Engine/Plugins/Runtime/Metasound/Source/MetasoundStandardNodes/Private/MetasoundADEnvelopeNode.cpp b/Engine/Plugins/Runtime/Metasound/Source/MetasoundStandardNodes/Private/MetasoundADEnvelopeNode.cpp
--- a/Engine/Plugins/Runtime/Metasound/Source/MetasoundStandardNodes/Private/MetasoundADEnvelopeNode.cpp
+++ b/Engine/Plugins/Runtime/Metasound/Source/MetasoundStandardNodes/Private/MetasoundADEnvelopeNode.cpp
@@ -28,6 +28,7 @@ namespace Metasound
 		METASOUND_PARAM(InputAttackCurve, "Attack Curve", "The exponential curve factor of the attack. 1.0 = linear growth, < 1.0 logorithmic growth, > 1.0 exponential growth.");
 		METASOUND_PARAM(InputDecayCurve, "Decay Curve", "The exponential curve factor of the decay. 1.0 = linear decay, < 1.0 exponential decay, > 1.0 logarithmic decay.");
 		METASOUND_PARAM(InputLooping, "Looping", "Set to true to enable looping of the envelope. This will allow the envelope to be an LFO or wave generator.");
+		METASOUND_PARAM(InputHardReset, "Hard Reset", "Set to true to always reset the envelope level to 0 when triggering.");
 
 		METASOUND_PARAM(OutputOnTrigger, "On Trigger", "Triggers when the envelope is triggered.");
 		METASOUND_PARAM(OutputOnDone, "On Done", "Triggers when the envelope finishes or loops back if looping is enabled.");
@@ -58,6 +59,7 @@ namespace Metasound
 			float CurrentEnvelopeValue = 0.0f;
 
 			bool bLooping = false;
+			bool bHardReset = false;
 		};
 
 		template<typename ValueType>
@@ -143,8 +145,15 @@ namespace Metasound
 						float AttackFraction = (float)InState.CurrentSampleIndex++ / InState.AttackSampleCount;
 						float EnvValue = FMath::Pow(AttackFraction, InState.AttackCurveFactor);
 						float TargetEnvelopeValue = InState.StartingEnvelopeValue + (1.0f - InState.StartingEnvelopeValue) * EnvValue;
-						InState.EnvEase.SetValue(TargetEnvelopeValue);
-						InState.CurrentEnvelopeValue = InState.EnvEase.GetNextValue();
+						if (InState.bHardReset)
+						{
+							InState.CurrentEnvelopeValue = TargetEnvelopeValue;
+						}
+						else
+						{
+							InState.EnvEase.SetValue(TargetEnvelopeValue);
+							InState.CurrentEnvelopeValue = InState.EnvEase.GetNextValue();
+						}
 						OutEnvPtr[i] = InState.CurrentEnvelopeValue;
 					}
 					else 
@@ -157,8 +166,15 @@ namespace Metasound
 							int32 SampleCountInDecayState = InState.CurrentSampleIndex++ - InState.AttackSampleCount;
 							float DecayFracton = (float)SampleCountInDecayState / InState.DecaySampleCount;
 							float TargetEnvelopeValue = 1.0f - FMath::Pow(DecayFracton, InState.DecayCurveFactor);
-							InState.EnvEase.SetValue(TargetEnvelopeValue);
-							InState.CurrentEnvelopeValue = InState.EnvEase.GetNextValue();
+							if (InState.bHardReset)
+							{
+								InState.CurrentEnvelopeValue = TargetEnvelopeValue;
+							}
+							else
+							{
+								InState.EnvEase.SetValue(TargetEnvelopeValue);
+								InState.CurrentEnvelopeValue = InState.EnvEase.GetNextValue();
+							}
 							OutEnvPtr[i] = InState.CurrentEnvelopeValue;
 						}
 						// We are looping so reset the sample index
@@ -171,11 +187,11 @@ namespace Metasound
 						}
 						else
 						{
-							InState.CurrentEnvelopeValue = InState.EnvEase.GetNextValue();
+							InState.CurrentEnvelopeValue = InState.bHardReset ? 0 : InState.EnvEase.GetNextValue();
 							OutEnvPtr[i] = InState.CurrentEnvelopeValue;
 
 							// Envelope is done
-							if (InState.EnvEase.IsDone())
+							if (InState.EnvEase.IsDone() || InState.bHardReset)
 							{
 								// Zero out the rest of the envelope
 								int32 NumSamplesLeft = EndFrame - i - 1;
@@ -211,7 +227,8 @@ namespace Metasound
 					TInputDataVertex<FTime>(METASOUND_GET_PARAM_NAME_AND_METADATA(InputDecayTime), 1.0f),
 					TInputDataVertex<float>(METASOUND_GET_PARAM_NAME_AND_METADATA(InputAttackCurve), 1.0f),
 					TInputDataVertex<float>(METASOUND_GET_PARAM_NAME_AND_METADATA(InputDecayCurve), 1.0f),
-					TInputDataVertex<bool>(METASOUND_GET_PARAM_NAME_AND_METADATA(InputLooping), false)
+					TInputDataVertex<bool>(METASOUND_GET_PARAM_NAME_AND_METADATA(InputLooping), false),
+					TInputDataVertex<bool>(METASOUND_GET_PARAM_NAME_AND_METADATA(InputHardReset), false)
 				),
 				FOutputVertexInterface(
 					TOutputDataVertex<FTrigger>(METASOUND_GET_PARAM_NAME_AND_METADATA(OutputOnTrigger)),
@@ -267,8 +284,9 @@ namespace Metasound
 			FFloatReadRef AttackCurveFactor = InParams.InputDataReferences.GetDataReadReferenceOrConstructWithVertexDefault<float>(InputInterface, METASOUND_GET_PARAM_NAME(InputAttackCurve), InParams.OperatorSettings);
 			FFloatReadRef DecayCurveFactor = InParams.InputDataReferences.GetDataReadReferenceOrConstructWithVertexDefault<float>(InputInterface, METASOUND_GET_PARAM_NAME(InputDecayCurve), InParams.OperatorSettings);
 			FBoolReadRef bLooping = InParams.InputDataReferences.GetDataReadReferenceOrConstructWithVertexDefault<bool>(InputInterface, METASOUND_GET_PARAM_NAME(InputLooping), InParams.OperatorSettings);
+			FBoolReadRef bHardReset = InParams.InputDataReferences.GetDataReadReferenceOrConstructWithVertexDefault<bool>(InputInterface, METASOUND_GET_PARAM_NAME(InputHardReset), InParams.OperatorSettings);
 
-			return MakeUnique<TADEnvelopeNodeOperator<ValueType>>(InParams.OperatorSettings, TriggerIn, AttackTime, DecayTime, AttackCurveFactor, DecayCurveFactor, bLooping);
+			return MakeUnique<TADEnvelopeNodeOperator<ValueType>>(InParams.OperatorSettings, TriggerIn, AttackTime, DecayTime, AttackCurveFactor, DecayCurveFactor, bLooping, bHardReset);
 		}
 
 		TADEnvelopeNodeOperator(const FOperatorSettings& InSettings,
@@ -277,13 +295,15 @@ namespace Metasound
 			const FTimeReadRef& InDecayTime,
 			const FFloatReadRef& InAttackCurveFactor,
 			const FFloatReadRef& InDecayeCurveFactor,
-			const FBoolReadRef& bInLooping)
+			const FBoolReadRef& bInLooping,
+			const FBoolReadRef& bInHardReset)
 			: TriggerAttackIn(InTriggerIn)
 			, AttackTime(InAttackTime)
 			, DecayTime(InDecayTime)
 			, AttackCurveFactor(InAttackCurveFactor)
 			, DecayCurveFactor(InDecayeCurveFactor)
 			, bLooping(bInLooping)
+			, bHardReset(bInHardReset)
 			, OnAttackTrigger(TDataWriteReferenceFactory<FTrigger>::CreateAny(InSettings))
 			, OnDone(TDataWriteReferenceFactory<FTrigger>::CreateAny(InSettings))
 			, OutputEnvelope(TDataWriteReferenceFactory<ValueType>::CreateAny(InSettings))
@@ -315,6 +335,7 @@ namespace Metasound
 			Inputs.AddDataReadReference(METASOUND_GET_PARAM_NAME(InputAttackCurve), AttackCurveFactor);
 			Inputs.AddDataReadReference(METASOUND_GET_PARAM_NAME(InputDecayCurve), DecayCurveFactor);
 			Inputs.AddDataReadReference(METASOUND_GET_PARAM_NAME(InputLooping), bLooping);
+			Inputs.AddDataReadReference(METASOUND_GET_PARAM_NAME(InputHardReset), bHardReset);
 
 			return Inputs;
 		}
@@ -340,6 +361,7 @@ namespace Metasound
 			EnvState.AttackCurveFactor = FMath::Max(KINDA_SMALL_NUMBER, *AttackCurveFactor);
 			EnvState.DecayCurveFactor = FMath::Max(KINDA_SMALL_NUMBER, *DecayCurveFactor);
 			EnvState.bLooping = *bLooping;
+			EnvState.bHardReset = *bHardReset;
 		}
 
 
@@ -373,7 +395,7 @@ namespace Metasound
 
 					// Set the sample index to the top of the envelope
 					EnvState.CurrentSampleIndex = 0;
-					EnvState.StartingEnvelopeValue = EnvState.CurrentEnvelopeValue;
+					EnvState.StartingEnvelopeValue = EnvState.bHardReset ? 0 : EnvState.CurrentEnvelopeValue;
 
 					// Generate the output (this will no-op if we're block rate)
 					TArray<int32> FinishedFrames;
@@ -397,6 +419,7 @@ namespace Metasound
 		FFloatReadRef AttackCurveFactor;
 		FFloatReadRef DecayCurveFactor;
 		FBoolReadRef bLooping;
+		FBoolReadRef bHardReset;
 
 		FTriggerWriteRef OnAttackTrigger;
 		FTriggerWriteRef OnDone;
diff --git a/Engine/Plugins/Runtime/Metasound/Source/MetasoundStandardNodes/Private/MetasoundSwitchAndRampNode.cpp b/Engine/Plugins/Runtime/Metasound/Source/MetasoundStandardNodes/Private/MetasoundSwitchAndRampNode.cpp
new file mode 100644
--- /dev/null
+++ b/Engine/Plugins/Runtime/Metasound/Source/MetasoundStandardNodes/Private/MetasoundSwitchAndRampNode.cpp
@@ -0,0 +1,204 @@
+// Copyright Epic Games, Inc. All Rights Reserved.
+// Background on switch-and-ramp: http://msp.ucsd.edu/techniques/v0.11/book-html/node63.html
+
+#include "Internationalization/Text.h"
+#include "MetasoundFacade.h"
+#include "MetasoundExecutableOperator.h"
+#include "MetasoundNodeRegistrationMacro.h"
+#include "MetasoundPrimitives.h"
+#include "MetasoundStandardNodesNames.h"
+#include "MetasoundStandardNodesCategories.h"
+#include "DSP/BufferVectorOperations.h"
+#include "MetasoundParamHelper.h"
+
+#define LOCTEXT_NAMESPACE "MetasoundStandardNodes_SwitchAndRampNode"
+
+namespace Metasound
+{
+	namespace SwitchAndRampVertexNames
+	{
+		METASOUND_PARAM(InputTrigger, "Trigger", "Trigger at the moment of the discontinuity that should be smoothed.");
+		METASOUND_PARAM(InputAudio, "In Audio", "The audio input.");
+		METASOUND_PARAM(InputSmoothTime, "Smooth Time", "The amount of time to smooth the discontinuity over.");
+
+		METASOUND_PARAM(OutputAudio, "Out Audio", "Output audio with the discontinuity smoothed.");
+	}
+
+	class FSwitchAndRampOperator : public TExecutableOperator<FSwitchAndRampOperator>
+	{
+	public:
+
+		static const FNodeClassMetadata& GetNodeInfo()
+		{
+			auto InitNodeInfo = []() -> FNodeClassMetadata
+			{
+				const FName OperatorName = TEXT("SwitchAndRamp");
+				const FText NodeDisplayName = METASOUND_LOCTEXT("Metasound_SwitchAndRampNodeDisplayName", "Switch And Ramp");
+				const FText NodeDescription = METASOUND_LOCTEXT("Metasound_SwitchAndRampNodeDescription", "Uses the switch-and-ramp technique to smooth out pops in audio, when triggered at the moment that a pop would occur.");
+
+				FNodeClassMetadata Info;
+				Info.ClassName = { StandardNodes::Namespace, OperatorName, TEXT("") };
+				Info.MajorVersion = 1;
+				Info.MinorVersion = 0;
+				Info.DisplayName = NodeDisplayName;
+				Info.Description = NodeDescription;
+				Info.Author = PluginAuthor;
+				Info.PromptIfMissing = PluginNodeMissingPrompt;
+				Info.DefaultInterface = GetVertexInterface();
+
+				return Info;
+			};
+
+			static const FNodeClassMetadata Info = InitNodeInfo();
+			return Info;
+		}
+
+		static const FVertexInterface& GetVertexInterface()
+		{
+			using namespace SwitchAndRampVertexNames;
+
+			static const FVertexInterface Interface(
+				FInputVertexInterface(
+					TInputDataVertex<FTrigger>(METASOUND_GET_PARAM_NAME_AND_METADATA(InputTrigger)),
+					TInputDataVertex<FAudioBuffer>(METASOUND_GET_PARAM_NAME_AND_METADATA(InputAudio)),
+					TInputDataVertex<FTime>(METASOUND_GET_PARAM_NAME_AND_METADATA(InputSmoothTime), 0.01f)
+				),
+				FOutputVertexInterface(
+					TOutputDataVertex<FAudioBuffer>(METASOUND_GET_PARAM_NAME_AND_METADATA(OutputAudio))
+				)
+			);
+
+			return Interface;
+		}
+
+		static TUniquePtr<IOperator> CreateOperator(const FCreateOperatorParams& InParams, FBuildErrorArray& OutErrors)
+		{
+			using namespace SwitchAndRampVertexNames;
+
+			const FDataReferenceCollection& InputCollection = InParams.InputDataReferences;
+			const FInputVertexInterface& InputInterface = GetVertexInterface().GetInputInterface();
+
+			FTriggerReadRef TriggerIn = InputCollection.GetDataReadReferenceOrConstruct<FTrigger>(METASOUND_GET_PARAM_NAME(InputTrigger), InParams.OperatorSettings);
+			FAudioBufferReadRef AudioInput = InputCollection.GetDataReadReferenceOrConstruct<FAudioBuffer>(METASOUND_GET_PARAM_NAME(InputAudio), InParams.OperatorSettings);
+			FTimeReadRef SmoothTime = InputCollection.GetDataReadReferenceOrConstructWithVertexDefault<FTime>(InputInterface, METASOUND_GET_PARAM_NAME(InputSmoothTime), InParams.OperatorSettings);
+
+			return MakeUnique<FSwitchAndRampOperator>(InParams.OperatorSettings, TriggerIn, AudioInput, SmoothTime);
+		}
+
+		FSwitchAndRampOperator(const FOperatorSettings& InSettings,
+			const FTriggerReadRef& InTriggerIn,
+			const FAudioBufferReadRef& InAudioInput,
+			const FTimeReadRef& InputSmoothTime)
+
+			: TriggerIn(InTriggerIn)
+			, AudioIn(InAudioInput)
+			, SmoothTime(InputSmoothTime)
+			, AudioOut(FAudioBufferWriteRef::CreateNew(InSettings))
+			, SampleRate(InSettings.GetSampleRate())
+		{
+		}
+
+		FDataReferenceCollection GetInputs() const
+		{
+			using namespace SwitchAndRampVertexNames;
+
+			FDataReferenceCollection InputDataReferences;
+			InputDataReferences.AddDataReadReference(METASOUND_GET_PARAM_NAME(InputTrigger), TriggerIn);
+			InputDataReferences.AddDataReadReference(METASOUND_GET_PARAM_NAME(InputAudio), AudioIn);
+			InputDataReferences.AddDataReadReference(METASOUND_GET_PARAM_NAME(InputSmoothTime), SmoothTime);
+
+			return InputDataReferences;
+		}
+
+		virtual FDataReferenceCollection GetOutputs() const
+		{
+			using namespace SwitchAndRampVertexNames;
+
+			FDataReferenceCollection OutputDataReferences;
+			OutputDataReferences.AddDataReadReference(METASOUND_GET_PARAM_NAME(OutputAudio), AudioOut);
+			return OutputDataReferences;
+		}
+
+		void Execute()
+		{
+			// Lambda is used for both pre and post trigger frames
+			auto RunRampLambda = [this](int32 StartFrame, int32 EndFrame)
+			{
+				const float* InputAudio = AudioIn->GetData();
+				float* OutputAudio = AudioOut->GetData();
+
+				//apply the ramp offset, if it's still running
+				int32 RampSamplesLeft = NumRampSamples - RampSampleIndex;
+				int32 LastRampFrame = FMath::Min(StartFrame + RampSamplesLeft, EndFrame);
+				for (int32 i = StartFrame; i < LastRampFrame; ++i)
+				{
+					float Adjustment = 0;
+					if (RampSampleIndex < NumRampSamples)
+					{
+						float RampFraction = (float)RampSampleIndex / NumRampSamples;
+						Adjustment = (1 - RampFraction) * DiscontinuityAmount;
+						++RampSampleIndex;
+					}
+					OutputAudio[i] = InputAudio[i] + Adjustment;
+				}
+
+				//copy the rest of the audio as passthrough
+				int32 StartCopyFrame = FMath::Max(StartFrame, LastRampFrame);
+				if (StartCopyFrame < EndFrame)
+					FMemory::Memcpy(&OutputAudio[StartCopyFrame], &InputAudio[StartCopyFrame], sizeof(float) * (EndFrame - StartCopyFrame));
+
+				MostRecentOutputValue = OutputAudio[EndFrame - 1];
+			};
+
+			TriggerIn->ExecuteBlock(
+				RunRampLambda,
+				// OnTrigger
+				[this, &RunRampLambda](int32 StartFrame, int32 EndFrame)
+				{
+					float RampFromValue = MostRecentOutputValue;
+					float RampToValue = AudioIn->GetData()[StartFrame];
+					DiscontinuityAmount = RampFromValue - RampToValue;
+
+					float SmoothTimeSeconds = SmoothTime->GetSeconds();
+					NumRampSamples = FMath::Max(1, SampleRate * SmoothTimeSeconds);
+					RampSampleIndex = 0;
+
+					RunRampLambda(StartFrame, EndFrame);
+				}
+			);
+		}
+
+	private:
+		FTriggerReadRef TriggerIn;
+		
+		// The input audio buffer
+		FAudioBufferReadRef AudioIn;
+		
+		FTimeReadRef SmoothTime;
+
+		// Audio output
+		FAudioBufferWriteRef AudioOut;
+
+		float SampleRate = 0.0f;
+		float MostRecentOutputValue = 0.0f;
+		float DiscontinuityAmount = 0.0f;
+		int32 NumRampSamples = 0;
+		int32 RampSampleIndex = 0;
+	};
+
+	class METASOUNDSTANDARDNODES_API FSwitchAndRampNode : public FNodeFacade
+	{
+	public:
+		/**
+		 * Constructor used by the Metasound Frontend.
+		 */
+		FSwitchAndRampNode(const FNodeInitData& InitData)
+			: FNodeFacade(InitData.InstanceName, InitData.InstanceID, TFacadeOperatorClass<FSwitchAndRampOperator>())
+		{
+		}
+	};
+
+	METASOUND_REGISTER_NODE(FSwitchAndRampNode)
+}
+
+#undef LOCTEXT_NAMESPACE
