add "hard reset" capability to metasounds ADSR envelope node, to match recent "hard reset" functionality in AD node from 22886510
additionally, simplify implementation of AD node's hard reset functionality to retain ease.

diff --git a/Engine/Plugins/Runtime/Metasound/Source/MetasoundStandardNodes/Private/MetasoundADEnvelopeNode.cpp b/Engine/Plugins/Runtime/Metasound/Source/MetasoundStandardNodes/Private/MetasoundADEnvelopeNode.cpp
--- a/Engine/Plugins/Runtime/Metasound/Source/MetasoundStandardNodes/Private/MetasoundADEnvelopeNode.cpp
+++ b/Engine/Plugins/Runtime/Metasound/Source/MetasoundStandardNodes/Private/MetasoundADEnvelopeNode.cpp
@@ -145,15 +145,8 @@ namespace Metasound
 						float AttackFraction = (float)InState.CurrentSampleIndex++ / InState.AttackSampleCount;
 						float EnvValue = FMath::Pow(AttackFraction, InState.AttackCurveFactor);
 						float TargetEnvelopeValue = InState.StartingEnvelopeValue + (1.0f - InState.StartingEnvelopeValue) * EnvValue;
-						if (InState.bHardReset)
-						{
-							InState.CurrentEnvelopeValue = TargetEnvelopeValue;
-						}
-						else
-						{
-							InState.EnvEase.SetValue(TargetEnvelopeValue);
-							InState.CurrentEnvelopeValue = InState.EnvEase.GetNextValue();
-						}
+						InState.EnvEase.SetValue(TargetEnvelopeValue);
+						InState.CurrentEnvelopeValue = InState.EnvEase.GetNextValue();
 						OutEnvPtr[i] = InState.CurrentEnvelopeValue;
 					}
 					else 
@@ -166,15 +159,8 @@ namespace Metasound
 							int32 SampleCountInDecayState = InState.CurrentSampleIndex++ - InState.AttackSampleCount;
 							float DecayFracton = (float)SampleCountInDecayState / InState.DecaySampleCount;
 							float TargetEnvelopeValue = 1.0f - FMath::Pow(DecayFracton, InState.DecayCurveFactor);
-							if (InState.bHardReset)
-							{
-								InState.CurrentEnvelopeValue = TargetEnvelopeValue;
-							}
-							else
-							{
-								InState.EnvEase.SetValue(TargetEnvelopeValue);
-								InState.CurrentEnvelopeValue = InState.EnvEase.GetNextValue();
-							}
+							InState.EnvEase.SetValue(TargetEnvelopeValue);
+							InState.CurrentEnvelopeValue = InState.EnvEase.GetNextValue();
 							OutEnvPtr[i] = InState.CurrentEnvelopeValue;
 						}
 						// We are looping so reset the sample index
@@ -187,11 +173,11 @@ namespace Metasound
 						}
 						else
 						{
-							InState.CurrentEnvelopeValue = InState.bHardReset ? 0 : InState.EnvEase.GetNextValue();
+							InState.CurrentEnvelopeValue = InState.EnvEase.GetNextValue();
 							OutEnvPtr[i] = InState.CurrentEnvelopeValue;
 
 							// Envelope is done
-							if (InState.EnvEase.IsDone() || InState.bHardReset)
+							if (InState.EnvEase.IsDone())
 							{
 								// Zero out the rest of the envelope
 								int32 NumSamplesLeft = EndFrame - i - 1;
@@ -396,6 +382,7 @@ namespace Metasound
 					// Set the sample index to the top of the envelope
 					EnvState.CurrentSampleIndex = 0;
 					EnvState.StartingEnvelopeValue = EnvState.bHardReset ? 0 : EnvState.CurrentEnvelopeValue;
+					EnvState.EnvEase.SetValue(EnvState.StartingEnvelopeValue, true /* bInit */);
 
 					// Generate the output (this will no-op if we're block rate)
 					TArray<int32> FinishedFrames;
diff --git a/Engine/Plugins/Runtime/Metasound/Source/MetasoundStandardNodes/Private/MetasoundADSREnvelopeNode.cpp b/Engine/Plugins/Runtime/Metasound/Source/MetasoundStandardNodes/Private/MetasoundADSREnvelopeNode.cpp
--- a/Engine/Plugins/Runtime/Metasound/Source/MetasoundStandardNodes/Private/MetasoundADSREnvelopeNode.cpp
+++ b/Engine/Plugins/Runtime/Metasound/Source/MetasoundStandardNodes/Private/MetasoundADSREnvelopeNode.cpp
@@ -34,6 +34,8 @@ namespace Metasound
 		METASOUND_PARAM(InputDecayCurve, "Decay Curve", "The exponential curve factor of the decay. 1.0 = linear decay, < 1.0 exponential decay, > 1.0 logorithmic decay.");
 		METASOUND_PARAM(InputReleaseCurve, "Release Curve", "The exponential curve factor of the release. 1.0 = linear release, < 1.0 exponential release, > 1.0 logorithmic release.");
 
+		METASOUND_PARAM(InputHardReset, "Hard Reset", "Set to true to always reset the envelope level to 0 when triggering.");
+
 		METASOUND_PARAM(OutputOnAttackTrigger, "On Attack Triggered", "Triggers when the envelope attack is triggered.");
 		METASOUND_PARAM(OutputOnDecayTrigger, "On Decay Triggered", "Triggers when the envelope decay begins and attack is finished.");
 		METASOUND_PARAM(OutputOnSustainTrigger, "On Sustain Triggered", "Triggers when the envelope sustain begins and attack is finished.");
@@ -72,6 +74,7 @@ namespace Metasound
 			float StartingEnvelopeValue = 0.0f;
 			float CurrentEnvelopeValue = 0.0f;
 			float EnvelopeValueAtReleaseStart = 0.0f;
+			bool bHardReset = false;
 
 			// If this is set, we are in release mode
 			bool bIsInRelease = false;
@@ -461,7 +464,8 @@ namespace Metasound
 					TInputDataVertex<FTime>(METASOUND_GET_PARAM_NAME_AND_METADATA(InputReleaseTime), 1.0f),
 					TInputDataVertex<float>(METASOUND_GET_PARAM_NAME_AND_METADATA(InputAttackCurve), 1.0f),
 					TInputDataVertex<float>(METASOUND_GET_PARAM_NAME_AND_METADATA(InputDecayCurve), 1.0f),
-					TInputDataVertex<float>(METASOUND_GET_PARAM_NAME_AND_METADATA(InputReleaseCurve), 1.0f)
+					TInputDataVertex<float>(METASOUND_GET_PARAM_NAME_AND_METADATA(InputReleaseCurve), 1.0f),
+					TInputDataVertex<bool>(METASOUND_GET_PARAM_NAME_AND_METADATA(InputHardReset), false)
 				),
 				FOutputVertexInterface(
 					TOutputDataVertex<FTrigger>(METASOUND_GET_PARAM_NAME_AND_METADATA(OutputOnAttackTrigger)),
@@ -520,6 +524,7 @@ namespace Metasound
 			FFloatReadRef AttackCurveFactor;
 			FFloatReadRef DecayCurveFactor;
 			FFloatReadRef ReleaseCurveFactor;
+			FBoolReadRef bInHardReset;
 		};
 
 		static TUniquePtr<IOperator> CreateOperator(const FCreateOperatorParams& InParams, TArray<TUniquePtr<IOperatorBuildError>>& OutErrors)
@@ -539,7 +544,8 @@ namespace Metasound
 				InParams.InputDataReferences.GetDataReadReferenceOrConstructWithVertexDefault<FTime>(InputInterface, METASOUND_GET_PARAM_NAME(InputReleaseTime), InParams.OperatorSettings),
 				InParams.InputDataReferences.GetDataReadReferenceOrConstructWithVertexDefault<float>(InputInterface, METASOUND_GET_PARAM_NAME(InputAttackCurve), InParams.OperatorSettings),
 				InParams.InputDataReferences.GetDataReadReferenceOrConstructWithVertexDefault<float>(InputInterface, METASOUND_GET_PARAM_NAME(InputDecayCurve), InParams.OperatorSettings),
-				InParams.InputDataReferences.GetDataReadReferenceOrConstructWithVertexDefault<float>(InputInterface, METASOUND_GET_PARAM_NAME(InputReleaseCurve), InParams.OperatorSettings)
+				InParams.InputDataReferences.GetDataReadReferenceOrConstructWithVertexDefault<float>(InputInterface, METASOUND_GET_PARAM_NAME(InputReleaseCurve), InParams.OperatorSettings),
+				InParams.InputDataReferences.GetDataReadReferenceOrConstructWithVertexDefault<bool>(InputInterface, METASOUND_GET_PARAM_NAME(InputHardReset), InParams.OperatorSettings)
 			};
 
 			return MakeUnique<TADSREnvelopeNodeOperator<EnvelopeClass, ValueType>>(Args);
@@ -555,6 +561,7 @@ namespace Metasound
 			, AttackCurveFactor(InArgs.AttackCurveFactor)
 			, DecayCurveFactor(InArgs.DecayCurveFactor)
 			, ReleaseCurveFactor(InArgs.ReleaseCurveFactor)
+			, bHardReset(InArgs.bInHardReset)
 			, OnDecayTrigger(TDataWriteReferenceFactory<FTrigger>::CreateAny(InArgs.OperatorSettings))
 			, OnSustainTrigger(TDataWriteReferenceFactory<FTrigger>::CreateAny(InArgs.OperatorSettings))
 			, OnDone(TDataWriteReferenceFactory<FTrigger>::CreateAny(InArgs.OperatorSettings))
@@ -582,6 +589,7 @@ namespace Metasound
 			Inputs.AddDataReadReference(METASOUND_GET_PARAM_NAME(InputAttackCurve), AttackCurveFactor);
 			Inputs.AddDataReadReference(METASOUND_GET_PARAM_NAME(InputDecayCurve), DecayCurveFactor);
 			Inputs.AddDataReadReference(METASOUND_GET_PARAM_NAME(InputReleaseCurve), ReleaseCurveFactor);
+			Inputs.AddDataReadReference(METASOUND_GET_PARAM_NAME(InputHardReset), bHardReset);
 
 			return Inputs;
 		}
@@ -613,6 +621,7 @@ namespace Metasound
 			EnvState.AttackCurveFactor = FMath::Max(KINDA_SMALL_NUMBER, *AttackCurveFactor);
 			EnvState.DecayCurveFactor = FMath::Max(KINDA_SMALL_NUMBER, *DecayCurveFactor);
 			EnvState.ReleaseCurveFactor = FMath::Max(KINDA_SMALL_NUMBER, *ReleaseCurveFactor);
+			EnvState.bHardReset = *bHardReset;
 		}
 
 		void ProcessEnvelopeOutput(int32 InStartFrame, int32 InEndFrame)
@@ -669,7 +678,8 @@ namespace Metasound
 					UpdateParams();
 					// Set the sample index to the top of the envelope
 					EnvState.CurrentSampleIndex = 0;
-					EnvState.StartingEnvelopeValue = EnvState.CurrentEnvelopeValue;
+					EnvState.StartingEnvelopeValue = EnvState.bHardReset ? 0 : EnvState.CurrentEnvelopeValue;
+					EnvState.EnvEase.SetValue(EnvState.StartingEnvelopeValue, true /* bInit */);
 					EnvState.bIsInRelease = false;
 
 						break;
@@ -700,6 +710,7 @@ namespace Metasound
 		FFloatReadRef AttackCurveFactor;
 		FFloatReadRef DecayCurveFactor;
 		FFloatReadRef ReleaseCurveFactor;
+		FBoolReadRef bHardReset;
 
 		FTriggerWriteRef OnDecayTrigger;
 		FTriggerWriteRef OnSustainTrigger;
