Metasound Audio Bus Writer Node

Adds a metasound node that sends its audio input to an audio bus.
Modifies the audio mixer interface to support the node.

FMixerSourceManager::StartAudioBus no longer sets bIsPlaying to false on sources that use an audio bus when the audio bus starts.
FMixerAudioBus::MixBuffer mixes its patch inputs into its MixedSourceData buffer, so that the patch input data is available when CopyCurrentBuffer is used to retrieve the mix.

diff --git a/Engine/Plugins/Runtime/Metasound/Source/MetasoundEngine/Private/MetasoundAudioBusReaderNode.cpp b/Engine/Plugins/Runtime/Metasound/Source/MetasoundEngine/Private/MetasoundAudioBusReaderNode.cpp
--- a/Engine/Plugins/Runtime/Metasound/Source/MetasoundEngine/Private/MetasoundAudioBusReaderNode.cpp
+++ b/Engine/Plugins/Runtime/Metasound/Source/MetasoundEngine/Private/MetasoundAudioBusReaderNode.cpp
@@ -228,12 +228,12 @@ namespace Metasound
 	};
 
 #define REGISTER_AUDIO_BUS_READER_NODE(ChannelCount) \
-	using FTriggerSequenceNode_##ChannelCount = TAudioBusReaderNode<ChannelCount>; \
-	METASOUND_REGISTER_NODE(FTriggerSequenceNode_##ChannelCount) \
+	using FAudioBusReaderNode_##ChannelCount = TAudioBusReaderNode<ChannelCount>; \
+	METASOUND_REGISTER_NODE(FAudioBusReaderNode_##ChannelCount) \
 
 	
-	REGISTER_AUDIO_BUS_READER_NODE(1)
-	REGISTER_AUDIO_BUS_READER_NODE(2)
+	REGISTER_AUDIO_BUS_READER_NODE(1);
+	REGISTER_AUDIO_BUS_READER_NODE(2);
 }
 
 #undef LOCTEXT_NAMESPACE
diff --git a/Engine/Plugins/Runtime/Metasound/Source/MetasoundEngine/Private/MetasoundAudioBusWriterNode.cpp b/Engine/Plugins/Runtime/Metasound/Source/MetasoundEngine/Private/MetasoundAudioBusWriterNode.cpp
new file mode 100644
--- /dev/null
+++ b/Engine/Plugins/Runtime/Metasound/Source/MetasoundEngine/Private/MetasoundAudioBusWriterNode.cpp
@@ -0,0 +1,233 @@
+// Copyright Epic Games, Inc. All Rights Reserved.
+
+#include "AudioMixerDevice.h"
+#include "Internationalization/Text.h"
+#include "MediaPacket.h"
+#include "MetasoundAudioBuffer.h"
+#include "MetasoundAudioBus.h"
+#include "MetasoundEngineNodesNames.h"
+#include "MetasoundExecutableOperator.h"
+#include "MetasoundFacade.h"
+#include "MetasoundNodeRegistrationMacro.h"
+#include "MetasoundParamHelper.h"
+#include "MetasoundStandardNodesCategories.h"
+
+#define LOCTEXT_NAMESPACE "MetasoundAudioBusWriterNode"
+
+namespace Metasound
+{
+	namespace AudioBusWriterNode
+	{
+		METASOUND_PARAM(InParamAudioBusOutput, "Audio Bus", "Audio Bus Asset.");
+		METASOUND_PARAM(InParamAudio, "In {0}", "Audio input for channel {0}.");
+	}
+
+	template<uint32 NumChannels>
+	class TAudioBusWriterOperator : public TExecutableOperator<TAudioBusWriterOperator<NumChannels>>
+	{
+	public:
+		static const FNodeClassMetadata& GetNodeInfo()
+		{
+			auto InitNodeInfo = []() -> FNodeClassMetadata
+			{
+				FName OperatorName = *FString::Printf(TEXT("Audio Bus Writer (%d)"), NumChannels);
+				FText NodeDisplayName = METASOUND_LOCTEXT_FORMAT("AudioBusWriterDisplayNamePattern", "Audio Bus Writer ({0})", NumChannels);
+
+				FNodeClassMetadata Info;
+				Info.ClassName = { EngineNodes::Namespace, OperatorName, TEXT("") };
+				Info.MajorVersion = 1;
+				Info.MinorVersion = 0;
+				Info.DisplayName = NodeDisplayName;
+				Info.Description = METASOUND_LOCTEXT("AudioBusWriter_Description", "Sends audio data to the audio bus asset.");
+				Info.Author = PluginAuthor;
+				Info.PromptIfMissing = PluginNodeMissingPrompt;
+				Info.DefaultInterface = GetVertexInterface();
+				Info.CategoryHierarchy.Emplace(NodeCategories::Io);
+				return Info;
+			};
+
+			static const FNodeClassMetadata Info = InitNodeInfo();
+
+			return Info;
+		}
+
+		static const FVertexInterface& GetVertexInterface()
+		{
+			using namespace AudioBusWriterNode;
+
+			auto CreateVertexInterface = []() -> FVertexInterface
+			{
+				FInputVertexInterface InputInterface;
+				InputInterface.Add(TInputDataVertex<FAudioBusAsset>(METASOUND_GET_PARAM_NAME_AND_METADATA(InParamAudioBusOutput)));
+				for (uint32 i = 0; i < NumChannels; ++i)
+				{
+					InputInterface.Add(TInputDataVertex<FAudioBuffer>(METASOUND_GET_PARAM_NAME_WITH_INDEX_AND_METADATA(InParamAudio, i)));
+				}
+
+				FOutputVertexInterface OutputInterface;
+
+				return FVertexInterface(InputInterface, OutputInterface);
+			};
+
+			static const FVertexInterface Interface = CreateVertexInterface();
+			return Interface;
+		}
+
+		static TUniquePtr<IOperator> CreateOperator(const FCreateOperatorParams& InParams, FBuildErrorArray& OutErrors)
+		{
+			using namespace Frontend;
+
+			using namespace AudioBusWriterNode;
+			const FDataReferenceCollection& InputCollection = InParams.InputDataReferences;
+
+			Audio::FDeviceId AudioDeviceId = InParams.Environment.GetValue<Audio::FDeviceId>(SourceInterface::Environment::DeviceID);
+			int32 AudioMixerOutputFrames = InParams.Environment.GetValue<int32>(SourceInterface::Environment::AudioMixerNumOutputFrames);
+
+			FAudioBusAssetReadRef AudioBusIn = InputCollection.GetDataReadReferenceOrConstruct<FAudioBusAsset>(METASOUND_GET_PARAM_NAME(InParamAudioBusOutput));
+
+			TArray<FAudioBufferReadRef> AudioInputs;
+			for (int32 ChannelIndex = 0; ChannelIndex < NumChannels; ++ChannelIndex)
+			{
+				AudioInputs.Add(InputCollection.GetDataReadReferenceOrConstruct<FAudioBuffer>(METASOUND_GET_PARAM_NAME_WITH_INDEX(InParamAudio, ChannelIndex), InParams.OperatorSettings));
+			}
+
+			return MakeUnique<TAudioBusWriterOperator<NumChannels>>(InParams, AudioMixerOutputFrames, AudioDeviceId, MoveTemp(AudioBusIn), MoveTemp(AudioInputs));
+		}
+
+		TAudioBusWriterOperator(const FCreateOperatorParams& InParams, int32 InAudioMixerOutputFrames, Audio::FDeviceId InAudioDeviceId, FAudioBusAssetReadRef InAudioBusAsset, TArray<FAudioBufferReadRef> InAudioInputs) :
+			AudioBusAsset(MoveTemp(InAudioBusAsset)),
+			AudioInputs(MoveTemp(InAudioInputs)),
+			AudioMixerOutputFrames(InAudioMixerOutputFrames),
+			AudioDeviceId(InAudioDeviceId),
+			SampleRate(InParams.OperatorSettings.GetSampleRate())
+		{
+			const FAudioBusProxyPtr& AudioBusProxy = AudioBusAsset->GetAudioBusProxy();
+			if (AudioBusProxy.IsValid())
+			{
+				if (FAudioDeviceManager* ADM = FAudioDeviceManager::Get())
+				{
+					if (FAudioDevice* AudioDevice = ADM->GetAudioDeviceRaw(AudioDeviceId))
+					{
+						if (AudioDevice->IsAudioMixerEnabled())
+						{
+							AudioBusChannels = AudioBusProxy->NumChannels;
+							BlockSizeFrames = AudioInputs[0]->Num();
+							InterleavedBuffer.AddUninitialized(BlockSizeFrames * AudioBusChannels);
+
+							Audio::FMixerDevice& MixerDevice(static_cast<Audio::FMixerDevice&>(*AudioDevice));
+							MixerDevice.StartAudioBus(AudioBusProxy->AudioBusId, NumChannels, false);
+
+							// Create a bus patch input with enough room for the number of samples we expect and some buffering
+							AudioBusPatchInput = Audio::FPatchInput(MakeShared<Audio::FPatchOutput, ESPMode::ThreadSafe>(BlockSizeFrames * FMath::DivideAndRoundUp(BlockSizeFrames, AudioMixerOutputFrames) * AudioBusChannels * 4, 1.f));
+							MixerDevice.AddPatchInputForAudioBus(AudioBusPatchInput, AudioBusProxy->AudioBusId);
+						}
+					}
+				}
+			}
+		}
+
+		virtual ~TAudioBusWriterOperator()
+		{
+		}
+
+		virtual FDataReferenceCollection GetInputs() const override
+		{
+			using namespace AudioBusWriterNode;
+
+			FDataReferenceCollection InputDataReferences;
+
+			InputDataReferences.AddDataReadReference(METASOUND_GET_PARAM_NAME(InParamAudioBusOutput), AudioBusAsset);
+
+			for (int32 ChannelIndex = 0; ChannelIndex < NumChannels; ++ChannelIndex)
+			{
+				InputDataReferences.AddDataReadReference(METASOUND_GET_PARAM_NAME_WITH_INDEX(InParamAudio, ChannelIndex), FAudioBufferReadRef(AudioInputs[ChannelIndex]));
+			}
+
+			return InputDataReferences;
+		}
+
+		virtual FDataReferenceCollection GetOutputs() const override
+		{
+			using namespace AudioBusWriterNode;
+
+			FDataReferenceCollection OutputDataReferences;
+			return OutputDataReferences;
+		}
+
+		void Execute()
+		{
+			if (!AudioBusPatchInput.IsOutputStillActive())
+			{
+				return;
+			}
+
+			int32 NumSamplesToPush = InterleavedBuffer.Num();
+			uint32 MinChannels = NumChannels;
+			if (MinChannels < AudioBusChannels)
+			{
+				// There are fewer inputs than outputs, so zero the interleaved buffer
+				FMemory::Memzero(InterleavedBuffer.GetData(), NumSamplesToPush * sizeof(float));
+			}
+			else
+			{
+				MinChannels = AudioBusChannels;
+			}
+
+			// Retrieve input and interleaved buffer pointers
+			const float* AudioInputBufferPtrs[NumChannels];
+			for (uint32 ChannelIndex = 0; ChannelIndex < MinChannels; ++ChannelIndex)
+			{
+				AudioInputBufferPtrs[ChannelIndex] = AudioInputs[ChannelIndex]->GetData();
+			}
+			float* InterleavedBufferPtr = InterleavedBuffer.GetData();
+
+			// Interleave the inputs
+			// Writing the channels of the interleaved buffer sequentially should improve
+			// cache utilization compared to writing each input's frames sequentially.
+			// There is more likely to be a cache line for each buffer than for the
+			// entirety of the interleaved buffer.
+			for (int32 FrameIndex = 0; FrameIndex < BlockSizeFrames; ++FrameIndex)
+			{
+				for (uint32 ChannelIndex = 0; ChannelIndex < MinChannels; ++ChannelIndex)
+				{
+					InterleavedBufferPtr[ChannelIndex] = *AudioInputBufferPtrs[ChannelIndex]++;
+				}
+				InterleavedBufferPtr += AudioBusChannels;
+			}
+
+			// Pushes the interleaved data to the audio bus
+			AudioBusPatchInput.PushAudio(InterleavedBuffer.GetData(), NumSamplesToPush);
+		}
+
+	private:
+		FAudioBusAssetReadRef AudioBusAsset;
+		TArray<FAudioBufferReadRef> AudioInputs;
+
+		TArray<float> InterleavedBuffer;
+		int32 AudioMixerOutputFrames = INDEX_NONE;
+		Audio::FDeviceId AudioDeviceId = INDEX_NONE;
+		float SampleRate = 0.0f;
+		Audio::FPatchInput AudioBusPatchInput;
+		uint32 AudioBusChannels = INDEX_NONE;
+		int32 BlockSizeFrames = 0;
+	};
+
+	template<uint32 NumChannels>
+	class TAudioBusWriterNode : public FNodeFacade
+	{
+	public:
+		TAudioBusWriterNode(const FNodeInitData& InitData)
+			: FNodeFacade(InitData.InstanceName, InitData.InstanceID, TFacadeOperatorClass<TAudioBusWriterOperator<NumChannels>>())
+		{
+		}
+	};
+
+#define REGISTER_AUDIO_BUS_WRITER_NODE(ChannelCount) \
+	using FAudioBusWriterNode_##ChannelCount = TAudioBusWriterNode<ChannelCount>; \
+	METASOUND_REGISTER_NODE(FAudioBusWriterNode_##ChannelCount) \
+
+	REGISTER_AUDIO_BUS_WRITER_NODE(1);
+	REGISTER_AUDIO_BUS_WRITER_NODE(2);
+}
+
+#undef LOCTEXT_NAMESPACE
diff --git a/Engine/Source/Runtime/AudioMixer/Private/AudioMixerBus.cpp b/Engine/Source/Runtime/AudioMixer/Private/AudioMixerBus.cpp
--- a/Engine/Source/Runtime/AudioMixer/Private/AudioMixerBus.cpp
+++ b/Engine/Source/Runtime/AudioMixer/Private/AudioMixerBus.cpp
@@ -17,9 +17,6 @@ namespace Audio
 		, bIsAutomatic(bInIsAutomatic)
 	{
 		SetNumOutputChannels(NumChannels);
-
-		// Make a patch bus input to push mixed audio sent to audio bus from source manager
-		AudioBusInput = PatchMixerSplitter.AddNewInput(2*InNumChannels*NumFrames, 1.0f);
 	}
 
 	void FMixerAudioBus::SetNumOutputChannels(int32 InNumOutputChannels)
@@ -90,14 +87,26 @@ namespace Audio
 
 	void FMixerAudioBus::MixBuffer()
 	{
-		// Reset and zero the mixed source data buffer for this bus
+		// Mix the patch mixer's inputs into the source data
+		const int32 MaxNumSamplesToPop = PatchMixer.MaxNumberOfSamplesThatCanBePopped();
 		const int32 NumSamples = NumFrames * NumChannels;
+		const int32 NumOutputFrames = SourceManager->GetNumOutputFrames();
 
-		MixedSourceData[CurrentBufferIndex].Reset();
-		MixedSourceData[CurrentBufferIndex].AddZeroed(NumSamples);
+		FAlignedFloatBuffer& MixBuffer = MixedSourceData[CurrentBufferIndex];
+		float* BusDataBufferPtr = MixBuffer.GetData();
 
-		float* BusDataBufferPtr = MixedSourceData[CurrentBufferIndex].GetData();
-		const int32 NumOutputFrames = SourceManager->GetNumOutputFrames();
+		if (MaxNumSamplesToPop < NumSamples)
+		{
+			// Zero buffer when there are no inputs connected,
+			// or one of the inputs has not pushed any audio yet.
+			FMemory::Memzero(BusDataBufferPtr, NumSamples * sizeof(float));
+		}
+		else
+		{
+			// Mix down inputs
+			int32 PopResult = PatchMixer.PopAudio(BusDataBufferPtr, NumSamples, false);
+			check(PopResult == NumSamples);
+		}
 
 		for (int32 BusSendType = 0; BusSendType < (int32)EBusSendType::Count; ++BusSendType)
 		{
@@ -146,11 +155,39 @@ namespace Audio
 			}
 		}
 
-		// Push the mixed data to the patch splitter
-		AudioBusInput.PushAudio(BusDataBufferPtr, NumOutputFrames * NumChannels);
+		// Send the mix to the patch splitter's outputs
+		int32 NumSamplesToPush = PatchSplitter.MaxNumberOfSamplesThatCanBePushed();
+		if (NumSamplesToPush <= 0)
+		{
+			// Early exit when there are no outputs.
+			return;
+		}
+
+		int SampleCacheSize = SampleCache.Num();
+		if (0 <= SampleCacheSize)
+		{
+			float* SampleCacheBufferPtr = SampleCache.GetData();
+			PatchSplitter.PushAudio(SampleCacheBufferPtr, FMath::Min(SampleCacheSize, NumSamplesToPush));
+
+			if (NumSamplesToPush < SampleCacheSize)
+			{
+				const int32 NumSamplesToCache = SampleCacheSize - NumSamplesToPush;
+				FMemory::Memmove(SampleCacheBufferPtr, &SampleCache[NumSamplesToPush], NumSamplesToCache);
+				SampleCache.SetNum(NumSamplesToCache, false);
+				SampleCache.Append(MixBuffer);
+				return;
+			}
+
+			NumSamplesToPush -= SampleCacheSize;
+			SampleCache.Reset();
+		}
+
+		PatchSplitter.PushAudio(BusDataBufferPtr, FMath::Min(NumSamples, NumSamplesToPush));
 
-		// Process the audio in the patch mixer splitter
-		PatchMixerSplitter.ProcessAudio();
+		if (NumSamplesToPush < NumSamples)
+		{
+			SampleCache.Append(&MixBuffer[NumSamplesToPush], NumSamples - NumSamplesToPush);
+		}
 	}
 
 	void FMixerAudioBus::CopyCurrentBuffer(Audio::FAlignedFloatBuffer& InChannelMap, int32 InNumOutputChannels, FAlignedFloatBuffer& OutBuffer, int32 NumOutputFrames) const
@@ -180,20 +217,16 @@ namespace Audio
 
 	void FMixerAudioBus::AddNewPatchOutput(const FPatchOutputStrongPtr& InPatchOutputStrongPtr)
 	{
-		PatchMixerSplitter.AddNewOutput(InPatchOutputStrongPtr);
+		PatchSplitter.AddNewPatch(InPatchOutputStrongPtr);
 	}
 
 	void FMixerAudioBus::AddNewPatchInput(FPatchInput& InPatchInput)
 	{
-		return PatchMixerSplitter.AddNewInput(InPatchInput);
+		return PatchMixer.AddNewInput(InPatchInput);
 	}
 
 	void FMixerAudioBus::RemovePatchInput(const FPatchInput& PatchInput)
 	{
-		return PatchMixerSplitter.RemovePatch(PatchInput);
+		return PatchMixer.RemovePatch(PatchInput);
 	}
-
-
-
 }
-
diff --git a/Engine/Source/Runtime/AudioMixer/Private/AudioMixerDevice.cpp b/Engine/Source/Runtime/AudioMixer/Private/AudioMixerDevice.cpp
--- a/Engine/Source/Runtime/AudioMixer/Private/AudioMixerDevice.cpp
+++ b/Engine/Source/Runtime/AudioMixer/Private/AudioMixerDevice.cpp
@@ -2719,7 +2719,82 @@ namespace Audio
 		return StrongOutputPtr;
 	}
 
-	Audio::FPatchOutputStrongPtr FMixerDevice::AddPatchForSubmix(uint32 InObjectId, float InPatchGain)
+	void FMixerDevice::AddPatchInputForAudioBus(const FPatchInput& InPatchInput, uint32 InAudioBusId, float InPatchGain)
+	{
+		// This function is supporting adding audio bus patches from multiple threads (AT, ART, GT, and tasks) and is currently
+		// depending on a number of places where data lives, which accounts for the complexity here.
+		// This code needs a clean up to refactor everything into a true MPSC model, along with an MPSC refactor of the source manager
+		// and our command queues. Once we do that we can remove the code which branches based on the thread the request is coming from. 
+
+		if (IsInGameThread())
+		{
+			if (ActiveAudioBuses_GameThread.Find(InAudioBusId))
+			{
+				FAudioThread::RunCommandOnAudioThread([this, PatchInput = InPatchInput, InAudioBusId]() mutable
+				{
+					SourceManager->AddPatchInputForAudioBus_AudioThread(InAudioBusId, PatchInput);
+				});
+				return;
+			}
+			UE_LOG(LogAudioMixer, Warning, TEXT("Unable to add a patch input for audio bus because audio bus id '%d' is not active."), InAudioBusId);
+			return;
+		}
+		else if (IsInAudioThread())
+		{
+			SourceManager->AddPatchInputForAudioBus_AudioThread(InAudioBusId, InPatchInput);
+			return;
+		}
+		else if (IsAudioRenderingThread())
+		{
+			check(SourceManager);
+			FPatchInput PatchInput(InPatchInput);
+			SourceManager->AddPatchInputForAudioBus(InAudioBusId, PatchInput);
+			return;
+		}
+		else
+		{
+			GameThreadMPSCCommand([this, PatchInput = InPatchInput, InAudioBusId]
+			{
+				if (ActiveAudioBuses_GameThread.Find(InAudioBusId))
+				{
+					FAudioThread::RunCommandOnAudioThread([this, InAudioBusId, PatchInput = PatchInput]() mutable
+					{
+						SourceManager->AddPatchInputForAudioBus_AudioThread(InAudioBusId, PatchInput);
+					});
+				}
+				else
+				{
+					UE_LOG(LogAudioMixer, Warning, TEXT("Unable to add a patch output for audio bus because audio bus id '%d' is not active."), InAudioBusId);
+				}
+			});
+			return;
+		}
+	}
+
+	void FMixerDevice::AddPatchInputForAudioBus_GameThread(const FPatchInput& InPatchInput, uint32 InAudioBusId, float InPatchGain)
+	{
+		if (IsInGameThread())
+		{
+			FActiveBusData* BusData = ActiveAudioBuses_GameThread.Find(InAudioBusId);
+			if (BusData)
+			{
+				FAudioThread::RunCommandOnAudioThread([this, InAudioBusId, PatchInput = InPatchInput]() mutable
+				{
+					SourceManager->AddPatchInputForAudioBus_AudioThread(InAudioBusId, PatchInput);
+				});
+			}
+			else
+			{
+				UE_LOG(LogAudioMixer, Warning, TEXT("Unable to add a patch input for audio bus because audio bus id '%d' is not active."), InAudioBusId);
+			}
+		}
+		else
+		{
+			UE_LOG(LogAudioMixer, Warning, TEXT("AddPatchForAudioBus can only be called from the game thread."));
+		}
+	}
+
+	FPatchOutputStrongPtr FMixerDevice::AddPatchForSubmix(uint32 InObjectId, float InPatchGain)
 	{
 		if (!ensure(IsAudioRenderingThread()))
 		{
diff --git a/Engine/Source/Runtime/AudioMixer/Private/AudioMixerSourceManager.cpp b/Engine/Source/Runtime/AudioMixer/Private/AudioMixerSourceManager.cpp
--- a/Engine/Source/Runtime/AudioMixer/Private/AudioMixerSourceManager.cpp
+++ b/Engine/Source/Runtime/AudioMixer/Private/AudioMixerSourceManager.cpp
@@ -1128,18 +1128,6 @@ namespace Audio
 
 				AudioBuses.Add(InAudioBusId, NewBusData);
 			}
-
-			//  Now add any existing playing sources to this audio bus as sends if they exist
-			for (FSourceInfo& SourceInfo : SourceInfos)
-			{
-				if (SourceInfo.AudioBusId == InAudioBusId)
-				{
-					SourceInfo.bIsPlaying = false;
-					SourceInfo.bIsPaused = false;
-					SourceInfo.bIsActive = false;
-					SourceInfo.bIsStopping = false;
-				}
-			}
 		}, AUDIO_MIXER_THREAD_COMMAND_STRING("StartAudioBus"));
 	}
 
@@ -1223,6 +1211,35 @@ namespace Audio
 		}, AUDIO_MIXER_THREAD_COMMAND_STRING("AddPatchOutputForAudioBus_AudioThread()"));
 	}
 
+	void FMixerSourceManager::AddPatchInputForAudioBus(uint32 InAudioBusId, FPatchInput& InPatchInput)
+	{
+		AUDIO_MIXER_CHECK_AUDIO_PLAT_THREAD(MixerDevice);
+		if (MixerDevice->IsAudioRenderingThread())
+		{
+			TSharedPtr<FMixerAudioBus> AudioBusPtr = AudioBuses.FindRef(InAudioBusId);
+			if (AudioBusPtr.IsValid())
+			{
+				AudioBusPtr->AddNewPatchInput(InPatchInput);
+			}
+		}
+		else
+		{
+			// Queue up the command via MPSC command queue
+			AudioMixerThreadMPSCCommand([this, InAudioBusId, InPatchInput]() mutable
+			{
+				AddPatchInputForAudioBus(InAudioBusId, InPatchInput);
+			});
+		}
+	}
+
+	void FMixerSourceManager::AddPatchInputForAudioBus_AudioThread(uint32 InAudioBusId, const FPatchInput& InPatchInput)
+	{
+		AudioMixerThreadCommand([this, InAudioBusId, PatchInput = InPatchInput]() mutable
+		{
+			AddPatchInputForAudioBus(InAudioBusId, PatchInput);
+		}, AUDIO_MIXER_THREAD_COMMAND_STRING("AddPatchInputForAudioBus_AudioThread()"));
+	}
+
 	void FMixerSourceManager::Play(const int32 SourceId)
 	{
 		AUDIO_MIXER_CHECK(SourceId < NumTotalSources);
diff --git a/Engine/Source/Runtime/AudioMixer/Private/AudioMixerSourceManager.h b/Engine/Source/Runtime/AudioMixer/Private/AudioMixerSourceManager.h
--- a/Engine/Source/Runtime/AudioMixer/Private/AudioMixerSourceManager.h
+++ b/Engine/Source/Runtime/AudioMixer/Private/AudioMixerSourceManager.h
@@ -210,6 +210,9 @@ namespace Audio
 		// Adds a patch input for an audio bus
 		void AddPatchInputForAudioBus(uint32 InAudioBusId, FPatchInput& InPatchInput);
 
+		// Adds a patch input for an audio bus from the Audio Thread
+		void AddPatchInputForAudioBus_AudioThread(uint32 InAudioBusId, const FPatchInput& InPatchInput);
+
 		void Play(const int32 SourceId);
 		void Stop(const int32 SourceId);
 		void CancelQuantizedSound(const int32 SourceId);
diff --git a/Engine/Source/Runtime/AudioMixer/Public/AudioMixerBus.h b/Engine/Source/Runtime/AudioMixer/Public/AudioMixerBus.h
--- a/Engine/Source/Runtime/AudioMixer/Public/AudioMixerBus.h
+++ b/Engine/Source/Runtime/AudioMixer/Public/AudioMixerBus.h
@@ -106,8 +106,9 @@ namespace Audio
 		FMixerSourceManager* SourceManager;
 
 		// Multiple places can produce and consume from audio buses
-		Audio::FPatchInput AudioBusInput;
-		Audio::FPatchMixerSplitter PatchMixerSplitter;
+		Audio::FPatchMixer PatchMixer;
+		Audio::FPatchSplitter PatchSplitter;
+		FAlignedFloatBuffer SampleCache;
 
 		// Was created manually, not via source buses.
 		bool bIsAutomatic;
diff --git a/Engine/Source/Runtime/AudioMixer/Public/AudioMixerDevice.h b/Engine/Source/Runtime/AudioMixer/Public/AudioMixerDevice.h
--- a/Engine/Source/Runtime/AudioMixer/Public/AudioMixerDevice.h
+++ b/Engine/Source/Runtime/AudioMixer/Public/AudioMixerDevice.h
@@ -302,6 +302,9 @@ namespace Audio
 		virtual FPatchOutputStrongPtr AddPatchForAudioBus(uint32 InAudioBusId, float InPatchGain = 1.0f) override;
 		virtual FPatchOutputStrongPtr AddPatchForAudioBus_GameThread(uint32 InAudioBusId, float InPatchGain = 1.0f) override;
 
+		virtual void AddPatchInputForAudioBus(const FPatchInput& InPatchInput, uint32 InAudioBusId, float InPatchGain = 1.0f) override;
+		virtual void AddPatchInputForAudioBus_GameThread(const FPatchInput& InPatchInput, uint32 InAudioBusId, float InPatchGain = 1.0f) override;
+
 		// Clock Manager for quantized event handling on Audio Render Thread
 		FQuartzClockManager QuantizedEventClockManager;
 
diff --git a/Engine/Source/Runtime/Engine/Public/AudioDevice.h b/Engine/Source/Runtime/Engine/Public/AudioDevice.h
--- a/Engine/Source/Runtime/Engine/Public/AudioDevice.h
+++ b/Engine/Source/Runtime/Engine/Public/AudioDevice.h
@@ -970,6 +970,14 @@ public:
 		return nullptr;
 	}
 	
+	virtual void AddPatchInputForAudioBus(const Audio::FPatchInput& InPatchInput, uint32 InAudioBusId, float InPatchGain = 1.0f)
+	{
+	}
+
+	virtual void AddPatchInputForAudioBus_GameThread(const Audio::FPatchInput& InPatchInput, uint32 InAudioBusId, float InPatchGain = 1.0f)
+	{
+	}
+
 	virtual void InitSoundEffectPresets() {}
 
 	/**
diff --git a/Engine/Source/Runtime/SignalProcessing/Private/MultithreadedPatching.cpp b/Engine/Source/Runtime/SignalProcessing/Private/MultithreadedPatching.cpp
--- a/Engine/Source/Runtime/SignalProcessing/Private/MultithreadedPatching.cpp
+++ b/Engine/Source/Runtime/SignalProcessing/Private/MultithreadedPatching.cpp
@@ -561,8 +561,6 @@ namespace Audio
 		int32 PopResult = Mixer.PopAudio(IntermediateBuffer.GetData(), NumSamplesToForward, false);
 		check(PopResult == NumSamplesToForward);
 		
-		OnProcessAudio(TArrayView<const float>(IntermediateBuffer));
-
 		// Push audio to outputs
 		Splitter.PushAudio(IntermediateBuffer.GetData(), NumSamplesToForward);
 	}
diff --git a/Engine/Source/Runtime/SignalProcessing/Public/DSP/MultithreadedPatching.h b/Engine/Source/Runtime/SignalProcessing/Public/DSP/MultithreadedPatching.h
--- a/Engine/Source/Runtime/SignalProcessing/Public/DSP/MultithreadedPatching.h
+++ b/Engine/Source/Runtime/SignalProcessing/Public/DSP/MultithreadedPatching.h
@@ -223,8 +223,6 @@ namespace Audio
 	class SIGNALPROCESSING_API FPatchMixerSplitter
 	{
 	public:
-		virtual ~FPatchMixerSplitter() = default;
-
 		/**
 		 * Adds a new output. Calling this is thread safe, but individual instances of FPatchOutput are only safe to be used from one thread.
 		 * the returned FPatchOutputPtr can be safely destroyed at any point.
@@ -246,10 +244,6 @@ namespace Audio
 		/** Mixes audio from all inputs and pushes it to all outputs. Should be called regularly. */
 		void ProcessAudio();
 
-	protected:
-		/** This class can be subclassed with OnProcessAudio overridden. */
-		virtual void OnProcessAudio(TArrayView<const float> InAudio) { }
-
 	private:
 		FPatchMixer Mixer;
 		FPatchSplitter Splitter;
