Metasounds: Fix missing dependencies during EditorDomain save of Metasound objects by shifting WaitUntilAsyncLoadReferencedAssetsComplete to earlier in PreSave, and by extending it to recursively wait for all transitive dependencies.
This fixes EditorDomain saves because they occur immediately after load of the package so waiting for async dependencies to load is required.
#rnx

diff --git a/Engine/Plugins/Runtime/Metasound/Source/MetasoundEngine/Private/MetasoundAssetSubsystem.cpp b/Engine/Plugins/Runtime/Metasound/Source/MetasoundEngine/Private/MetasoundAssetSubsystem.cpp
--- a/Engine/Plugins/Runtime/Metasound/Source/MetasoundEngine/Private/MetasoundAssetSubsystem.cpp
+++ b/Engine/Plugins/Runtime/Metasound/Source/MetasoundEngine/Private/MetasoundAssetSubsystem.cpp
@@ -419,9 +419,18 @@ void UMetaSoundAssetSubsystem::RequestAsyncLoadReferencedAssets(FMetasoundAssetB
 
 void UMetaSoundAssetSubsystem::WaitUntilAsyncLoadReferencedAssetsComplete(FMetasoundAssetBase& InAssetBase)
 {
-	UObject* OwningAsset = InAssetBase.GetOwningAsset();
-	if (OwningAsset)
-	{
+	TSet<FMetasoundAssetBase*> TransitiveReferences;
+	TArray<FMetasoundAssetBase*> TransitiveReferencesQueue;
+	TransitiveReferences.Add(&InAssetBase);
+	TransitiveReferencesQueue.Add(&InAssetBase);
+	while (!TransitiveReferencesQueue.IsEmpty())
+	{
+		FMetasoundAssetBase* Reference = TransitiveReferencesQueue.Pop();
+		UObject* OwningAsset = Reference->GetOwningAsset();
+		if (!OwningAsset)
+		{
+			continue;
+		}
 		while (FMetaSoundAsyncAssetDependencies* LoadingDependency = FindLoadingDependencies(OwningAsset))
 		{
 			// Grab shared ptr to handle as LoadingDependencies may be deleted and have it's shared pointer removed. 
@@ -447,6 +456,16 @@ void UMetaSoundAssetSubsystem::WaitUntilAsyncLoadReferencedAssetsComplete(FMetas
 				StreamableHandle->CancelHandle();
 			}
 		}
+
+		for (FMetasoundAssetBase* NextReference : Reference->GetReferencedAssets())
+		{
+			bool bAlreadyInSet;
+			TransitiveReferences.Add(NextReference, &bAlreadyInSet);
+			if (!bAlreadyInSet)
+			{
+				TransitiveReferencesQueue.Add(NextReference);
+			}
+		}
 	}
 }
 
