Fix Audio Bus Writer Node patch input buffer size calculation

Calculation no longer assumes BlockSizeFrames >= AudioMixerOutputFrames.
Overall size reduced to three quarters because my testing suggests that's adequate.

diff --git a/Engine/Plugins/Runtime/Metasound/Source/MetasoundEngine/Private/MetasoundAudioBusWriterNode.cpp b/Engine/Plugins/Runtime/Metasound/Source/MetasoundEngine/Private/MetasoundAudioBusWriterNode.cpp
--- a/Engine/Plugins/Runtime/Metasound/Source/MetasoundEngine/Private/MetasoundAudioBusWriterNode.cpp
+++ b/Engine/Plugins/Runtime/Metasound/Source/MetasoundEngine/Private/MetasoundAudioBusWriterNode.cpp
@@ -118,7 +118,8 @@ namespace Metasound
 							MixerDevice.StartAudioBus(AudioBusProxy->AudioBusId, NumChannels, false);
 
 							// Create a bus patch input with enough room for the number of samples we expect and some buffering
-							AudioBusPatchInput = Audio::FPatchInput(MakeShared<Audio::FPatchOutput, ESPMode::ThreadSafe>(BlockSizeFrames * FMath::DivideAndRoundUp(BlockSizeFrames, AudioMixerOutputFrames) * AudioBusChannels * 4, 1.f));
+							int32 MaxSizeFrames = FMath::Max(BlockSizeFrames, AudioMixerOutputFrames), MinSizeFrames = FMath::Min(BlockSizeFrames, AudioMixerOutputFrames), BufferScale = 3;
+							AudioBusPatchInput = Audio::FPatchInput(MakeShared<Audio::FPatchOutput, ESPMode::ThreadSafe>(BufferScale * MinSizeFrames * FMath::DivideAndRoundUp(MaxSizeFrames, MinSizeFrames) * AudioBusChannels, 1.f));
 							MixerDevice.AddPatchInputForAudioBus(AudioBusPatchInput, AudioBusProxy->AudioBusId);
 						}
 					}
