Expose IMetasoundEditor interface
#9674

diff --git a/Engine/Plugins/Runtime/Metasound/Source/MetasoundEditor/Public/IMetasoundEditor.h b/Engine/Plugins/Runtime/Metasound/Source/MetasoundEditor/Public/IMetasoundEditor.h
--- a/Engine/Plugins/Runtime/Metasound/Source/MetasoundEditor/Public/IMetasoundEditor.h
+++ b/Engine/Plugins/Runtime/Metasound/Source/MetasoundEditor/Public/IMetasoundEditor.h
@@ -10,6 +10,7 @@
 
 class IMetasoundEditor : public FAssetEditorToolkit
 {
+public:
 	// Returns the editor's assigned MetaSound UObject.
 	virtual UObject* GetMetasoundObject() const = 0;
 
