Searchable tag update for Trigger Compare Node, so that it shows up for keywords "equal" and "not," in addition to some other variations

diff --git a/Engine/Plugins/Runtime/Metasound/Source/MetasoundStandardNodes/Private/MetasoundTriggerCompareNode.cpp b/Engine/Plugins/Runtime/Metasound/Source/MetasoundStandardNodes/Private/MetasoundTriggerCompareNode.cpp
--- a/Engine/Plugins/Runtime/Metasound/Source/MetasoundStandardNodes/Private/MetasoundTriggerCompareNode.cpp
+++ b/Engine/Plugins/Runtime/Metasound/Source/MetasoundStandardNodes/Private/MetasoundTriggerCompareNode.cpp
@@ -21,7 +21,8 @@ namespace Metasound
 				PluginNodeMissingPrompt,
 				InDefaultInterface,
 				{ NodeCategories::Trigger },
-				{ METASOUND_LOCTEXT("EqualsKeyword", "="), METASOUND_LOCTEXT("NotEqualsKeyword", "!="), METASOUND_LOCTEXT("LessThanKeyword", "<"), METASOUND_LOCTEXT("GreaterThanKeyword", ">"), METASOUND_LOCTEXT("LessEqualsKeyword", "<="), METASOUND_LOCTEXT("GreaterEqualsKeyword", ">=")},
+				{ METASOUND_LOCTEXT("EqualsKeyword", "="), METASOUND_LOCTEXT("NotEqualsKeyword", "!="), METASOUND_LOCTEXT("LessThanKeyword", "<"), METASOUND_LOCTEXT("GreaterThanKeyword", ">"), METASOUND_LOCTEXT("LessEqualsKeyword", "<="), METASOUND_LOCTEXT("GreaterEqualsKeyword", ">="),
+				 METASOUND_LOCTEXT("EqualsTextKeyword", "equal"), METASOUND_LOCTEXT("InequalsKeyword", "inequal"), METASOUND_LOCTEXT("NotKeyword", "not"), METASOUND_LOCTEXT("LessThanTextKeyword", "less"), METASOUND_LOCTEXT("GreatThanTextKeyword", "greater")},
 				FNodeDisplayStyle()
 			};
 
