Correct channel descriptions in Metasound Wave Player node.

diff --git a/Engine/Plugins/Runtime/Metasound/Source/MetasoundEngine/Private/MetasoundWavePlayerNode.cpp b/Engine/Plugins/Runtime/Metasound/Source/MetasoundEngine/Private/MetasoundWavePlayerNode.cpp
--- a/Engine/Plugins/Runtime/Metasound/Source/MetasoundEngine/Private/MetasoundWavePlayerNode.cpp
+++ b/Engine/Plugins/Runtime/Metasound/Source/MetasoundEngine/Private/MetasoundWavePlayerNode.cpp
@@ -48,13 +48,13 @@ namespace Metasound
 		METASOUND_PARAM(OutputAudioLeft, "Out Left", "The left channel audio output.")
 		METASOUND_PARAM(OutputAudioRight, "Out Right", "The right channel audio output.")
 		METASOUND_PARAM(OutputAudioFrontRight, "Out Front Right", "The front right channel audio output.")
-		METASOUND_PARAM(OutputAudioFrontLeft, "Out Front Left", "The front right channel audio output.")
+		METASOUND_PARAM(OutputAudioFrontLeft, "Out Front Left", "The front left channel audio output.")
 		METASOUND_PARAM(OutputAudioFrontCenter, "Out Front Center", "The front center channel audio output.")
 		METASOUND_PARAM(OutputAudioLowFrequency, "Out Low Frequency", "The low frequency channel audio output.")
 		METASOUND_PARAM(OutputAudioSideRight, "Out Side Right", "The side right channel audio output.")
-		METASOUND_PARAM(OutputAudioSideLeft, "Out Side Left", "The side right channel audio output.")
+		METASOUND_PARAM(OutputAudioSideLeft, "Out Side Left", "The side left channel audio output.")
 		METASOUND_PARAM(OutputAudioBackRight, "Out Back Right", "The back right channel audio output.")
-		METASOUND_PARAM(OutputAudioBackLeft, "Out Back Left", "The back right channel audio output.")
+		METASOUND_PARAM(OutputAudioBackLeft, "Out Back Left", "The back left channel audio output.")
 	}
 
 	namespace WavePlayerNodePrivate
