[MetaSounds] Unti test helper which builds a graph and operator from node names.

diff --git a/Engine/Plugins/Runtime/Metasound/Source/MetasoundEngineTest/Private/NodeTestGraphBuilder.cpp b/Engine/Plugins/Runtime/Metasound/Source/MetasoundEngineTest/Private/NodeTestGraphBuilder.cpp
new file mode 100644
--- /dev/null
+++ b/Engine/Plugins/Runtime/Metasound/Source/MetasoundEngineTest/Private/NodeTestGraphBuilder.cpp
@@ -0,0 +1,137 @@
+// Copyright Epic Games, Inc. All Rights Reserved.
+
+#include "NodeTestGraphBuilder.h"
+#include "Containers/Set.h"
+#include "MetasoundFrontendGraph.h"
+#include "MetasoundFrontendSearchEngine.h"
+#include "MetasoundOperatorBuilder.h"
+
+namespace Metasound::Test
+{
+	using namespace Metasound::Frontend;
+
+	FNodeTestGraphBuilder::FNodeTestGraphBuilder()
+	{
+		Document.RootGraph.Metadata.SetClassName({ "Metasound", "TestNodes", *LexToString(FGuid::NewGuid()) });
+		Document.RootGraph.Metadata.SetType(EMetasoundFrontendClassType::Graph);
+
+		DocumentHandle = IDocumentController::CreateDocumentHandle(Document);
+		RootGraph = DocumentHandle->GetRootGraph();
+		check(RootGraph->IsValid());
+	}
+
+	FNodeHandle FNodeTestGraphBuilder::AddNode(const FNodeClassName& ClassName, int32 MajorVersion)
+	{
+		check(RootGraph->IsValid());
+
+		FNodeHandle Node = INodeController::GetInvalidHandle();
+		FMetasoundFrontendClass NodeClass;
+		if (ISearchEngine::Get().FindClassWithHighestMinorVersion(ClassName, MajorVersion, NodeClass))
+		{
+			Node = RootGraph->AddNode(NodeClass.Metadata);
+		}
+
+		return Node;
+	}
+
+	FNodeHandle FNodeTestGraphBuilder::AddInput(const FName& InputName, const FName& TypeName)
+	{
+		check(RootGraph->IsValid());
+
+		FMetasoundFrontendClassInput Input;
+		Input.Name = InputName;
+		Input.TypeName = TypeName;
+		Input.VertexID = FGuid::NewGuid();
+		return RootGraph->AddInputVertex(Input);
+	}
+
+	FNodeHandle FNodeTestGraphBuilder::AddOutput(const FName& OutputName, const FName& TypeName)
+	{
+		check(RootGraph->IsValid());
+
+		FMetasoundFrontendClassOutput Output;
+		Output.Name = OutputName;
+		Output.TypeName = TypeName;
+		Output.VertexID = FGuid::NewGuid();
+		return RootGraph->AddOutputVertex(Output);
+	}
+
+	TUniquePtr<IOperator> FNodeTestGraphBuilder::BuildGraph()
+	{
+		TSet<FName> TransmittableInputNames;
+		const FString UnknownAsset = TEXT("UnknownAsset");
+		if (TUniquePtr<FFrontendGraph> Graph = FFrontendGraphBuilder::CreateGraph(Document, TransmittableInputNames, UnknownAsset))
+		{
+			FOperatorBuilderSettings BuilderSettings;
+			BuilderSettings.bFailOnAnyError = true;
+			BuilderSettings.bPopulateInternalDataReferences = true;
+			BuilderSettings.bValidateEdgeDataTypesMatch = true;
+			BuilderSettings.bValidateNoCyclesInGraph = true;
+			BuilderSettings.bValidateNoDuplicateInputs = true;
+			BuilderSettings.bValidateOperatorOutputsAreBound = true;
+			BuilderSettings.bValidateVerticesExist = true;
+			FOperatorBuilder Builder{ BuilderSettings };
+
+			FOperatorSettings OperatorSettings = FOperatorSettings(48000.0f, 256);
+			FInputVertexInterfaceData InterfaceData;
+			FMetasoundEnvironment Environment;
+			FBuildGraphOperatorParams BuildParams{ *Graph, OperatorSettings, InterfaceData, Environment };
+			FBuildResults Results;
+			return Builder.BuildGraphOperator(BuildParams, Results);
+		}
+
+		return nullptr;
+	}
+
+	TUniquePtr<IOperator> FNodeTestGraphBuilder::MakeSingleNodeGraph(const FNodeClassName& ClassName, int32 MajorVersion)
+	{
+		FNodeTestGraphBuilder Builder;
+		FNodeHandle NodeHandle = Builder.AddNode(ClassName, MajorVersion);
+
+		if (!NodeHandle->IsValid())
+		{
+			return nullptr;
+		}
+
+		// add the inputs and connect them
+		for (FInputHandle Input : NodeHandle->GetInputs())
+		{
+			FNodeHandle InputNode = Builder.AddInput(Input->GetName(), Input->GetDataType());
+
+			if (!InputNode->IsValid())
+			{
+				return nullptr;
+			}
+
+			FOutputHandle OutputToConnect = InputNode->GetOutputWithVertexName(Input->GetName());
+			FInputHandle InputToConnect = NodeHandle->GetInputWithVertexName(Input->GetName());
+
+			if (!InputToConnect->Connect(*OutputToConnect))
+			{
+				return nullptr;
+			}
+		}
+
+		// add the outputs and connect them
+		for (FOutputHandle Output : NodeHandle->GetOutputs())
+		{
+			FNodeHandle OutputNode = Builder.AddOutput(Output->GetName(), Output->GetDataType());
+
+			if (!OutputNode->IsValid())
+			{
+				return nullptr;
+			}
+
+			FOutputHandle OutputToConnect = NodeHandle->GetOutputWithVertexName(Output->GetName());
+			FInputHandle InputToConnect = OutputNode->GetInputWithVertexName(Output->GetName());
+
+			if (!InputToConnect->Connect(*OutputToConnect))
+			{
+				return nullptr;
+			}
+		}
+
+		// build the graph
+		return Builder.BuildGraph();
+	}
+}
diff --git a/Engine/Plugins/Runtime/Metasound/Source/MetasoundEngineTest/Public/NodeTestGraphBuilder.h b/Engine/Plugins/Runtime/Metasound/Source/MetasoundEngineTest/Public/NodeTestGraphBuilder.h
new file mode 100644
--- /dev/null
+++ b/Engine/Plugins/Runtime/Metasound/Source/MetasoundEngineTest/Public/NodeTestGraphBuilder.h
@@ -0,0 +1,38 @@
+// Copyright Epic Games, Inc. All Rights Reserved.
+
+#pragma once
+
+#include "MetasoundFrontendController.h"
+#include "MetasoundFrontendDocument.h"
+#include "MetasoundNodeInterface.h"
+#include "MetasoundOperatorInterface.h"
+
+namespace Metasound::Test
+{
+	/** Helper to make testing nodes simpler. */
+	class METASOUNDENGINETEST_API FNodeTestGraphBuilder
+	{
+	public:
+		FNodeTestGraphBuilder();
+
+		/** Add a node to the graph */
+		Frontend::FNodeHandle AddNode(const FNodeClassName& ClassName, int32 MajorVersion);
+
+		/** Add an input node to the graph */
+		Frontend::FNodeHandle AddInput(const FName& InputName, const FName& TypeName);
+
+		/** Add an output node to the grapn */
+		Frontend::FNodeHandle AddOutput(const FName& OutputName, const FName& TypeName);
+
+		/** Build the graph and get the graph's operator to work with */
+		TUniquePtr<IOperator> BuildGraph();
+
+		/** Helper that will add a single node, wire up the node's inputs and outputs, and hand back the graph's operator */
+		static TUniquePtr<IOperator> MakeSingleNodeGraph(const FNodeClassName& ClassName, int32 MajorVersion);
+
+	private:
+		FMetasoundFrontendDocument Document;
+		Frontend::FDocumentHandle DocumentHandle = Frontend::IDocumentController::GetInvalidHandle();
+		Frontend::FGraphHandle RootGraph = Frontend::IGraphController::GetInvalidHandle();
+	};
+}
