Have UMetasoundSource instances maintain a map  "AudioComponentIds -> weak pointers to FMetasoundGenerators" it has been asked to create. This will be useful in future checkins where parameter packs can be sent efficiently to playing metasound instances.

diff --git a/Engine/Plugins/Runtime/Metasound/Source/MetasoundEngine/Private/MetasoundSource.cpp b/Engine/Plugins/Runtime/Metasound/Source/MetasoundEngine/Private/MetasoundSource.cpp
--- a/Engine/Plugins/Runtime/Metasound/Source/MetasoundEngine/Private/MetasoundSource.cpp
+++ b/Engine/Plugins/Runtime/Metasound/Source/MetasoundEngine/Private/MetasoundSource.cpp
@@ -647,7 +647,16 @@ ISoundGeneratorPtr UMetaSoundSource::CreateSoundGenerator(const FSoundGeneratorI
 		MoveTemp(InDefaultParameters)
 	};
 
-	return ISoundGeneratorPtr(new FMetasoundGenerator(MoveTemp(InitParams)));
+	TSharedPtr<FMetasoundGenerator> Generator = MakeShared<FMetasoundGenerator>(MoveTemp(InitParams));
+	TrackGenerator(InParams.InstanceID, Generator);
+
+	return ISoundGeneratorPtr(Generator);
+}
+
+void UMetaSoundSource::OnEndGenerate(ISoundGeneratorPtr Generator)
+{
+	using namespace Metasound;
+	ForgetGenerator(Generator);
 }
 
 bool UMetaSoundSource::GetAllDefaultParameters(TArray<FAudioParameter>& OutParameters) const
@@ -985,5 +994,41 @@ const TArray<Metasound::FVertexName>& UMetaSoundSource::GetOutputAudioChannelOrd
 		return Empty;
 	}
 }
+
+void UMetaSoundSource::TrackGenerator(uint64 Id, TSharedPtr<Metasound::FMetasoundGenerator> Generator)
+{
+	FScopeLock GeneratorMapLock(&GeneratorMapCriticalSection);
+	Generators.Add(Id, Generator);
+	OnGeneratorInstanceCreated.Broadcast(Id, Generator);
+}
+
+void UMetaSoundSource::ForgetGenerator(ISoundGeneratorPtr Generator)
+{
+	using namespace Metasound;
+	FMetasoundGenerator* AsMetasoundGenerator = static_cast<FMetasoundGenerator*>(Generator.Get());
+	FScopeLock GeneratorMapLock(&GeneratorMapCriticalSection);
+	for (auto It = Generators.begin(); It != Generators.end(); ++It)
+	{
+		if ((*It).Value.HasSameObject(AsMetasoundGenerator))
+		{
+			OnGeneratorInstanceDestroyed.Broadcast((*It).Key, StaticCastSharedPtr<Metasound::FMetasoundGenerator>(Generator));
+			Generators.Remove((*It).Key);
+			return;
+		}
+	}
+}
+
+TWeakPtr<Metasound::FMetasoundGenerator> UMetaSoundSource::GetGeneratorForAudioComponent(uint64 ComponentId) const
+{
+	using namespace Metasound;
+	FScopeLock GeneratorMapLock(&GeneratorMapCriticalSection);
+	const TWeakPtr<FMetasoundGenerator>* Result = Generators.Find(ComponentId);
+	if (!Result)
+	{
+		return TWeakPtr<FMetasoundGenerator>(nullptr);
+	}
+	return *Result;
+}
+
 #undef LOCTEXT_NAMESPACE // MetaSound
 
diff --git a/Engine/Plugins/Runtime/Metasound/Source/MetasoundEngine/Public/MetasoundSource.h b/Engine/Plugins/Runtime/Metasound/Source/MetasoundEngine/Public/MetasoundSource.h
--- a/Engine/Plugins/Runtime/Metasound/Source/MetasoundEngine/Public/MetasoundSource.h
+++ b/Engine/Plugins/Runtime/Metasound/Source/MetasoundEngine/Public/MetasoundSource.h
@@ -3,6 +3,7 @@
 
 #include "CoreMinimal.h"
 #include "EdGraph/EdGraph.h"
+#include "HAL/CriticalSection.h"
 #include "IAudioParameterTransmitter.h"
 #include "Metasound.h"
 #include "MetasoundAssetBase.h"
@@ -21,6 +22,7 @@ namespace Metasound
 {
 	// Forward declare
 	struct FMetaSoundEngineAssetHelper;
+	class FMetasoundGenerator;
 }
 
 /** Declares the output audio format of the UMetaSoundSource */
@@ -36,6 +38,9 @@ enum class EMetasoundSourceAudioFormat : uint8
 	COUNT UMETA(Hidden)
 };
 
+DECLARE_MULTICAST_DELEGATE_TwoParams(FOnGeneratorInstanceCreated, uint64, TSharedPtr<Metasound::FMetasoundGenerator>);
+DECLARE_MULTICAST_DELEGATE_TwoParams(FOnGeneratorInstanceDestroyed, uint64, TSharedPtr<Metasound::FMetasoundGenerator>);
+
 /**
  * This Metasound type can be played as an audio source.
  */
@@ -184,11 +189,17 @@ public:
 	virtual float GetDuration() const override;
 	virtual bool ImplementsParameterInterface(Audio::FParameterInterfacePtr InInterface) const override;
 	virtual ISoundGeneratorPtr CreateSoundGenerator(const FSoundGeneratorInitParams& InParams, TArray<FAudioParameter>&& InDefaultParameters) override;
+	virtual void OnEndGenerate(ISoundGeneratorPtr Generator) override;
 	virtual TSharedPtr<Audio::IParameterTransmitter> CreateParameterTransmitter(Audio::FParameterTransmitterInitParams&& InParams) const override;
 	virtual bool IsParameterValid(const FAudioParameter& InParameter) const override;
 	virtual bool IsLooping() const override;
 	virtual bool IsOneShot() const override;
 	virtual bool EnableSubmixSendsOnPreview() const override { return true; }
+
+	TWeakPtr<Metasound::FMetasoundGenerator> GetGeneratorForAudioComponent(uint64 ComponentId) const;
+	FOnGeneratorInstanceCreated OnGeneratorInstanceCreated;
+	FOnGeneratorInstanceDestroyed OnGeneratorInstanceDestroyed;
+
 protected:
 
 	Metasound::Frontend::FDocumentAccessPtr GetDocument() override
@@ -223,4 +234,9 @@ private:
 	Metasound::FMetasoundEnvironment CreateEnvironment(const FSoundGeneratorInitParams& InParams) const;
 	Metasound::FMetasoundEnvironment CreateEnvironment(const Audio::FParameterTransmitterInitParams& InParams) const;
 	const TArray<Metasound::FVertexName>& GetOutputAudioChannelOrder() const;
+
+	mutable FCriticalSection GeneratorMapCriticalSection;
+	TSortedMap<uint64, TWeakPtr<Metasound::FMetasoundGenerator>> Generators;
+	void TrackGenerator(uint64 Id, TSharedPtr<Metasound::FMetasoundGenerator> Generator);
+	void ForgetGenerator(ISoundGeneratorPtr Generator);
 };
