MetasoundParameterPack V1

diff --git a/Engine/Plugins/Runtime/Metasound/Source/MetasoundEditor/Private/MetasoundNodeDetailCustomization.cpp b/Engine/Plugins/Runtime/Metasound/Source/MetasoundEditor/Private/MetasoundNodeDetailCustomization.cpp
--- a/Engine/Plugins/Runtime/Metasound/Source/MetasoundEditor/Private/MetasoundNodeDetailCustomization.cpp
+++ b/Engine/Plugins/Runtime/Metasound/Source/MetasoundEditor/Private/MetasoundNodeDetailCustomization.cpp
@@ -53,7 +53,8 @@ namespace Metasound
 			{
 				"Audio:Mono",
 				"Audio:Stereo",
-				GetMetasoundDataTypeName<Frontend::FAnalyzerAddress>()
+				GetMetasoundDataTypeName<Frontend::FAnalyzerAddress>(),
+				"MetasoundParameterPack"
 			};
 
 			static const FText OverrideInputDefaultText = LOCTEXT("OverridePresetInputDefault", "Override Inherited Default");
diff --git a/Engine/Plugins/Runtime/Metasound/Source/MetasoundEngine/Private/MetasoundGeneratorHandle.cpp b/Engine/Plugins/Runtime/Metasound/Source/MetasoundEngine/Private/MetasoundGeneratorHandle.cpp
new file mode 100644
--- /dev/null
+++ b/Engine/Plugins/Runtime/Metasound/Source/MetasoundEngine/Private/MetasoundGeneratorHandle.cpp
@@ -0,0 +1,158 @@
+// Copyright Epic Games, Inc. All Rights Reserved.
+
+#include "MetasoundGeneratorHandle.h"
+
+#include "Components/AudioComponent.h"
+#include "AudioDeviceManager.h"
+#include "MetasoundGenerator.h"
+#include "MetasoundSource.h"
+#include "MetasoundParameterPack.h"
+
+UMetasoundGeneratorHandle* UMetasoundGeneratorHandle::CreateMetaSoundGeneratorHandle(UAudioComponent* OnComponent)
+{
+	if (!OnComponent)
+	{
+		return nullptr;
+	}
+
+	UMetasoundGeneratorHandle* Result = NewObject<UMetasoundGeneratorHandle>();
+	Result->SetAudioComponent(OnComponent);
+	return Result;
+}
+
+void UMetasoundGeneratorHandle::ClearCachedData()
+{
+	DetachGeneratorDelegates();
+	AudioComponent        = nullptr;
+	AudioComponentId      = 0;
+	CachedMetasoundSource = nullptr;
+	CachedGeneratorPtr    = nullptr;
+	CachedParameterPack   = nullptr;
+}
+
+void UMetasoundGeneratorHandle::SetAudioComponent(UAudioComponent* InAudioComponent)
+{
+	if (InAudioComponent != AudioComponent)
+	{
+		ClearCachedData();
+		AudioComponent   = InAudioComponent;
+		AudioComponentId = InAudioComponent->GetAudioComponentID();
+	}
+}
+
+void UMetasoundGeneratorHandle::CacheMetasoundSource()
+{
+	if (!AudioComponent)
+	{
+		return;
+	}
+
+	UMetaSoundSource* CurrentMetasoundSource = Cast<UMetaSoundSource>(AudioComponent->GetSound());
+	if (CachedMetasoundSource == CurrentMetasoundSource)
+	{
+		return;
+	}
+
+	DetachGeneratorDelegates();
+	CachedGeneratorPtr    = nullptr;
+	CachedMetasoundSource = CurrentMetasoundSource;
+
+	if (CachedMetasoundSource)
+	{
+		AttachGeneratorDelegates();
+	}
+}
+
+void UMetasoundGeneratorHandle::AttachGeneratorDelegates()
+{
+	GeneratorCreatedDelegateHandle = CachedMetasoundSource->OnGeneratorInstanceCreated.AddLambda(
+		[=](uint64 InAudioComponentId, TSharedPtr<Metasound::FMetasoundGenerator> InGenerator)
+		{
+			OnSourceCreatedAGenerator(InAudioComponentId, InGenerator);
+		});
+	GeneratorDestroyedDelegateHandle = CachedMetasoundSource->OnGeneratorInstanceDestroyed.AddLambda(
+		[=](uint64 InAudioComponentId, TSharedPtr<Metasound::FMetasoundGenerator> InGenerator)
+		{
+			OnSourceDestroyedAGenerator(InAudioComponentId, InGenerator);
+		});
+}
+
+void UMetasoundGeneratorHandle::DetachGeneratorDelegates()
+{
+	CachedMetasoundSource->OnGeneratorInstanceCreated.Remove(GeneratorCreatedDelegateHandle);
+	GeneratorCreatedDelegateHandle.Reset();
+	CachedMetasoundSource->OnGeneratorInstanceDestroyed.Remove(GeneratorDestroyedDelegateHandle);
+	GeneratorDestroyedDelegateHandle.Reset();
+}
+
+TSharedPtr<Metasound::FMetasoundGenerator> UMetasoundGeneratorHandle::PinGenerator()
+{
+	TSharedPtr<Metasound::FMetasoundGenerator> PinnedGenerator = CachedGeneratorPtr.Pin();
+	if (PinnedGenerator.IsValid() || !CachedMetasoundSource)
+	{
+		return PinnedGenerator;
+	}
+
+	// The first attempt to pin failed, so reach out to the MetaSoundSource and see if it has a 
+	// generator for our AudioComponent...
+	CachedGeneratorPtr = CachedMetasoundSource->GetGeneratorForAudioComponent(AudioComponentId);
+	PinnedGenerator    = CachedGeneratorPtr.Pin();
+	return PinnedGenerator;
+}
+
+bool UMetasoundGeneratorHandle::ApplyParameterPack(UMetasoundParameterPack* Pack)
+{
+	if (!Pack)
+	{
+		return false;
+	}
+
+	// Create a copy of the parameter pack and cache it.
+	CachedParameterPack = Pack->GetCopyOfParameterStorage();
+
+	// No point in continuing if the parameter pack is not valid for any reason.
+	if (!CachedParameterPack.IsValid())
+	{
+		return false;
+	}
+
+	// Assure that our MetaSoundSource is up to date. It is possible that this has been 
+	// changed via script since we were first created.
+	CacheMetasoundSource();
+
+	// Now we can try to pin the generator.
+	TSharedPtr<Metasound::FMetasoundGenerator> PinnedGenerator = PinGenerator();
+
+	if (!PinnedGenerator.IsValid())
+	{
+		// Failed to pin the generator, but we have cached the parameter pack,
+		// so if our delegate gets called when a new generator is created we can 
+		// apply the cached parameters then.
+		return false;
+	}
+
+	// Finally... send down the parameter pack.
+	PinnedGenerator->QueueParameterPack(CachedParameterPack);
+	return true;
+}
+
+void UMetasoundGeneratorHandle::OnSourceCreatedAGenerator(uint64 InAudioComponentId, TSharedPtr<Metasound::FMetasoundGenerator> InGenerator)
+{
+	if (InAudioComponentId == AudioComponentId)
+	{
+		CachedGeneratorPtr = InGenerator;
+		if (InGenerator && CachedParameterPack)
+		{
+			InGenerator->QueueParameterPack(CachedParameterPack);
+		}
+	}
+}
+
+void UMetasoundGeneratorHandle::OnSourceDestroyedAGenerator(uint64 InAudioComponentId, TSharedPtr<Metasound::FMetasoundGenerator> InGenerator)
+{
+	if (InAudioComponentId == AudioComponentId)
+	{
+		CachedGeneratorPtr = nullptr;
+	}
+}
+
diff --git a/Engine/Plugins/Runtime/Metasound/Source/MetasoundEngine/Public/MetasoundGeneratorHandle.h b/Engine/Plugins/Runtime/Metasound/Source/MetasoundEngine/Public/MetasoundGeneratorHandle.h
new file mode 100644
--- /dev/null
+++ b/Engine/Plugins/Runtime/Metasound/Source/MetasoundEngine/Public/MetasoundGeneratorHandle.h
@@ -0,0 +1,71 @@
+// Copyright Epic Games, Inc. All Rights Reserved.
+
+#pragma once
+#include "Templates/SharedPointer.h"
+#include "MetasoundParameterPack.h"
+
+#include "MetasoundGeneratorHandle.generated.h"
+
+class UAudioComponent;
+class UMetaSoundSource;
+class UMetasoundParameterPack;
+namespace Metasound
+{
+	class FMetasoundGenerator;
+}
+
+UCLASS(BlueprintType,Category="MetaSound")
+class METASOUNDENGINE_API UMetasoundGeneratorHandle : public UObject
+{
+	GENERATED_BODY()
+
+public:
+
+	UFUNCTION(BlueprintCallable, Category="MetaSound")
+	static UMetasoundGeneratorHandle* CreateMetaSoundGeneratorHandle(UAudioComponent* OnComponent);
+
+	/**
+	 * Makes a copy of the supplied parameter pack and passes it to the MetaSoundGenerator
+	 * for asynchronous processing. IT ALSO caches this copy so that if the AudioComponent
+	 * is virtualized the parameter pack will be sent again when/if the AudioComponent is 
+	 * "unvirtualized".
+	 */
+	UFUNCTION(BlueprintCallable, Category="MetaSoundParameterPack")
+	bool ApplyParameterPack(UMetasoundParameterPack* Pack);
+
+private:
+
+	void SetAudioComponent(UAudioComponent* InAudioComponent);
+	void CacheMetasoundSource();
+	void ClearCachedData();
+
+	/**
+	 * Attempts to pin the weak generator pointer. If the first attempt fails it checks to see
+	 * if it can "recapture" a pointer to a generator for the current AudioComponent/MetaSoundSource
+	 * combination. 
+	 */
+	TSharedPtr<Metasound::FMetasoundGenerator> PinGenerator();
+
+	/**
+	 * Functions for adding and removing our MetaSoundGenerator lifecycle delegates
+	 */
+	void AttachGeneratorDelegates();
+	void DetachGeneratorDelegates();
+
+	/**
+	 * Generator creation and destruction delegates we register with the UMetaSoundSource
+	 */
+	void OnSourceCreatedAGenerator(uint64 InAudioComponentId, TSharedPtr<Metasound::FMetasoundGenerator> InGenerator);
+	void OnSourceDestroyedAGenerator(uint64 InAudioComponentId, TSharedPtr<Metasound::FMetasoundGenerator> InGenerator);
+
+	UPROPERTY(Transient)
+	TObjectPtr<UAudioComponent> AudioComponent;
+	uint64 AudioComponentId;
+	UPROPERTY(Transient)
+	TObjectPtr<UMetaSoundSource> CachedMetasoundSource;
+	TWeakPtr<Metasound::FMetasoundGenerator> CachedGeneratorPtr;
+	FSharedMetasoundParameterStoragePtr CachedParameterPack;
+
+	FDelegateHandle GeneratorCreatedDelegateHandle;
+	FDelegateHandle GeneratorDestroyedDelegateHandle;
+};
\ No newline at end of file
diff --git a/Engine/Plugins/Runtime/Metasound/Source/MetasoundFrontend/Private/MetasoundAssetBase.cpp b/Engine/Plugins/Runtime/Metasound/Source/MetasoundFrontend/Private/MetasoundAssetBase.cpp
--- a/Engine/Plugins/Runtime/Metasound/Source/MetasoundFrontend/Private/MetasoundAssetBase.cpp
+++ b/Engine/Plugins/Runtime/Metasound/Source/MetasoundFrontend/Private/MetasoundAssetBase.cpp
@@ -33,6 +33,7 @@
 #include "StructSerializer.h"
 #include "Templates/SharedPointer.h"
 #include "UObject/MetaData.h"
+#include "MetasoundParameterPack.h"
 
 #define LOCTEXT_NAMESPACE "MetaSound"
 
@@ -824,6 +825,9 @@ TArray<FMetasoundFrontendClassInput> FMetasoundAssetBase::GetPublicClassInputs()
 
 	Algo::CopyIf(Doc.RootGraph.Interface.Inputs, PublicInputs, IsPublic);
 
+	// Add the parameter pack input that ALL Metasounds have
+	PublicInputs.Add(UMetasoundParameterPack::GetClassInput());
+	
 	return PublicInputs;
 }
 
diff --git a/Engine/Plugins/Runtime/Metasound/Source/MetasoundFrontend/Private/MetasoundFrontendDataTypeRegistry.cpp b/Engine/Plugins/Runtime/Metasound/Source/MetasoundFrontend/Private/MetasoundFrontendDataTypeRegistry.cpp
--- a/Engine/Plugins/Runtime/Metasound/Source/MetasoundFrontend/Private/MetasoundFrontendDataTypeRegistry.cpp
+++ b/Engine/Plugins/Runtime/Metasound/Source/MetasoundFrontend/Private/MetasoundFrontendDataTypeRegistry.cpp
@@ -522,6 +522,8 @@ namespace Metasound
 				virtual TOptional<FAnyDataReference> CreateDataReference(const FName& InDataType, EDataReferenceAccessType InAccessType, const FLiteral& InLiteral, const FOperatorSettings& InOperatorSettings) const override;
 				virtual TSharedPtr<IDataChannel, ESPMode::ThreadSafe> CreateDataChannel(const FName& InDataType, const FOperatorSettings& InOperatorSettings) const override;
 
+				virtual const IParameterAssignmentFunction& GetRawAssignmentFunction(const FName& InDataType) const override;
+
 				virtual bool GetFrontendInputClass(const FName& InDataType, FMetasoundFrontendClass& OutClass) const override;
 				virtual bool GetFrontendConstructorInputClass(const FName& InDataType, FMetasoundFrontendClass& OutClass) const override;
 				virtual bool GetFrontendLiteralClass(const FName& InDataType, FMetasoundFrontendClass& OutClass) const override;
@@ -895,6 +897,16 @@ namespace Metasound
 				return nullptr;
 			}
 
+			const IParameterAssignmentFunction& FDataTypeRegistry::GetRawAssignmentFunction(const FName& InDataType) const
+			{
+				static IParameterAssignmentFunction NoOp;
+				if (const IDataTypeRegistryEntry* Entry = FindDataTypeEntry(InDataType))
+				{
+					return Entry->GetRawAssignmentFunction();
+				}
+				return NoOp;
+			}
+
 			bool FDataTypeRegistry::GetFrontendInputClass(const FName& InDataType, FMetasoundFrontendClass& OutClass) const
 			{
 				if (const IDataTypeRegistryEntry* Entry = FindDataTypeEntry(InDataType))
diff --git a/Engine/Plugins/Runtime/Metasound/Source/MetasoundFrontend/Private/MetasoundParameterPack.cpp b/Engine/Plugins/Runtime/Metasound/Source/MetasoundFrontend/Private/MetasoundParameterPack.cpp
new file mode 100644
--- /dev/null
+++ b/Engine/Plugins/Runtime/Metasound/Source/MetasoundFrontend/Private/MetasoundParameterPack.cpp
@@ -0,0 +1,180 @@
+// Copyright Epic Games, Inc. All Rights Reserved.
+
+#include "MetasoundParameterPack.h"
+
+#include "MetasoundFrontendDocumentAccessPtr.h"
+#include "MetasoundDataTypeRegistrationMacro.h"
+
+#include <limits>
+
+DEFINE_LOG_CATEGORY_STATIC(LogMetasoundParamPack, Log, All);
+
+REGISTER_METASOUND_DATATYPE(FMetasoundParameterStorageWrapper, "MetasoundParameterPack", Metasound::ELiteralType::UObjectProxy, UMetasoundParameterPack);
+
+namespace MetasoundParameterPack
+{
+	FName RoutingName("ParamPack");
+}
+
+TUniquePtr<Audio::IProxyData> UMetasoundParameterPack::CreateNewProxyData(const Audio::FProxyDataInitParams& InitParams)
+{
+	return MakeUnique<FMetasoundParameterPackProxy>(ParameterStorage);
+}
+
+Metasound::FSendAddress UMetasoundParameterPack::CreateSendAddressFromEnvironment(const Metasound::FMetasoundEnvironment& InEnvironment)
+{
+	return Metasound::FMetaSoundParameterTransmitter::CreateSendAddressFromEnvironment(InEnvironment, MetasoundParameterPack::RoutingName, Metasound::GetMetasoundDataTypeName<FMetasoundParameterStorageWrapper>());
+}
+
+FMetasoundFrontendClassInput UMetasoundParameterPack::GetClassInput()
+{
+	FMetasoundFrontendClassInput ClassInput;
+	ClassInput.Name = MetasoundParameterPack::RoutingName;
+	ClassInput.DefaultLiteral = FMetasoundFrontendLiteral({ MetasoundParameterPack::RoutingName, nullptr });
+	ClassInput.TypeName = Metasound::GetMetasoundDataTypeName<FMetasoundParameterStorageWrapper>();
+	return ClassInput;
+}
+
+UMetasoundParameterPack* UMetasoundParameterPack::MakeMetasoundParameterPack()
+{
+	return NewObject<UMetasoundParameterPack>();
+}
+
+ESetParamResult UMetasoundParameterPack::SetBool(FName Name, bool InValue, bool OnlyIfExists)
+{
+	return SetParameter<bool>(Name, Metasound::GetMetasoundDataTypeName<bool>(), InValue, OnlyIfExists);
+}
+
+ESetParamResult UMetasoundParameterPack::SetInt(FName Name, int32 InValue, bool OnlyIfExists)
+{
+	return SetParameter<int32>(Name, Metasound::GetMetasoundDataTypeName<int32>(), InValue, OnlyIfExists);
+}
+
+ESetParamResult UMetasoundParameterPack::SetFloat(FName Name, float InValue, bool OnlyIfExists)
+{
+	return SetParameter<float>(Name, Metasound::GetMetasoundDataTypeName<float>(), InValue, OnlyIfExists);
+}
+
+ESetParamResult UMetasoundParameterPack::SetString(FName Name, const FString& InValue, bool OnlyIfExists)
+{
+	return SetParameter<FString>(Name, Metasound::GetMetasoundDataTypeName<FString>(), InValue, OnlyIfExists);
+}
+
+ESetParamResult UMetasoundParameterPack::SetTrigger(FName Name, bool OnlyIfExists)
+{
+	return SetParameter<bool>(Name, Metasound::GetMetasoundDataTypeName<Metasound::FTrigger>(), true, OnlyIfExists);
+}
+
+bool UMetasoundParameterPack::GetBool(FName Name, ESetParamResult& Result)
+{
+	return GetParameter<bool>(Name, Metasound::GetMetasoundDataTypeName<bool>(), Result);
+}
+
+int32 UMetasoundParameterPack::GetInt(FName Name, ESetParamResult& Result)
+{
+	return GetParameter<int32>(Name, Metasound::GetMetasoundDataTypeName<int32>(), Result);
+}
+
+float UMetasoundParameterPack::GetFloat(FName Name, ESetParamResult& Result)
+{
+	return GetParameter<float>(Name, Metasound::GetMetasoundDataTypeName<float>(), Result);
+}
+
+FString UMetasoundParameterPack::GetString(FName Name, ESetParamResult& Result)
+{
+	return GetParameter<FString>(Name, Metasound::GetMetasoundDataTypeName<FString>(), Result);
+}
+
+bool UMetasoundParameterPack::GetTrigger(FName Name, ESetParamResult& Result)
+{
+	return GetParameter<bool>(Name, Metasound::GetMetasoundDataTypeName<Metasound::FTrigger>(), Result);
+}
+
+bool* UMetasoundParameterPack::GetBoolParameterPtr(FName Name)
+{
+	return GetParameterPtr<bool>(Name, Metasound::GetMetasoundDataTypeName<bool>());
+}
+
+int32* UMetasoundParameterPack::GetIntParameterPtr(FName Name)
+{
+	return GetParameterPtr<int32>(Name, Metasound::GetMetasoundDataTypeName<int32>());
+}
+
+float* UMetasoundParameterPack::GetFloatParameterPtr(FName Name)
+{
+	return GetParameterPtr<float>(Name, Metasound::GetMetasoundDataTypeName<float>());
+}
+
+FString* UMetasoundParameterPack::GetStringParameterPtr(FName Name)
+{
+	return GetParameterPtr<FString>(Name, Metasound::GetMetasoundDataTypeName<FString>());
+}
+
+bool* UMetasoundParameterPack::GetTriggerParameterPtr(FName Name)
+{
+	return GetParameterPtr<bool>(Name, Metasound::GetMetasoundDataTypeName<Metasound::FTrigger>());
+}
+
+bool UMetasoundParameterPack::HasBool(FName Name)
+{
+	return HasParameter<bool>(Name, Metasound::GetMetasoundDataTypeName<bool>());
+}
+
+bool UMetasoundParameterPack::HasInt(FName Name)
+{
+	return HasParameter<int32>(Name, Metasound::GetMetasoundDataTypeName<int32>());
+}
+
+bool UMetasoundParameterPack::HasFloat(FName Name)
+{
+	return HasParameter<float>(Name, Metasound::GetMetasoundDataTypeName<float>());
+}
+
+bool UMetasoundParameterPack::HasString(FName Name)
+{
+	return HasParameter<FString>(Name, Metasound::GetMetasoundDataTypeName<FString>());
+}
+
+bool UMetasoundParameterPack::HasTrigger(FName Name)
+{
+	return HasParameter<bool>(Name, Metasound::GetMetasoundDataTypeName<Metasound::FTrigger>());
+}
+
+void UMetasoundParameterPack::AddBoolParameter(FName Name, bool InValue)
+{
+	AddParameter<bool>(Name, Metasound::GetMetasoundDataTypeName<bool>(), InValue);
+}
+
+void UMetasoundParameterPack::AddIntParameter(FName Name, int32 InValue)
+{
+	AddParameter<int32>(Name, Metasound::GetMetasoundDataTypeName<int32>(), InValue);
+}
+
+void UMetasoundParameterPack::AddFloatParameter(FName Name, float InValue)
+{
+	AddParameter<float>(Name, Metasound::GetMetasoundDataTypeName<float>(), InValue);
+}
+
+void UMetasoundParameterPack::AddStringParameter(FName Name, const FString& InValue)
+{
+	AddParameter<FString>(Name, Metasound::GetMetasoundDataTypeName<FString>(), InValue);
+}
+
+void UMetasoundParameterPack::AddTriggerParameter(FName Name, bool InValue)
+{
+	AddParameter<bool>(Name, Metasound::GetMetasoundDataTypeName<Metasound::FTrigger>(), InValue);
+}
+
+FSharedMetasoundParameterStoragePtr UMetasoundParameterPack::GetParameterStorage() const
+{
+	return ParameterStorage;
+}
+
+FSharedMetasoundParameterStoragePtr UMetasoundParameterPack::GetCopyOfParameterStorage() const
+{
+	if (!ParameterStorage.IsValid())
+	{
+		return nullptr;
+	}
+	return MakeShared<FMetasoundParameterPackStorage>(*ParameterStorage);
+}
diff --git a/Engine/Plugins/Runtime/Metasound/Source/MetasoundFrontend/Private/MetasoundTrigger.cpp b/Engine/Plugins/Runtime/Metasound/Source/MetasoundFrontend/Private/MetasoundTrigger.cpp
--- a/Engine/Plugins/Runtime/Metasound/Source/MetasoundFrontend/Private/MetasoundTrigger.cpp
+++ b/Engine/Plugins/Runtime/Metasound/Source/MetasoundFrontend/Private/MetasoundTrigger.cpp
@@ -29,6 +29,13 @@ namespace Metasound
 		: FTrigger(InSettings, false)
 	{}
 
+	void FTrigger::AssignRawParameter(const void* ParamPackPayload)
+	{
+		// We don't even need to check the payload. If the parameter pack contained
+		// a Trigger bool, it must be true. There is no concept of "un-trigger" :-) 
+		TriggerFrame(0);
+	}
+
 	void FTrigger::TriggerFrame(int32 InFrameToTrigger)
 	{
 		// Insert trigger frame index into sorted triggered frames
diff --git a/Engine/Plugins/Runtime/Metasound/Source/MetasoundFrontend/Public/MetasoundDataTypeRegistrationMacro.h b/Engine/Plugins/Runtime/Metasound/Source/MetasoundFrontend/Public/MetasoundDataTypeRegistrationMacro.h
--- a/Engine/Plugins/Runtime/Metasound/Source/MetasoundFrontend/Public/MetasoundDataTypeRegistrationMacro.h
+++ b/Engine/Plugins/Runtime/Metasound/Source/MetasoundFrontend/Public/MetasoundDataTypeRegistrationMacro.h
@@ -212,7 +212,7 @@ namespace Metasound
 			return RegistryInfo;
 		}
 
-		/** Returns an IEnumDataTypeInterface pointer for the data type. If the 
+		/** Returns an IEnumDataTypeInterface pointer for the data type. If the
 		 * data type has no IEnumDataTypeInterface, the returned pointer will be
 		 * invalid.
 		 *
@@ -270,9 +270,16 @@ namespace Metasound
 			return EnumInterfacePtr;
 		}
 
+		template <typename T, typename V = void>
+		struct HasRawParameterAssignmentOp { bool value = false; };
+
+		template <typename T>
+		struct HasRawParameterAssignmentOp<T, decltype(&T::AssignRawParameter, void())> { bool value = true; };
+
+
 		/** Registers a data type with the MetaSound Frontend. This allows the data type
-		 * to be used in Input and Output nodes by informing the Frontend how to 
-		 * instantiate an instance. 
+		 * to be used in Input and Output nodes by informing the Frontend how to
+		 * instantiate an instance.
 		 *
 		 * @tparam TDataType - The data type to register.
 		 * @tparam PreferredArgType - The preferred constructor argument type to use when creating an instance of the data type.
@@ -311,9 +318,18 @@ namespace Metasound
 			public:
 
 				FDataTypeRegistryEntry()
-				: Info(CreateDataTypeInfo<TDataType, PreferredArgType, UClassToUse>())
-				, EnumInterface(GetEnumDataTypeInterface<TDataType>())
+					: Info(CreateDataTypeInfo<TDataType, PreferredArgType, UClassToUse>())
+					, EnumInterface(GetEnumDataTypeInterface<TDataType>())
 				{
+					if constexpr (HasRawParameterAssignmentOp<TDataType>().value)
+					{
+						RawAssignmentFunction = [](const void* src, void* dest) { reinterpret_cast<TDataType*>(dest)->AssignRawParameter(src); };
+					}
+					else if constexpr (std::is_copy_assignable_v<TDataType>)
+					{
+						RawAssignmentFunction = [](const void* src, void* dest) { *(reinterpret_cast<TDataType*>(dest)) = *(reinterpret_cast<const TDataType*>(src)); };
+					}
+
 					// Create class info using prototype node
 					// TODO: register nodes with static class info instead of prototype instance.
 
@@ -589,6 +605,11 @@ namespace Metasound
 					return MakeUnique<FDataTypeRegistryEntry>();
 				}
 
+				virtual const Frontend::IParameterAssignmentFunction& GetRawAssignmentFunction() const override
+				{
+					return RawAssignmentFunction;
+				}
+
 			private:
 				Frontend::FDataTypeRegistryInfo Info;
 				FMetasoundFrontendClass InputClass;
@@ -601,6 +622,7 @@ namespace Metasound
 				FMetasoundFrontendClass VariableAccessorClass;
 				FMetasoundFrontendClass VariableDeferredAccessorClass;
 				TSharedPtr<Frontend::IEnumDataTypeInterface> EnumInterface;
+				Frontend::IParameterAssignmentFunction RawAssignmentFunction;
 			};
 
 			bool bSucceeded = Frontend::IDataTypeRegistry::Get().RegisterDataType(MakeUnique<FDataTypeRegistryEntry>());
diff --git a/Engine/Plugins/Runtime/Metasound/Source/MetasoundFrontend/Public/MetasoundFrontendDataTypeRegistry.h b/Engine/Plugins/Runtime/Metasound/Source/MetasoundFrontend/Public/MetasoundFrontendDataTypeRegistry.h
--- a/Engine/Plugins/Runtime/Metasound/Source/MetasoundFrontend/Public/MetasoundFrontendDataTypeRegistry.h
+++ b/Engine/Plugins/Runtime/Metasound/Source/MetasoundFrontend/Public/MetasoundFrontendDataTypeRegistry.h
@@ -18,6 +18,8 @@ namespace Metasound
 {
 	namespace Frontend
 	{
+		using IParameterAssignmentFunction = TUniqueFunction<void(const void*,void*)>;
+
 		/** FDataTypeRegsitryInfo contains runtime inspectable behavior of a registered
 		 * MetaSound data type.
 		 */
@@ -184,6 +186,9 @@ namespace Metasound
 			/** Create a receive node for this data type. */
 			virtual TUniquePtr<INode> CreateReceiveNode(const FNodeInitData&) const = 0;
 
+			/* Get function that should be used when a parameter pack contains a setting this data type. */ 
+			virtual const IParameterAssignmentFunction& GetRawAssignmentFunction() const = 0;
+
 			/** Create a init variable node for this data type. 
 			 *
 			 *  @param InInitParams - Contains a literal used to create the variable.
@@ -266,6 +271,9 @@ namespace Metasound
 
 			virtual TSharedPtr<IDataChannel, ESPMode::ThreadSafe> CreateDataChannel(const FName& InDataType, const FOperatorSettings& InOperatorSettings) const = 0;
 
+			/* Get function that should be used when a parameter pack contains a setting for the specified data type. */ 
+			virtual const IParameterAssignmentFunction& GetRawAssignmentFunction(const FName& InDataType) const = 0;
+
 			/** Return an FMetasoundFrontendClass representing an input node of the data type. */
 			virtual bool GetFrontendInputClass(const FName& InDataType, FMetasoundFrontendClass& OutClass) const = 0;
 
diff --git a/Engine/Plugins/Runtime/Metasound/Source/MetasoundFrontend/Public/MetasoundParameterPack.h b/Engine/Plugins/Runtime/Metasound/Source/MetasoundFrontend/Public/MetasoundParameterPack.h
new file mode 100644
--- /dev/null
+++ b/Engine/Plugins/Runtime/Metasound/Source/MetasoundFrontend/Public/MetasoundParameterPack.h
@@ -0,0 +1,425 @@
+// Copyright Epic Games, Inc. All Rights Reserved.
+
+#pragma once
+#include "Templates/AlignmentTemplates.h"
+
+#include "MetasoundFrontendLiteral.h"
+#include "MetasoundDataReferenceMacro.h"
+#include "MetasoundRouter.h"
+#include "MetasoundFrontendDocument.h"
+
+#include "MetasoundParameterPack.generated.h"
+
+class UMetaSoundSource;
+
+// A structure that encapsulates a "parameter tag" in the raw 'bag-o-bytes'
+// that will be the parameter storage. It lets us walk through memory to 
+// find a parameter with a specific name and specific type. 
+struct FMetasoundParameterPackItemBase
+{
+	FMetasoundParameterPackItemBase(const FName& InName, const FName& InTypeName, uint16 InNextOffset, uint16 InValueOffset)
+		: Name(InName)
+		, TypeName(InTypeName)
+		, NextHeaderOffset(InNextOffset)
+		, ValueOffset(InValueOffset)
+	{}
+	FMetasoundParameterPackItemBase(const FMetasoundParameterPackItemBase& Other) = default;
+	void* GetPayload()
+	{
+		uint8* ptr = reinterpret_cast<uint8*>(this);
+		return ptr + ValueOffset;
+	}
+	FName  Name;
+	FName  TypeName;
+	uint16 NextHeaderOffset;
+	uint16 ValueOffset;
+};
+
+// A template that lets us stick an arbitrary data type into the raw 
+// 'bag-o-bytes' that is the parameter storage. Uses the base class 
+// above.
+template <typename T>
+struct FMetasoundParameterPackItem : public FMetasoundParameterPackItemBase
+{
+	FMetasoundParameterPackItem(const FName& InName, const FName& InTypeName, const T& InValue)
+		: FMetasoundParameterPackItemBase(InName, InTypeName, AlignArbitrary<uint16>((uint16)sizeof(FMetasoundParameterPackItem<T>), alignof(std::max_align_t)), offsetof(FMetasoundParameterPackItem<T>,Value))
+		, Value(InValue)
+	{}
+	FMetasoundParameterPackItem(const FMetasoundParameterPackItem& Other) = default;
+	T      Value;
+};
+
+// This next class owns the 'bag-o-bytes' that holds all of the parameters. 
+struct FMetasoundParameterPackStorage
+{
+	TArray<uint8> Storage;	
+	int32         StorageBytesInUse = 0;
+
+	FMetasoundParameterPackStorage() = default;
+	FMetasoundParameterPackStorage(const FMetasoundParameterPackStorage& Other)
+		: Storage(Other.Storage)
+		, StorageBytesInUse(Other.StorageBytesInUse)
+	{}
+	FMetasoundParameterPackStorage(FMetasoundParameterPackStorage&& Other)
+		: Storage(MoveTemp(Other.Storage))
+		, StorageBytesInUse(Other.StorageBytesInUse)
+	{
+		Other.StorageBytesInUse = 0;
+	}
+
+	struct ParameterIterator
+	{
+		uint8* CurrentNode;
+		uint8* EndAddress;
+
+		ParameterIterator(uint8* Start, uint8* End)
+			: CurrentNode(Start)
+			, EndAddress(End)
+		{}
+
+		ParameterIterator& operator++()
+		{
+			if (!CurrentNode)
+			{
+				return *this; 
+			}
+			FMetasoundParameterPackItemBase* ParamAddress = reinterpret_cast<FMetasoundParameterPackItemBase*>(CurrentNode);
+			CurrentNode += ParamAddress->NextHeaderOffset;
+			if (CurrentNode >= EndAddress)
+			{
+				CurrentNode = nullptr;
+			}
+			return *this;
+		}
+
+		ParameterIterator operator++(int)
+		{
+			ParameterIterator ReturnValue = *this;
+			++*this;
+			return ReturnValue;
+		}
+
+		bool operator==(const ParameterIterator& Other) const
+		{
+			return CurrentNode == Other.CurrentNode;
+		}
+
+		bool operator!=(const ParameterIterator& Other) const
+		{
+			return CurrentNode != Other.CurrentNode;
+		}
+
+		FMetasoundParameterPackItemBase* operator->()
+		{
+			return reinterpret_cast<FMetasoundParameterPackItemBase*>(CurrentNode);
+		}
+
+		FMetasoundParameterPackItemBase* GetParameterBase()
+		{
+			return reinterpret_cast<FMetasoundParameterPackItemBase*>(CurrentNode);
+		}
+
+	};
+
+	ParameterIterator begin()
+	{
+		if (Storage.IsEmpty())
+		{
+			return ParameterIterator(nullptr, nullptr);
+		}
+		return ParameterIterator(&Storage[0], &Storage[0] + StorageBytesInUse);
+	}
+
+	ParameterIterator end()
+	{
+		return ParameterIterator(nullptr, nullptr);
+	}
+
+	template<typename T>
+	FMetasoundParameterPackItem<T>* FindParameter(const FName& ParamName, const FName& InTypeName)
+	{
+		if (Storage.IsEmpty())
+		{
+			return nullptr;
+		}
+		FMetasoundParameterPackItemBase* ParamWalker = reinterpret_cast<FMetasoundParameterPackItemBase*>(&Storage[0]);
+		int32 StorageOffset = 0;
+		while (ParamWalker)
+		{
+			if (ParamWalker->Name == ParamName)
+			{
+				if (ParamWalker->TypeName == InTypeName)
+				{
+					return reinterpret_cast<FMetasoundParameterPackItem<T>*>(ParamWalker);
+				}
+				return nullptr;
+			}
+			StorageOffset += ParamWalker->NextHeaderOffset;
+			if (StorageOffset < StorageBytesInUse)
+			{
+				ParamWalker = reinterpret_cast<FMetasoundParameterPackItemBase*>(&Storage[StorageOffset]);
+			}
+			else
+			{
+				break;
+			}
+		}
+		return nullptr;
+	}
+
+	template<typename T>
+	std::enable_if_t<!std::is_class_v<T>,T*> AddParameter(FMetasoundParameterPackItem<T>& NewParam)
+	{
+		int32 NextParamLocation = AlignArbitrary<int32>(StorageBytesInUse, alignof(std::max_align_t));
+		int32 TotalStorageNeeded = NextParamLocation + NewParam.NextHeaderOffset;
+		if (TotalStorageNeeded > Storage.Num())
+		{
+			Storage.AddUninitialized(TotalStorageNeeded - Storage.Num());
+		}
+		FMetasoundParameterPackItem<T>* Destination = reinterpret_cast<FMetasoundParameterPackItem<T>*>(&Storage[NextParamLocation]);
+		*Destination = NewParam;
+		StorageBytesInUse = NextParamLocation + NewParam.NextHeaderOffset;
+		return &Destination->Value;
+	}
+
+	template<typename T>
+	std::enable_if_t<std::is_class_v<T>, T*> AddParameter(FMetasoundParameterPackItem<T>& NewParam)
+	{
+		int32 NextParamLocation = AlignArbitrary<int32>(StorageBytesInUse, alignof(std::max_align_t));
+		int32 TotalStorageNeeded = NextParamLocation + NewParam.NextHeaderOffset;
+		if (TotalStorageNeeded > Storage.Num())
+		{
+			Storage.AddUninitialized(TotalStorageNeeded - Storage.Num());
+		}
+		FMetasoundParameterPackItem<T>* Destination = new (reinterpret_cast<FMetasoundParameterPackItem<T>*>(&Storage[NextParamLocation])) FMetasoundParameterPackItem<T>(NewParam);
+		StorageBytesInUse = NextParamLocation + NewParam.NextHeaderOffset;
+		return &Destination->Value;
+	}
+};
+
+using FSharedMetasoundParameterStoragePtr = TSharedPtr<FMetasoundParameterPackStorage, ESPMode::ThreadSafe>;
+
+UENUM(BlueprintType)
+enum class ESetParamResult : uint8
+{
+	Succeeded,
+	Failed
+};
+
+// Here is the UObject BlueprintType that can be used in c++ and blueprint code. It holds a FMetasoundParamStorage 
+// instance and can pass it along to the audio system's SetObjectParameter function via an AudioProxy.
+UCLASS(BlueprintType,meta = (DisplayName = "MetaSoundParameterPack"))
+class METASOUNDFRONTEND_API UMetasoundParameterPack : public UObject, public IAudioProxyDataFactory
+{
+	GENERATED_BODY()
+
+public:
+
+	UFUNCTION(BlueprintCallable, Category = "MetaSoundParameterPack")
+	static UMetasoundParameterPack* MakeMetasoundParameterPack();
+
+	UFUNCTION(BlueprintCallable, Category = "MetaSoundParameterPack", meta = (ExpandEnumAsExecs = ReturnValue))
+	ESetParamResult SetBool(FName ParameterName, bool InValue, bool OnlyIfExists = true);
+	UFUNCTION(BlueprintCallable, Category = "MetaSoundParameterPack", meta = (ExpandEnumAsExecs = ReturnValue))
+	ESetParamResult SetInt(FName ParameterName, int32 InValue, bool OnlyIfExists = true);
+	UFUNCTION(BlueprintCallable, Category = "MetaSoundParameterPack", meta = (ExpandEnumAsExecs = ReturnValue))
+	ESetParamResult SetFloat(FName ParameterName, float InValue, bool OnlyIfExists = true);
+	UFUNCTION(BlueprintCallable, Category = "MetaSoundParameterPack", meta = (ExpandEnumAsExecs = ReturnValue))
+	ESetParamResult SetString(FName ParameterName, const FString& InValue, bool OnlyIfExists = true);
+	UFUNCTION(BlueprintCallable, Category = "MetaSoundParameterPack", meta = (ExpandEnumAsExecs = ReturnValue))
+	ESetParamResult SetTrigger(FName ParameterName, bool OnlyIfExists = true);
+
+	UFUNCTION(BlueprintCallable, Category = "MetaSoundParameterPack", meta = (ExpandEnumAsExecs = "Result"))
+	bool GetBool(FName ParameterName, ESetParamResult& Result);
+	UFUNCTION(BlueprintCallable, Category = "MetaSoundParameterPack", meta = (ExpandEnumAsExecs = "Result"))
+	int32 GetInt(FName ParameterName, ESetParamResult& Result);
+	UFUNCTION(BlueprintCallable, Category = "MetaSoundParameterPack", meta = (ExpandEnumAsExecs = "Result"))
+	float GetFloat(FName ParameterName, ESetParamResult& Result);
+	UFUNCTION(BlueprintCallable, Category = "MetaSoundParameterPack", meta = (ExpandEnumAsExecs = "Result"))
+	FString GetString(FName ParameterName, ESetParamResult& Result);
+	UFUNCTION(BlueprintCallable, Category = "MetaSoundParameterPack", meta = (ExpandEnumAsExecs = "Result"))
+	bool GetTrigger(FName ParameterName, ESetParamResult& Result);
+
+	UFUNCTION(BlueprintCallable, Category = "MetaSoundParameterPack", meta = (ExpandEnumAsExecs = ReturnValue))
+	bool HasBool(FName ParameterName);
+	UFUNCTION(BlueprintCallable, Category = "MetaSoundParameterPack", meta = (ExpandEnumAsExecs = ReturnValue))
+	bool HasInt(FName ParameterName);
+	UFUNCTION(BlueprintCallable, Category = "MetaSoundParameterPack", meta = (ExpandEnumAsExecs = ReturnValue))
+	bool HasFloat(FName ParameterName);
+	UFUNCTION(BlueprintCallable, Category = "MetaSoundParameterPack", meta = (ExpandEnumAsExecs = ReturnValue))
+	bool HasString(FName ParameterName);
+	UFUNCTION(BlueprintCallable, Category = "MetaSoundParameterPack", meta = (ExpandEnumAsExecs = ReturnValue))
+	bool HasTrigger(FName ParameterName);
+
+	void AddBoolParameter(FName Name, bool InValue);
+	void AddIntParameter(FName Name, int32 InValue);
+	void AddFloatParameter(FName Name, float InValue);
+	void AddStringParameter(FName Name, const FString& InValue);
+	void AddTriggerParameter(FName Name, bool InValue = true);
+
+	// YIKES! BEWARE: If the returned pointers are only valid until another parameter is added!
+	bool*    GetBoolParameterPtr(FName Name);
+	int32*   GetIntParameterPtr(FName Name);
+	float*   GetFloatParameterPtr(FName Name);
+	FString* GetStringParameterPtr(FName Name);
+	bool*    GetTriggerParameterPtr(FName Name);
+
+	template<typename T>
+	ESetParamResult SetParameter(const FName& Name, const FName& InTypeName, const T& InValue, bool OnlyIfExists = true)
+	{
+		T* TheParam = GetParameterPtr<T>(Name, InTypeName);
+		if (TheParam)
+		{
+			*TheParam = InValue;
+			return ESetParamResult::Succeeded;
+		}
+		else if (!OnlyIfExists)
+		{
+			// TODO Add test to see if we are allowed to create it!
+			AddParameter<T>(Name, InTypeName, InValue);
+			return ESetParamResult::Succeeded;
+		}
+		return ESetParamResult::Failed;
+	}
+
+	template<typename T>
+	bool HasParameter(const FName& Name, const FName& InTypeName)
+	{
+		if (!ParameterStorage.IsValid())
+		{
+			return false;
+		}
+		return ParameterStorage->FindParameter<T>(Name, InTypeName) != nullptr;
+	}
+
+	template<typename T>
+	void AddParameter(const FName& Name, const FName& InTypeName, const T& InValue)
+	{
+		if (!ParameterStorage.IsValid())
+		{
+			ParameterStorage = MakeShared<FMetasoundParameterPackStorage>();
+		}
+		FMetasoundParameterPackItem<T> NewParameter(Name, InTypeName, InValue);
+		ParameterStorage->AddParameter(NewParameter);
+	}
+
+	template<typename T>
+	T* GetParameterPtr(const FName& Name, const FName& InTypeName)
+	{
+		if (!ParameterStorage.IsValid())
+		{
+			return nullptr;
+		}
+		auto TheParameter = ParameterStorage->FindParameter<T>(Name, InTypeName);
+		if (!TheParameter)
+		{
+			return nullptr;
+		}
+		return &TheParameter->Value;
+	}
+
+	template<typename T>
+	T GetParameter(const FName& Name, const FName& InTypeName, ESetParamResult& Result)
+	{
+		T* TheParameter = GetParameterPtr<T>(Name, InTypeName);
+		if (!TheParameter)
+		{
+			Result = ESetParamResult::Failed;
+			return T();
+		}
+		Result = ESetParamResult::Succeeded;
+		return *TheParameter;
+	}
+
+	TUniquePtr<Audio::IProxyData> CreateNewProxyData(const Audio::FProxyDataInitParams& InitParams) override;
+	
+	// A couple of utilities for use by MetasoundAssetBase and MetasoundGenerator to set
+	// up the routing for parameter packs
+	static Metasound::FSendAddress CreateSendAddressFromEnvironment(const Metasound::FMetasoundEnvironment& InEnvironment);
+	static FMetasoundFrontendClassInput GetClassInput();
+
+	FSharedMetasoundParameterStoragePtr GetParameterStorage() const;
+	FSharedMetasoundParameterStoragePtr GetCopyOfParameterStorage() const;
+
+private:
+	TSharedPtr<FMetasoundParameterPackStorage> ParameterStorage;
+};
+
+// Here is the proxy that UMetasoundParameterPack creates when asked for a proxy by the audio system
+class METASOUNDFRONTEND_API FMetasoundParameterPackProxy : public Audio::TProxyData<FMetasoundParameterPackProxy>
+{
+public:
+	IMPL_AUDIOPROXY_CLASS(FMetasoundParameterPackProxy);
+
+	explicit FMetasoundParameterPackProxy(FSharedMetasoundParameterStoragePtr& Data)
+		: ParameterStoragePtr(Data)
+	{}
+
+	FMetasoundParameterPackProxy(const FMetasoundParameterPackProxy& Other) = default;
+
+	Audio::IProxyDataPtr Clone() const override
+	{
+		return MakeUnique<FMetasoundParameterPackProxy>(*this);
+	}
+
+	FSharedMetasoundParameterStoragePtr GetParamStorage()
+	{
+		return ParameterStoragePtr;
+	}
+
+private:
+	FSharedMetasoundParameterStoragePtr ParameterStoragePtr;
+};
+
+// And finally... A type we can register with Metasound that holds a TSharedPtr to an FMetasoundParamStorage instance.
+// It can be created by various systems in Metasound given a proxy.
+class METASOUNDFRONTEND_API FMetasoundParameterStorageWrapper
+{
+	FSharedMetasoundParameterStoragePtr ParameterStoragePtr;
+
+public:
+
+	FMetasoundParameterStorageWrapper() = default;
+	FMetasoundParameterStorageWrapper(const FMetasoundParameterStorageWrapper&) = default;
+	FMetasoundParameterStorageWrapper& operator=(const FMetasoundParameterStorageWrapper& Other) = default;
+	FMetasoundParameterStorageWrapper(const TUniquePtr<Audio::IProxyData>& InInitData)
+	{
+		if (InInitData.IsValid())
+		{
+			if (InInitData->CheckTypeCast<FMetasoundParameterPackProxy>())
+			{
+				// should we be getting handed a SharedPtr here?
+				FMetasoundParameterPackProxy& AsParameterPack = InInitData->GetAs<FMetasoundParameterPackProxy>();
+				ParameterStoragePtr = AsParameterPack.GetParamStorage();
+			}
+		}
+	}
+
+	FSharedMetasoundParameterStoragePtr Get() const
+	{
+		return ParameterStoragePtr;
+	}
+
+	bool IsPackValid() const
+	{
+		return ParameterStoragePtr.IsValid();
+	}
+
+	const FSharedMetasoundParameterStoragePtr& GetPackProxy() const
+	{
+		return ParameterStoragePtr;
+	}
+
+	const FMetasoundParameterPackStorage* operator->() const
+	{
+		return ParameterStoragePtr.Get();
+	}
+
+	FMetasoundParameterPackStorage* operator->()
+	{
+		return ParameterStoragePtr.Get();
+	}
+};
+
+DECLARE_METASOUND_DATA_REFERENCE_TYPES(FMetasoundParameterStorageWrapper, METASOUNDFRONTEND_API, FMetasoundParameterStorageWrapperTypeInfo, FMetasoundParameterStorageWrapperReadRef, FMetasoundParameterStorageWrapperWriteRef)
+
diff --git a/Engine/Plugins/Runtime/Metasound/Source/MetasoundFrontend/Public/MetasoundTrigger.h b/Engine/Plugins/Runtime/Metasound/Source/MetasoundFrontend/Public/MetasoundTrigger.h
--- a/Engine/Plugins/Runtime/Metasound/Source/MetasoundFrontend/Public/MetasoundTrigger.h
+++ b/Engine/Plugins/Runtime/Metasound/Source/MetasoundFrontend/Public/MetasoundTrigger.h
@@ -41,6 +41,12 @@ namespace Metasound
 			 */
 			FTrigger(const FOperatorSettings& InSettings);
 
+			/** For use when a Trigger request is found in a Parameter Pack
+			 *
+			 * @param ParamPackPayload - A pointer to a bool which should always be 'true'.
+			 */
+			void AssignRawParameter(const void* ParamPackPayload);
+
 			/** Trigger a specific frame in the future.
 			 *
 			 * @param InFrameToTrigger - Index of frame to trigger.
diff --git a/Engine/Plugins/Runtime/Metasound/Source/MetasoundGenerator/Private/MetasoundGenerator.cpp b/Engine/Plugins/Runtime/Metasound/Source/MetasoundGenerator/Private/MetasoundGenerator.cpp
--- a/Engine/Plugins/Runtime/Metasound/Source/MetasoundGenerator/Private/MetasoundGenerator.cpp
+++ b/Engine/Plugins/Runtime/Metasound/Source/MetasoundGenerator/Private/MetasoundGenerator.cpp
@@ -14,6 +14,7 @@
 #include "MetasoundTrace.h"
 #include "MetasoundTrigger.h"
 #include "MetasoundVertexData.h"
+#include "MetasoundFrontendDataTypeRegistry.h"
 
 
 namespace Metasound
@@ -87,11 +88,29 @@ namespace Metasound
 				FinishTrigger = VertexData.GetOutputs().GetOrConstructDataReadReference<FTrigger>(SourceOneShotInterface::Outputs::OnFinished, InitParams.OperatorSettings, false);
 			}
 
+			// Create the parameter setter map so parameter packs can be cracked
+			// open and distributed as appropriate...
+			TMap<FName, FParameterSetter> ParameterSetters;
+			FInputVertexInterfaceData& GraphInputs = VertexData.GetInputs();
+			for (auto InputIterator = GraphInputs.begin(); InputIterator != GraphInputs.end(); ++InputIterator)
+			{
+				const FInputDataVertex& InputVertex = (*InputIterator).GetVertex();
+				const Frontend::IParameterAssignmentFunction& Setter = IDataTypeRegistry::Get().GetRawAssignmentFunction(InputVertex.DataTypeName);
+				if (Setter)
+				{
+					FParameterSetter ParameterSetter(InputVertex.DataTypeName,
+						(*InputIterator).GetDataReference()->GetRaw(),
+						Setter);
+					ParameterSetters.Add(InputVertex.VertexName, ParameterSetter);
+				}
+			}
+
 			// Set data needed for graph
 			FMetasoundGeneratorData GeneratorData
 			{
 				InitParams.OperatorSettings,
 				MoveTemp(GraphOperator),
+				MoveTemp(ParameterSetters),
 				MoveTemp(GraphAnalyzer),
 				MoveTemp(OutputBuffers),
 				MoveTemp(PlayTrigger),
@@ -212,6 +231,10 @@ namespace Metasound
 		NumFramesPerExecute = InParams.OperatorSettings.GetNumFramesPerBlock();
 		NumSamplesPerExecute = NumChannels * NumFramesPerExecute;
 
+		// Create the routing for parameter packs
+		ParameterPackSendAddress = UMetasoundParameterPack::CreateSendAddressFromEnvironment(InParams.Environment);
+		ParameterPackReceiver = FDataTransmissionCenter::Get().RegisterNewReceiver<FMetasoundParameterStorageWrapper>(ParameterPackSendAddress, FReceiverInitParams{ InParams.OperatorSettings });
+
 		BuilderTask = MakeUnique<FBuilderTask>(this, MoveTemp(InParams), true /* bTriggerGenerator */);
 		
 		if (Metasound::ConsoleVariables::bEnableAsyncMetaSoundGeneratorBuilder)
@@ -236,6 +259,10 @@ namespace Metasound
 			BuilderTask->EnsureCompletion();
 			BuilderTask = nullptr;
 		}
+
+		// Remove routing for parameter packs
+		ParameterPackReceiver.Reset();
+		FDataTransmissionCenter::Get().UnregisterDataChannelIfUnconnected(ParameterPackSendAddress);
 	}
 
 	void FMetasoundGenerator::SetPendingGraph(FMetasoundGeneratorData&& InData, bool bTriggerGraph)
@@ -327,12 +354,19 @@ namespace Metasound
 
 		GraphAnalyzer = MoveTemp(InData->GraphAnalyzer);
 
+		ParameterSetters = MoveTemp(InData->ParameterSetters);
+
 		if (bTriggerGraph)
 		{
 			OnPlayTriggerRef->TriggerFrame(0);
 		}
 	}
 
+	void FMetasoundGenerator::QueueParameterPack(TSharedPtr<FMetasoundParameterPackStorage> ParameterPack)
+	{
+		ParameterPackQueue.Enqueue(ParameterPack);
+	}
+
 	int32 FMetasoundGenerator::GetNumChannels() const
 	{
 		return NumChannels;
@@ -384,6 +418,8 @@ namespace Metasound
 
 		while (NumSamplesRemaining > 0)
 		{
+			UnpackAndTransmitUpdatedParameters();
+
 			// Call metasound graph operator.
 			RootExecuter.Execute();
 
@@ -486,4 +522,41 @@ namespace Metasound
 		}
 		// TODO: memcpy for single channel. 
 	}
+
+	void FMetasoundGenerator::UnpackAndTransmitUpdatedParameters()
+	{
+		using namespace Frontend;
+
+		TUniqueFunction<void(FMetasoundParameterPackStorage*)> ProcessPack = [&](FMetasoundParameterPackStorage* Pack)
+			{
+				for (auto Walker = Pack->begin(); Walker != Pack->end(); ++Walker)
+				{
+					if (const FParameterSetter* ParameterSetter = ParameterSetters.Find(Walker->Name))
+					{
+						if (ParameterSetter->DataType == Walker->TypeName)
+						{
+							ParameterSetter->SetParameterWithPayload(Walker->GetPayload());
+						}
+					}
+				}
+			};
+
+		// First handle packs that have come from the IAudioParameterInterface system...
+		if (ParameterPackReceiver.IsValid())
+		{
+			FMetasoundParameterStorageWrapper Pack;
+			while (ParameterPackReceiver->CanPop())
+			{
+				ParameterPackReceiver->Pop(Pack);
+				ProcessPack(Pack.Get().Get());
+			}
+		}
+
+		// Now handle packs that came from QueueParameterPack...
+		TSharedPtr<FMetasoundParameterPackStorage> QueuedParameterPack;
+		while (ParameterPackQueue.Dequeue(QueuedParameterPack))
+		{
+			ProcessPack(QueuedParameterPack.Get());
+		}
+	}
 } 
diff --git a/Engine/Plugins/Runtime/Metasound/Source/MetasoundGenerator/Public/MetasoundGenerator.h b/Engine/Plugins/Runtime/Metasound/Source/MetasoundGenerator/Public/MetasoundGenerator.h
--- a/Engine/Plugins/Runtime/Metasound/Source/MetasoundGenerator/Public/MetasoundGenerator.h
+++ b/Engine/Plugins/Runtime/Metasound/Source/MetasoundGenerator/Public/MetasoundGenerator.h
@@ -11,6 +11,7 @@
 #include "MetasoundGraphOperator.h"
 #include "MetasoundOperatorBuilder.h"
 #include "MetasoundOperatorInterface.h"
+#include "MetasoundParameterPack.h"
 #include "MetasoundPrimitives.h"
 #include "MetasoundRouter.h"
 #include "MetasoundTrigger.h"
@@ -19,7 +20,6 @@
 #include "Sound/SoundGenerator.h"
 #include "Tickable.h"
 
-
 namespace Metasound
 {
 	// Struct needed for building the metasound graph
@@ -36,10 +36,28 @@ namespace Metasound
 		void Release();
 	};
 
+	// A struct that provides a method of pushing "raw" data from a parameter pack into a specific metasound input node.
+	struct FParameterSetter
+	{
+		FName DataType;
+		void* Destination;
+		const Frontend::IParameterAssignmentFunction& Setter;
+		FParameterSetter(FName InDataType, void* InDestination, const Frontend::IParameterAssignmentFunction& InSetter)
+			: DataType(InDataType)
+			, Destination(InDestination)
+			, Setter(InSetter)
+		{}
+		void SetParameterWithPayload(const void* ParameterPayload) const
+		{
+			Setter(ParameterPayload, Destination);
+		}
+	};
+
 	struct FMetasoundGeneratorData
 	{
 		FOperatorSettings OperatorSettings;
 		TUniquePtr<IOperator> GraphOperator;
+		TMap<FName, FParameterSetter> ParameterSetters;
 		TUniquePtr<Frontend::FGraphAnalyzer> GraphAnalyzer;
 		TArray<TDataReadReference<FAudioBuffer>> OutputBuffers;
 		FTriggerWriteRef TriggerOnPlayRef;
@@ -140,6 +158,8 @@ namespace Metasound
 			}
 		}
 
+		void QueueParameterPack(TSharedPtr<FMetasoundParameterPackStorage> ParameterPack);
+
 		/** Return the number of audio channels. */
 		int32 GetNumChannels() const;
 
@@ -172,6 +192,8 @@ namespace Metasound
 		// Metasound creates deinterleaved audio while sound generator requires interleaved audio.
 		void InterleaveGeneratedAudio();
 		
+		void UnpackAndTransmitUpdatedParameters();
+
 		FExecuter RootExecuter;
 
 		bool bIsGraphBuilding;
@@ -205,5 +227,20 @@ namespace Metasound
 		bool bIsWaitingForFirstGraph;
 
 		TUniquePtr<Frontend::FGraphAnalyzer> GraphAnalyzer;
+
+		// These next items are needed to provide a destination for the FAudioDevice, etc. to
+		// send parameter packs to. Every playing metasound will have a parameter destination
+		// that can accept parameter packs.
+		FSendAddress ParameterPackSendAddress;
+		TReceiverPtr<FMetasoundParameterStorageWrapper> ParameterPackReceiver;
+		
+		// This map provides setters for all of the input nodes in the metasound graph. 
+		// It is used when processing named parameters in a parameter pack.
+		TMap<FName, FParameterSetter> ParameterSetters;
+
+		// While parameter packs may arrive via the IAudioParameterInterface system,
+		// a faster method of sending parameters is via the QueueParameterPack function 
+		// and this queue.
+		TMpscQueue<TSharedPtr<FMetasoundParameterPackStorage>> ParameterPackQueue;
 	};
 }
diff --git a/Engine/Plugins/Runtime/Metasound/Source/MetasoundGraphCore/Public/MetasoundDataReference.h b/Engine/Plugins/Runtime/Metasound/Source/MetasoundGraphCore/Public/MetasoundDataReference.h
--- a/Engine/Plugins/Runtime/Metasound/Source/MetasoundGraphCore/Public/MetasoundDataReference.h
+++ b/Engine/Plugins/Runtime/Metasound/Source/MetasoundGraphCore/Public/MetasoundDataReference.h
@@ -123,6 +123,9 @@ namespace Metasound
 
 		/** Creates a copy of the parameter type. */
 		virtual TUniquePtr<IDataReference> Clone() const = 0;
+
+		/** provides a raw pointer to the storage where the data actually resides. */
+		virtual void* GetRaw() const = 0;
 	};
 	
 	/** Test if an IDataReference contains the same data type as the template
@@ -208,6 +211,12 @@ namespace Metasound
 			return GetMetasoundDataTypeId<DataType>();
 		}
 
+		/** Return a raw pointer to the data. */
+		virtual void* GetRaw() const override
+		{
+			return &ObjectReference.Get();
+		}
+
 	protected:
 
 		// Protected object reference is utilized by subclasses which define what
@@ -599,6 +608,13 @@ namespace Metasound
 			return DataRefPtr->Clone();
 		}
 
+		/** Return a raw pointer to the data. */
+		void* GetRaw() const
+		{
+			check(DataRefPtr.IsValid());
+			return DataRefPtr->GetRaw();
+		}
+
 		/** Returns the current value of a reference. 
 		 *
 		 * This method's behavior is undefined and an assert will be called if
