Fix non-unity build.

diff --git a/Engine/Plugins/Runtime/Metasound/Source/MetasoundGenerator/Public/MetasoundGenerator.h b/Engine/Plugins/Runtime/Metasound/Source/MetasoundGenerator/Public/MetasoundGenerator.h
--- a/Engine/Plugins/Runtime/Metasound/Source/MetasoundGenerator/Public/MetasoundGenerator.h
+++ b/Engine/Plugins/Runtime/Metasound/Source/MetasoundGenerator/Public/MetasoundGenerator.h
@@ -4,6 +4,7 @@
 #include "Analysis/MetasoundFrontendGraphAnalyzer.h"
 #include "Async/Async.h"
 #include "Async/AsyncWork.h"
+#include "Containers/MpscQueue.h"
 #include "DSP/Dsp.h"
 #include "MetasoundAudioFormats.h"
 #include "MetasoundExecutableOperator.h"
