Fix compile error caused by merge.

diff --git a/Engine/Plugins/Runtime/Metasound/Source/MetasoundFrontend/Public/MetasoundParameterPack.h b/Engine/Plugins/Runtime/Metasound/Source/MetasoundFrontend/Public/MetasoundParameterPack.h
--- a/Engine/Plugins/Runtime/Metasound/Source/MetasoundFrontend/Public/MetasoundParameterPack.h
+++ b/Engine/Plugins/Runtime/Metasound/Source/MetasoundFrontend/Public/MetasoundParameterPack.h
@@ -357,11 +357,6 @@ public:
 
 	FMetasoundParameterPackProxy(const FMetasoundParameterPackProxy& Other) = default;
 
-	Audio::IProxyDataPtr Clone() const override
-	{
-		return MakeUnique<FMetasoundParameterPackProxy>(*this);
-	}
-
 	FSharedMetasoundParameterStoragePtr GetParamStorage()
 	{
 		return ParameterStoragePtr;
