Fix hololens compile

diff --git a/Engine/Plugins/Runtime/Metasound/Source/MetasoundFrontend/Public/MetasoundParameterPack.h b/Engine/Plugins/Runtime/Metasound/Source/MetasoundFrontend/Public/MetasoundParameterPack.h
--- a/Engine/Plugins/Runtime/Metasound/Source/MetasoundFrontend/Public/MetasoundParameterPack.h
+++ b/Engine/Plugins/Runtime/Metasound/Source/MetasoundFrontend/Public/MetasoundParameterPack.h
@@ -150,7 +150,7 @@ struct FMetasoundParameterPackStorage
 			{
 				if (ParamWalker->TypeName == InTypeName)
 				{
-					return reinterpret_cast<FMetasoundParameterPackItem<T>*>(ParamWalker);
+					return static_cast<FMetasoundParameterPackItem<T>*>(ParamWalker);
 				}
 				return nullptr;
 			}
