CIS fix for deprecation warning in ue5/main with parameter pack

diff --git a/Engine/Plugins/Runtime/Metasound/Source/MetasoundFrontend/Private/MetasoundParameterPack.cpp b/Engine/Plugins/Runtime/Metasound/Source/MetasoundFrontend/Private/MetasoundParameterPack.cpp
--- a/Engine/Plugins/Runtime/Metasound/Source/MetasoundFrontend/Private/MetasoundParameterPack.cpp
+++ b/Engine/Plugins/Runtime/Metasound/Source/MetasoundFrontend/Private/MetasoundParameterPack.cpp
@@ -16,9 +16,9 @@ namespace MetasoundParameterPack
 	FName RoutingName("ParamPack");
 }
 
-TUniquePtr<Audio::IProxyData> UMetasoundParameterPack::CreateNewProxyData(const Audio::FProxyDataInitParams& InitParams)
+TSharedPtr<Audio::IProxyData> UMetasoundParameterPack::CreateProxyData(const Audio::FProxyDataInitParams& InitParams)
 {
-	return MakeUnique<FMetasoundParameterPackProxy>(ParameterStorage);
+	return MakeShared<FMetasoundParameterPackProxy>(ParameterStorage);
 }
 
 Metasound::FSendAddress UMetasoundParameterPack::CreateSendAddressFromEnvironment(const Metasound::FMetasoundEnvironment& InEnvironment)
diff --git a/Engine/Plugins/Runtime/Metasound/Source/MetasoundFrontend/Public/MetasoundParameterPack.h b/Engine/Plugins/Runtime/Metasound/Source/MetasoundFrontend/Public/MetasoundParameterPack.h
--- a/Engine/Plugins/Runtime/Metasound/Source/MetasoundFrontend/Public/MetasoundParameterPack.h
+++ b/Engine/Plugins/Runtime/Metasound/Source/MetasoundFrontend/Public/MetasoundParameterPack.h
@@ -331,7 +331,7 @@ public:
 		return *TheParameter;
 	}
 
-	TUniquePtr<Audio::IProxyData> CreateNewProxyData(const Audio::FProxyDataInitParams& InitParams) override;
+	TSharedPtr<Audio::IProxyData> CreateProxyData(const Audio::FProxyDataInitParams& InitParams) override;
 	
 	// A couple of utilities for use by MetasoundAssetBase and MetasoundGenerator to set
 	// up the routing for parameter packs
@@ -377,7 +377,7 @@ public:
 	FMetasoundParameterStorageWrapper() = default;
 	FMetasoundParameterStorageWrapper(const FMetasoundParameterStorageWrapper&) = default;
 	FMetasoundParameterStorageWrapper& operator=(const FMetasoundParameterStorageWrapper& Other) = default;
-	FMetasoundParameterStorageWrapper(const TUniquePtr<Audio::IProxyData>& InInitData)
+	FMetasoundParameterStorageWrapper(const TSharedPtr<Audio::IProxyData>& InInitData)
 	{
 		if (InInitData.IsValid())
 		{
