Fix MetaSound constructor pin icon not appearing properly for enum types

diff --git a/Engine/Plugins/Runtime/Metasound/Source/MetasoundEditor/Private/SMetasoundGraphEnumPin.cpp b/Engine/Plugins/Runtime/Metasound/Source/MetasoundEditor/Private/SMetasoundGraphEnumPin.cpp
--- a/Engine/Plugins/Runtime/Metasound/Source/MetasoundEditor/Private/SMetasoundGraphEnumPin.cpp
+++ b/Engine/Plugins/Runtime/Metasound/Source/MetasoundEditor/Private/SMetasoundGraphEnumPin.cpp
@@ -14,6 +14,7 @@ namespace Metasound
 		void SMetasoundGraphEnumPin::Construct(const FArguments& InArgs, UEdGraphPin* InGraphPinObj)
 		{
 			SGraphPin::Construct(SGraphPin::FArguments(), InGraphPinObj);
+			CacheAccessType();
 		}
 
 		TSharedRef<SWidget> SMetasoundGraphEnumPin::GetDefaultValueWidget()
