Fix double clicking on a MetaSound input node causing inputs to be reset to default, remove false positive warning, typo fixes

#rnx

diff --git a/Engine/Plugins/Runtime/Metasound/Source/MetasoundEditor/Private/MetasoundEditor.cpp b/Engine/Plugins/Runtime/Metasound/Source/MetasoundEditor/Private/MetasoundEditor.cpp
--- a/Engine/Plugins/Runtime/Metasound/Source/MetasoundEditor/Private/MetasoundEditor.cpp
+++ b/Engine/Plugins/Runtime/Metasound/Source/MetasoundEditor/Private/MetasoundEditor.cpp
@@ -1680,22 +1680,7 @@ namespace Metasound
 				return;
 			}
 
-			if (UMetasoundEditorGraphInputNode* InputNode = Cast<UMetasoundEditorGraphInputNode>(InNode))
-			{
-				if (IsPlaying())
-				{
-					if (UAudioComponent* PreviewComponent = GEditor->GetPreviewAudioComponent())
-					{
-						FConstNodeHandle NodeHandle = InputNode->GetConstNodeHandle();
-						NodeHandle->IterateConstInputs([ComponentInterface = PreviewComponent](FConstInputHandle Input)
-						{
-							FVertexName VertexName = Input->GetName();
-							ComponentInterface->SetTriggerParameter(VertexName);
-						});
-					}
-				}
-			}
-			else if (UMetasoundEditorGraphExternalNode* ExternalNode = Cast<UMetasoundEditorGraphExternalNode>(InNode))
+			if (UMetasoundEditorGraphExternalNode* ExternalNode = Cast<UMetasoundEditorGraphExternalNode>(InNode))
 			{
 				FConstNodeHandle NodeHandle = ExternalNode->GetConstNodeHandle();
 				FNodeRegistryKey Key = FMetasoundFrontendRegistryContainer::Get()->GetRegistryKey(NodeHandle->GetClassMetadata());
diff --git a/Engine/Plugins/Runtime/Metasound/Source/MetasoundEditor/Private/MetasoundEditorGraph.cpp b/Engine/Plugins/Runtime/Metasound/Source/MetasoundEditor/Private/MetasoundEditorGraph.cpp
--- a/Engine/Plugins/Runtime/Metasound/Source/MetasoundEditor/Private/MetasoundEditorGraph.cpp
+++ b/Engine/Plugins/Runtime/Metasound/Source/MetasoundEditor/Private/MetasoundEditorGraph.cpp
@@ -1932,7 +1932,7 @@ void UMetasoundEditorGraph::ValidateInternal(Metasound::Editor::FGraphValidation
 		// Validate there is only 1 editor node per guid 
 		// Input nodes are currently not 1:1 with their frontend representation
 		// but when they are, they can be validated here as well 
-		if (!Node->IsA<UMetasoundEditorGraphInputNode>())
+		if (!Node->IsA<UMetasoundEditorGraphInputNode>() && !Node->IsA<UMetasoundEditorGraphVariableNode>())
 		{
 			NodeGuids.Add(Node->GetNodeID(), &bNodeIdFound);
 			if (bNodeIdFound)
diff --git a/Engine/Plugins/Runtime/Metasound/Source/MetasoundStandardNodes/Private/MetasoundCrossfadeNode.cpp b/Engine/Plugins/Runtime/Metasound/Source/MetasoundStandardNodes/Private/MetasoundCrossfadeNode.cpp
--- a/Engine/Plugins/Runtime/Metasound/Source/MetasoundStandardNodes/Private/MetasoundCrossfadeNode.cpp
+++ b/Engine/Plugins/Runtime/Metasound/Source/MetasoundStandardNodes/Private/MetasoundCrossfadeNode.cpp
@@ -27,7 +27,7 @@ namespace Metasound
 	namespace CrossfadeVertexNames
 	{
 		METASOUND_PARAM(InputCrossfadeValue, "Crossfade Value", "Crossfade value to crossfade across inputs. Output will be the float value between adjacent whole number values.")
-		METASOUND_PARAM(OutputTrigger, "Out", "Triggered when all input triggers have been triggered. Call Reset to reset the state or use \"Auto Reset\"")
+		METASOUND_PARAM(OutputTrigger, "Out", "Output value.")
 
 		const FVertexName GetInputName(uint32 InIndex)
 		{
diff --git a/Engine/Plugins/Runtime/Metasound/Source/MetasoundStandardNodes/Private/MetasoundMapRangeNode.cpp b/Engine/Plugins/Runtime/Metasound/Source/MetasoundStandardNodes/Private/MetasoundMapRangeNode.cpp
--- a/Engine/Plugins/Runtime/Metasound/Source/MetasoundStandardNodes/Private/MetasoundMapRangeNode.cpp
+++ b/Engine/Plugins/Runtime/Metasound/Source/MetasoundStandardNodes/Private/MetasoundMapRangeNode.cpp
@@ -25,7 +25,7 @@ namespace Metasound
 		METASOUND_PARAM(InputOutRangeBName, "Out Range B", "The max output value range.");
 		METASOUND_PARAM(InputClampedName, "Clamped", "Whether or not to clamp the input to the specified input range.");
 
-		METASOUND_PARAM(OutputValueName, "Out Value", "Triggers when the envelope is triggered.");
+		METASOUND_PARAM(OutputValueName, "Out Value", "Mapped output value.");
 	}
 
 	namespace MetasoundMapRangeNodePrivate
