Metasound Multi-Oscillator Node

diff --git a/Engine/Plugins/Runtime/Metasound/Source/MetasoundStandardNodes/Private/MetasoundOscillatorNodes.cpp b/Engine/Plugins/Runtime/Metasound/Source/MetasoundStandardNodes/Private/MetasoundOscillatorNodes.cpp
--- a/Engine/Plugins/Runtime/Metasound/Source/MetasoundStandardNodes/Private/MetasoundOscillatorNodes.cpp
+++ b/Engine/Plugins/Runtime/Metasound/Source/MetasoundStandardNodes/Private/MetasoundOscillatorNodes.cpp
@@ -6,9 +6,11 @@
 #include "MetasoundAudioBuffer.h"
 #include "MetasoundEnumRegistrationMacro.h"
 #include "MetasoundExecutableOperator.h"
+#include "MetasoundEnumRegistrationMacro.h"
 #include "MetasoundFacade.h"
 #include "MetasoundNodeRegistrationMacro.h"
 #include "MetasoundDataTypeRegistrationMacro.h"
+#include "MetasoundOscillators.h"
 #include "MetasoundParamHelper.h"
 #include "MetasoundPrimitives.h"
 #include "MetasoundStandardNodesCategories.h"
@@ -23,512 +25,6 @@ namespace Metasound
 {			
 #pragma region Common
 
-	// Contained here are a small selection of block generators and sample generators.
-	// They are templatized to allow quick construction of new oscillators, but aren't considered
-	// optimal (as they contain many branches) and they should be only as a default 
-	// until Vectorized versions replace them.
-
-	namespace Generators
-	{
-		// Standard params passed to the Generate Block templates below
-		struct FGeneratorArgs
-		{
-			float SampleRate = 0.f;
-			float FrequencyHz = 0.f;
-			float GlideEaseFactor = 0.0f;
-			float PulseWidth = 0.f;
-			bool BiPolar = true;
-			TArrayView<float> AlignedBuffer;
-			TArrayView<const float> FM;
-		};
-
-		// Functor used when we need to wrap the phase of some of the oscillators. Sinf for example does not care
-		struct FWrapPhase
-		{
-			FORCEINLINE void operator()(float& InPhase) const
-			{
-				while (InPhase >= 1.f)
-				{
-					InPhase -= 1.f;
-				}
-
-				while (InPhase < 0.0f)
-				{
-					InPhase += 1.f;
-				}
-			}
-		};
-
-		// Functor that does nothing to the phase and lets it climb indefinitely 
-		struct FPhaseLetClimb
-		{
-			FORCEINLINE void operator()(const float&) const {};
-		};
-
-		// Smooth out the edges of the saw based on its current frequency
-		// using a polynomial to smooth it at the discontinuity. This 
-		// limits aliasing by avoiding the infinite frequency at the discontinuity.
-		FORCEINLINE float PolySmoothSaw(const float InPhase, const float InPhaseDelta)
-		{	
-			float Output = 0.0f;
-			float AbsolutePhaseDelta = FMath::Abs(InPhaseDelta);
-
-			// The current phase is on the left side of discontinuity
-			if (InPhase > 1.0f - AbsolutePhaseDelta)
-			{
-				const float Dist = (InPhase - 1.0f) / AbsolutePhaseDelta;
-				Output = -Dist * Dist - 2.0f * Dist - 1.0f;
-			}
-			// The current phase is on the right side of the discontinuity
-			else if (InPhase < AbsolutePhaseDelta)
-			{
-				// Distance into polynomial
-				const float Dist = InPhase / AbsolutePhaseDelta;
-				Output = Dist * Dist - 2.0f * Dist + 1.0f;
-			}
-			return Output;
-		}
-
-		// Functor that does the a generic generate block with a supplied Oscillator.
-		// The PhaseWrap functor controls how the phase is accumulated and wrapped
-		template<
-			typename Oscillator,
-			typename PhaseWrap = FWrapPhase
-		>
-		struct TGenerateBlock
-		{
-			float Phase = 0.f;
-			Oscillator Osc;
-			PhaseWrap Wrap;
-			float CurrentFade = 0.0f;
-			float FadeSmooth = 0.01f;
-			float CurrentFreq = -1.f;
-		
-			void operator()(const FGeneratorArgs& InArgs)
-			{
-				int32 RemainingSamplesInBlock = InArgs.AlignedBuffer.Num();
-				float* Out = InArgs.AlignedBuffer.GetData();
-				const float OneOverSampleRate = 1.f / InArgs.SampleRate;
-				
-				// TODO: break this into separate calls because there is a lot of code duplication to prevent per sample branching 
-				if (InArgs.BiPolar)
-				{
-					// Constant Freq between this an last update? 
-					if (FMath::IsNearlyEqual(InArgs.FrequencyHz, CurrentFreq) || CurrentFreq < 0.f || FMath::IsNearlyEqual(InArgs.GlideEaseFactor, 1.0f))
-					{
-						CurrentFreq = InArgs.FrequencyHz;
-						const float DeltaPhase = InArgs.FrequencyHz * OneOverSampleRate;
-						while (RemainingSamplesInBlock > 0)
-						{
-							Wrap(Phase);
-
-							float Output = Osc(Phase, DeltaPhase, InArgs);
-							Output *= CurrentFade;
-							*Out++ = Output;
-
-							CurrentFade = CurrentFade + FadeSmooth * (1.0f - CurrentFade);
-
-							Phase += DeltaPhase;
-
-							RemainingSamplesInBlock--;
-						}
-					}
-					else
-					{
-						while (RemainingSamplesInBlock > 0)
-						{
-							Wrap(Phase);
-
-							const float DeltaPhase = CurrentFreq * OneOverSampleRate;
-
-							float Output = Osc(Phase, DeltaPhase, InArgs);
-							Output *= CurrentFade;
-							*Out++ = Output;
-
-							CurrentFade = CurrentFade + FadeSmooth * (1.0f - CurrentFade);
-
-							Phase += DeltaPhase;
-
-							// lerp frequency based on the glide factor
-							CurrentFreq = CurrentFreq + InArgs.GlideEaseFactor * (InArgs.FrequencyHz - CurrentFreq);
-
-							RemainingSamplesInBlock--;
-						}
-					}
-				}
-				else // unipolar
-				{
-					// Constant Freq between this an last update? 
-					if (FMath::IsNearlyEqual(InArgs.FrequencyHz, CurrentFreq) || CurrentFreq < 0.f || FMath::IsNearlyEqual(InArgs.GlideEaseFactor, 1.0f))
-					{
-						CurrentFreq = InArgs.FrequencyHz;
-						const float DeltaPhase = InArgs.FrequencyHz * OneOverSampleRate;
-						while (RemainingSamplesInBlock > 0)
-						{
-							Wrap(Phase);
-
-							float Output = Osc(Phase, DeltaPhase, InArgs);
-							Output = Audio::GetUnipolar(Output);
-							Output *= CurrentFade;
-							*Out++ = Output;
-
-							CurrentFade = CurrentFade + FadeSmooth * (1.0f - CurrentFade);
-
-							Phase += DeltaPhase;
-
-							RemainingSamplesInBlock--;
-						}
-					}
-					else
-					{
-						while (RemainingSamplesInBlock > 0)
-						{
-							Wrap(Phase);
-
-							const float DeltaPhase = CurrentFreq * OneOverSampleRate;
-
-							float Output = Osc(Phase, DeltaPhase, InArgs);
-							Output = Audio::GetUnipolar(Output);
-							Output *= CurrentFade;
-							*Out++ = Output;
-
-							CurrentFade = CurrentFade + FadeSmooth * (1.0f - CurrentFade);
-
-							Phase += DeltaPhase;
-
-							// lerp frequency based on the glide factor
-							CurrentFreq = CurrentFreq + InArgs.GlideEaseFactor * (InArgs.FrequencyHz - CurrentFreq);
-
-							RemainingSamplesInBlock--;
-						}
-					}
-				}
-			}
-		};
-
-		// Functor that does the a generic generate block with a supplied Oscillator and applies an FM signal
-		// to the frequency per sample.
-		template<
-			typename Oscillator,
-			typename PhaseWrap = FWrapPhase
-		>
-		struct TGenerateBlockFM
-		{
-			float Phase = 0.f;
-			Oscillator Osc;
-			PhaseWrap Wrap;
-
-			float CurrentFade = 0.0f;
-			float FadeSmooth = 0.01f;
-			float CurrentFreq = -1.f;
-
-			void operator()(const FGeneratorArgs& InArgs)
-			{
-				float Nyquist = InArgs.SampleRate / 2.0f;
-				int32 RemainingSamplesInBlock = InArgs.AlignedBuffer.Num();
-				float* Out = InArgs.AlignedBuffer.GetData();
-				const float* FM = InArgs.FM.GetData();
-				const float OneOverSampleRate = 1.f / InArgs.SampleRate;
-
-				// TODO: break this into separate calls because there is a lot of code duplication to prevent per sample branching 
-				if (InArgs.BiPolar)
-				{
-					// Constant Base Freq between this and last update? 
-					if (FMath::IsNearlyEqual(InArgs.FrequencyHz, CurrentFreq) || CurrentFreq < 0.f || FMath::IsNearlyEqual(InArgs.GlideEaseFactor, 1.0f))
-					{
-						while (RemainingSamplesInBlock > 0)
-						{
-							const float PerSampleFreq = FMath::Clamp(InArgs.FrequencyHz + *FM++, -Nyquist, Nyquist);
-							const float DeltaPhase = PerSampleFreq * OneOverSampleRate;
-
-							Wrap(Phase);
-
-							float Output = Osc(Phase, DeltaPhase, InArgs);
-							Output *= CurrentFade;
-							*Out++ = Output;
-
-							CurrentFade = CurrentFade + FadeSmooth * (1.0f - CurrentFade);
-
-							Phase += DeltaPhase;
-
-							RemainingSamplesInBlock--;
-						}
-
-						CurrentFreq = InArgs.FrequencyHz;
-					}
-					else
-					{
-						while (RemainingSamplesInBlock > 0)
-						{
-							const float ModulatedFreqSum = FMath::Clamp(CurrentFreq + *FM++, -Nyquist, Nyquist);
-							const float DeltaPhase = ModulatedFreqSum * OneOverSampleRate;
-
-							Wrap(Phase);
-
-							float Output = Osc(Phase, DeltaPhase, InArgs);
-							Output *= CurrentFade;
-							*Out++ = Output;
-
-							CurrentFade = CurrentFade + FadeSmooth * (1.0f - CurrentFade);
-
-							Phase += DeltaPhase;
-
-							// lerp frequency based on the glide factor
-							CurrentFreq = CurrentFreq + InArgs.GlideEaseFactor * (InArgs.FrequencyHz - CurrentFreq);
-
-							RemainingSamplesInBlock--;
-						}
-					}
-				}
-				else // unipolar
-				{
-					// Constant Base Freq between this and last update? 
-					if (FMath::IsNearlyEqual(InArgs.FrequencyHz, CurrentFreq) || CurrentFreq < 0.f || FMath::IsNearlyEqual(InArgs.GlideEaseFactor, 1.0f))
-					{
-						while (RemainingSamplesInBlock > 0)
-						{
-							const float PerSampleFreq = InArgs.FrequencyHz + *FM++;
-							const float DeltaPhase = PerSampleFreq * OneOverSampleRate;
-
-							Wrap(Phase);
-
-							float Output = Osc(Phase, DeltaPhase, InArgs);
-							Output = Audio::GetUnipolar(Output);
-							Output *= CurrentFade;
-							*Out++ = Output;
-
-							CurrentFade = CurrentFade + FadeSmooth * (1.0f - CurrentFade);
-
-							Phase += DeltaPhase;
-
-							RemainingSamplesInBlock--;
-						}
-
-						CurrentFreq = InArgs.FrequencyHz;
-					}
-					else
-					{
-						while (RemainingSamplesInBlock > 0)
-						{
-							const float ModulatedFreqSum = CurrentFreq + *FM++;
-							const float DeltaPhase = ModulatedFreqSum * OneOverSampleRate;
-
-							Wrap(Phase);
-
-							float Output = Osc(Phase, DeltaPhase, InArgs);
-							Output = Audio::GetUnipolar(Output);
-							Output *= CurrentFade;
-							*Out++ = Output;
-
-							CurrentFade = CurrentFade + FadeSmooth * (1.0f - CurrentFade);
-
-							Phase += DeltaPhase;
-
-							// lerp frequency based on the glide factor
-							CurrentFreq = CurrentFreq + InArgs.GlideEaseFactor * (InArgs.FrequencyHz - CurrentFreq);
-
-							RemainingSamplesInBlock--;
-						}
-					}
-				}
-			}
-		};	
-		
-		// Sine types and generators.
-		struct FSinfGenerator
-		{
-			FORCEINLINE float operator()(float InPhase, float, const FGeneratorArgs&)
-			{
-				static constexpr float TwoPi = 2.f * PI;
-				return FMath::Sin(InPhase * TwoPi);
-			}
-		};
-		struct FBhaskaraGenerator
-		{
-			float operator()(float InPhase, float, const FGeneratorArgs&) const
-			{
-				static const float TwoPi = 2.f * PI;
-				float PhaseRadians = InPhase * TwoPi;
-
-				float InBhaskaraDomain = (PhaseRadians < 0) ? PhaseRadians + PI : PhaseRadians - PI;
-				return Audio::FastSin3(InBhaskaraDomain); // Expects [-PI, PI] 
-			}
-		};
-		struct FSineWaveTableGenerator
-		{
-			static const TArray<float>& GetWaveTable()
-			{
-				auto MakeSineTable = []() -> const TArray<float>
-				{
-					int32 TableSize = 4096;
-					// Generate the table
-					TArray<float> WaveTable;
-					WaveTable.AddUninitialized(TableSize);
-					float* WaveTableData = WaveTable.GetData();
-					for (int32 i = 0; i < TableSize; ++i)
-					{
-						static const float TwoPi = 2.f * PI;
-						float Phase = (float)i / TableSize;
-						WaveTableData[i] = FMath::Sin(Phase * TwoPi);
-					}
-					return WaveTable;
-				};
-
-				static const TArray<float> SineWaveTable = MakeSineTable();
-				return SineWaveTable;
-			}
-
-			float operator()(float InPhase, float, const FGeneratorArgs&) const
-			{			
-				const TArray<float>& WaveTable = GetWaveTable();
-				int32 LastIndex = WaveTable.Num() - 1;
-				int32 TableIndex = (int32)(InPhase * (float)(LastIndex));
-				TableIndex = FMath::Wrap(TableIndex, 0, LastIndex);
-				return WaveTable[TableIndex];
-			}
-		};
-
-		struct F2DRotatorGenerateBlock
-		{
-			Audio::FSinOsc2DRotation Rotator;
-			F2DRotatorGenerateBlock(float InStartingLinearPhase)
-				: Rotator{ 2.f * PI * InStartingLinearPhase }
-			{}
-			void operator()(const FGeneratorArgs& Args)
-			{
-				Rotator.GenerateBuffer(Args.SampleRate, Args.FrequencyHz, Args.AlignedBuffer.GetData(), Args.AlignedBuffer.Num());
-				if (!Args.BiPolar)
-				{
-					Audio::ConvertBipolarBufferToUnipolar(Args.AlignedBuffer.GetData(), Args.AlignedBuffer.Num());
-				}
-			}
-		};		
-		
-		using FSinf						= TGenerateBlock<FSinfGenerator>;
-		using FSinfWithFm				= TGenerateBlockFM<FSinfGenerator>;
-		
-		using FBhaskara					= TGenerateBlock<FBhaskaraGenerator>;
-		using FBhaskaraWithFm			= TGenerateBlockFM<FBhaskaraGenerator>;
-
-		using FSineWaveTable			= TGenerateBlock<FSineWaveTableGenerator>;
-		using FSineWaveTableWithFm		= TGenerateBlockFM<FSineWaveTableGenerator>;
-
-		// Saws.
-		struct FSawGenerator
-		{
-			FORCEINLINE float operator()(float InPhase, float, const FGeneratorArgs&)
-			{
-				return -1.f + (2.f * InPhase);
-			}
-		};
-		struct FSawPolySmoothGenerator
-		{
-			FORCEINLINE float operator()(float InPhase, float InPhaseDelta, const FGeneratorArgs&)
-			{
-				// Two-sided wave-shaped sawtooth
-				static const float A = Audio::FastTanh(1.5f);
-				float Output = Audio::GetBipolar(InPhase);
-				Output = Audio::FastTanh(1.5f * Output) / A;
-				Output += PolySmoothSaw(InPhase, InPhaseDelta);
-				return Output;
-			}
-		};
-
-		using FSaw							= TGenerateBlock<FSawGenerator>;
-		using FSawWithFm					= TGenerateBlockFM<FSawGenerator>;
-
-		using FSawPolysmooth				= TGenerateBlock<FSawPolySmoothGenerator>;
-		using FSawPolysmoothWithFm			= TGenerateBlockFM<FSawPolySmoothGenerator>;
-
-		// Square.
-		struct FSquareGenerator
-		{
-			FORCEINLINE float operator()(float InPhase, float InPhaseDelta, const FGeneratorArgs& InArgs)
-			{
-				float PulseWidthToUse = (InPhaseDelta >= 0) ? InArgs.PulseWidth : 1 - InArgs.PulseWidth;
-
-				return InPhase >= PulseWidthToUse ? 1.f : -1.f;
-			}
-		};
-		struct FSquarePolysmoothGenerator
-		{
-			float operator()(float InPhase, float InPhaseDelta, const FGeneratorArgs& InArgs)
-			{
-				// Taken piecemeal from Osc.cpp
-				// Lots of branches.
-
-				// First generate a smoothed sawtooth
-				float SquareSaw1 = Audio::GetBipolar(InPhase);
-				SquareSaw1 += PolySmoothSaw(InPhase, InPhaseDelta);
-
-				// Create a second sawtooth that is phase-shifted based on the pulsewidth
-				float NewPhase = 0.0f;
-				if (InPhaseDelta > 0.0f)
-				{
-					NewPhase = InPhase + InArgs.PulseWidth;
-					if (NewPhase >= 1.0f)
-					{
-						NewPhase -= 1.0f;
-					}
-				}
-				else
-				{
-					NewPhase = InPhase - InArgs.PulseWidth;
-					if (NewPhase <= 0.0f)
-					{
-						NewPhase += 1.0f;
-					}
-				}
-
-				float SquareSaw2 = Audio::GetBipolar(NewPhase);
-				SquareSaw2 += PolySmoothSaw(NewPhase, InPhaseDelta);
-
-				// Subtract 2 saws, then apply DC correction 
-				// Simplified version of 
-				// float Output = 0.5f * SquareSaw1 - 0.5f * SquareSaw2;
-				// Output = 2.0f * (Output + InArgs.PulseWidth) - 1.0f;
-				if (InPhaseDelta > 0.0)
-				{
-					return SquareSaw1 - SquareSaw2 + 2.0f * (InArgs.PulseWidth - 0.5f);
-				}
-				else
-				{
-					return SquareSaw1 - SquareSaw2 + 2.0f * (0.5f - InArgs.PulseWidth);
-				}
-				
-			}
-		};
-
-		using FSquare						= TGenerateBlock<FSquareGenerator>;
-		using FSquareWithFm					= TGenerateBlockFM<FSquareGenerator>;
-
-		using FSquarePolysmooth				= TGenerateBlock<FSquarePolysmoothGenerator>;
-		using FSquarePolysmoothWithFm		= TGenerateBlockFM<FSquarePolysmoothGenerator>;
-
-		// Triangle.
-		struct FTriangleGenerator
-		{
-			float TriangleSign = -1.0f;
-			float PrevPhase = -1.0f;
-			float DPW_z1 = 0.0f;
-
-			float operator()(float InPhase, float InPhaseDelta, const FGeneratorArgs& InArgs)
-			{
-				constexpr float FASTASIN_HALF_PI{ 1.5707963050f };
-				float PhaseRadians = (InPhase * PI * 2.f);
-				float Output = FMath::FastAsin(Audio::FastSin3(PhaseRadians - PI)) / FASTASIN_HALF_PI;
-				return Output;
-			}
-		};
-
-		using FTriangle						= TGenerateBlock<FTriangleGenerator>;
-		using FTriangleWithFm				= TGenerateBlockFM<FTriangleGenerator>;
-
-		// TODO: make a true polysmooth version
-		using FTrianglePolysmooth = TGenerateBlock<FTriangleGenerator>;
-		using FTrianglePolysmoothWithFm = TGenerateBlockFM<FTriangleGenerator>;
-	}
-
 	namespace OscillatorCommonVertexNames
 	{
 		METASOUND_PARAM(EnabledPin, "Enabled", "Enable the oscillator.")
@@ -1328,17 +824,6 @@ namespace Metasound
 
 #pragma region LFO
 
-	enum class ELfoWaveshapeType
-	{
-		Sine,
-		Saw,
-		Triangle,
-		Square,
-	};
-
-	DECLARE_METASOUND_ENUM(ELfoWaveshapeType, ELfoWaveshapeType::Sine, METASOUNDSTANDARDNODES_API,
-		FEnumLfoWaveshapeType, FEnumLfoWaveshapeTypeInfo, FEnumLfoWaveshapeTypeReadRef, FEnumLfoWaveshapeTypeWriteRef);
-
 	DEFINE_METASOUND_ENUM_BEGIN(ELfoWaveshapeType, FEnumLfoWaveshapeType, "LfoWaveshapeType")
 		DEFINE_METASOUND_ENUM_ENTRY(ELfoWaveshapeType::Sine, "LfoWaveShapeSineDescription", "Sine", "LfoWaveShapeSineDescriptionTT", "Sinewave Low Frequency Oscillator"),
 		DEFINE_METASOUND_ENUM_ENTRY(ELfoWaveshapeType::Saw, "LfoWaveShapeSawDescription", "Saw", "LfoWaveShapeSawDescriptionTT", "Sawtooth Low Frequency Oscillator"),
diff --git a/Engine/Plugins/Runtime/Metasound/Source/MetasoundStandardNodes/Private/MetasoundOscillators.h b/Engine/Plugins/Runtime/Metasound/Source/MetasoundStandardNodes/Private/MetasoundOscillators.h
new file mode 100644
--- /dev/null
+++ b/Engine/Plugins/Runtime/Metasound/Source/MetasoundStandardNodes/Private/MetasoundOscillators.h
@@ -0,0 +1,523 @@
+// Copyright Epic Games, Inc. All Rights Reserved.
+
+#pragma once
+
+#include "DSP/Dsp.h"
+#include "MetasoundAudioBuffer.h"
+#include "MetasoundEnumRegistrationMacro.h"
+#include "MetasoundPrimitives.h"
+
+namespace Metasound
+{
+	enum class ELfoWaveshapeType
+	{
+		Sine,
+		Saw,
+		Triangle,
+		Square,
+	};
+	
+	DECLARE_METASOUND_ENUM(ELfoWaveshapeType, ELfoWaveshapeType::Sine, METASOUNDSTANDARDNODES_API,
+		FEnumLfoWaveshapeType, FEnumLfoWaveshapeTypeInfo, FEnumLfoWaveshapeTypeReadRef, FEnumLfoWaveshapeTypeWriteRef);
+
+	// Contained here are a small selection of block generators and sample generators.
+	// They are templatized to allow quick construction of new oscillators, but aren't considered
+	// optimal (as they contain many branches) and they should be only as a default 
+	// until Vectorized versions replace them.
+	namespace Generators
+	{
+		// Standard params passed to the Generate Block templates below
+		struct FGeneratorArgs
+		{
+			float SampleRate = 0.f;
+			float FrequencyHz = 0.f;
+			float GlideEaseFactor = 0.0f;
+			float PulseWidth = 0.f;
+			bool BiPolar = true;
+			TArrayView<float> AlignedBuffer;
+			TArrayView<const float> FM;
+		};
+
+		// Functor used when we need to wrap the phase of some of the oscillators. Sinf for example does not care
+		struct FWrapPhase
+		{
+			FORCEINLINE void operator()(float& InPhase) const
+			{
+				while (InPhase >= 1.f)
+				{
+					InPhase -= 1.f;
+				}
+
+				while (InPhase < 0.0f)
+				{
+					InPhase += 1.f;
+				}
+			}
+		};
+
+		// Functor that does nothing to the phase and lets it climb indefinitely 
+		struct FPhaseLetClimb
+		{
+			FORCEINLINE void operator()(const float&) const {};
+		};
+
+		// Smooth out the edges of the saw based on its current frequency
+		// using a polynomial to smooth it at the discontinuity. This 
+		// limits aliasing by avoiding the infinite frequency at the discontinuity.
+		FORCEINLINE float PolySmoothSaw(const float InPhase, const float InPhaseDelta)
+		{	
+			float Output = 0.0f;
+			float AbsolutePhaseDelta = FMath::Abs(InPhaseDelta);
+
+			// The current phase is on the left side of discontinuity
+			if (InPhase > 1.0f - AbsolutePhaseDelta)
+			{
+				const float Dist = (InPhase - 1.0f) / AbsolutePhaseDelta;
+				Output = -Dist * Dist - 2.0f * Dist - 1.0f;
+			}
+			// The current phase is on the right side of the discontinuity
+			else if (InPhase < AbsolutePhaseDelta)
+			{
+				// Distance into polynomial
+				const float Dist = InPhase / AbsolutePhaseDelta;
+				Output = Dist * Dist - 2.0f * Dist + 1.0f;
+			}
+			
+			return Output;
+		}
+
+		// Functor that does the a generic generate block with a supplied Oscillator.
+		// The PhaseWrap functor controls how the phase is accumulated and wrapped
+		template<
+			typename Oscillator,
+			typename PhaseWrap = FWrapPhase
+		>
+		struct TGenerateBlock
+		{
+			float Phase = 0.f;
+			Oscillator Osc;
+			PhaseWrap Wrap;
+			float CurrentFade = 0.0f;
+			float FadeSmooth = 0.01f;
+			float CurrentFreq = -1.f;
+		
+			void operator()(const FGeneratorArgs& InArgs)
+			{
+				int32 RemainingSamplesInBlock = InArgs.AlignedBuffer.Num();
+				float* Out = InArgs.AlignedBuffer.GetData();
+				const float OneOverSampleRate = 1.f / InArgs.SampleRate;
+				
+				// TODO: break this into separate calls because there is a lot of code duplication to prevent per sample branching 
+				if (InArgs.BiPolar)
+				{
+					// Constant Freq between this an last update? 
+					if (FMath::IsNearlyEqual(InArgs.FrequencyHz, CurrentFreq) || CurrentFreq < 0.f || FMath::IsNearlyEqual(InArgs.GlideEaseFactor, 1.0f))
+					{
+						CurrentFreq = InArgs.FrequencyHz;
+						const float DeltaPhase = InArgs.FrequencyHz * OneOverSampleRate;
+						while (RemainingSamplesInBlock > 0)
+						{
+							Wrap(Phase);
+
+							float Output = Osc(Phase, DeltaPhase, InArgs);
+							Output *= CurrentFade;
+							*Out++ = Output;
+
+							CurrentFade = CurrentFade + FadeSmooth * (1.0f - CurrentFade);
+
+							Phase += DeltaPhase;
+
+							RemainingSamplesInBlock--;
+						}
+					}
+					else
+					{
+						while (RemainingSamplesInBlock > 0)
+						{
+							Wrap(Phase);
+
+							const float DeltaPhase = CurrentFreq * OneOverSampleRate;
+
+							float Output = Osc(Phase, DeltaPhase, InArgs);
+							Output *= CurrentFade;
+							*Out++ = Output;
+
+							CurrentFade = CurrentFade + FadeSmooth * (1.0f - CurrentFade);
+
+							Phase += DeltaPhase;
+
+							// lerp frequency based on the glide factor
+							CurrentFreq = CurrentFreq + InArgs.GlideEaseFactor * (InArgs.FrequencyHz - CurrentFreq);
+
+							RemainingSamplesInBlock--;
+						}
+					}
+				}
+				else // unipolar
+				{
+					// Constant Freq between this an last update? 
+					if (FMath::IsNearlyEqual(InArgs.FrequencyHz, CurrentFreq) || CurrentFreq < 0.f || FMath::IsNearlyEqual(InArgs.GlideEaseFactor, 1.0f))
+					{
+						CurrentFreq = InArgs.FrequencyHz;
+						const float DeltaPhase = InArgs.FrequencyHz * OneOverSampleRate;
+						while (RemainingSamplesInBlock > 0)
+						{
+							Wrap(Phase);
+
+							float Output = Osc(Phase, DeltaPhase, InArgs);
+							Output = Audio::GetUnipolar(Output);
+							Output *= CurrentFade;
+							*Out++ = Output;
+
+							CurrentFade = CurrentFade + FadeSmooth * (1.0f - CurrentFade);
+
+							Phase += DeltaPhase;
+
+							RemainingSamplesInBlock--;
+						}
+					}
+					else
+					{
+						while (RemainingSamplesInBlock > 0)
+						{
+							Wrap(Phase);
+
+							const float DeltaPhase = CurrentFreq * OneOverSampleRate;
+
+							float Output = Osc(Phase, DeltaPhase, InArgs);
+							Output = Audio::GetUnipolar(Output);
+							Output *= CurrentFade;
+							*Out++ = Output;
+
+							CurrentFade = CurrentFade + FadeSmooth * (1.0f - CurrentFade);
+
+							Phase += DeltaPhase;
+
+							// lerp frequency based on the glide factor
+							CurrentFreq = CurrentFreq + InArgs.GlideEaseFactor * (InArgs.FrequencyHz - CurrentFreq);
+
+							RemainingSamplesInBlock--;
+						}
+					}
+				}
+			}
+		};
+
+		// Functor that does the a generic generate block with a supplied Oscillator and applies an FM signal
+		// to the frequency per sample.
+		template<
+			typename Oscillator,
+			typename PhaseWrap = FWrapPhase
+		>
+		struct TGenerateBlockFM
+		{
+			float Phase = 0.f;
+			Oscillator Osc;
+			PhaseWrap Wrap;
+
+			float CurrentFade = 0.0f;
+			float FadeSmooth = 0.01f;
+			float CurrentFreq = -1.f;
+
+			void operator()(const FGeneratorArgs& InArgs)
+			{
+				float Nyquist = InArgs.SampleRate / 2.0f;
+				int32 RemainingSamplesInBlock = InArgs.AlignedBuffer.Num();
+				float* Out = InArgs.AlignedBuffer.GetData();
+				const float* FM = InArgs.FM.GetData();
+				const float OneOverSampleRate = 1.f / InArgs.SampleRate;
+
+				// TODO: break this into separate calls because there is a lot of code duplication to prevent per sample branching 
+				if (InArgs.BiPolar)
+				{
+					// Constant Base Freq between this and last update? 
+					if (FMath::IsNearlyEqual(InArgs.FrequencyHz, CurrentFreq) || CurrentFreq < 0.f || FMath::IsNearlyEqual(InArgs.GlideEaseFactor, 1.0f))
+					{
+						while (RemainingSamplesInBlock > 0)
+						{
+							const float PerSampleFreq = FMath::Clamp(InArgs.FrequencyHz + *FM++, -Nyquist, Nyquist);
+							const float DeltaPhase = PerSampleFreq * OneOverSampleRate;
+
+							Wrap(Phase);
+
+							float Output = Osc(Phase, DeltaPhase, InArgs);
+							Output *= CurrentFade;
+							*Out++ = Output;
+
+							CurrentFade = CurrentFade + FadeSmooth * (1.0f - CurrentFade);
+
+							Phase += DeltaPhase;
+
+							RemainingSamplesInBlock--;
+						}
+
+						CurrentFreq = InArgs.FrequencyHz;
+					}
+					else
+					{
+						while (RemainingSamplesInBlock > 0)
+						{
+							const float ModulatedFreqSum = FMath::Clamp(CurrentFreq + *FM++, -Nyquist, Nyquist);
+							const float DeltaPhase = ModulatedFreqSum * OneOverSampleRate;
+
+							Wrap(Phase);
+
+							float Output = Osc(Phase, DeltaPhase, InArgs);
+							Output *= CurrentFade;
+							*Out++ = Output;
+
+							CurrentFade = CurrentFade + FadeSmooth * (1.0f - CurrentFade);
+
+							Phase += DeltaPhase;
+
+							// lerp frequency based on the glide factor
+							CurrentFreq = CurrentFreq + InArgs.GlideEaseFactor * (InArgs.FrequencyHz - CurrentFreq);
+
+							RemainingSamplesInBlock--;
+						}
+					}
+				}
+				else // unipolar
+				{
+					// Constant Base Freq between this and last update? 
+					if (FMath::IsNearlyEqual(InArgs.FrequencyHz, CurrentFreq) || CurrentFreq < 0.f || FMath::IsNearlyEqual(InArgs.GlideEaseFactor, 1.0f))
+					{
+						while (RemainingSamplesInBlock > 0)
+						{
+							const float PerSampleFreq = InArgs.FrequencyHz + *FM++;
+							const float DeltaPhase = PerSampleFreq * OneOverSampleRate;
+
+							Wrap(Phase);
+
+							float Output = Osc(Phase, DeltaPhase, InArgs);
+							Output = Audio::GetUnipolar(Output);
+							Output *= CurrentFade;
+							*Out++ = Output;
+
+							CurrentFade = CurrentFade + FadeSmooth * (1.0f - CurrentFade);
+
+							Phase += DeltaPhase;
+
+							RemainingSamplesInBlock--;
+						}
+
+						CurrentFreq = InArgs.FrequencyHz;
+					}
+					else
+					{
+						while (RemainingSamplesInBlock > 0)
+						{
+							const float ModulatedFreqSum = CurrentFreq + *FM++;
+							const float DeltaPhase = ModulatedFreqSum * OneOverSampleRate;
+
+							Wrap(Phase);
+
+							float Output = Osc(Phase, DeltaPhase, InArgs);
+							Output = Audio::GetUnipolar(Output);
+							Output *= CurrentFade;
+							*Out++ = Output;
+
+							CurrentFade = CurrentFade + FadeSmooth * (1.0f - CurrentFade);
+
+							Phase += DeltaPhase;
+
+							// lerp frequency based on the glide factor
+							CurrentFreq = CurrentFreq + InArgs.GlideEaseFactor * (InArgs.FrequencyHz - CurrentFreq);
+
+							RemainingSamplesInBlock--;
+						}
+					}
+				}
+			}
+		};	
+		
+		// Sine types and generators.
+		struct FSinfGenerator
+		{
+			FORCEINLINE float operator()(float InPhase, float, const FGeneratorArgs&)
+			{
+				return FMath::Sin(InPhase * UE_TWO_PI);
+			}
+		};
+		struct FBhaskaraGenerator
+		{
+			FORCEINLINE float operator()(float InPhase, float, const FGeneratorArgs&) const
+			{
+				const float PhaseRadians = InPhase * UE_TWO_PI;
+
+				float InBhaskaraDomain = (PhaseRadians < 0) ? PhaseRadians + UE_PI : PhaseRadians - UE_PI;
+				return Audio::FastSin3(InBhaskaraDomain); // Expects [-PI, PI] 
+			}
+		};
+		struct FSineWaveTableGenerator
+		{
+			static const TArray<float>& GetWaveTable()
+			{
+				auto MakeSineTable = []() -> const TArray<float>
+				{
+					int32 TableSize = 4096;
+					// Generate the table
+					TArray<float> WaveTable;
+					WaveTable.AddUninitialized(TableSize);
+					float* WaveTableData = WaveTable.GetData();
+					for (int32 i = 0; i < TableSize; ++i)
+					{
+						float Phase = (float)i / TableSize;
+						WaveTableData[i] = FMath::Sin(Phase * UE_TWO_PI);
+					}
+					return WaveTable;
+				};
+
+				static const TArray<float> SineWaveTable = MakeSineTable();
+				return SineWaveTable;
+			}
+
+			float operator()(float InPhase, float, const FGeneratorArgs&) const
+			{
+				const TArray<float>& WaveTable = GetWaveTable();
+				int32 LastIndex = WaveTable.Num() - 1;
+				int32 TableIndex = (int32)(InPhase * (float)(LastIndex));
+				TableIndex = FMath::Wrap(TableIndex, 0, LastIndex);
+				return WaveTable[TableIndex];
+			}
+		};
+
+		struct F2DRotatorGenerateBlock
+		{
+			Audio::FSinOsc2DRotation Rotator;
+			F2DRotatorGenerateBlock(float InStartingLinearPhase)
+				: Rotator{ UE_TWO_PI * InStartingLinearPhase }
+			{}
+			void operator()(const FGeneratorArgs& Args)
+			{
+				Rotator.GenerateBuffer(Args.SampleRate, Args.FrequencyHz, Args.AlignedBuffer.GetData(), Args.AlignedBuffer.Num());
+				if (!Args.BiPolar)
+				{
+					Audio::ConvertBipolarBufferToUnipolar(Args.AlignedBuffer.GetData(), Args.AlignedBuffer.Num());
+				}
+			}
+		};		
+		
+		using FSinf						= TGenerateBlock<FSinfGenerator>;
+		using FSinfWithFm				= TGenerateBlockFM<FSinfGenerator>;
+		
+		using FBhaskara					= TGenerateBlock<FBhaskaraGenerator>;
+		using FBhaskaraWithFm			= TGenerateBlockFM<FBhaskaraGenerator>;
+
+		using FSineWaveTable			= TGenerateBlock<FSineWaveTableGenerator>;
+		using FSineWaveTableWithFm		= TGenerateBlockFM<FSineWaveTableGenerator>;
+
+		// Saws.
+		struct FSawGenerator
+		{
+			FORCEINLINE float operator()(float InPhase, float, const FGeneratorArgs&)
+			{
+				return -1.f + (2.f * InPhase);
+			}
+		};
+		struct FSawPolySmoothGenerator
+		{
+			FORCEINLINE float operator()(float InPhase, float InPhaseDelta, const FGeneratorArgs&)
+			{
+				// Two-sided wave-shaped sawtooth
+				constexpr float OneOverTanhOnePointFive = 1.10479139f;
+				float Output = Audio::GetBipolar(InPhase);
+				Output = Audio::FastTanh(1.5f * Output) * OneOverTanhOnePointFive;
+				Output += PolySmoothSaw(InPhase, InPhaseDelta);
+				return Output;
+			}
+		};
+
+		using FSaw							= TGenerateBlock<FSawGenerator>;
+		using FSawWithFm					= TGenerateBlockFM<FSawGenerator>;
+
+		using FSawPolysmooth				= TGenerateBlock<FSawPolySmoothGenerator>;
+		using FSawPolysmoothWithFm			= TGenerateBlockFM<FSawPolySmoothGenerator>;
+
+		// Square.
+		struct FSquareGenerator
+		{
+			FORCEINLINE float operator()(float InPhase, float InPhaseDelta, const FGeneratorArgs& InArgs)
+			{
+				const float PulseWidthToUse = (InPhaseDelta >= 0) ? InArgs.PulseWidth : 1 - InArgs.PulseWidth;
+                
+                return InPhase >= PulseWidthToUse ? 1.f : -1.f;
+			}
+		};
+		struct FSquarePolysmoothGenerator
+		{
+			float operator()(float InPhase, float InPhaseDelta, const FGeneratorArgs& InArgs)
+			{
+				// Taken piecemeal from Osc.cpp
+				// Lots of branches.
+
+				// First generate a smoothed sawtooth
+				float SquareSaw1 = Audio::GetBipolar(InPhase);
+				SquareSaw1 += PolySmoothSaw(InPhase, InPhaseDelta);
+
+				// Create a second sawtooth that is phase-shifted based on the pulsewidth
+				float NewPhase = 0.0f;
+				if (InPhaseDelta > 0.0f)
+				{
+					NewPhase = InPhase + InArgs.PulseWidth;
+					if (NewPhase >= 1.0f)
+					{
+						NewPhase -= 1.0f;
+					}
+				}
+				else
+				{
+					NewPhase = InPhase - InArgs.PulseWidth;
+					if (NewPhase <= 0.0f)
+					{
+						NewPhase += 1.0f;
+					}
+				}
+
+				float SquareSaw2 = Audio::GetBipolar(NewPhase);
+				SquareSaw2 += PolySmoothSaw(NewPhase, InPhaseDelta);
+
+				// Subtract 2 saws, then apply DC correction 
+				// Simplified version of 
+				// float Output = 0.5f * SquareSaw1 - 0.5f * SquareSaw2;
+				// Output = 2.0f * (Output + InArgs.PulseWidth) - 1.0f;
+				if (InPhaseDelta > 0.0)
+				{
+					return SquareSaw1 - SquareSaw2 + 2.0f * (InArgs.PulseWidth - 0.5f);
+				}
+				else
+				{
+					return SquareSaw1 - SquareSaw2 + 2.0f * (0.5f - InArgs.PulseWidth);
+				}
+			}
+		};
+
+		using FSquare						= TGenerateBlock<FSquareGenerator>;
+		using FSquareWithFm					= TGenerateBlockFM<FSquareGenerator>;
+
+		using FSquarePolysmooth				= TGenerateBlock<FSquarePolysmoothGenerator>;
+		using FSquarePolysmoothWithFm		= TGenerateBlockFM<FSquarePolysmoothGenerator>;
+
+		// Triangle.
+		struct FTriangleGenerator
+		{
+			float TriangleSign = -1.0f;
+			float PrevPhase = -1.0f;
+			float DPW_z1 = 0.0f;
+
+			FORCEINLINE float operator()(float InPhase, float InPhaseDelta, const FGeneratorArgs& InArgs)
+			{
+				constexpr float OneOverFastAsinHalfPi = 1.f / 1.5707963050f;
+				const float PhaseRadians = (InPhase * UE_TWO_PI);
+				return FMath::FastAsin(Audio::FastSin3(PhaseRadians - UE_PI)) * OneOverFastAsinHalfPi;
+			}
+		};
+
+		using FTriangle						= TGenerateBlock<FTriangleGenerator>;
+		using FTriangleWithFm				= TGenerateBlockFM<FTriangleGenerator>;
+
+		// TODO: make a true polysmooth version
+		using FTrianglePolysmooth = TGenerateBlock<FTriangleGenerator>;
+		using FTrianglePolysmoothWithFm = TGenerateBlockFM<FTriangleGenerator>;
+	}
+}
diff --git a/Engine/Plugins/Runtime/Metasound/Source/MetasoundStandardNodes/Private/MetasoundSuperOscillatorNode.cpp b/Engine/Plugins/Runtime/Metasound/Source/MetasoundStandardNodes/Private/MetasoundSuperOscillatorNode.cpp
new file mode 100644
--- /dev/null
+++ b/Engine/Plugins/Runtime/Metasound/Source/MetasoundStandardNodes/Private/MetasoundSuperOscillatorNode.cpp
@@ -0,0 +1,681 @@
+// Copyright Epic Games, Inc. All Rights Reserved.
+
+#include "Containers/Union.h"
+#include "Misc/TVariant.h"
+#include "DSP/Dsp.h"
+#include "DSP/DynamicsProcessor.h"
+#include "DSP/FloatArrayMath.h"
+#include "MetasoundAudioBuffer.h"
+#include "MetasoundDataTypeRegistrationMacro.h"
+#include "MetasoundEnumRegistrationMacro.h"
+#include "MetasoundExecutableOperator.h"
+#include "MetasoundFacade.h"
+#include "MetasoundNodeRegistrationMacro.h"
+#include "MetasoundOscillators.h"
+#include "MetasoundParamHelper.h"
+#include "MetasoundPrimitives.h"
+#include "MetasoundStandardNodesCategories.h"
+#include "MetasoundStandardNodesNames.h"
+#include "MetasoundVertex.h"
+
+#define LOCTEXT_NAMESPACE "MetasoundStandardNodes_SuperOscillatorNode"
+
+namespace Metasound
+{			
+	namespace Generators
+	{
+		// Standard params passed to the Generate Block templates below
+		struct FVoiceGeneratorArgs
+		{
+			float SampleRate = 0.f;
+			float FrequencyHz = 0.f;
+			float DetuneMultiplier = 0.f;
+			float Entropy = 0.f;
+			float GlideEaseFactor = 0.0f;
+			float PulseWidth = 0.f;
+			float BlendGain = 0.f;
+			float StereoWidth = 0.f;
+			TArrayView<float> AlignedBuffer;
+			TArrayView<const float> Fm;
+		};
+
+		struct FVoiceGeneratorInitArgs
+		{
+			const float SampleRate = 0.f;
+			const float Phase = 0.f;
+			const float DetuneRatio = 0.f;
+			const float PanRatio = 0.f;
+			const float Entropy = 0.f;
+		};
+		
+		// wraps a TGenerateBlock with extra members it needs to be an indiviual voice with its own pan/phase/pitch
+		template<typename GeneratorType>
+		struct TVoiceGenerator
+		{
+			// [-1, 1] - how strongly this voice should be detuned relative to the current max detune amount
+			float DetuneRatio = 0.f; 
+
+			// [-1, 1] - where this voice should be panned to if at full stereo width
+			float PanRatio = 0.f; 
+
+			// [-1, 1] - where on either side of our target values we should shift if at full entropy
+			float EntropyTarget = 0.f;
+			
+			GeneratorType Generator;
+			
+			TVoiceGenerator() = default;
+			
+			explicit TVoiceGenerator(const FVoiceGeneratorInitArgs& InitArgs)
+				: DetuneRatio(InitArgs.DetuneRatio)
+				, PanRatio(InitArgs.PanRatio)
+				, EntropyTarget(InitArgs.Entropy)
+			{
+				Generator.Phase = InitArgs.Phase;
+			}
+
+			void Process(const FVoiceGeneratorArgs& InArgs)
+			{
+				const float Entropy = 1.f + InArgs.Entropy * EntropyTarget;
+				const float FreqOffset = FMath::Pow(FMath::Pow(InArgs.DetuneMultiplier, Entropy), DetuneRatio);
+				const float TargetFreq = InArgs.FrequencyHz * FreqOffset;
+				
+				FGeneratorArgs Args =
+				{
+					InArgs.SampleRate,
+					TargetFreq,
+					InArgs.GlideEaseFactor,
+					InArgs.PulseWidth,
+					true,
+					InArgs.AlignedBuffer,
+					InArgs.Fm
+				};
+				
+				Generator(Args);
+			}
+		};
+			
+		using FSuperOsc	= TVariant<TVoiceGenerator<FSineWaveTable>, TVoiceGenerator<FSineWaveTableWithFm>,
+						TVoiceGenerator<FTrianglePolysmooth>,TVoiceGenerator<FTrianglePolysmoothWithFm>,
+						TVoiceGenerator<FSquarePolysmooth>, TVoiceGenerator<FSquarePolysmoothWithFm>,
+						TVoiceGenerator<FSawPolysmooth>, TVoiceGenerator<FSawPolysmoothWithFm>>;
+	
+		#define MAX_SUPEROSCILLATOR_VOICES 16
+		using FGeneratorArray = TArray<FSuperOsc, TInlineAllocator<MAX_SUPEROSCILLATOR_VOICES>>;
+
+		template<typename OscillatorType>
+		static void FillGeneratorsArray(FGeneratorArray& GeneratorArray, const int32 NumVoicesToInit, const float SampleRate)
+		{
+			GeneratorArray.Reset(NumVoicesToInit);
+
+			TArray<FVoiceGeneratorInitArgs, TInlineAllocator<MAX_SUPEROSCILLATOR_VOICES>> VoiceArgs;
+
+			VoiceArgs.Add({ SampleRate });
+
+			// alternate +/- offset values starting outside from 1 and working in
+			const float NumOffSetsNeeded = FMath::CeilToInt32((float)NumVoicesToInit * 0.5f);
+			for (int32 OffsetIdx = 0; OffsetIdx < NumOffSetsNeeded; OffsetIdx++)
+			{
+				const float Offset = (NumOffSetsNeeded - OffsetIdx) / NumOffSetsNeeded;
+
+				// flip pan every other voice to more evenly distribute +/- detunes to the left and right
+				const float PanSign = (OffsetIdx % 2 == 0) ? 1.f : -1.f;
+
+				FVoiceGeneratorInitArgs PosArgs =
+				{
+					SampleRate,
+					Offset,
+					Offset,
+					Offset * PanSign,
+					FMath::FRandRange(-1.f, 1.f)
+				};
+
+				VoiceArgs.Add(PosArgs);
+
+				FVoiceGeneratorInitArgs NegArgs =
+				{
+					SampleRate,
+					Offset + 0.5f,
+					-Offset,
+					-Offset * PanSign,
+					FMath::FRandRange(-1.f, 1.f)
+				};
+				
+				VoiceArgs.Add(NegArgs);
+			}
+
+			check(VoiceArgs.Num() >= NumVoicesToInit);
+			
+			for (int32 VoiceIdx = 0; VoiceIdx < NumVoicesToInit; VoiceIdx++)
+			{
+				FSuperOsc Osc = FSuperOsc(TInPlaceType<OscillatorType>(), VoiceArgs[VoiceIdx]);
+				
+				GeneratorArray.Add(Osc); 
+			}
+		}
+	}
+
+	namespace SuperOscillatorVertexNames
+	{
+		METASOUND_PARAM(AudioOutPin, "Audio", "The output audio")
+		METASOUND_PARAM(AudioLeftPin, "Left", "Left channel output audio")
+		METASOUND_PARAM(AudioRightPin, "Right", "Right channel output audio")
+
+		METASOUND_PARAM(BaseFrequencyPin, "Frequency", "Base Frequency of Oscillator in Hz.")
+		METASOUND_PARAM(DetuneVolumePin, "Blend", "Volume in decibels of detuned voices, relative to the primary voice")
+		METASOUND_PARAM(EnabledPin, "Enabled", "Enable the oscillator.")
+		METASOUND_PARAM(LimitOutputPin, "Limit Output", "Enables an internal limiter to keep output volume in check.")
+		METASOUND_PARAM(EntropyPin, "Entropy", "[0,1] Controls how evenly the voices are distributed in pitch")
+		METASOUND_PARAM(FrequencyModPin, "Modulation","Modulation Frequency Input (for doing FM)")
+		METASOUND_PARAM(GlideFactorPin, "Glide", "The amount of glide to use when changing frequencies. 0.0 = no glide, 1.0 = lots of glide.")
+		METASOUND_PARAM(MaxDetunePin, "Detune", "Max pitch offset of any Oscillators in Semitones. Only oscillators 2+ are detuned")
+		METASOUND_PARAM(NumVoicesPin, "Voices", "[1,16] The number of Oscillators")
+		METASOUND_PARAM(PulseWidthPin, "Pulse Width", "The Width of the square part of the wave. Only used for square waves.")
+		METASOUND_PARAM(StereoWidthPin, "Width", "[0,1] Stereo Width of the oscillators")
+		METASOUND_PARAM(WaveTypePin, "Type", "Shape of the Oscillator")
+	}
+
+	// Common set of construct params for each of the Oscillators.
+	struct FSuperOscillatorOperatorConstructParams
+	{
+		const FOperatorSettings& Settings;
+		FBoolReadRef Enabled;
+		FBoolReadRef Limiter;
+		int32 NumVoices;
+		FFloatReadRef BaseFrequency;
+		FFloatReadRef DetunedVoiceVolumeDb;
+		FFloatReadRef MaxDetuneSemitones;
+		FFloatReadRef Entropy;
+		FFloatReadRef GlideFactor;
+		FFloatReadRef PulseWidth;
+		bool bUseModulation;
+		FAudioBufferReadRef Fm;
+		FEnumLfoWaveshapeType OscType;
+	};
+	
+	class FSuperOscillatorOperator : public TExecutableOperator<FSuperOscillatorOperator>
+	{		
+	public:
+		FSuperOscillatorOperator(const FSuperOscillatorOperatorConstructParams& InConstructParams)
+			: SampleRate(InConstructParams.Settings.GetSampleRate())
+			, Nyquist(InConstructParams.Settings.GetSampleRate() / 2.0f)
+			, bEnabled(InConstructParams.Enabled)
+			, bLimit(InConstructParams.Limiter)
+			, NumVoices(FMath::Clamp(InConstructParams.NumVoices, 1, 16))
+			, BaseFrequency(InConstructParams.BaseFrequency)
+			, MaxDetune(InConstructParams.MaxDetuneSemitones)
+			, Entropy(InConstructParams.Entropy)
+			, DetuneDb(InConstructParams.DetunedVoiceVolumeDb)
+			, GlideFactor(InConstructParams.GlideFactor)
+			, PulseWidth(InConstructParams.PulseWidth)
+			, Fm(InConstructParams.Fm)
+			, OscType(InConstructParams.OscType)
+		{
+			using namespace Generators;
+
+			const int32 NumVoicesToInit = NumVoices;
+			ScratchBuffer.SetNumZeroed(InConstructParams.Settings.GetNumFramesPerBlock());
+			
+			// init generators
+			if (InConstructParams.bUseModulation)
+			{
+				// fm generators
+				switch (OscType)
+				{
+				case ELfoWaveshapeType::Sine:
+					FillGeneratorsArray<TVoiceGenerator<FSineWaveTableWithFm>>(Generators, NumVoicesToInit, SampleRate);
+					break;
+				case ELfoWaveshapeType::Triangle:
+					FillGeneratorsArray<TVoiceGenerator<FTrianglePolysmoothWithFm>>(Generators, NumVoicesToInit, SampleRate);
+					break;
+				case ELfoWaveshapeType::Square:
+					FillGeneratorsArray<TVoiceGenerator<FSquarePolysmoothWithFm>>(Generators, NumVoicesToInit, SampleRate);
+					break;
+				case ELfoWaveshapeType::Saw:
+					FillGeneratorsArray<TVoiceGenerator<FSawPolysmoothWithFm>>(Generators, NumVoicesToInit, SampleRate);
+					break;
+				default: 
+					check(!"Unsupported Wave Shape in FSuperOscillatorOperator");
+					return;
+				}
+			}
+			else
+			{
+				//non-fm generators
+				switch (OscType)
+				{
+				case ELfoWaveshapeType::Sine:
+					FillGeneratorsArray<TVoiceGenerator<FSineWaveTable>>(Generators, NumVoicesToInit, SampleRate);
+					break;
+				case ELfoWaveshapeType::Triangle:
+					FillGeneratorsArray<TVoiceGenerator<FTrianglePolysmooth>>(Generators, NumVoicesToInit, SampleRate);
+					break;
+				case ELfoWaveshapeType::Square:
+					FillGeneratorsArray<TVoiceGenerator<FSquarePolysmooth>>(Generators, NumVoicesToInit, SampleRate);
+					break;
+				case ELfoWaveshapeType::Saw:
+					FillGeneratorsArray<TVoiceGenerator<FSawPolysmooth>>(Generators, NumVoicesToInit, SampleRate);
+					break;
+				default: 
+					check(!"Unsupported Wave Shape in FSuperOscillatorOperator");
+					return;
+				}
+			}
+		}
+		
+		virtual void Bind(FVertexInterfaceData& InVertexData) const override
+		{
+			using namespace SuperOscillatorVertexNames; 
+
+			FInputVertexInterfaceData& Inputs = InVertexData.GetInputs();
+
+			Inputs.BindReadVertex(METASOUND_GET_PARAM_NAME(EnabledPin), bEnabled);
+			Inputs.BindReadVertex(METASOUND_GET_PARAM_NAME(LimitOutputPin), bLimit);
+			Inputs.BindReadVertex(METASOUND_GET_PARAM_NAME(BaseFrequencyPin), BaseFrequency);
+			Inputs.SetValue(METASOUND_GET_PARAM_NAME(NumVoicesPin), NumVoices);
+			Inputs.BindReadVertex(METASOUND_GET_PARAM_NAME(MaxDetunePin), MaxDetune);
+			Inputs.BindReadVertex(METASOUND_GET_PARAM_NAME(DetuneVolumePin), DetuneDb);
+			Inputs.BindReadVertex(METASOUND_GET_PARAM_NAME(GlideFactorPin), GlideFactor);
+			Inputs.BindReadVertex(METASOUND_GET_PARAM_NAME(PulseWidthPin), PulseWidth);
+			Inputs.BindReadVertex(METASOUND_GET_PARAM_NAME(FrequencyModPin), Fm);
+			Inputs.BindReadVertex(METASOUND_GET_PARAM_NAME(EntropyPin), Entropy);
+			Inputs.SetValue(METASOUND_GET_PARAM_NAME(WaveTypePin), OscType);
+		}
+		
+		virtual FDataReferenceCollection GetInputs() const override
+		{
+			checkNoEntry();
+
+			FDataReferenceCollection InputDataReferences;
+			return InputDataReferences;
+		}
+
+		virtual FDataReferenceCollection GetOutputs() const override
+		{
+			checkNoEntry();
+
+			FDataReferenceCollection OutputDataReferences;
+			return OutputDataReferences;
+		}
+
+		virtual void FillOutputs(const float Gain = 1.f, const float LinearPan = 0.f) {}
+		virtual void ZeroOutputs() {}
+		virtual float GetStereoWidth() const { return 0.f; }
+		
+		void SetLimiterSettings(Audio::FDynamicsProcessor& InLimiter) const
+		{
+			InLimiter.Init(SampleRate, 1);
+			InLimiter.SetLookaheadMsec(0.0f);
+			InLimiter.SetAttackTime(20.0f);
+			InLimiter.SetReleaseTime(200.0f);
+			InLimiter.SetThreshold(-6.0f);
+			InLimiter.SetRatio(10.f);
+			InLimiter.SetOutputGain(0.0f);
+			InLimiter.SetAnalogMode(true);
+			InLimiter.SetPeakMode(Audio::EPeakMode::Peak);
+			InLimiter.SetProcessingMode(Audio::EDynamicsProcessingMode::Compressor);
+		}
+
+		virtual void LimitOutput() {}
+		
+		void Execute()
+		{
+			using namespace Generators;
+			
+			if (Generators.IsEmpty())
+			{
+				return;
+			}
+			
+			const float ClampedFreq = FMath::Clamp(*BaseFrequency, -Nyquist, Nyquist);
+			const float ClampedEntropy = FMath::Clamp(*Entropy, 0.f, 1.f);
+			const float ClampedGlideEase = Audio::GetLogFrequencyClamped(*GlideFactor, { 0.0f, 1.0f }, { 1.0f, 0.0001f });
+			const float DetuneRatio = Audio::GetFrequencyMultiplier(*MaxDetune);
+			const float DetuneChannelGain = Audio::ConvertToLinear(*DetuneDb);
+			const float ClampedPulsueWidth = FMath::Clamp(*PulseWidth, 0.01f, 0.99f);
+			
+			const int32 NumFrames = ScratchBuffer.Num();
+			ScratchBuffer.SetNumZeroed(ScratchBuffer.Num());
+
+			ZeroOutputs();
+			
+			if (*bEnabled && NumFrames > 0)
+			{
+				const FVoiceGeneratorArgs Args = 
+				{
+					SampleRate,
+					ClampedFreq,
+					DetuneRatio,
+					ClampedEntropy,
+					ClampedGlideEase,
+					ClampedPulsueWidth,
+					DetuneChannelGain,
+					GetStereoWidth(),
+					MakeArrayView(ScratchBuffer.GetData(), NumFrames),
+					MakeArrayView(Fm->GetData(), NumFrames)
+				};
+
+				const auto GenerateFunc = [&Args](auto& Osc) -> float
+				{
+					Osc.Process(Args);
+					return Osc.PanRatio;
+				};
+
+				bool bFirstVoiceProcessed = false;
+
+				for (FSuperOsc& Generator : Generators)
+				{
+					const float Gain = bFirstVoiceProcessed ? Args.BlendGain : 1.f;
+					const float PanRatio = Visit(GenerateFunc, Generator);
+
+					FillOutputs(Gain, PanRatio * Args.StereoWidth);
+
+					bFirstVoiceProcessed = true;
+
+					if (*bLimit)
+					{
+						LimitOutput();
+					}
+				}
+			}
+		}
+
+		float SampleRate = 1.f;
+		float Nyquist = 1.f;
+
+		// apply this gain whether we are limiting or not, to avoid large volume jumps when toggling the limiter
+		const float PreLimiterGain = 0.707f;
+
+		FBoolReadRef bEnabled;
+		FBoolReadRef bLimit;
+		int32 NumVoices = 0;
+		FFloatReadRef BaseFrequency;
+		FFloatReadRef MaxDetune;
+		FFloatReadRef Entropy;
+		FFloatReadRef DetuneDb;
+		FFloatReadRef GlideFactor;
+		FFloatReadRef PulseWidth;
+
+		FAudioBufferReadRef Fm;
+		FEnumLfoWaveshapeType OscType;
+
+		Audio::FAlignedFloatBuffer ScratchBuffer;
+
+		TArray<Generators::FSuperOsc, TInlineAllocator<MAX_SUPEROSCILLATOR_VOICES>> Generators;
+	};
+
+	class FSuperOscillatorOperatorMono : public FSuperOscillatorOperator
+	{
+		using Super = FSuperOscillatorOperator;
+	public:
+		FSuperOscillatorOperatorMono(const FSuperOscillatorOperatorConstructParams& InConstructParams)
+			: FSuperOscillatorOperator(InConstructParams)
+			, AudioBuffer(FAudioBufferWriteRef::CreateNew(InConstructParams.Settings))
+		{
+			SetLimiterSettings(Limiter);
+		}
+		
+		static const FNodeClassMetadata& GetNodeInfo()
+		{
+			auto InitNodeInfo = []() -> FNodeClassMetadata
+			{
+				FNodeClassMetadata Info = FNodeClassMetadata::GetEmpty();
+				Info.ClassName = { StandardNodes::Namespace, TEXT("SuperOscillatorMono"), StandardNodes::AudioVariant };
+				Info.MajorVersion = 1;
+				Info.MinorVersion = 0;
+				Info.DisplayName = METASOUND_LOCTEXT("Metasound_SuperOscillatorMonoNodeDisplayName", "SuperOscillator (Mono)");
+				Info.Description = METASOUND_LOCTEXT("Metasound_SuperOscillatorNodeDescription", "Generates audio using multiple internal oscillators.");
+				Info.Author = PluginAuthor;
+				Info.PromptIfMissing = PluginNodeMissingPrompt;
+				Info.DefaultInterface = GetVertexInterface();
+				Info.CategoryHierarchy.Emplace(NodeCategories::Generators);
+				Info.Keywords = { METASOUND_LOCTEXT("OscKeyword", "Osc"), METASOUND_LOCTEXT("FMKeyword", "FM"), METASOUND_LOCTEXT("SynthesisKeyword", "Synthesis") };
+				return Info;
+			};
+			static const FNodeClassMetadata Info = InitNodeInfo();
+			return Info;
+		}
+		
+		static const FVertexInterface& GetVertexInterface()
+		{
+			using namespace SuperOscillatorVertexNames; 
+
+			static const FVertexInterface Interface
+			{
+				FInputVertexInterface{
+					TInputDataVertex<bool>(METASOUND_GET_PARAM_NAME_AND_METADATA(EnabledPin), true),
+					TInputDataVertex<bool>(METASOUND_GET_PARAM_NAME_AND_METADATA(LimitOutputPin), true),
+					TInputConstructorVertex<int32>(METASOUND_GET_PARAM_NAME_AND_METADATA(NumVoicesPin), 3),
+					TInputDataVertex<float>(METASOUND_GET_PARAM_NAME_AND_METADATA(BaseFrequencyPin), 440.f),
+					TInputDataVertex<FAudioBuffer>(METASOUND_GET_PARAM_NAME_AND_METADATA(FrequencyModPin)),
+					TInputDataVertex<float>(METASOUND_GET_PARAM_NAME_AND_METADATA(MaxDetunePin), -0.25f),
+					TInputDataVertex<float>(METASOUND_GET_PARAM_NAME_AND_METADATA(EntropyPin), 0.f),
+					TInputDataVertex<float>(METASOUND_GET_PARAM_NAME_AND_METADATA(DetuneVolumePin), 0.f),
+					TInputDataVertex<float>(METASOUND_GET_PARAM_NAME_AND_METADATA(GlideFactorPin), 0.f),
+					TInputDataVertex<float>(METASOUND_GET_PARAM_NAME_AND_METADATA(PulseWidthPin), 0.5f),
+					TInputConstructorVertex<FEnumLfoWaveshapeType>(METASOUND_GET_PARAM_NAME_AND_METADATA(WaveTypePin))
+				},
+				FOutputVertexInterface{
+					TOutputDataVertex<FAudioBuffer>(METASOUND_GET_PARAM_NAME_AND_METADATA(AudioOutPin))
+				}
+			};
+			
+			return Interface;
+		}
+		
+		static TUniquePtr<IOperator> CreateOperator(const FBuildOperatorParams& InParams, FBuildResults& OutResults)
+		{
+			const FInputVertexInterfaceData& InputData = InParams.InputData;
+
+			const FOperatorSettings& Settings = InParams.OperatorSettings;
+
+			using namespace SuperOscillatorVertexNames; 
+			FSuperOscillatorOperatorConstructParams OpParams
+			{
+				Settings,
+				InputData.GetOrConstructDataReadReference<bool>(METASOUND_GET_PARAM_NAME(EnabledPin), true),
+				InputData.GetOrConstructDataReadReference<bool>(METASOUND_GET_PARAM_NAME(LimitOutputPin), true),
+				InputData.GetOrCreateDefaultValue<int32>(METASOUND_GET_PARAM_NAME(NumVoicesPin), Settings),
+				InputData.GetOrConstructDataReadReference<float>(METASOUND_GET_PARAM_NAME(BaseFrequencyPin), 440.f),
+				InputData.GetOrConstructDataReadReference<float>(METASOUND_GET_PARAM_NAME(DetuneVolumePin), 0.f),
+				InputData.GetOrConstructDataReadReference<float>(METASOUND_GET_PARAM_NAME(MaxDetunePin), -0.25f),
+				InputData.GetOrConstructDataReadReference<float>(METASOUND_GET_PARAM_NAME(EntropyPin), 0.0f),
+				InputData.GetOrConstructDataReadReference<float>(METASOUND_GET_PARAM_NAME(GlideFactorPin), 0.f),
+				InputData.GetOrConstructDataReadReference<float>(METASOUND_GET_PARAM_NAME(PulseWidthPin), 0.5f),
+				InputData.IsVertexBound(METASOUND_GET_PARAM_NAME(FrequencyModPin)),
+				InputData.GetOrConstructDataReadReference<FAudioBuffer>(METASOUND_GET_PARAM_NAME(FrequencyModPin), Settings),
+				InputData.GetOrCreateDefaultValue<FEnumLfoWaveshapeType>(METASOUND_GET_PARAM_NAME(WaveTypePin), Settings)
+			};
+			
+			return MakeUnique<FSuperOscillatorOperatorMono>(OpParams);
+		}
+
+		virtual void ZeroOutputs() override
+		{
+			AudioBuffer->Zero();
+		}
+		
+		virtual void FillOutputs(const float Gain = 1.f, const float LinearPan = 0.f) final override
+		{
+			Audio::ArrayMixIn(ScratchBuffer, *AudioBuffer, Gain * PreLimiterGain);
+		}
+
+		virtual void LimitOutput() override
+		{
+			Limiter.ProcessAudio(AudioBuffer->GetData(), AudioBuffer->Num(), AudioBuffer->GetData());
+		}
+
+		virtual void Bind(FVertexInterfaceData& InVertexData) const override
+		{
+			using namespace SuperOscillatorVertexNames;
+			Super::Bind(InVertexData);
+			
+			FOutputVertexInterfaceData& Outputs = InVertexData.GetOutputs();
+			Outputs.BindReadVertex(METASOUND_GET_PARAM_NAME(AudioOutPin), AudioBuffer);
+		}
+		
+		FAudioBufferWriteRef AudioBuffer;
+		
+		Audio::FDynamicsProcessor Limiter;
+	};
+	
+	class FSuperOscillatorOperatorStereo : public FSuperOscillatorOperator
+	{
+		using Super = FSuperOscillatorOperator;
+	public:
+		FSuperOscillatorOperatorStereo(const FSuperOscillatorOperatorConstructParams& InConstructParams, FFloatReadRef InStereoWidth)
+			: FSuperOscillatorOperator(InConstructParams)
+			, StereoWidth(InStereoWidth)
+			, AudioLeft(FAudioBufferWriteRef::CreateNew(InConstructParams.Settings))
+			, AudioRight(FAudioBufferWriteRef::CreateNew(InConstructParams.Settings))
+		{
+			SetLimiterSettings(LimiterLeft);
+			SetLimiterSettings(LimiterRight);
+		}
+
+		static const FNodeClassMetadata& GetNodeInfo()
+		{
+			auto InitNodeInfo = []() -> FNodeClassMetadata
+			{
+				FNodeClassMetadata Info = FNodeClassMetadata::GetEmpty();
+				Info.ClassName = { StandardNodes::Namespace, TEXT("SuperOscillatorStereo"), StandardNodes::AudioVariant };
+				Info.MajorVersion = 1;
+				Info.MinorVersion = 0;
+				Info.DisplayName = METASOUND_LOCTEXT("Metasound_SuperOscillatorNodeStereoDisplayName", "SuperOscillator (Stereo)");
+				Info.Description = METASOUND_LOCTEXT("Metasound_SuperOscillatorNodeDescription", "Generates audio using multiple internal oscillators.");
+				Info.Author = PluginAuthor;
+				Info.PromptIfMissing = PluginNodeMissingPrompt;
+				Info.DefaultInterface = GetVertexInterface();
+				Info.CategoryHierarchy.Emplace(NodeCategories::Generators);
+				Info.Keywords = { METASOUND_LOCTEXT("OscKeyword", "Osc"), METASOUND_LOCTEXT("FMKeyword", "FM"), METASOUND_LOCTEXT("SynthesisKeyword", "Synthesis") };
+				return Info;
+			};
+			static const FNodeClassMetadata Info = InitNodeInfo();
+			return Info;
+		}
+		
+		static const FVertexInterface& GetVertexInterface()
+		{
+			using namespace SuperOscillatorVertexNames; 
+
+			static const FVertexInterface Interface
+			{
+				FInputVertexInterface{
+					TInputDataVertex<bool>(METASOUND_GET_PARAM_NAME_AND_METADATA(EnabledPin), true),
+					TInputDataVertex<bool>(METASOUND_GET_PARAM_NAME_AND_METADATA(LimitOutputPin), true),
+					TInputConstructorVertex<int32>(METASOUND_GET_PARAM_NAME_AND_METADATA(NumVoicesPin), 3),
+					TInputDataVertex<float>(METASOUND_GET_PARAM_NAME_AND_METADATA(BaseFrequencyPin), 440.f),
+					TInputDataVertex<FAudioBuffer>(METASOUND_GET_PARAM_NAME_AND_METADATA(FrequencyModPin)),
+					TInputDataVertex<float>(METASOUND_GET_PARAM_NAME_AND_METADATA(MaxDetunePin), -0.25f),
+					TInputDataVertex<float>(METASOUND_GET_PARAM_NAME_AND_METADATA(EntropyPin), 0.f),
+					TInputDataVertex<float>(METASOUND_GET_PARAM_NAME_AND_METADATA(DetuneVolumePin), 0.f),
+					TInputDataVertex<float>(METASOUND_GET_PARAM_NAME_AND_METADATA(GlideFactorPin), 0.f),
+					TInputDataVertex<float>(METASOUND_GET_PARAM_NAME_AND_METADATA(PulseWidthPin), 0.5f),
+					TInputDataVertex<float>(METASOUND_GET_PARAM_NAME_AND_METADATA(StereoWidthPin), 0.5f),
+					TInputConstructorVertex<FEnumLfoWaveshapeType>(METASOUND_GET_PARAM_NAME_AND_METADATA(WaveTypePin))
+				},
+				FOutputVertexInterface{
+					TOutputDataVertex<FAudioBuffer>(METASOUND_GET_PARAM_NAME_AND_METADATA(AudioLeftPin)),
+					TOutputDataVertex<FAudioBuffer>(METASOUND_GET_PARAM_NAME_AND_METADATA(AudioRightPin))
+				}
+			};
+			
+			return Interface;
+		}
+
+		static TUniquePtr<IOperator> CreateOperator(const FBuildOperatorParams& InParams, FBuildResults& OutResults)
+		{
+			const FInputVertexInterfaceData& InputData = InParams.InputData;
+
+			const FOperatorSettings& Settings = InParams.OperatorSettings;
+
+			using namespace SuperOscillatorVertexNames; 
+			FSuperOscillatorOperatorConstructParams OpParams
+			{
+				Settings,
+				InputData.GetOrConstructDataReadReference<bool>(METASOUND_GET_PARAM_NAME(EnabledPin), true),
+				InputData.GetOrConstructDataReadReference<bool>(METASOUND_GET_PARAM_NAME(LimitOutputPin), true),
+				InputData.GetOrCreateDefaultValue<int32>(METASOUND_GET_PARAM_NAME(NumVoicesPin), Settings),
+				InputData.GetOrConstructDataReadReference<float>(METASOUND_GET_PARAM_NAME(BaseFrequencyPin), 440.f),
+				InputData.GetOrConstructDataReadReference<float>(METASOUND_GET_PARAM_NAME(DetuneVolumePin), 0.f),
+				InputData.GetOrConstructDataReadReference<float>(METASOUND_GET_PARAM_NAME(MaxDetunePin), -0.25f),
+				InputData.GetOrConstructDataReadReference<float>(METASOUND_GET_PARAM_NAME(EntropyPin), 0.f),
+				InputData.GetOrConstructDataReadReference<float>(METASOUND_GET_PARAM_NAME(GlideFactorPin), 0.f),
+				InputData.GetOrConstructDataReadReference<float>(METASOUND_GET_PARAM_NAME(PulseWidthPin), 0.5f),
+				InputData.IsVertexBound(METASOUND_GET_PARAM_NAME(FrequencyModPin)),
+				InputData.GetOrConstructDataReadReference<FAudioBuffer>(METASOUND_GET_PARAM_NAME(FrequencyModPin), Settings),
+				InputData.GetOrCreateDefaultValue<FEnumLfoWaveshapeType>(METASOUND_GET_PARAM_NAME(WaveTypePin), Settings)
+			};
+
+			FFloatReadRef StereoWidthInput = InputData.GetOrConstructDataReadReference<float>(METASOUND_GET_PARAM_NAME(StereoWidthPin), 0.5f);
+			
+			return MakeUnique<FSuperOscillatorOperatorStereo>(OpParams, StereoWidthInput);
+		}
+
+		virtual void ZeroOutputs() override
+		{
+			AudioLeft->Zero();
+			AudioRight->Zero();
+		}
+
+		virtual float GetStereoWidth() const override { return *StereoWidth; }
+		
+		virtual void FillOutputs(const float Gain = 1.f, const float LinearPan = 0.f) final override
+		{
+			float PanLeftGain;
+			float PanRightGain;
+			
+			Audio::GetStereoPan(LinearPan, PanLeftGain, PanRightGain);
+
+			Audio::ArrayMixIn(ScratchBuffer, *AudioLeft, Gain * PanLeftGain * PreLimiterGain);
+			Audio::ArrayMixIn(ScratchBuffer, *AudioRight, Gain * PanRightGain * PreLimiterGain);
+		}
+
+		virtual void LimitOutput() override
+		{
+			LimiterLeft.ProcessAudio(AudioLeft->GetData(), AudioLeft->Num(), AudioLeft->GetData());
+			LimiterRight.ProcessAudio(AudioRight->GetData(), AudioRight->Num(), AudioRight->GetData());
+		}
+
+		virtual void Bind(FVertexInterfaceData& InVertexData) const override
+		{
+			using namespace SuperOscillatorVertexNames;
+			Super::Bind(InVertexData);
+
+			FInputVertexInterfaceData& Inputs = InVertexData.GetInputs();
+			Inputs.BindReadVertex(METASOUND_GET_PARAM_NAME(StereoWidthPin), StereoWidth);
+			
+			FOutputVertexInterfaceData& Outputs = InVertexData.GetOutputs();
+			Outputs.BindReadVertex(METASOUND_GET_PARAM_NAME(AudioLeftPin), AudioLeft);
+			Outputs.BindReadVertex(METASOUND_GET_PARAM_NAME(AudioRightPin), AudioRight);
+		}
+
+		FFloatReadRef StereoWidth;
+
+		FAudioBufferWriteRef AudioLeft;
+		FAudioBufferWriteRef AudioRight;
+
+		Audio::FDynamicsProcessor LimiterLeft;
+		Audio::FDynamicsProcessor LimiterRight;
+	};
+	
+	class FSuperOscillatorNodeMono : public FNodeFacade
+	{
+	public:
+		FSuperOscillatorNodeMono(const FNodeInitData& InitData)
+		: FNodeFacade(InitData.InstanceName, InitData.InstanceID, TFacadeOperatorClass<FSuperOscillatorOperatorMono>())
+		{}
+	};
+
+	class FSuperOscillatorNodeStereo : public FNodeFacade
+	{
+	public:
+		FSuperOscillatorNodeStereo(const FNodeInitData& InitData)
+		: FNodeFacade(InitData.InstanceName, InitData.InstanceID, TFacadeOperatorClass<FSuperOscillatorOperatorStereo>())
+		{}
+	};
+
+	METASOUND_REGISTER_NODE(FSuperOscillatorNodeMono);
+	METASOUND_REGISTER_NODE(FSuperOscillatorNodeStereo);
+} // namespace Metasound
+
+#undef LOCTEXT_NAMESPACE //MetasoundStandardNodes
\ No newline at end of file
diff --git a/Engine/Plugins/Runtime/Metasound/Source/MetasoundStandardNodes/Public/MetasoundOscillatorNodes.h b/Engine/Plugins/Runtime/Metasound/Source/MetasoundStandardNodes/Public/MetasoundOscillatorNodes.h
--- a/Engine/Plugins/Runtime/Metasound/Source/MetasoundStandardNodes/Public/MetasoundOscillatorNodes.h
+++ b/Engine/Plugins/Runtime/Metasound/Source/MetasoundStandardNodes/Public/MetasoundOscillatorNodes.h
@@ -12,7 +12,7 @@
 #include "MetasoundNode.h"
 #include "MetasoundOperatorInterface.h"
 #include "MetasoundVertex.h"
-
+#include "MetasoundPrimitives.h"
 
 namespace Metasound
 {
