Work-in-progress:
Seek Request
Loop Configuration

diff --git a/Engine/Plugins/Runtime/Metasound/Source/MetasoundFrontend/Private/MetasoundParameterPack.cpp b/Engine/Plugins/Runtime/Metasound/Source/MetasoundFrontend/Private/MetasoundParameterPack.cpp
--- a/Engine/Plugins/Runtime/Metasound/Source/MetasoundFrontend/Private/MetasoundParameterPack.cpp
+++ b/Engine/Plugins/Runtime/Metasound/Source/MetasoundFrontend/Private/MetasoundParameterPack.cpp
@@ -176,5 +176,10 @@ FSharedMetasoundParameterStoragePtr UMetasoundParameterPack::GetCopyOfParameterS
 	{
 		return nullptr;
 	}
-	return MakeShared<FMetasoundParameterPackStorage>(*ParameterStorage);
+	FSharedMetasoundParameterStoragePtr NewCopy = MakeShared<FMetasoundParameterPackStorage>(*ParameterStorage);
+	for (auto Walker = ParameterStorage->begin(); Walker != ParameterStorage->end(); ++Walker)
+	{
+		Walker->OnPostApplied();
+	}
+	return NewCopy;
 }
diff --git a/Engine/Plugins/Runtime/Metasound/Source/MetasoundFrontend/Public/MetasoundParameterPack.h b/Engine/Plugins/Runtime/Metasound/Source/MetasoundFrontend/Public/MetasoundParameterPack.h
--- a/Engine/Plugins/Runtime/Metasound/Source/MetasoundFrontend/Public/MetasoundParameterPack.h
+++ b/Engine/Plugins/Runtime/Metasound/Source/MetasoundFrontend/Public/MetasoundParameterPack.h
@@ -12,6 +12,12 @@
 
 class UMetaSoundSource;
 
+template <typename T, typename V = void>
+struct TIsOneshotDataType { bool value = false; };
+
+template <typename T>
+struct TIsOneshotDataType<T, decltype(&T::OneshotSubmitted, void())> { bool value = true; };
+
 // A structure that encapsulates a "parameter tag" in the raw 'bag-o-bytes'
 // that will be the parameter storage. It lets us walk through memory to 
 // find a parameter with a specific name and specific type. 
@@ -29,6 +35,7 @@ struct FMetasoundParameterPackItemBase
 		uint8* ptr = reinterpret_cast<uint8*>(this);
 		return ptr + ValueOffset;
 	}
+	virtual void OnPostApplied() {};
 	FName  Name;
 	FName  TypeName;
 	uint16 NextHeaderOffset;
@@ -47,6 +54,15 @@ struct FMetasoundParameterPackItem : public FMetasoundParameterPackItemBase
 	{}
 	FMetasoundParameterPackItem(const FMetasoundParameterPackItem& Other) = default;
 	T      Value;
+	virtual ~FMetasoundParameterPackItem() = default;
+
+	virtual void OnPostApplied() override
+	{
+		if constexpr (TIsOneshotDataType<T>().value)
+		{
+			Value.OneshotSubmitted(); 
+		}
+	}
 };
 
 // This next class owns the 'bag-o-bytes' that holds all of the parameters. 
diff --git a/Engine/Plugins/Runtime/Metasound/Source/MetasoundGenerator/Private/MetasoundGenerator.cpp b/Engine/Plugins/Runtime/Metasound/Source/MetasoundGenerator/Private/MetasoundGenerator.cpp
--- a/Engine/Plugins/Runtime/Metasound/Source/MetasoundGenerator/Private/MetasoundGenerator.cpp
+++ b/Engine/Plugins/Runtime/Metasound/Source/MetasoundGenerator/Private/MetasoundGenerator.cpp
@@ -536,6 +536,7 @@ namespace Metasound
 						if (ParameterSetter->DataType == Walker->TypeName)
 						{
 							ParameterSetter->SetParameterWithPayload(Walker->GetPayload());
+							Walker->OnPostApplied();
 						}
 					}
 				}
