Fix MetaSound vertex sort order and description being reset when changing data type or constructor pin state

diff --git a/Engine/Plugins/Runtime/Metasound/Source/MetasoundEditor/Private/MetasoundEditorGraph.cpp b/Engine/Plugins/Runtime/Metasound/Source/MetasoundEditor/Private/MetasoundEditorGraph.cpp
--- a/Engine/Plugins/Runtime/Metasound/Source/MetasoundEditor/Private/MetasoundEditorGraph.cpp
+++ b/Engine/Plugins/Runtime/Metasound/Source/MetasoundEditor/Private/MetasoundEditorGraph.cpp
@@ -368,6 +368,8 @@ void UMetasoundEditorGraphVertex::SetDataType(FName InNewType, bool bPostTransac
 	FNodeHandle NodeHandle = GetNodeHandle();
 	const FName NodeName = NodeHandle->GetNodeName();
 	const FText NodeDisplayName = NodeHandle->GetDisplayName();
+	const FText NodeDescription = NodeHandle->GetDescription();
+	const int32 SortOrderIndex = GetSortOrderIndex();
 
 	// Remove the current nodes and vertex
 	Graph->RemoveMemberNodes(*this);
@@ -382,6 +384,16 @@ void UMetasoundEditorGraphVertex::SetDataType(FName InNewType, bool bPostTransac
 	FNodeHandle NewNodeHandle = AddNodeHandle(NodeName, VertexParams);
 	NewNodeHandle->SetNodeName(NodeName);
 	NewNodeHandle->SetDisplayName(NodeDisplayName);
+	NewNodeHandle->SetDescription(NodeDescription);
+	FGraphHandle GraphHandle = Graph->GetGraphHandle();
+	if (IsA<UMetasoundEditorGraphInput>())
+	{
+		GraphHandle->SetSortOrderIndexForInput(NodeName, SortOrderIndex);
+	}
+	else if (IsA<UMetasoundEditorGraphOutput>())
+	{
+		GraphHandle->SetSortOrderIndexForOutput(NodeName, SortOrderIndex);
+	}
 
 	if (!ensure(NewNodeHandle->IsValid()))
 	{
@@ -455,7 +467,9 @@ void UMetasoundEditorGraphVertex::SetVertexAccessType(EMetasoundFrontendVertexAc
 	FNodeHandle NodeHandle = GetNodeHandle();
 	const FName NodeName = NodeHandle->GetNodeName();
 	const FText NodeDisplayName = NodeHandle->GetDisplayName();
+	const FText NodeDescription = NodeHandle->GetDescription();
 	const FName NodeType = GetDataType();
+	const int32 SortOrderIndex = GetSortOrderIndex();
 	
 	// 3. Remove the current nodes and vertex
 	Graph->RemoveMemberNodes(*this);
@@ -470,7 +484,17 @@ void UMetasoundEditorGraphVertex::SetVertexAccessType(EMetasoundFrontendVertexAc
 	FNodeHandle NewNodeHandle = AddNodeHandle(NodeName, VertexParams);
 	NewNodeHandle->SetNodeName(NodeName);
 	NewNodeHandle->SetDisplayName(NodeDisplayName);
-
+	NewNodeHandle->SetDescription(NodeDescription);
+	FGraphHandle GraphHandle = Graph->GetGraphHandle();
+	if (IsA<UMetasoundEditorGraphInput>())
+	{
+		GraphHandle->SetSortOrderIndexForInput(NodeName, SortOrderIndex);
+	}
+	else if (IsA<UMetasoundEditorGraphOutput>())
+	{
+		GraphHandle->SetSortOrderIndexForOutput(NodeName, SortOrderIndex);
+	}
+	
 	if (!ensure(NewNodeHandle->IsValid()))
 	{
 		return;
