Backing out in-progress work on "transient parameter" handling

diff --git a/Engine/Plugins/Runtime/Metasound/Source/MetasoundFrontend/Private/MetasoundParameterPack.cpp b/Engine/Plugins/Runtime/Metasound/Source/MetasoundFrontend/Private/MetasoundParameterPack.cpp
--- a/Engine/Plugins/Runtime/Metasound/Source/MetasoundFrontend/Private/MetasoundParameterPack.cpp
+++ b/Engine/Plugins/Runtime/Metasound/Source/MetasoundFrontend/Private/MetasoundParameterPack.cpp
@@ -176,10 +176,5 @@ FSharedMetasoundParameterStoragePtr UMetasoundParameterPack::GetCopyOfParameterS
 	{
 		return nullptr;
 	}
-	FSharedMetasoundParameterStoragePtr NewCopy = MakeShared<FMetasoundParameterPackStorage>(*ParameterStorage);
-	for (auto Walker = ParameterStorage->begin(); Walker != ParameterStorage->end(); ++Walker)
-	{
-		Walker->OnPostApplied();
-	}
-	return NewCopy;
+	return MakeShared<FMetasoundParameterPackStorage>(*ParameterStorage);
 }
diff --git a/Engine/Plugins/Runtime/Metasound/Source/MetasoundFrontend/Public/MetasoundParameterPack.h b/Engine/Plugins/Runtime/Metasound/Source/MetasoundFrontend/Public/MetasoundParameterPack.h
--- a/Engine/Plugins/Runtime/Metasound/Source/MetasoundFrontend/Public/MetasoundParameterPack.h
+++ b/Engine/Plugins/Runtime/Metasound/Source/MetasoundFrontend/Public/MetasoundParameterPack.h
@@ -12,12 +12,6 @@
 
 class UMetaSoundSource;
 
-template <typename T, typename V = void>
-struct TIsOneshotDataType { bool value = false; };
-
-template <typename T>
-struct TIsOneshotDataType<T, decltype(&T::OneshotSubmitted, void())> { bool value = true; };
-
 // A structure that encapsulates a "parameter tag" in the raw 'bag-o-bytes'
 // that will be the parameter storage. It lets us walk through memory to 
 // find a parameter with a specific name and specific type. 
@@ -35,7 +29,6 @@ struct FMetasoundParameterPackItemBase
 		uint8* ptr = reinterpret_cast<uint8*>(this);
 		return ptr + ValueOffset;
 	}
-	virtual void OnPostApplied() {};
 	FName  Name;
 	FName  TypeName;
 	uint16 NextHeaderOffset;
@@ -54,15 +47,6 @@ struct FMetasoundParameterPackItem : public FMetasoundParameterPackItemBase
 	{}
 	FMetasoundParameterPackItem(const FMetasoundParameterPackItem& Other) = default;
 	T      Value;
-	virtual ~FMetasoundParameterPackItem() = default;
-
-	virtual void OnPostApplied() override
-	{
-		if constexpr (TIsOneshotDataType<T>().value)
-		{
-			Value.OneshotSubmitted(); 
-		}
-	}
 };
 
 // This next class owns the 'bag-o-bytes' that holds all of the parameters. 
diff --git a/Engine/Plugins/Runtime/Metasound/Source/MetasoundGenerator/Private/MetasoundGenerator.cpp b/Engine/Plugins/Runtime/Metasound/Source/MetasoundGenerator/Private/MetasoundGenerator.cpp
--- a/Engine/Plugins/Runtime/Metasound/Source/MetasoundGenerator/Private/MetasoundGenerator.cpp
+++ b/Engine/Plugins/Runtime/Metasound/Source/MetasoundGenerator/Private/MetasoundGenerator.cpp
@@ -534,7 +534,6 @@ namespace Metasound
 						if (ParameterSetter->DataType == Walker->TypeName)
 						{
 							ParameterSetter->SetParameterWithPayload(Walker->GetPayload());
-							Walker->OnPostApplied();
 						}
 					}
 				}
