[MetaSounds] - Inputs, Outputs, and Variables are not allowing spaces in the member name

diff --git a/Engine/Plugins/Runtime/Metasound/Source/MetasoundEditor/Private/MetasoundEditorGraph.cpp b/Engine/Plugins/Runtime/Metasound/Source/MetasoundEditor/Private/MetasoundEditorGraph.cpp
--- a/Engine/Plugins/Runtime/Metasound/Source/MetasoundEditor/Private/MetasoundEditorGraph.cpp
+++ b/Engine/Plugins/Runtime/Metasound/Source/MetasoundEditor/Private/MetasoundEditorGraph.cpp
@@ -24,9 +24,10 @@
 
 #define LOCTEXT_NAMESPACE "MetaSoundEditor"
 
-// Parameter names do not support '/' in addition to standard
-// name characters As analyzers use this to pack data into address.
-#define INVALID_PARAMETER_NAME_CHARACTERS INVALID_NAME_CHARACTERS METASOUND_ANALYZER_PATH_SEPARATOR
+// Parameter names do not support analyzer path separator, but do support
+// spaces (to be as consistent as possible with other systems such as Blueprint)
+#define INVALID_PARAMETER_NAME_CHARACTERS TEXT("\"',\n\r\t") METASOUND_ANALYZER_PATH_SEPARATOR
+
 
 namespace Metasound
 {
