Fix ensure when renaming a metasound asset

diff --git a/Engine/Plugins/Runtime/Metasound/Source/MetasoundEngine/Private/MetasoundAssetSubsystem.cpp b/Engine/Plugins/Runtime/Metasound/Source/MetasoundEngine/Private/MetasoundAssetSubsystem.cpp
--- a/Engine/Plugins/Runtime/Metasound/Source/MetasoundEngine/Private/MetasoundAssetSubsystem.cpp
+++ b/Engine/Plugins/Runtime/Metasound/Source/MetasoundEngine/Private/MetasoundAssetSubsystem.cpp
@@ -78,6 +78,27 @@ namespace Metasound
 
 			return bSuccess;
 		}
+
+		// Remove the Map entry only if the key and value are equal.
+		//
+		// This protects against scenarios where a metasound is renamed or moved 
+		// and the new entry was being erroneously removed from the PathMap
+		bool RemoveIfExactMatch(TMap<Frontend::FNodeRegistryKey, FSoftObjectPath>& InMap, const Frontend::FNodeRegistryKey& InKeyToRemove, const FSoftObjectPath& InPathToRemove)
+		{
+			if (const FSoftObjectPath* Path = InMap.Find(InKeyToRemove))
+			{
+				if (*Path == InPathToRemove)
+				{
+					InMap.Remove(InKeyToRemove);
+					return true;
+				}
+				else
+				{
+					UE_LOG(LogMetaSound, VeryVerbose, TEXT("Object paths do not match. Skipping removal of %s:%s from the asset subsystem."), *InKeyToRemove, *InPathToRemove.ToString());
+				}
+			}
+			return false;
+		}
 	}
 }
 
@@ -559,7 +580,7 @@ void UMetaSoundAssetSubsystem::RemoveAsset(const UObject& InObject)
 	{
 		const FNodeClassInfo ClassInfo = MetaSoundAsset->GetAssetClassInfo();
 		FNodeRegistryKey RegistryKey = FMetasoundFrontendRegistryContainer::Get()->GetRegistryKey(ClassInfo);
-		PathMap.Remove(RegistryKey);
+		AssetSubsystemPrivate::RemoveIfExactMatch(PathMap, RegistryKey, FSoftObjectPath(&InObject));
 	}
 }
 
@@ -572,7 +593,7 @@ void UMetaSoundAssetSubsystem::RemoveAsset(const FAssetData& InAssetData)
 	if (ensureAlways(AssetSubsystemPrivate::GetAssetClassInfo(InAssetData, ClassInfo)))
 	{
 		FNodeRegistryKey RegistryKey = FMetasoundFrontendRegistryContainer::Get()->GetRegistryKey(ClassInfo);
-		PathMap.Remove(RegistryKey);
+		AssetSubsystemPrivate::RemoveIfExactMatch(PathMap, RegistryKey, InAssetData.GetSoftObjectPath());
 	}
 }
 
@@ -697,7 +718,7 @@ void UMetaSoundAssetSubsystem::UnregisterAssetClassesInDirectories(const TArray<
 				if (bIsRegistered)
 				{
 					FMetasoundFrontendRegistryContainer::Get()->UnregisterNode(RegistryKey);
-					PathMap.Remove(RegistryKey);
+					AssetSubsystemPrivate::RemoveIfExactMatch(PathMap, RegistryKey, AssetData.GetSoftObjectPath());
 				}
 			}
 		}
