Add Trigger type to EAudioParameter so triggers can be differentiated from Boolean paramters

diff --git a/Engine/Plugins/Runtime/Metasound/Source/MetasoundEngine/Private/MetasoundSource.cpp b/Engine/Plugins/Runtime/Metasound/Source/MetasoundEngine/Private/MetasoundSource.cpp
--- a/Engine/Plugins/Runtime/Metasound/Source/MetasoundEngine/Private/MetasoundSource.cpp
+++ b/Engine/Plugins/Runtime/Metasound/Source/MetasoundEngine/Private/MetasoundSource.cpp
@@ -407,6 +407,12 @@ void UMetaSoundSource::InitParameters(TArray<FAudioParameter>& ParametersToInit,
 
 		switch (Parameter.ParamType)
 		{
+			case EAudioParameterType::Trigger:
+			{
+				Parameter = FAudioParameter(Parameter.ParamName, EAudioParameterType::Trigger);
+			}
+			break;
+			
 			case EAudioParameterType::Boolean:
 			{
 				Parameter = FAudioParameter(Parameter.ParamName, Parameter.BoolParam);
@@ -677,7 +683,16 @@ bool UMetaSoundSource::GetAllDefaultParameters(TArray<FAudioParameter>& OutParam
 		{
 			case EMetasoundFrontendLiteralType::Boolean:
 			{
-				Params.ParamType = EAudioParameterType::Boolean;
+				static const FName TriggerName = "Trigger";
+				if (Params.TypeName == TriggerName)
+				{
+					Params.ParamType = EAudioParameterType::Trigger;
+				}
+				else
+				{
+					Params.ParamType = EAudioParameterType::Boolean;
+				}
+					
 				ensure(Input.DefaultLiteral.TryGet(Params.BoolParam));
 			}
 			break;
@@ -789,6 +804,7 @@ bool UMetaSoundSource::IsParameterValid(const FAudioParameter& InParameter, cons
 	bool bIsValid = false;
 	switch (InParameter.ParamType)
 	{
+		case EAudioParameterType::Trigger:
 		case EAudioParameterType::Boolean:
 		{
 			FDataTypeRegistryInfo DataTypeInfo;
diff --git a/Engine/Plugins/Runtime/Metasound/Source/MetasoundFrontend/Private/MetasoundFrontendLiteral.cpp b/Engine/Plugins/Runtime/Metasound/Source/MetasoundFrontend/Private/MetasoundFrontendLiteral.cpp
--- a/Engine/Plugins/Runtime/Metasound/Source/MetasoundFrontend/Private/MetasoundFrontendLiteral.cpp
+++ b/Engine/Plugins/Runtime/Metasound/Source/MetasoundFrontend/Private/MetasoundFrontendLiteral.cpp
@@ -83,6 +83,7 @@ FMetasoundFrontendLiteral::FMetasoundFrontendLiteral(const FAudioParameter& InPa
 {
 	switch (InParameter.ParamType)
 	{
+		case EAudioParameterType::Trigger:
 		case EAudioParameterType::Boolean:
 		{
 			Set(InParameter.BoolParam);
@@ -157,7 +158,7 @@ FMetasoundFrontendLiteral::FMetasoundFrontendLiteral(const FAudioParameter& InPa
 
 		default:
 		{
-			static_assert(static_cast<int32>(EAudioParameterType::COUNT) == 12, "Possible missing switch case coverage");
+			static_assert(static_cast<int32>(EAudioParameterType::COUNT) == 13, "Possible missing switch case coverage");
 			checkNoEntry();
 		}
 	}
@@ -511,7 +512,6 @@ void FMetasoundFrontendLiteral::SetType(EMetasoundFrontendLiteralType InType)
 			Set(0.f);
 		}
 		break;
-
 	
 		case EMetasoundFrontendLiteralType::String:
 		{
diff --git a/Engine/Plugins/Runtime/Metasound/Source/MetasoundFrontend/Private/MetasoundParameterTransmitter.cpp b/Engine/Plugins/Runtime/Metasound/Source/MetasoundFrontend/Private/MetasoundParameterTransmitter.cpp
--- a/Engine/Plugins/Runtime/Metasound/Source/MetasoundFrontend/Private/MetasoundParameterTransmitter.cpp
+++ b/Engine/Plugins/Runtime/Metasound/Source/MetasoundFrontend/Private/MetasoundParameterTransmitter.cpp
@@ -25,6 +25,7 @@ namespace Metasound
 		{
 			switch (InValue.ParamType)
 			{
+				case EAudioParameterType::Trigger:
 				case EAudioParameterType::Boolean:
 				{
 					return FLiteral(InValue.BoolParam);
@@ -94,7 +95,7 @@ namespace Metasound
 
 				default:
 				{
-					static_assert(static_cast<int32>(EAudioParameterType::COUNT) == 12, "Possible missing switch case coverage");
+					static_assert(static_cast<int32>(EAudioParameterType::COUNT) == 13, "Possible missing switch case coverage");
 					checkNoEntry();
 				}
 			}
@@ -106,6 +107,7 @@ namespace Metasound
 		{
 			switch (InValue.ParamType)
 			{
+				case EAudioParameterType::Trigger:
 				case EAudioParameterType::Boolean:
 				{
 					return FLiteral(InValue.BoolParam);
@@ -175,7 +177,7 @@ namespace Metasound
 
 				default:
 				{
-					static_assert(static_cast<int32>(EAudioParameterType::COUNT) == 12, "Possible missing switch case coverage");
+					static_assert(static_cast<int32>(EAudioParameterType::COUNT) == 13, "Possible missing switch case coverage");
 					checkNoEntry();
 				}
 			}
@@ -187,6 +189,7 @@ namespace Metasound
 		{
 			switch (InParameterType)
 			{
+				case EAudioParameterType::Trigger:
 				case EAudioParameterType::Boolean:
 					return GetMetasoundDataTypeName<bool>();
 				case EAudioParameterType::BooleanArray:
@@ -212,7 +215,7 @@ namespace Metasound
 				case EAudioParameterType::NoneArray:
 				default:
 					ensureAlwaysMsgf(false, TEXT("Failed to convert AudioParameterType to POD MetaSound DataType"));
-					static_assert(static_cast<int32>(EAudioParameterType::COUNT) == 12, "Possible missing case coverage");
+					static_assert(static_cast<int32>(EAudioParameterType::COUNT) == 13, "Possible missing case coverage");
 					return FName();
 			}
 		}
diff --git a/Engine/Plugins/Runtime/Metasound/Source/MetasoundFrontend/Public/MetasoundFrontendLiteral.h b/Engine/Plugins/Runtime/Metasound/Source/MetasoundFrontend/Public/MetasoundFrontendLiteral.h
--- a/Engine/Plugins/Runtime/Metasound/Source/MetasoundFrontend/Public/MetasoundFrontendLiteral.h
+++ b/Engine/Plugins/Runtime/Metasound/Source/MetasoundFrontend/Public/MetasoundFrontendLiteral.h
@@ -48,7 +48,6 @@ static_assert(static_cast<uint8>(EMetasoundFrontendLiteralType::IntegerArray) ==
 static_assert(static_cast<uint8>(EMetasoundFrontendLiteralType::FloatArray) == static_cast<uint8>(EAudioParameterType::FloatArray), "Type 'FloatArray' value must match");
 static_assert(static_cast<uint8>(EMetasoundFrontendLiteralType::StringArray) == static_cast<uint8>(EAudioParameterType::StringArray), "Type 'StringArray' value must match");
 static_assert(static_cast<uint8>(EMetasoundFrontendLiteralType::UObjectArray) == static_cast<uint8>(EAudioParameterType::ObjectArray), "Type 'UObjectProxyArray' value must match");
-static_assert(static_cast<uint8>(EMetasoundFrontendLiteralType::Invalid) == static_cast<uint8>(EAudioParameterType::COUNT), "Enum EMetasoundFrontendLiteralType' count must match 'EAudioParameterType'");
 
 // Check that the static_cast<>s above are using the correct type.
 static_assert(std::is_same<uint8, std::underlying_type<EMetasoundFrontendLiteralType>::type>::value, "Update type in static_cast<TYPE> from Metasound::ELiteralType to EMetasoundFrontendLiteralType in EMetasoundFrontendLiteralType declaration.");
diff --git a/Engine/Source/Runtime/AudioExtensions/Public/AudioParameter.h b/Engine/Source/Runtime/AudioExtensions/Public/AudioParameter.h
--- a/Engine/Source/Runtime/AudioExtensions/Public/AudioParameter.h
+++ b/Engine/Source/Runtime/AudioExtensions/Public/AudioParameter.h
@@ -81,6 +81,9 @@ enum class EAudioParameterType : uint8
 	// Array of object values (not supported by legacy SoundCue system)
 	ObjectArray UMETA(DisplayName = "Object (Array)"),
 
+	// Trigger value
+	Trigger,
+
 	COUNT UMETA(Hidden)
 };
 
@@ -202,6 +205,16 @@ struct AUDIOEXTENSIONS_API FAudioParameter
 	{
 	}
 
+	FAudioParameter(FName InName, EAudioParameterType Type)
+		: ParamName(InName)
+		, ParamType(Type)
+	{
+		if (Type == EAudioParameterType::Trigger)
+		{
+			BoolParam = true;
+		}
+	}
+
 	// Static function to avoid int32 constructor collision
 	static FAudioParameter CreateDefaultArray(FName InName, int32 InNum)
 	{
diff --git a/Engine/Source/Runtime/Engine/Private/Audio/SoundParameterControllerInterface.cpp b/Engine/Source/Runtime/Engine/Private/Audio/SoundParameterControllerInterface.cpp
--- a/Engine/Source/Runtime/Engine/Private/Audio/SoundParameterControllerInterface.cpp
+++ b/Engine/Source/Runtime/Engine/Private/Audio/SoundParameterControllerInterface.cpp
@@ -46,7 +46,8 @@ void ISoundParameterControllerInterface::SetTriggerParameter(FName InName)
 	{
 		if (FAudioDevice* AudioDevice = GetAudioDevice())
 		{
-			FAudioParameter ParamToSet = FAudioParameter(InName, true);
+			FAudioParameter ParamToSet = FAudioParameter(InName, EAudioParameterType::Trigger);
+			
 			if (USoundBase* Sound = GetSound())
 			{
 				TArray<FAudioParameter> Params = { MoveTemp(ParamToSet) };
diff --git a/Engine/Source/Runtime/Engine/Private/SoundCue.cpp b/Engine/Source/Runtime/Engine/Private/SoundCue.cpp
--- a/Engine/Source/Runtime/Engine/Private/SoundCue.cpp
+++ b/Engine/Source/Runtime/Engine/Private/SoundCue.cpp
@@ -828,14 +828,10 @@ bool FSoundCueParameterTransmitter::SetParameters(TArray<FAudioParameter>&& InPa
 	TArray<FAudioParameter> TempParams = InParameters;
 	FAudioParameter::Merge(MoveTemp(TempParams), ParamsToSet);
 
-	// todo: filter trigger types from being saved in the arrray of persistent params. triggers are not detectable at this point
-	// currently, triggers will be set on sounds as they start if they had been executed at any point in the past
-	// solvable if triggers were their own enum type, but are currently just bool types
-	//
-	// InParameters.FilterByPredicate([](const FAudioParameter& Param)
-	// {
-	// 	return Param.ParamType != EAudioParameterType::Trigger;
-	// });
+	InParameters.FilterByPredicate([](const FAudioParameter& Param)
+	{
+		return Param.ParamType != EAudioParameterType::Trigger;
+	});
 		
 	return Audio::FParameterTransmitterBase::SetParameters(MoveTemp(InParameters));
 }
\ No newline at end of file
diff --git a/Engine/Source/Runtime/MovieSceneTracks/Private/Sections/MovieSceneAudioSection.cpp b/Engine/Source/Runtime/MovieSceneTracks/Private/Sections/MovieSceneAudioSection.cpp
--- a/Engine/Source/Runtime/MovieSceneTracks/Private/Sections/MovieSceneAudioSection.cpp
+++ b/Engine/Source/Runtime/MovieSceneTracks/Private/Sections/MovieSceneAudioSection.cpp
@@ -152,25 +152,21 @@ void UMovieSceneAudioSection::SetupSoundInputParameters(const USoundBase* InSoun
 
 		for (const FAudioParameter& Param : DefaultParams)
 		{
-			switch(Param.ParamType)
+			switch (Param.ParamType)
 			{
 			case EAudioParameterType::Float:
 			{
 				Inputs_Float.FindOrAdd(Param.ParamName, FMovieSceneFloatChannel{}).SetDefault(Param.FloatParam);
 				break;
 			}
+			case EAudioParameterType::Trigger:
+			{
+				Inputs_Trigger.FindOrAdd(Param.ParamName, FMovieSceneAudioTriggerChannel{});
+				break;
+			}
 			case EAudioParameterType::Boolean:
 			{
-				// Triggers are fundamentally just booleans outside of Metasound.
-				static const FName TriggerName = FName(TEXT("Trigger")); // MOVE ME.
-				if (Param.TypeName == TriggerName)
-				{
-					Inputs_Trigger.FindOrAdd(Param.ParamName, FMovieSceneAudioTriggerChannel{});
-				}
-				else
-				{
-					Inputs_Bool.FindOrAdd(Param.ParamName, FMovieSceneBoolChannel{}).SetDefault(Param.BoolParam);
-				}
+				Inputs_Bool.FindOrAdd(Param.ParamName, FMovieSceneBoolChannel{}).SetDefault(Param.BoolParam);
 				break;
 			}
 			case EAudioParameterType::Integer:
