Cleanup metasound input operators
#rnx

diff --git a/Engine/Plugins/Runtime/Metasound/Source/MetasoundFrontend/Private/MetasoundInputNode.cpp b/Engine/Plugins/Runtime/Metasound/Source/MetasoundFrontend/Private/MetasoundInputNode.cpp
--- a/Engine/Plugins/Runtime/Metasound/Source/MetasoundFrontend/Private/MetasoundInputNode.cpp
+++ b/Engine/Plugins/Runtime/Metasound/Source/MetasoundFrontend/Private/MetasoundInputNode.cpp
@@ -9,37 +9,37 @@
 
 namespace Metasound
 {
-	FDataReferenceCollection FInputValueOperator::GetInputs() const
+	namespace MetasoundInputNodePrivate
 	{
-		checkNoEntry();
-		return {};
-	}
-
-	FDataReferenceCollection FInputValueOperator::GetOutputs() const
-	{
-		checkNoEntry();
-		return {};
-	}
+		FDataReferenceCollection FInputOperatorBase::GetInputs() const
+		{
+			// This is slated to be deprecated and removed.
+			checkNoEntry();
+			return {};
+		}
 
-	void FInputValueOperator::Bind(FVertexInterfaceData& InVertexData) const
-	{
-		if (const FAnyDataReference* Ref = InVertexData.GetInputs().FindDataReference(VertexName))
+		FDataReferenceCollection FInputOperatorBase::GetOutputs() const
 		{
-			checkf(EDataReferenceAccessType::Value == Ref->GetAccessType(), TEXT("Expected bound reference to have %s access. Actual access was %s"), *LexToString(EDataReferenceAccessType::Value), *LexToString(Ref->GetAccessType()));
+			// This is slated to be deprecated and removed.
+			checkNoEntry();
+			return {};
+		}
 
-			// Pass through input to output
-			InVertexData.GetOutputs().BindVertex(VertexName, *Ref);
+		FNonExecutableInputOperatorBase::FNonExecutableInputOperatorBase(const FVertexName& InVertexName, FAnyDataReference&& InDataRef)
+		: VertexName(InVertexName)
+		, DataRef(MoveTemp(InDataRef))
+		{
 		}
-		else
+
+		void FNonExecutableInputOperatorBase::Bind(FVertexInterfaceData& InVertexData) const
 		{
-			// Use stored default value
-			InVertexData.GetInputs().BindVertex(VertexName, Default);
-			InVertexData.GetOutputs().BindVertex(VertexName, Default);
+			InVertexData.GetInputs().BindVertex(VertexName, DataRef);
+			InVertexData.GetOutputs().BindVertex(VertexName, DataRef);
 		}
-	}
 
-	IOperator::FExecuteFunction FInputValueOperator::GetExecuteFunction()
-	{
-		return nullptr;
+		IOperator::FExecuteFunction FNonExecutableInputOperatorBase::GetExecuteFunction()
+		{
+			return nullptr;
+		}
 	}
 }
diff --git a/Engine/Plugins/Runtime/Metasound/Source/MetasoundFrontend/Public/MetasoundInputNode.h b/Engine/Plugins/Runtime/Metasound/Source/MetasoundFrontend/Public/MetasoundInputNode.h
--- a/Engine/Plugins/Runtime/Metasound/Source/MetasoundFrontend/Public/MetasoundInputNode.h
+++ b/Engine/Plugins/Runtime/Metasound/Source/MetasoundFrontend/Public/MetasoundInputNode.h
@@ -5,6 +5,7 @@
 #include "MetasoundBuilderInterface.h"
 #include "MetasoundBuildError.h"
 #include "MetasoundDataReference.h"
+#include "MetasoundExecutableOperator.h"
 #include "MetasoundFrontendDataTypeTraits.h"
 #include "MetasoundNodeConstructorParams.h"
 #include "MetasoundLiteral.h"
@@ -23,252 +24,204 @@
 
 namespace Metasound
 {
-	/** A writable input and a readable output. */
-	template<typename DataType>
-	class TInputOperator : public IOperator
+	namespace MetasoundInputNodePrivate
 	{
+		class METASOUNDFRONTEND_API FInputOperatorBase : public IOperator
+		{
 		public:
-			using FDataWriteReference = TDataWriteReference<DataType>;
+			virtual FDataReferenceCollection GetInputs() const override;
+			virtual FDataReferenceCollection GetOutputs() const override;
+		};
 
-			TInputOperator(const FVertexName& InDataReferenceName, FDataWriteReference InDataReference)
-				: DataReferenceName(InDataReferenceName)
-				// Executable DataTypes require a copy of the output to operate on whereas non-executable
-				// types do not. Avoid copy by assigning to reference for non-executable types.
-				, InputValue(InDataReference)
-				, OutputValue(TExecutableDataType<DataType>::bIsExecutable ? FDataWriteReference::CreateNew(*InDataReference) : InDataReference)
+		class METASOUNDFRONTEND_API FNonExecutableInputOperatorBase : public FInputOperatorBase
+		{	
+		public:
+			virtual void Bind(FVertexInterfaceData& InVertexData) const override;
+			virtual IOperator::FExecuteFunction GetExecuteFunction() override;
+
+		protected:
+			FNonExecutableInputOperatorBase(const FVertexName& InVertexName, FAnyDataReference&& InDataRef);
+
+		private:
+			FVertexName VertexName;
+			FAnyDataReference DataRef;
+		};
+
+		class METASOUNDFRONTEND_API FNonExecutableInputPassThroughOperator : public FNonExecutableInputOperatorBase
+		{
+		public:
+			template<typename DataType>
+			FNonExecutableInputPassThroughOperator(const FVertexName& InVertexName, const TDataReadReference<DataType>& InDataRef)
+			: FNonExecutableInputOperatorBase(InVertexName, FAnyDataReference{InDataRef})
 			{
 			}
 
-			virtual ~TInputOperator() = default;
+			template<typename DataType>
+			FNonExecutableInputPassThroughOperator(const FVertexName& InVertexName, const TDataWriteReference<DataType>& InDataRef)
+			: FNonExecutableInputOperatorBase(InVertexName, FAnyDataReference{InDataRef})
+			{
+			}
+		};
 
-			virtual FDataReferenceCollection GetInputs() const override
+		/** FInputValueOperator provides an input for value references. */
+		class METASOUNDFRONTEND_API FInputValueOperator : public FNonExecutableInputOperatorBase
+		{
+		public:
+			/** Construct an FInputValueOperator with the name of the vertex and the 
+			 * value reference associated with input. 
+			 */
+			template<typename DataType>
+			explicit FInputValueOperator(const FName& InVertexName, const TDataValueReference<DataType>& InValueRef)
+			: FNonExecutableInputOperatorBase(InVertexName, FAnyDataReference{InValueRef})
 			{
-				// This is slated to be deprecated and removed.
-				checkNoEntry();
-				return {};
 			}
+		};
 
-			virtual FDataReferenceCollection GetOutputs() const override
+		/** A writable input and a readable output. */
+		template<typename DataType>
+		class TExecutableInputOperator : public FInputOperatorBase
+		{
+			static_assert(TExecutableDataType<DataType>::bIsExecutable, "TExecutableInputOperatorBase should only be used with executable data types");
+		public:
+			using FDataWriteReference = TDataWriteReference<DataType>;
+
+			TExecutableInputOperator(const FVertexName& InDataReferenceName, FDataWriteReference InDataReference)
+				: DataReferenceName(InDataReferenceName)
+				// Executable DataTypes require a copy of the output to operate on whereas non-executable
+				// types do not. Avoid copy by assigning to reference for non-executable types.
+				, InputValue(InDataReference)
+				, OutputValue(FDataWriteReference::CreateNew(*InDataReference))
 			{
-				// This is slated to be deprecated and removed.
-				checkNoEntry();
-				return {};
 			}
 
 			virtual void Bind(FVertexInterfaceData& InOutVertexData) const override
 			{
-				// TODO: Expose a readable reference instead of a writable reference.
-				//
-				// If data needs to be written to, outside entities should create
-				// it and pass it in as a readable reference. Currently, the workflow
-				// is to have the input node create a writable reference which is then
-				// queried by the outside world. Exposing writable references causes 
-				// code maintainability issues where TInputNode<> specializations need
-				// to handle multiple situations which can happen in an input node.
-				//
-				// The only reason that this code is not removed immediately is because
-				// of the `TExecutableDataType<>` which primarily supports the FTrigger.
-				// The TExecutableDataType<> advances the trigger within the graph. But,
-				// with graph composition, the owner of the data type becomes more 
-				// complicated and hence triggers advancing should be managed by a 
-				// different object. Preferably the graph operator itself, or an
-				// explicit trigger manager tied to the environment.
 				InOutVertexData.GetInputs().BindWriteVertex(DataReferenceName, InputValue);
-
 				InOutVertexData.GetOutputs().BindReadVertex(DataReferenceName, OutputValue);
 			}
 
-			void Execute()
+			virtual FExecuteFunction GetExecuteFunction() override
 			{
-				TExecutableDataType<DataType>::Execute(*InputValue, *OutputValue);
+				return &Execute;
 			}
 
-			static void ExecuteFunction(IOperator* InOperator)
-			{
-				static_cast<TInputOperator<DataType>*>(InOperator)->Execute();
-			}
+		private:
 
-			virtual FExecuteFunction GetExecuteFunction() override
+			static void Execute(IOperator* InOperator)
 			{
-				if (TExecutableDataType<DataType>::bIsExecutable)
-				{
-					return &TInputOperator<DataType>::ExecuteFunction;
-				}
-				return nullptr;
+				using FExecutableInputOperator = TExecutableInputOperator<DataType>;
+
+				FExecutableInputOperator* DerivedOperator = static_cast<FExecutableInputOperator*>(InOperator);
+				check(nullptr != DerivedOperator);
+
+				TExecutableDataType<DataType>::Execute(*(DerivedOperator->InputValue), *(DerivedOperator->OutputValue));
 			}
 
-		protected:
 			FVertexName DataReferenceName;
 
 			FDataWriteReference InputValue;
 			FDataWriteReference OutputValue;
-	};
-
-	/** TPassThroughOperator supplies a readable input and a readable output. 
-	 *
-	 * It does *not* invoke executable data types (see `TExecutableDataType<>`).
-	 */
-	template<typename DataType>
-	class TPassThroughOperator : public TInputOperator<DataType>
-	{
-	public:
-		using FDataReadReference = TDataReadReference<DataType>;
-		using Super = TInputOperator<DataType>;
-
-		TPassThroughOperator(const FVertexName& InDataReferenceName, FDataReadReference InDataReference)
-		: TInputOperator<DataType>(InDataReferenceName, WriteCast(InDataReference)) // Write cast is safe because `GetExecuteFunction() and Bind() are overridden, ensuring that data is not written.
-		{
-		}
+		};
 
-		virtual ~TPassThroughOperator() = default;
 
-		static void ExecuteFunction(IOperator* InOperator)
+		/** FInputValueOperator provides support for transmittable inputs. */
+		template<typename DataType>
+		class TInputReceiverOperator : public FInputOperatorBase
 		{
-			using FPassThroughOperator = TPassThroughOperator<DataType>;
-			FPassThroughOperator* PassThroughOperator = static_cast<FPassThroughOperator*>(InOperator);
-			*(PassThroughOperator->OutputValue) = *(PassThroughOperator->InputValue);
-		}
+		public:
+			using FDataReadReference = TDataReadReference<DataType>;
+			using FDataWriteReference = TDataWriteReference<DataType>;
+			using FInputReceiverOperator = TInputReceiverOperator<DataType>;
 
-		virtual IOperator::FExecuteFunction GetExecuteFunction() override
-		{
-			// TODO: this is a hack until we can remove TExecutableOperator<>.
-			//
-			// The primary contention is that we would like to allow developers 
-			// to specialize `TInputNode<>` as in `TInputNode<FStereoAudioFormat>`.
-			// `TExecutableOperator<>` adds in a level of complexity that makes it 
-			// difficult to allow specialization of TInputNode and to derive from
-			// TInputNode to create the TPassThroughOperator. Particularly because
-			// TExecutableOperator<> alters which output data reference is used. 
-			// Specializations of TInputNode also tend to alter the output data
-			// references. Supporting both is likely to cause issues. 
-			//
-			// We may need to ensure that input nodes do not provide execution
-			// functions. Or we may need a more explicit way of only allowing
-			// outputs to be modified. Likely a mix of the `final` keyword
-			// and disabling template specialization of a base class. 
-			//
-			// namespace Private
-			// {
-			//     class TInputNodePrivate<>
-			//     {
-			//         GetInputs() final
-			//         GetExecutionFunction() final
-			//         GetOutputs()
-			//     }
-			// }
-			// 
-			// template<DataType>
-			// using TInputNodeBase<DataType> = TInputNodePrivate<DataType>; // Do not allow specialization of TInputNodePrivate<> or TInputNodeBase<> (this works because you can't specialize a template alias)
-			//
-			// // DO ALLOW specialization of TInputNode
-			// template<DataType>
-			// class TInputNode<DataType> : public TInputNodeBase<DataType>
-			// {
-			// };
-			// 
-			// template<>
-			// class TInputNode<MyType> : public TInputNodeBase<MyType>
-			// {
-			// 	 GetOutputs() <-- OK to override
-			// }
-			//
-			if (TExecutableDataType<DataType>::bIsExecutable)
+			/** Construct an TInputReceiverOperator with the name of the vertex, value reference associated with input, SendAddress, & Receiver
+			 */
+			TInputReceiverOperator(const FVertexName& InDataReferenceName, FDataReadReference InDataReference, FSendAddress&& InSendAddress, TReceiverPtr<DataType>&& InReceiver)
+				: DataReferenceName(InDataReferenceName) 
+				, InputValue(InDataReference)
+				, OutputValue(FDataWriteReference::CreateNew(*InDataReference))
+				, SendAddress(MoveTemp(InSendAddress))
+				, Receiver(MoveTemp(InReceiver))
 			{
-				return &TPassThroughOperator<DataType>::ExecuteFunction;
 			}
 
-			return nullptr;
-		}
-	};
-
-	/** FInputValueOperator provides support for transmittable inputs. */
-	template<typename DataType>
-	class TInputReceiverOperator : public TInputOperator<DataType>
-	{
-	public:
-		using FDataReadReference = TDataReadReference<DataType>;
-		using FInputReceiverOperator = TInputReceiverOperator<DataType>;
-		using Super = TInputOperator<DataType>;
-
-		/** Construct an TInputReceiverOperator with the name of the vertex, value reference associated with input, SendAddress, & Receiver
-		 */
-		TInputReceiverOperator(const FVertexName& InDataReferenceName, FDataReadReference InDataReference, FSendAddress&& InSendAddress, TReceiverPtr<DataType>&& InReceiver)
-			: Super(InDataReferenceName, WriteCast(InDataReference)) // Write cast is safe because `GetExecuteFunction() and Bind() are overridden, ensuring that data is not written.
-			, SendAddress(MoveTemp(InSendAddress))
-			, Receiver(MoveTemp(InReceiver))
-		{
-		}
-
-		virtual ~TInputReceiverOperator()
-		{
-			Receiver.Reset();
-			FDataTransmissionCenter::Get().UnregisterDataChannelIfUnconnected(SendAddress);
-		}
-
-		void Execute()
-		{
-			DataType& OutputData = *Super::OutputValue;
-
-			bool bHasNewData = Receiver->CanPop();
-			if (bHasNewData)
+			virtual ~TInputReceiverOperator()
 			{
-				Receiver->Pop(OutputData);
-				bHasNotReceivedData = false;
+				Receiver.Reset();
+				FDataTransmissionCenter::Get().UnregisterDataChannelIfUnconnected(SendAddress);
 			}
 
-			if (bHasNotReceivedData)
+			virtual void Bind(FVertexInterfaceData& InOutVertexData) const override
 			{
-				OutputData = *Super::InputValue;
-				bHasNewData = true;
+				InOutVertexData.GetInputs().BindReadVertex(DataReferenceName, InputValue);
+				InOutVertexData.GetOutputs().BindReadVertex(DataReferenceName, OutputValue);
 			}
 
-			if constexpr (TExecutableDataType<DataType>::bIsExecutable)
+			virtual FExecuteFunction GetExecuteFunction() override
 			{
-				TExecutableDataType<DataType>::ExecuteInline(OutputData, bHasNewData);
+				return &StaticExecute;
 			}
-		}
 
-		static void ExecuteFunction(IOperator* InOperator)
-		{
-			static_cast<FInputReceiverOperator*>(InOperator)->Execute();
-		}
+		private:
+			void Execute()
+			{
+				DataType& OutputData = *OutputValue;
 
-		IOperator::FExecuteFunction GetExecuteFunction()
-		{
-			return &FInputReceiverOperator::ExecuteFunction;
-		}
+				bool bHasNewData = Receiver->CanPop();
+				if (bHasNewData)
+				{
+					Receiver->Pop(OutputData);
+					bHasNotReceivedData = false;
+				}
 
-	private:
-		FSendAddress SendAddress;
-		TReceiverPtr<DataType> Receiver;
-		bool bHasNotReceivedData = true;
-	};
+				if (bHasNotReceivedData)
+				{
+					OutputData = *InputValue;
+					bHasNewData = true;
+				}
 
-	/** FInputValueOperator provides an input for value references. */
-	class METASOUNDFRONTEND_API FInputValueOperator : public IOperator
-	{
-	public:
-		/** Construct an FInputValueOperator with the name of the vertex and the 
-		 * value reference associated with input. 
-		 */
-		template<typename DataType>
-		explicit FInputValueOperator(const FName& InVertexName, const TDataValueReference<DataType>& InValueRef)
-		: VertexName(InVertexName)
-		, Default(InValueRef)
-		{
-		}
+				if constexpr (TExecutableDataType<DataType>::bIsExecutable)
+				{
+					TExecutableDataType<DataType>::ExecuteInline(OutputData, bHasNewData);
+				}
+			}
 
-		virtual ~FInputValueOperator() = default;
+			static void StaticExecute(IOperator* InOperator)
+			{
+				using FOperator = TInputReceiverOperator<DataType>;
 
-		virtual FDataReferenceCollection GetInputs() const override;
-		virtual FDataReferenceCollection GetOutputs() const override;
-		virtual void Bind(FVertexInterfaceData& InVertexData) const override;
-		virtual FExecuteFunction GetExecuteFunction() override;
+				FOperator* Operator = static_cast<FOperator*>(InOperator);
+				check(nullptr != Operator);
+				Operator->Execute();
+			}
 
-	private:
-		FName VertexName;
-		FAnyDataReference Default;
-	};
+			FVertexName DataReferenceName;
+			FDataReadReference InputValue;
+			FDataWriteReference OutputValue;
+			FSendAddress SendAddress;
+			TReceiverPtr<DataType> Receiver;
+			bool bHasNotReceivedData = true;
+		};
+	}
+
+	/** Choose input operator based upon data type and access type */
+	template<typename DataType, EVertexAccessType VertexAccess=EVertexAccessType::Reference>
+	using TInputOperator = std::conditional_t<
+		VertexAccess == EVertexAccessType::Value,
+		MetasoundInputNodePrivate::FInputValueOperator,
+		std::conditional_t<
+			TExecutableDataType<DataType>::bIsExecutable,
+			MetasoundInputNodePrivate::TExecutableInputOperator<DataType>,
+			MetasoundInputNodePrivate::FNonExecutableInputPassThroughOperator
+		>
+	>;
+
+	/** Choose pass through operator based upon data type and access type */
+	template<typename DataType, EVertexAccessType VertexAccess=EVertexAccessType::Reference>
+	using TPassThroughOperator = std::conditional_t<
+		VertexAccess == EVertexAccessType::Value,
+		MetasoundInputNodePrivate::FInputValueOperator,
+		MetasoundInputNodePrivate::FNonExecutableInputPassThroughOperator
+	>;
 
 	/** Data type creation policy to create by copy construction. */
 	template<typename DataType>
@@ -443,9 +396,9 @@ namespace Metasound
 
 			virtual TUniquePtr<IOperator> CreateOperator(const FBuildOperatorParams& InParams, FBuildResults& OutResults) override
 			{
+				using namespace MetasoundInputNodePrivate;
+
 				using FInputNodeType = TInputNode<DataType, VertexAccess>;
-				using FOwnedInputOperatorType = std::conditional_t<bIsReferenceVertexAccess, TInputOperator<DataType>, FInputValueOperator>;
-				using FPassThroughInputOperatorType = std::conditional_t<bIsReferenceVertexAccess, TPassThroughOperator<DataType>, FInputValueOperator>;
 
 				checkf(!(bEnableTransmission && bIsValueVertexAccess), TEXT("Input cannot enable transmission for vertex with access 'EVertexAccessType::Reference'"));
 
@@ -478,13 +431,12 @@ namespace Metasound
 				if (const FAnyDataReference* Ref = InParams.InputData.FindDataReference(VertexKey))
 				{
 					// Pass through input value
-					return MakeUnique<FPassThroughInputOperatorType>(VertexKey, ReferenceCreator->CreateDataReference(*Ref));
+					return MakeUnique<TPassThroughOperator<DataType, VertexAccess>>(VertexKey, ReferenceCreator->CreateDataReference(*Ref));
 				}
 
-
 				// Owned input value
 				FDataReference DataRef = ReferenceCreator->CreateDataReference(InParams.OperatorSettings);
-				return MakeUnique<FOwnedInputOperatorType>(VertexKey, DataRef);
+				return MakeUnique<TInputOperator<DataType, VertexAccess>>(VertexKey, DataRef);
 			}
 
 		private:
