Perlin "value" noise metasound node

diff --git a/Engine/Plugins/Runtime/Metasound/Source/MetasoundGraphCore/Public/MetasoundFacade.h b/Engine/Plugins/Runtime/Metasound/Source/MetasoundGraphCore/Public/MetasoundFacade.h
--- a/Engine/Plugins/Runtime/Metasound/Source/MetasoundGraphCore/Public/MetasoundFacade.h
+++ b/Engine/Plugins/Runtime/Metasound/Source/MetasoundGraphCore/Public/MetasoundFacade.h
@@ -165,4 +165,17 @@ namespace Metasound
 
 			FVertexInterface VertexInterface;
 	};
+
+
+	/** TNodeFacade further reduces boilerplate code by allowing shorthand node implementations.
+	 *  in the form: TNodeFacade<FMyOperator>
+	 */
+	template<typename OperatorType>
+	class TNodeFacade : public FNodeFacade
+	{
+	public:
+		TNodeFacade(const FNodeInitData& InInitData)
+			: FNodeFacade(InInitData.InstanceName, InInitData.InstanceID, TFacadeOperatorClass<OperatorType>())
+		{}
+	};
 }
diff --git a/Engine/Plugins/Runtime/Metasound/Source/MetasoundStandardNodes/Private/MetasoundPerlinNoiseNode.cpp b/Engine/Plugins/Runtime/Metasound/Source/MetasoundStandardNodes/Private/MetasoundPerlinNoiseNode.cpp
new file mode 100644
--- /dev/null
+++ b/Engine/Plugins/Runtime/Metasound/Source/MetasoundStandardNodes/Private/MetasoundPerlinNoiseNode.cpp
@@ -0,0 +1,299 @@
+// Copyright Epic Games, Inc. All Rights Reserved.
+#include "MetasoundAudioBuffer.h"
+#include "MetasoundEnumRegistrationMacro.h"
+#include "MetasoundExecutableOperator.h"
+#include "MetasoundFacade.h"
+#include "MetasoundNodeRegistrationMacro.h"
+#include "MetasoundDataTypeRegistrationMacro.h"
+#include "MetasoundParamHelper.h"
+#include "MetasoundPrimitives.h"
+#include "MetasoundStandardNodesCategories.h"
+#include "MetasoundStandardNodesNames.h"
+#include "MetasoundTime.h"
+#include "MetasoundTrigger.h"
+#include "MetasoundVertex.h"
+#include "DSP/PerlinNoise.h"
+#include "Math/UnrealMathUtility.h"
+
+#define LOCTEXT_NAMESPACE "MetasoundStandardNodes_PerlinNoise"
+
+// Control/Audio rate Perlin "Value" Noise: 
+// Take any value "X" and produces deterministic smooth noise.
+// Each "Octave" adds another layer of noise at higher frequency and lower amplitude 
+// If you have multiple nodes and which different results for each, use the offset pin to start at different location (similar to a seed).
+
+namespace Metasound
+{		
+	namespace PerlinNoiseVertexNames
+	{
+		METASOUND_PARAM(XPin, "X", "Input value to the Perlin function. (Will use internal clock, in seconds, if not connected)")	
+		METASOUND_PARAM(OctavesPin, "Layers", "Number of layers (or Octaves) of noise to sum")
+		METASOUND_PARAM(OffsetPin, "Seed", "Seed (or offset) into the noise. (-1 will use system rand())")
+		METASOUND_PARAM(MinValuePin, "Min Value", "Minimum output value")
+		METASOUND_PARAM(MaxValuePin, "Max Value", "Maximum output value")
+		METASOUND_PARAM(OutputScaledPin, "Output", "Scaled output based on Min and Max values")
+		METASOUND_PARAM(OutputNormalizedPin, "Normalized", "Normalized output from Perlin function (-1...1)")
+	}
+
+	template<typename TDataClass, typename TThisType>
+	class TBasePerlinNoiseOperator : public TExecutableOperator<TThisType>
+	{
+	public:	
+		static const int32 MaxOctaves = 8; // Careful here as this can be expensive, each octave adds a full layer of noise
+
+		using TDataClassReadRef = TDataReadReference<TDataClass>;
+		using TDataClassWriteRef = TDataWriteReference<TDataClass>;
+
+		static const FVertexInterface& GetVertexInterface()
+		{
+			using namespace PerlinNoiseVertexNames;
+
+			static const FVertexInterface Interface
+			{
+				FInputVertexInterface
+				{
+					TInputDataVertex<TDataClass>(METASOUND_GET_PARAM_NAME_AND_METADATA(XPin), 0),
+					TInputDataVertex<int32>(METASOUND_GET_PARAM_NAME_AND_METADATA(OctavesPin), 1),
+					TInputConstructorVertex<float>(METASOUND_GET_PARAM_NAME_AND_METADATA(OffsetPin), -1.0f),
+					TInputDataVertex<float>(METASOUND_GET_PARAM_NAME_AND_METADATA(MinValuePin), -1.0f),
+					TInputDataVertex<float>(METASOUND_GET_PARAM_NAME_AND_METADATA(MaxValuePin), 1.0f),
+				},
+				FOutputVertexInterface
+				{
+					TOutputDataVertex<TDataClass>(METASOUND_GET_PARAM_NAME_AND_METADATA(OutputScaledPin)),
+					TOutputDataVertex<TDataClass>(METASOUND_GET_PARAM_NAME_AND_METADATA(OutputNormalizedPin)),
+				}
+			};
+			return Interface;
+		}	
+
+		static const FNodeClassMetadata& GetNodeInfo()
+		{
+			// Common stuff
+			auto InitNodeInfo = []() -> FNodeClassMetadata
+			{
+				FNodeClassMetadata Info = TThisType::CreateNodeInfo();
+				Info.MajorVersion = 1;
+				Info.MinorVersion = 0;
+				Info.Author = PluginAuthor;
+				Info.PromptIfMissing = PluginNodeMissingPrompt;
+				Info.DefaultInterface = GetVertexInterface();
+				Info.CategoryHierarchy.Emplace(NodeCategories::Generators);
+				Info.CategoryHierarchy.Emplace(NodeCategories::RandomUtils);
+				Info.Keywords.Emplace(METASOUND_LOCTEXT("Metasound_PerlinNoiseNodeKeywordRandom", "Random"));
+				Info.Keywords.Emplace(METASOUND_LOCTEXT("Metasound_PerlinNoiseNodeKeywordRand", "Rand"));
+				return Info;
+			};
+			static const FNodeClassMetadata Info = InitNodeInfo();
+			return Info;
+		}
+
+		struct FPinReadRefs
+		{
+			TDataClassReadRef X;
+			FInt32ReadRef Octaves;
+			FFloatReadRef MinValue;
+			FFloatReadRef MaxValue;
+			bool bXConnected = false;
+			float Offset = 0.f;
+		};
+
+		static TUniquePtr<IOperator> CreateOperator(const FBuildOperatorParams& InParams, FBuildResults& OutResults)
+		{		
+			const FOperatorSettings& Settings = InParams.OperatorSettings;
+			const FInputVertexInterfaceData& InputData = InParams.InputData;
+
+			using namespace PerlinNoiseVertexNames;
+
+			// Constructor pin.
+			// If Offset is -1, use system rand as a "seed"
+			float Offset = InputData.GetOrCreateDefaultValue<float>(METASOUND_GET_PARAM_NAME(OffsetPin), InParams.OperatorSettings); 
+			if (FMath::IsNearlyEqual(Offset, -1.f))
+			{
+				Offset = FMath::FRand();
+			}
+
+			FPinReadRefs Pins
+			{
+				  InputData.GetOrCreateDefaultDataReadReference<TDataClass>(METASOUND_GET_PARAM_NAME(XPin), Settings)
+				, InputData.GetOrCreateDefaultDataReadReference<int32>(METASOUND_GET_PARAM_NAME(OctavesPin), Settings)
+				, InputData.GetOrCreateDefaultDataReadReference<float>(METASOUND_GET_PARAM_NAME(MinValuePin), Settings)
+				, InputData.GetOrCreateDefaultDataReadReference<float>(METASOUND_GET_PARAM_NAME(MaxValuePin), Settings)
+				, InputData.IsVertexBound(METASOUND_GET_PARAM_NAME(XPin))
+				, Offset
+			};
+
+			return MakeUnique<TThisType>(Settings, MoveTemp(Pins));
+		}
+
+		TBasePerlinNoiseOperator(const FOperatorSettings& InSettings, FPinReadRefs&& InPins)
+			: SampleRate{ InSettings.GetSampleRate() }
+			, NumFramesPerBlock{ InSettings.GetNumFramesPerBlock() }
+			, Pins{ MoveTemp(InPins) }
+			, OutputScaled{ TThisType::CreateOutputWriteRef(InSettings) }
+			, OutputNormalized{ TThisType::CreateOutputWriteRef(InSettings) }
+		{
+			check(SampleRate > 0);
+			check(NumFramesPerBlock > 0);
+		}
+
+		FDataReferenceCollection GetInputs() const override { return {}; }
+		FDataReferenceCollection GetOutputs() const override { return {}; }
+		
+		void Bind(FVertexInterfaceData& InVertexData) const override
+		{
+			using namespace PerlinNoiseVertexNames;
+			FInputVertexInterfaceData& Inputs = InVertexData.GetInputs();
+			Inputs.BindReadVertex(METASOUND_GET_PARAM_NAME(XPin), Pins.X);
+			Inputs.SetValue(METASOUND_GET_PARAM_NAME(OffsetPin), Pins.Offset);
+			Inputs.BindReadVertex(METASOUND_GET_PARAM_NAME(OctavesPin), Pins.Octaves);
+			
+			FOutputVertexInterfaceData& Outputs = InVertexData.GetOutputs();
+			Outputs.BindReadVertex(METASOUND_GET_PARAM_NAME(OutputScaledPin), OutputScaled);
+			Outputs.BindReadVertex(METASOUND_GET_PARAM_NAME(OutputNormalizedPin), OutputNormalized);
+		}
+
+		void AccumulateTime()
+		{
+			NumFramesElapsed += NumFramesPerBlock;
+		}
+		float GetElapsedTimeInSeconds() const 
+		{
+			return static_cast<float>(NumFramesElapsed / SampleRate);
+		}
+
+	protected:
+		float SampleRate = 0;
+		int32 NumFramesPerBlock = 0;
+		int32 NumFramesElapsed = 0;
+		FPinReadRefs Pins;
+		TDataClassWriteRef OutputScaled;
+		TDataClassWriteRef OutputNormalized;
+	};
+
+	class FPerlinControlRate : public TBasePerlinNoiseOperator<float, FPerlinControlRate>
+	{
+	public:
+		using BaseClass = TBasePerlinNoiseOperator<float, FPerlinControlRate>;
+		using BaseClass::TBasePerlinNoiseOperator;
+		using BaseClass::CreateOperator;
+
+		static FFloatWriteRef CreateOutputWriteRef(const FOperatorSettings& InSettings) { return FFloatWriteRef::CreateNew(0.f); }
+		static FNodeClassMetadata CreateNodeInfo()
+		{
+			FNodeClassMetadata Info;
+			Info.ClassName = { StandardNodes::Namespace, TEXT("Perlin Noise (float)"), StandardNodes::AudioVariant };
+			Info.DisplayName = METASOUND_LOCTEXT("Metasound_PerlinNoiseFloatNodeDisplayName", "Perlin Noise (float)");
+			Info.Description = METASOUND_LOCTEXT("Metasound_PerlinNoiseFloatNodeDescription", "Generates 1D Perlin \"Value noise\" at control rate");
+			return Info;
+		}
+
+		void Execute()
+		{
+			// Use input pin if its connected, otherwise use elapsed time in seconds.
+			const float X = Pins.bXConnected ?
+				*Pins.X :
+				GetElapsedTimeInSeconds();
+
+			int32 OctavesClamped = FMath::Clamp(*Pins.Octaves, 1, MaxOctaves);
+
+			// Generate noise (this is multiple layers of Perlin noise at different frequencies).
+			const float ValueNoise = Audio::PerlinValueNoise1D(X + Pins.Offset, OctavesClamped);
+
+			// Output the scaled version.
+			const float ScaledValue = FMath::GetMappedRangeValueClamped(
+				FVector2f{ -1.0f, 1.0f },
+				FVector2f{ *Pins.MinValue, *Pins.MaxValue },
+				ValueNoise);
+
+			*OutputScaled = ScaledValue;
+			*OutputNormalized = ValueNoise;
+
+			AccumulateTime();
+		}		
+	};
+
+	class FPerlinAudioRate : public TBasePerlinNoiseOperator<FAudioBuffer, FPerlinAudioRate>
+	{
+	public:
+		using BaseClass = TBasePerlinNoiseOperator<FAudioBuffer, FPerlinAudioRate>;
+		using BaseClass::CreateOperator;
+
+		static FAudioBufferWriteRef CreateOutputWriteRef(const FOperatorSettings& InSettings) { return FAudioBufferWriteRef::CreateNew(InSettings); }
+
+		static FNodeClassMetadata CreateNodeInfo()
+		{
+			FNodeClassMetadata Info;
+			Info.ClassName = { StandardNodes::Namespace, TEXT("Perlin Noise (audio)"), StandardNodes::AudioVariant };
+			Info.DisplayName = METASOUND_LOCTEXT("Metasound_PerlinNoiseAudioNodeDisplayName", "Perlin Noise (audio)");
+			Info.Description = METASOUND_LOCTEXT("Metasound_PerlinNoiseAudioNodeDescription", "Generates 1D Perlin \"Value noise\" at audio rate");
+			return Info;
+		}
+
+		FPerlinAudioRate(const FOperatorSettings& InSettings, FPinReadRefs&& InPins)
+			: BaseClass{ InSettings, MoveTemp(InPins) }
+			, TimeBuffer{ InPins.bXConnected ? 0 : InSettings.GetNumFramesPerBlock() }	// Create a scratch buffer if X is not connected.
+		{
+		}
+
+		void GenerateTimeArray()
+		{
+			// Not integrating here, for fear of discontinuities with floating point error
+			float* TimePtr = TimeBuffer.GetData();
+			const float SampleRateReciprocal = 1.f / SampleRate;
+			const int32 StopTime = NumFramesElapsed + NumFramesPerBlock;
+			for (int32 i = NumFramesElapsed; i < StopTime; ++i)
+			{
+				*TimePtr++ = static_cast<float>(i * SampleRateReciprocal);
+			}
+		}
+
+		void GenerateScaledOutput()
+		{
+			float* ScaledPtr = OutputScaled->GetData();
+			const float* NormalizedPtr = OutputNormalized->GetData();
+			const FVector2f MinMax{ *Pins.MinValue, *Pins.MaxValue };
+			const FVector2f NormRange{ -1.0f, 1.0f };
+
+			for (int32 i = 0; i < NumFramesPerBlock; ++i)
+			{
+				// Output the scaled version.
+				*ScaledPtr++ = FMath::GetMappedRangeValueClamped(NormRange, MinMax, *NormalizedPtr++);
+			}
+		}
+
+		void Execute()
+		{
+			// Use input pin if its connected, otherwise generate an array of increasing time
+			const FAudioBuffer* XBuffer = Pins.X.Get();
+			if (!Pins.bXConnected)
+			{
+				GenerateTimeArray();
+				XBuffer = &TimeBuffer;
+			}
+		
+			int32 OctavesClamped = FMath::Clamp(*Pins.Octaves,1,MaxOctaves);
+
+			Audio::PerlinValueNoise1DBuffer(
+				*XBuffer,
+				Pins.Offset, 
+				OctavesClamped, 
+				*OutputNormalized
+			);
+			
+			GenerateScaledOutput();
+			AccumulateTime();
+		}
+
+	private:
+		FAudioBuffer TimeBuffer;
+	};
+	
+	using FControlRatePerlinNode = TNodeFacade<FPerlinControlRate>;
+	using FAudioRatePerlinNode = TNodeFacade<FPerlinAudioRate>;
+
+	METASOUND_REGISTER_NODE(FControlRatePerlinNode);
+	METASOUND_REGISTER_NODE(FAudioRatePerlinNode);
+}
+
+#undef LOCTEXT_NAMESPACE
diff --git a/Engine/Source/Runtime/SignalProcessing/Private/PerlinNoise.cpp b/Engine/Source/Runtime/SignalProcessing/Private/PerlinNoise.cpp
new file mode 100644
--- /dev/null
+++ b/Engine/Source/Runtime/SignalProcessing/Private/PerlinNoise.cpp
@@ -0,0 +1,41 @@
+// Copyright Epic Games, Inc. All Rights Reserved.
+
+#include "DSP/PerlinNoise.h"
+#include "Math/UnrealMathUtility.h"
+
+namespace Audio
+{
+	float PerlinValueNoise1D(const float InX, const int32 InNumOctaves)
+	{
+		check(InNumOctaves > 0);
+		
+		float Frequency = 1.f;
+		float NoiseSum = 0.f;
+		float NormalizeRange = 0;
+
+		// To normalize we add up the max range of each octave, to keep everything in the range [-1,1]
+		for (int32 Octave = 0; Octave < InNumOctaves; Octave++, Frequency *= 2.f)
+		{
+			const float FreqReciprocal = 1.f / Frequency;
+			NoiseSum += (FMath::PerlinNoise1D(InX * Frequency) * FreqReciprocal);
+			NormalizeRange += FreqReciprocal;
+		}
+	
+		return NoiseSum / NormalizeRange;
+	}
+
+	void PerlinValueNoise1DBuffer(TArrayView<const float> InXBuffer, const float InOffset, const int32 InNumOctaves, TArrayView<float> OutNoiseBuffer)
+	{
+		float* OutNoisePtr = OutNoiseBuffer.GetData();
+		const float* XPtr = InXBuffer.GetData();
+		int32 Size = OutNoiseBuffer.Num();
+		
+		check(InXBuffer.Num() <= Size)
+
+		for (int32 i = 0; i < Size; ++i)
+		{
+			const float X = *XPtr++;
+			*OutNoisePtr++ = PerlinValueNoise1D(X + InOffset, InNumOctaves);
+		}
+	}
+}
\ No newline at end of file
diff --git a/Engine/Source/Runtime/SignalProcessing/Public/DSP/PerlinNoise.h b/Engine/Source/Runtime/SignalProcessing/Public/DSP/PerlinNoise.h
new file mode 100644
--- /dev/null
+++ b/Engine/Source/Runtime/SignalProcessing/Public/DSP/PerlinNoise.h
@@ -0,0 +1,28 @@
+// Copyright Epic Games, Inc. All Rights Reserved.
+
+#pragma once
+
+#include "SignalProcessingModule.h"
+
+namespace Audio
+{
+	/**
+	* PerlinValueNoise1D. Computes 1 sample of "value" noise using the Perlin 1D function with N octaves.
+	* Each "Octave" is a new layer and comes from the doubling of frequencies (through the noise), hence octave.
+	* The higher frequencies have less amplitude, creating higher "detail".
+	*
+	* @param InX: The input "X" Value to hash into Perlin noise space. Time or an "X" coordinate (in the case of 1d) is often used.
+	* @param InNumOctaves: The number of octaves (layers) to compute. Each octave has higher frequency and less amplitude. Value needs to be > 0.
+	*/
+	SIGNALPROCESSING_API float PerlinValueNoise1D(const float InX, const int32 InNumOctaves);
+
+	/**
+	* PerlinValueNoise1DBuffer. Whole buffer version of Value noise.
+	*
+	* @param InXBuffer: Buffer to compute value noise on.
+	* @param InOffset: Offset each value is offset by. This is normally used to give different results when using multiple calls with the same input.
+	* @param InNumOctaves: The number of octaves (layers) to compute. Each octave has higher frequency and less amplitude.
+	* @param OutNoiseBuffer: Buffer to write the noise too. Should be as big as the input buffer.
+	*/
+	SIGNALPROCESSING_API void PerlinValueNoise1DBuffer(TArrayView<const float> InXBuffer, const float InOffset, const int32 InNumOctaves, TArrayView<float> OutNoiseBuffer);
+}
\ No newline at end of file
