Metasound automated node tests and fixes

diff --git a/Engine/Plugins/Runtime/AudioModulation/Source/AudioModulation/Private/MetaSoundNodes/GetModulatorValueNode.cpp b/Engine/Plugins/Runtime/AudioModulation/Source/AudioModulation/Private/MetaSoundNodes/GetModulatorValueNode.cpp
--- a/Engine/Plugins/Runtime/AudioModulation/Source/AudioModulation/Private/MetaSoundNodes/GetModulatorValueNode.cpp
+++ b/Engine/Plugins/Runtime/AudioModulation/Source/AudioModulation/Private/MetaSoundNodes/GetModulatorValueNode.cpp
@@ -68,10 +68,18 @@ namespace AudioModulation
 			const FInputVertexInterface& InputInterface = InParams.Node.GetVertexInterface().GetInputInterface();
 			const FDataReferenceCollection& InputCollection = InParams.InputDataReferences;
 
-			FSoundModulatorAssetReadRef ModulatorReadRef = InputCollection.GetDataReadReferenceOrConstruct<FSoundModulatorAsset>("Modulator");
-			FBoolReadRef NormalizedReadRef = InputCollection.GetDataReadReferenceOrConstructWithVertexDefault<bool>(InputInterface, "Normalized", InParams.OperatorSettings);
+			if (InParams.Environment.Contains<Audio::FDeviceId>(Metasound::Frontend::SourceInterface::Environment::DeviceID))
+			{
+				FSoundModulatorAssetReadRef ModulatorReadRef = InputCollection.GetDataReadReferenceOrConstruct<FSoundModulatorAsset>("Modulator");
+				FBoolReadRef NormalizedReadRef = InputCollection.GetDataReadReferenceOrConstructWithVertexDefault<bool>(InputInterface, "Normalized", InParams.OperatorSettings);
 
-			return MakeUnique<FGetModulatorValueNodeOperator>(InParams, ModulatorReadRef, NormalizedReadRef);
+				return MakeUnique<FGetModulatorValueNodeOperator>(InParams, ModulatorReadRef, NormalizedReadRef);
+			}
+			else
+			{
+				UE_LOG(LogMetaSound, Warning, TEXT("Cannot create GetModulatorValueNode when no audio device ID supplied. Expected metasound environment variable '%s'"), *Metasound::Frontend::SourceInterface::Environment::DeviceID.ToString());
+				return nullptr;
+			}
 		}
 
 		FGetModulatorValueNodeOperator(const Metasound::FCreateOperatorParams& InParams, const FSoundModulatorAssetReadRef& InModulator, const Metasound::FBoolReadRef& InNormalized)
@@ -162,4 +170,4 @@ namespace AudioModulation
 } // namespace AudioModulation
 
 #undef LOCTEXT_NAMESPACE
-#endif // WITH_AUDIOMODULATION_METASOUND_SUPPORT
\ No newline at end of file
+#endif // WITH_AUDIOMODULATION_METASOUND_SUPPORT
diff --git a/Engine/Plugins/Runtime/AudioModulation/Source/AudioModulation/Private/MetaSoundNodes/MixModulatorsNode.cpp b/Engine/Plugins/Runtime/AudioModulation/Source/AudioModulation/Private/MetaSoundNodes/MixModulatorsNode.cpp
--- a/Engine/Plugins/Runtime/AudioModulation/Source/AudioModulation/Private/MetaSoundNodes/MixModulatorsNode.cpp
+++ b/Engine/Plugins/Runtime/AudioModulation/Source/AudioModulation/Private/MetaSoundNodes/MixModulatorsNode.cpp
@@ -70,13 +70,21 @@ namespace AudioModulation
 
 			const FInputVertexInterface& InputInterface = InParams.Node.GetVertexInterface().GetInputInterface();
 			const FDataReferenceCollection& InputCollection = InParams.InputDataReferences;
+				
+			if (InParams.Environment.Contains<Audio::FDeviceId>(Metasound::Frontend::SourceInterface::Environment::DeviceID))
+			{
+				FSoundModulatorAssetReadRef Modulator1ReadRef = InputCollection.GetDataReadReferenceOrConstruct<FSoundModulatorAsset>("In1");
+				FSoundModulatorAssetReadRef Modulator2ReadRef = InputCollection.GetDataReadReferenceOrConstruct<FSoundModulatorAsset>("In2");
+				FSoundModulationParameterAssetReadRef ParameterReadRef = InputCollection.GetDataReadReferenceOrConstruct<FSoundModulationParameterAsset>("MixParameter");
+				FBoolReadRef NormalizedReadRef = InputCollection.GetDataReadReferenceOrConstructWithVertexDefault<bool>(InputInterface, "Normalized", InParams.OperatorSettings);
 
-			FSoundModulatorAssetReadRef Modulator1ReadRef = InputCollection.GetDataReadReferenceOrConstruct<FSoundModulatorAsset>("In1");
-			FSoundModulatorAssetReadRef Modulator2ReadRef = InputCollection.GetDataReadReferenceOrConstruct<FSoundModulatorAsset>("In2");
-			FSoundModulationParameterAssetReadRef ParameterReadRef = InputCollection.GetDataReadReferenceOrConstruct<FSoundModulationParameterAsset>("MixParameter");
-			FBoolReadRef NormalizedReadRef = InputCollection.GetDataReadReferenceOrConstructWithVertexDefault<bool>(InputInterface, "Normalized", InParams.OperatorSettings);
-
-			return MakeUnique<FMixModulatorsNodeOperator>(InParams, Modulator1ReadRef, Modulator2ReadRef, ParameterReadRef, NormalizedReadRef);
+				return MakeUnique<FMixModulatorsNodeOperator>(InParams, Modulator1ReadRef, Modulator2ReadRef, ParameterReadRef, NormalizedReadRef);
+			}
+			else
+			{
+				UE_LOG(LogMetaSound, Warning, TEXT("Cannot create mix modulators node when no audio device ID supplied. Expected metasound environment variable '%s'"), *Metasound::Frontend::SourceInterface::Environment::DeviceID.ToString());
+				return nullptr;
+			}
 		}
 
 		FMixModulatorsNodeOperator(
diff --git a/Engine/Plugins/Runtime/Metasound/Source/MetasoundEngine/Private/MetasoundAudioBusReaderNode.cpp b/Engine/Plugins/Runtime/Metasound/Source/MetasoundEngine/Private/MetasoundAudioBusReaderNode.cpp
--- a/Engine/Plugins/Runtime/Metasound/Source/MetasoundEngine/Private/MetasoundAudioBusReaderNode.cpp
+++ b/Engine/Plugins/Runtime/Metasound/Source/MetasoundEngine/Private/MetasoundAudioBusReaderNode.cpp
@@ -84,12 +84,22 @@ namespace Metasound
 			using namespace AudioBusReaderNode; 
 			const FDataReferenceCollection& InputCollection = InParams.InputDataReferences;
 
-			Audio::FDeviceId AudioDeviceId = InParams.Environment.GetValue<Audio::FDeviceId>(SourceInterface::Environment::DeviceID);
-			int32 AudioMixerOutputFrames = InParams.Environment.GetValue<int32>(SourceInterface::Environment::AudioMixerNumOutputFrames);
-
-			FAudioBusAssetReadRef AudioBusIn = InputCollection.GetDataReadReferenceOrConstruct<FAudioBusAsset>(METASOUND_GET_PARAM_NAME(InParamAudioBusInput));
+			bool bHasEnvironmentVars = InParams.Environment.Contains<Audio::FDeviceId>(SourceInterface::Environment::DeviceID) && InParams.Environment.Contains<int32>(SourceInterface::Environment::AudioMixerNumOutputFrames);
 			
-			return MakeUnique<TAudioBusReaderOperator<NumChannels>>(InParams.OperatorSettings, AudioMixerOutputFrames, AudioDeviceId, AudioBusIn);
+			if (bHasEnvironmentVars)
+			{
+				Audio::FDeviceId AudioDeviceId = InParams.Environment.GetValue<Audio::FDeviceId>(SourceInterface::Environment::DeviceID);
+				int32 AudioMixerOutputFrames = InParams.Environment.GetValue<int32>(SourceInterface::Environment::AudioMixerNumOutputFrames);
+
+				FAudioBusAssetReadRef AudioBusIn = InputCollection.GetDataReadReferenceOrConstruct<FAudioBusAsset>(METASOUND_GET_PARAM_NAME(InParamAudioBusInput));
+				
+				return MakeUnique<TAudioBusReaderOperator<NumChannels>>(InParams.OperatorSettings, AudioMixerOutputFrames, AudioDeviceId, AudioBusIn);
+			}
+			else
+			{
+				UE_LOG(LogMetaSound, Warning, TEXT("Audio bus reader node requires audio device ID '%s' and audio mixer num output frames '%s' environment variables"), *SourceInterface::Environment::DeviceID.ToString(), *SourceInterface::Environment::AudioMixerNumOutputFrames.ToString());
+				return nullptr;
+			}
 		}
 
 		TAudioBusReaderOperator(const FOperatorSettings& InSettings, int32 InAudioMixerOutputFrames, Audio::FDeviceId InAudioDeviceId, const FAudioBusAssetReadRef& InAudioBusAsset)
diff --git a/Engine/Plugins/Runtime/Metasound/Source/MetasoundEngine/Private/MetasoundAudioBusWriterNode.cpp b/Engine/Plugins/Runtime/Metasound/Source/MetasoundEngine/Private/MetasoundAudioBusWriterNode.cpp
--- a/Engine/Plugins/Runtime/Metasound/Source/MetasoundEngine/Private/MetasoundAudioBusWriterNode.cpp
+++ b/Engine/Plugins/Runtime/Metasound/Source/MetasoundEngine/Private/MetasoundAudioBusWriterNode.cpp
@@ -81,18 +81,28 @@ namespace Metasound
 			using namespace AudioBusWriterNode;
 			const FDataReferenceCollection& InputCollection = InParams.InputDataReferences;
 
-			Audio::FDeviceId AudioDeviceId = InParams.Environment.GetValue<Audio::FDeviceId>(SourceInterface::Environment::DeviceID);
-			int32 AudioMixerOutputFrames = InParams.Environment.GetValue<int32>(SourceInterface::Environment::AudioMixerNumOutputFrames);
+			bool bHasEnvironmentVars = InParams.Environment.Contains<Audio::FDeviceId>(SourceInterface::Environment::DeviceID) && InParams.Environment.Contains<int32>(SourceInterface::Environment::AudioMixerNumOutputFrames);
+			
+			if (bHasEnvironmentVars)
+			{
+				Audio::FDeviceId AudioDeviceId = InParams.Environment.GetValue<Audio::FDeviceId>(SourceInterface::Environment::DeviceID);
+				int32 AudioMixerOutputFrames = InParams.Environment.GetValue<int32>(SourceInterface::Environment::AudioMixerNumOutputFrames);
 
-			FAudioBusAssetReadRef AudioBusIn = InputCollection.GetDataReadReferenceOrConstruct<FAudioBusAsset>(METASOUND_GET_PARAM_NAME(InParamAudioBusOutput));
+				FAudioBusAssetReadRef AudioBusIn = InputCollection.GetDataReadReferenceOrConstruct<FAudioBusAsset>(METASOUND_GET_PARAM_NAME(InParamAudioBusOutput));
 
-			TArray<FAudioBufferReadRef> AudioInputs;
-			for (int32 ChannelIndex = 0; ChannelIndex < NumChannels; ++ChannelIndex)
+				TArray<FAudioBufferReadRef> AudioInputs;
+				for (int32 ChannelIndex = 0; ChannelIndex < NumChannels; ++ChannelIndex)
+				{
+					AudioInputs.Add(InputCollection.GetDataReadReferenceOrConstruct<FAudioBuffer>(METASOUND_GET_PARAM_NAME_WITH_INDEX(InParamAudio, ChannelIndex), InParams.OperatorSettings));
+				}
+
+				return MakeUnique<TAudioBusWriterOperator<NumChannels>>(InParams, AudioMixerOutputFrames, AudioDeviceId, MoveTemp(AudioBusIn), MoveTemp(AudioInputs));
+			}
+			else
 			{
-				AudioInputs.Add(InputCollection.GetDataReadReferenceOrConstruct<FAudioBuffer>(METASOUND_GET_PARAM_NAME_WITH_INDEX(InParamAudio, ChannelIndex), InParams.OperatorSettings));
+				UE_LOG(LogMetaSound, Warning, TEXT("Audio bus writer node requires audio device ID '%s' and audio mixer num output frames '%s' environment variables"), *SourceInterface::Environment::DeviceID.ToString(), *SourceInterface::Environment::AudioMixerNumOutputFrames.ToString());
+				return nullptr;
 			}
-
-			return MakeUnique<TAudioBusWriterOperator<NumChannels>>(InParams, AudioMixerOutputFrames, AudioDeviceId, MoveTemp(AudioBusIn), MoveTemp(AudioInputs));
 		}
 
 		TAudioBusWriterOperator(const FCreateOperatorParams& InParams, int32 InAudioMixerOutputFrames, Audio::FDeviceId InAudioDeviceId, FAudioBusAssetReadRef InAudioBusAsset, TArray<FAudioBufferReadRef> InAudioInputs) :
diff --git a/Engine/Plugins/Runtime/Metasound/Source/MetasoundEngineTest/MetasoundEngineTest.Build.cs b/Engine/Plugins/Runtime/Metasound/Source/MetasoundEngineTest/MetasoundEngineTest.Build.cs
--- a/Engine/Plugins/Runtime/Metasound/Source/MetasoundEngineTest/MetasoundEngineTest.Build.cs
+++ b/Engine/Plugins/Runtime/Metasound/Source/MetasoundEngineTest/MetasoundEngineTest.Build.cs
@@ -17,6 +17,7 @@ namespace UnrealBuildTool.Rules
 					"CoreUObject",
 					"Engine",
 					"AudioExtensions",
+					"AudioMixer",
 					"MetasoundGraphCore",
 					"MetasoundGenerator",
 					"MetasoundFrontend",
diff --git a/Engine/Plugins/Runtime/Metasound/Source/MetasoundEngineTest/Private/EngineTestAutomatedNodeTests.cpp b/Engine/Plugins/Runtime/Metasound/Source/MetasoundEngineTest/Private/EngineTestAutomatedNodeTests.cpp
new file mode 100644
--- /dev/null
+++ b/Engine/Plugins/Runtime/Metasound/Source/MetasoundEngineTest/Private/EngineTestAutomatedNodeTests.cpp
@@ -0,0 +1,605 @@
+// Copyright Epic Games, Inc. All Rights Reserved.
+
+#include "AudioDevice.h"
+#include "AudioDeviceManager.h"
+#include "AudioMixerDevice.h"
+#include "Containers/Array.h"
+#include "Containers/UnrealString.h"
+#include "Interfaces/MetasoundFrontendSourceInterface.h"
+#include "Math/NumericLimits.h"
+#include "Misc/AutomationTest.h"
+#include "MetasoundAudioBuffer.h"
+#include "MetasoundFrontendDataTypeRegistry.h"
+#include "MetasoundFrontendRegistries.h"
+#include "MetasoundFrontendSearchEngine.h"
+#include "MetasoundLog.h"
+#include "MetasoundPrimitives.h"
+#include "MetasoundVertex.h"
+#include "MetasoundVertexData.h"
+#include "Templates/UniquePtr.h"
+#include "Tests/AutomationCommon.h"
+
+#if WITH_DEV_AUTOMATION_TESTS
+
+#if WITH_EDITORONLY_DATA
+
+namespace Metasound::MetasoundEngineTestPrivate {
+
+	// Return audio mixer device if one is available
+	Audio::FMixerDevice* GetAudioMixerDevice()
+	{
+		if (FAudioDeviceManager* DeviceManager = FAudioDeviceManager::Get())
+		{
+			if (FAudioDevice* AudioDevice = DeviceManager->GetMainAudioDeviceRaw())
+			{
+				if (AudioDevice->IsAudioMixerEnabled())
+				{
+					return static_cast<Audio::FMixerDevice*>(AudioDevice);
+				}
+			}
+		}
+		return nullptr;
+	}
+
+	// Create an example environment that generally exists for a UMetaSoundSoruce
+	FMetasoundEnvironment GetSourceEnvironment()
+	{
+		using namespace Frontend;
+
+		FMetasoundEnvironment Environment;
+
+		Environment.SetValue<uint32>(SourceInterface::Environment::SoundUniqueID, 0);
+		Environment.SetValue<bool>(SourceInterface::Environment::IsPreview, false);
+		Environment.SetValue<uint64>(SourceInterface::Environment::TransmitterID, 0);
+		Environment.SetValue<FString>(SourceInterface::Environment::GraphName, TEXT("ENGINE_TEST_REGISTERED_NODES"));
+
+		if (Audio::FMixerDevice* MixerDevice = GetAudioMixerDevice())
+		{
+			Environment.SetValue<Audio::FDeviceId>(SourceInterface::Environment::DeviceID, MixerDevice->DeviceID);
+			Environment.SetValue<int32>(SourceInterface::Environment::AudioMixerNumOutputFrames, MixerDevice->GetNumOutputFrames());
+		}
+
+		return Environment;
+	}
+
+	// TTestTypeInfo converts test types to strings
+	template<typename DataType>
+	struct TTestTypeInfo
+	{
+		static FString ToString(const DataType& InData)
+		{
+			return ::LexToString(InData);
+		}
+	};
+
+	template<typename ElementType>
+	struct TTestTypeInfo<TArray<ElementType>>
+	{
+		static FString ToString(const TArray<ElementType>& InData)
+		{
+			return FString::Printf(TEXT("[%s]"), *FString::JoinBy(InData, TEXT(","), &TTestTypeInfo<ElementType>::ToString));
+		}
+	};
+
+	template<>
+	struct TTestTypeInfo<FTime>
+	{
+		static FString ToString(const FTime& InData)
+		{
+			return ::LexToString(InData.GetSeconds());
+		}
+	};
+
+	// TTestTypeValues should return basic bounds for tested input data types.
+	// Similar to TNumericLimits<>
+	template<typename DataType>
+	struct TTestTypeValues
+	{};
+
+	// TArray specialization to defer to a single element array with array element's values
+	template<typename ElementType>
+	struct TTestTypeValues<TArray<ElementType>>
+	{
+		static TArray<ElementType> Min(const FOperatorSettings& InSettings) 
+		{
+			return TArray<ElementType>({TTestTypeValues<ElementType>::Min(InSettings)}); 
+		}
+
+		static TArray<ElementType> Max(const FOperatorSettings& InSettings) 
+		{ 
+			return TArray<ElementType>({TTestTypeValues<ElementType>::Max(InSettings)}); 
+		}
+
+		static TArray<ElementType> Default(const FOperatorSettings& InSettings) 
+		{ 
+			return TArray<ElementType>({TTestTypeValues<ElementType>::Default(InSettings)}); 
+		}
+
+		static TArray<ElementType> Random(const FOperatorSettings& InSettings)
+		{ 
+			return TArray<ElementType>({TTestTypeValues<ElementType>::Random(InSettings)}); 
+		}
+	};
+
+	template<>
+	struct TTestTypeValues<bool>
+	{
+		static bool Min(const FOperatorSettings& InSettings) { return false; }
+		static bool Max(const FOperatorSettings& InSettings) { return true; }
+		static bool Default(const FOperatorSettings& InSettings) { return true; }
+		static bool Random(const FOperatorSettings& InSettings) { return FMath::RandRange(0.f, 1.f) > 0.5; }
+	};
+
+	template<>
+	struct TTestTypeValues<int32>
+	{
+		static int32 Min(const FOperatorSettings& InSettings) { return TNumericLimits<int32>::Min(); }
+		static int32 Max(const FOperatorSettings& InSettings) { return TNumericLimits<int32>::Max(); }
+		static int32 Default(const FOperatorSettings& InSettings) { return 0; }
+		static int32 Random(const FOperatorSettings& InSettings) { return FMath::RandRange(TNumericLimits<int32>::Min(), TNumericLimits<int32>::Max()); }
+	};
+
+	template<>
+	struct TTestTypeValues<float>
+	{
+		static float Min(const FOperatorSettings& InSettings) { return TNumericLimits<float>::Min(); }
+		static float Max(const FOperatorSettings& InSettings) { return TNumericLimits<float>::Max(); }
+		static float Default(const FOperatorSettings& InSettings) { return 0.f; }
+		static float Random(const FOperatorSettings& InSettings) { return FMath::RandRange(TNumericLimits<float>::Min(), TNumericLimits<float>::Max()); }
+	};
+
+	template<>
+	struct TTestTypeValues<FTime>
+	{
+		static FTime Min(const FOperatorSettings& InSettings) { return FTime{TNumericLimits<float>::Min()}; }
+		static FTime Max(const FOperatorSettings& InSettings) { return FTime{TNumericLimits<float>::Max()}; }
+		static FTime Default(const FOperatorSettings& InSettings) { return FTime{0.f}; }
+		static FTime Random(const FOperatorSettings& InSettings) { return FTime{FMath::RandRange(TNumericLimits<float>::Min(), TNumericLimits<float>::Max())}; }
+	};
+
+	template<>
+	struct TTestTypeValues<FTrigger>
+	{
+		static FTrigger Min(const FOperatorSettings& InSettings) { return FTrigger{InSettings, false}; }
+		static FTrigger Max(const FOperatorSettings& InSettings) 
+		{ 
+			FTrigger Trigger{InSettings, false};
+			for (int32 i = 0; i < InSettings.GetNumFramesPerBlock(); i++)
+			{
+				Trigger.TriggerFrame(i);
+			}
+			return Trigger;
+		}
+
+		static FTrigger Default(const FOperatorSettings& InSettings) { return FTrigger{InSettings, true}; }
+		static FTrigger Random(const FOperatorSettings& InSettings) 
+		{ 
+			FTrigger Trigger{InSettings, false};
+			int32 NumTriggers = FMath::RandRange(0, InSettings.GetNumFramesPerBlock());
+			while (NumTriggers > 0)
+			{
+				Trigger.TriggerFrame(FMath::RandRange(0, InSettings.GetNumFramesPerBlock()));
+				NumTriggers--;
+			}
+			return Trigger;
+		}
+	};
+
+	template<>
+	struct TTestTypeValues<FString>
+	{
+		static FString Min(const FOperatorSettings& InSettings) { return TEXT(""); }
+		static FString Max(const FOperatorSettings& InSettings) { return TEXT("THIS IS SUPPOSED TO REPRESENT A MAXIMUM STRING BUT THERE IS NO SUCH THING SO?"); }
+		static FString Default(const FOperatorSettings& InSettings) { return TEXT("TestString"); }
+		static FString Random(const FOperatorSettings& InSettings) { return TEXT("We should probably implement a random string."); }
+	};
+
+	// Interface for mutating data references
+	struct IDataReferenceMutator
+	{
+		virtual ~IDataReferenceMutator() = default;
+		virtual void SetDefault(const FOperatorSettings& InSettings, const FAnyDataReference& InDataRef) const = 0;
+		virtual void SetMax(const FOperatorSettings& InSettings, const FAnyDataReference& InDataRef) const = 0;
+		virtual void SetMin(const FOperatorSettings& InSettings, const FAnyDataReference& InDataRef) const = 0; 
+		virtual void SetRandom(const FOperatorSettings& InSettings, const FAnyDataReference& InDataRef) const = 0;
+		virtual FString ToString(const FAnyDataReference& InDataRef) const = 0;
+	};
+
+	template<typename DataType>
+	struct TDataReferenceMutator : IDataReferenceMutator
+	{
+		virtual void SetDefault(const FOperatorSettings& InSettings, const FAnyDataReference& InDataRef) const override
+		{
+			*InDataRef.GetDataWriteReference<DataType>() = TTestTypeValues<DataType>::Default(InSettings);
+		}
+
+		virtual void SetMax(const FOperatorSettings& InSettings, const FAnyDataReference& InDataRef) const override
+		{
+			*InDataRef.GetDataWriteReference<DataType>() = TTestTypeValues<DataType>::Max(InSettings);
+		}
+
+		virtual void SetMin(const FOperatorSettings& InSettings, const FAnyDataReference& InDataRef) const override
+		{
+			*InDataRef.GetDataWriteReference<DataType>() = TTestTypeValues<DataType>::Min(InSettings);
+		}
+
+		virtual void SetRandom(const FOperatorSettings& InSettings, const FAnyDataReference& InDataRef) const override
+		{
+			*InDataRef.GetDataWriteReference<DataType>() = TTestTypeValues<DataType>::Random(InSettings);
+		}
+
+		virtual FString ToString(const FAnyDataReference& InDataRef) const override
+		{
+			if (const DataType* Data = InDataRef.GetValue<DataType>())
+			{
+				return FString::Printf(TEXT("%s:%s"), *GetMetasoundDataTypeString<DataType>(), *TTestTypeInfo<DataType>::ToString(*Data));
+			}
+			else
+			{
+				// Data references should never be null
+				UE_LOG(LogMetaSound, Error, TEXT("Failed to get data type value of type %s"), *GetMetasoundDataTypeString<DataType>());
+				return TEXT("");
+			}
+		}
+	};
+
+
+	template<typename DataType>
+	void AddDataReferenceMutatorEntryToMap(TMap<FName, TSharedPtr<const IDataReferenceMutator>>& InMap)
+	{
+		InMap.Add(GetMetasoundDataTypeName<DataType>(), MakeShared<TDataReferenceMutator<DataType>>());
+	}
+
+	// Returns map of mutable input types
+	const TMap<FName, TSharedPtr<const IDataReferenceMutator>>& GetDataTypeGeneratorMap()
+	{
+		static TMap<FName, TSharedPtr<const IDataReferenceMutator>> Map;
+
+		AddDataReferenceMutatorEntryToMap<bool>(Map);
+		AddDataReferenceMutatorEntryToMap<int32>(Map);
+		AddDataReferenceMutatorEntryToMap<float>(Map);
+		AddDataReferenceMutatorEntryToMap<FString>(Map);
+		AddDataReferenceMutatorEntryToMap<FTime>(Map);
+		AddDataReferenceMutatorEntryToMap<FTrigger>(Map);
+		AddDataReferenceMutatorEntryToMap<TArray<bool>>(Map);
+		AddDataReferenceMutatorEntryToMap<TArray<int32>>(Map);
+		AddDataReferenceMutatorEntryToMap<TArray<float>>(Map);
+		AddDataReferenceMutatorEntryToMap<TArray<FString>>(Map);
+		AddDataReferenceMutatorEntryToMap<TArray<FTime>>(Map);
+		AddDataReferenceMutatorEntryToMap<TArray<FTrigger>>(Map);
+
+		return Map;
+	}
+
+	// Convenience class for setting node input data reference values to default, min, max or random values. 
+	struct FInputVertexDataTestController
+	{
+		struct FMutableInput
+		{
+			FAnyDataReference DataReference;
+			FVertexName VertexName;
+			TSharedPtr<const IDataReferenceMutator> DataReferenceMutator;
+		};
+
+		FInputVertexDataTestController(const FOperatorSettings& InSettings, const FInputVertexInterface& InInputInterface, const FInputVertexInterfaceData& InInputData)
+		: Settings(InSettings)
+		{
+			const TMap<FName, TSharedPtr<const IDataReferenceMutator>>& GeneratorMap = GetDataTypeGeneratorMap();
+
+			for (const FInputDataVertex& Vertex : InInputInterface)
+			{
+				if (GeneratorMap.Contains(Vertex.DataTypeName))
+				{
+					if (const FAnyDataReference* Ref = InInputData.FindDataReference(Vertex.VertexName))
+					{
+						if (EDataReferenceAccessType::Write == Ref->GetAccessType())
+						{
+							MutableInputs.Add(FMutableInput{*Ref, Vertex.VertexName, GeneratorMap[Vertex.DataTypeName]});
+						}
+					}
+				}
+			}
+		}
+
+		int32 GetNumMutableInputs() const
+		{
+			return MutableInputs.Num();
+		}
+
+		void SetMutableInputsToMin()
+		{
+			for (const FMutableInput& MutableInput : MutableInputs)
+			{
+				MutableInput.DataReferenceMutator->SetMin(Settings, MutableInput.DataReference);
+			}
+			UE_LOG(LogMetaSound, Verbose, TEXT("Setting operator input values:%s%s"), LINE_TERMINATOR, *FString::Join(GetInputValueStrings(), LINE_TERMINATOR));
+		}
+
+		void SetMutableInputsToMax()
+		{
+			for (const FMutableInput& MutableInput : MutableInputs)
+			{
+				MutableInput.DataReferenceMutator->SetMax(Settings, MutableInput.DataReference);
+			}
+			UE_LOG(LogMetaSound, Verbose, TEXT("Setting operator input values:%s%s"), LINE_TERMINATOR, *FString::Join(GetInputValueStrings(), LINE_TERMINATOR));
+		}
+
+		void SetMutableInputsToDefault()
+		{
+			for (const FMutableInput& MutableInput : MutableInputs)
+			{
+				MutableInput.DataReferenceMutator->SetDefault(Settings, MutableInput.DataReference);
+			}
+			UE_LOG(LogMetaSound, Verbose, TEXT("Setting operator input values:%s%s"), LINE_TERMINATOR, *FString::Join(GetInputValueStrings(), LINE_TERMINATOR));
+		}
+
+		void SetMutableInputsToRandom()
+		{
+			for (const FMutableInput& MutableInput : MutableInputs)
+			{
+				MutableInput.DataReferenceMutator->SetRandom(Settings, MutableInput.DataReference);
+			}
+			UE_LOG(LogMetaSound, Verbose, TEXT("Setting operator input values:%s%s"), LINE_TERMINATOR, *FString::Join(GetInputValueStrings(), LINE_TERMINATOR));
+		}
+
+		TArray<FString> GetInputValueStrings() const
+		{
+			TArray<FString> ValueStrings;
+			for (const FMutableInput& MutableInput : MutableInputs)
+			{
+				ValueStrings.Add(FString::Printf(TEXT("%s %s"), *MutableInput.VertexName.ToString(), *MutableInput.DataReferenceMutator->ToString(MutableInput.DataReference)));
+			}
+			return ValueStrings;
+		}
+
+	private:
+
+		FOperatorSettings Settings;
+		TArray<FMutableInput> MutableInputs;
+	};
+
+	static const FLazyName TestNodeName{"TEST_NODE"};
+	static const FLazyName TestVertexName{"TEXT_VERTEX"};
+	static const FGuid TestNodeID{0xA5A5A5A5, 0xA5A5A5A5, 0xA5A5A5A5, 0xA5A5A5A5};
+
+	// Create a node from a node registry key
+	TUniquePtr<INode> CreateNode(const Frontend::FNodeRegistryKey& InNodeRegistryKey)
+	{
+		using namespace Frontend;
+
+		TUniquePtr<INode> Node;
+
+		FMetasoundFrontendRegistryContainer* NodeRegistry = FMetasoundFrontendRegistryContainer::Get();
+		check(nullptr != NodeRegistry);
+		IDataTypeRegistry& DataTypeRegistry = IDataTypeRegistry::Get();
+
+		// Lookup node class metadata to determine how to create this node.
+		FMetasoundFrontendClass NodeClass;
+		if (!NodeRegistry->FindFrontendClassFromRegistered(InNodeRegistryKey, NodeClass))
+		{
+			UE_LOG(LogMetaSound, Error, TEXT("Failed to find registered class with registry key %s"), *InNodeRegistryKey);
+			return MoveTemp(Node);
+		}
+
+		// Build node differently dependent upon node type
+		switch (NodeClass.Metadata.GetType())
+		{
+			case EMetasoundFrontendClassType::VariableDeferredAccessor:
+			case EMetasoundFrontendClassType::VariableAccessor:
+			case EMetasoundFrontendClassType::VariableMutator:
+			case EMetasoundFrontendClassType::External:
+			case EMetasoundFrontendClassType::Graph:
+			{
+				FNodeInitData NodeInitData{TestNodeName, TestNodeID};
+				Node = NodeRegistry->CreateNode(InNodeRegistryKey, NodeInitData);
+			}
+			break;
+
+			case EMetasoundFrontendClassType::Input:
+			{
+				FName DataTypeName = NodeClass.Metadata.GetClassName().Name;
+				FInputNodeConstructorParams NodeInitData
+				{
+					TestNodeName, 
+					TestNodeID, 
+					TestVertexName,
+					NodeClass.Interface.Inputs[0].DefaultLiteral.ToLiteral(DataTypeName)
+				};
+
+				Node = DataTypeRegistry.CreateInputNode(DataTypeName, MoveTemp(NodeInitData));
+			}
+			break;
+
+			case EMetasoundFrontendClassType::Variable:
+			{
+				FName DataTypeName = NodeClass.Metadata.GetClassName().Name;
+				FDefaultLiteralNodeConstructorParams NodeInitData{TestNodeName, TestNodeID, DataTypeRegistry.CreateDefaultLiteral(DataTypeName)};
+				Node = DataTypeRegistry.CreateVariableNode(DataTypeName, MoveTemp(NodeInitData));
+			}
+			break;
+
+			case EMetasoundFrontendClassType::Literal:
+			{
+				FName DataTypeName = NodeClass.Metadata.GetClassName().Name;
+				FDefaultLiteralNodeConstructorParams NodeInitData{TestNodeName, TestNodeID, DataTypeRegistry.CreateDefaultLiteral(DataTypeName)};
+				Node = DataTypeRegistry.CreateLiteralNode(DataTypeName, MoveTemp(NodeInitData));
+			}
+			break;
+
+			case EMetasoundFrontendClassType::Output:
+			{
+				FName DataTypeName = NodeClass.Metadata.GetClassName().Name;
+				FDefaultNamedVertexNodeConstructorParams NodeInitData{TestNodeName, TestNodeID, TestVertexName};
+				Node = DataTypeRegistry.CreateOutputNode(DataTypeName, MoveTemp(NodeInitData));
+			}
+			break;
+
+			case EMetasoundFrontendClassType::Template:
+			default:
+			static_assert(static_cast<int32>(EMetasoundFrontendClassType::Invalid) == 10, "Possible missed EMetasoundFrontendClassType case coverage");
+		}
+
+		return MoveTemp(Node);
+	}
+
+	// Create input vertex interface data for a node.
+	FInputVertexInterfaceData CreateInputVertexInterfaceData(const INode& InNode, const FOperatorSettings& InOperatorSettings)
+	{
+		Frontend::IDataTypeRegistry& DataTypeRegistry = Frontend::IDataTypeRegistry::Get();
+
+		// Populate inputs to node. 
+		FVertexInterface NodeVertexInterface = InNode.GetVertexInterface();
+		FInputVertexInterfaceData NodeInputVertexInterfaceData{NodeVertexInterface.GetInputInterface()};
+
+		for (const FInputDataVertex& InputVertex : NodeVertexInterface.GetInputInterface())
+		{
+			if (InputVertex.AccessType != EVertexAccessType::Reference)
+			{
+				// Not testing constructor inputs.
+				continue;
+			}
+
+			// input data type must be registered in order to create it.
+			if (DataTypeRegistry.IsRegistered(InputVertex.DataTypeName))
+			{
+				Frontend::FDataTypeRegistryInfo DataTypeInfo;
+				ensure(DataTypeRegistry.GetDataTypeInfo(InputVertex.DataTypeName, DataTypeInfo));
+
+				// Can only create data types that are parsable from a literal
+				if (DataTypeInfo.bIsParsable)
+				{
+					FLiteral DefaultLiteral = DataTypeRegistry.CreateDefaultLiteral(InputVertex.DataTypeName);
+					TOptional<FAnyDataReference> DataReference = DataTypeRegistry.CreateDataReference(InputVertex.DataTypeName, EDataReferenceAccessType::Write, DefaultLiteral, InOperatorSettings);
+
+					if (!DataReference.IsSet())
+					{
+						UE_LOG(LogMetaSound, Error, TEXT("Failed to create data reference for data type %s "), *InputVertex.DataTypeName.ToString());
+						continue;
+					}
+
+					NodeInputVertexInterfaceData.BindVertex(InputVertex.VertexName, *DataReference);
+				}
+			}
+		}
+		return NodeInputVertexInterfaceData;
+	}
+}
+
+IMPLEMENT_COMPLEX_AUTOMATION_TEST(FMetasoundAutomatedNodeTest, "Audio.Metasound.AutomatedNodeTest", EAutomationTestFlags::EditorContext | EAutomationTestFlags::StressFilter)
+void FMetasoundAutomatedNodeTest::GetTests(TArray<FString>& OutBeautifiedNames, TArray <FString>& OutTestCommands) const
+{
+	using namespace Metasound;
+
+	// Get all the classes that have been registered
+	Frontend::ISearchEngine& NodeSearchEngine = Frontend::ISearchEngine::Get();
+	TArray<FMetasoundFrontendClass> AllClasses = NodeSearchEngine.FindAllClasses(true /* IncludeAllVersions */);
+
+	FMetasoundFrontendRegistryContainer* NodeRegistry = FMetasoundFrontendRegistryContainer::Get();
+	check(nullptr != NodeRegistry);
+
+	for (const FMetasoundFrontendClass& NodeClass : AllClasses)
+	{
+		// Exclude template classes because they cannot be created directly from the node registry
+		if (NodeClass.Metadata.GetType() == EMetasoundFrontendClassType::Template)
+		{
+			continue;
+		}
+
+		Frontend::FNodeRegistryKey NodeRegistryKey = NodeRegistry->GetRegistryKey(NodeClass.Metadata);
+		
+		OutBeautifiedNames.Add(FString::Printf(TEXT("%s %s"), *NodeClass.Metadata.GetClassName().ToString(), *NodeClass.Metadata.GetVersion().ToString()));
+		// Test commands are node registry keys
+		OutTestCommands.Add(NodeRegistryKey);
+	}
+
+	UE_LOG(LogMetaSound, Verbose, TEXT("Found %d metasound nodes to test"), OutTestCommands.Num());
+}
+
+bool FMetasoundAutomatedNodeTest::RunTest(const FString& InRegistryKey)
+{
+	using namespace Metasound;
+	using namespace MetasoundEngineTestPrivate;
+
+	static const FOperatorSettings OperatorSettings{48000  /* samplerate */, 100.f /* block rate */};
+	static const FMetasoundEnvironment SourceEnvironment = GetSourceEnvironment();
+
+	TUniquePtr<INode> Node = CreateNode(InRegistryKey);
+	if (!Node.IsValid())
+	{
+		AddError(FString::Printf(TEXT("Failed to create node %s from registry"), *InRegistryKey));
+		return false;
+	}
+
+	// Populate inputs to node. 
+	FInputVertexInterfaceData NodeInputVertexInterfaceData = CreateInputVertexInterfaceData(*Node, OperatorSettings);
+	FInputVertexInterface InputInterface = Node->GetVertexInterface().GetInputInterface();
+	FInputVertexDataTestController InputTester(OperatorSettings, InputInterface, NodeInputVertexInterfaceData);
+
+	// Create operator 
+	FBuildOperatorParams BuildParams
+	{
+		*Node,
+		OperatorSettings,
+		NodeInputVertexInterfaceData,
+		SourceEnvironment
+	};
+
+	// Convenience function for testing entire lifecycle of an individual operator 
+	// with a variety of inputs.
+	auto RunTestIteration = [&]() -> bool
+	{
+		FBuildResults BuildResults;
+		TUniquePtr<IOperator> Operator = Node->GetDefaultOperatorFactory()->CreateOperator(BuildParams, BuildResults);
+
+		if (!Operator.IsValid())
+		{
+			AddError(FString::Printf(TEXT("Failed to create operator from node %s."), *InRegistryKey));
+			return false;
+		}
+		
+		IOperator::FExecuteFunction OpExecFunc = Operator->GetExecuteFunction();
+
+		if (OpExecFunc)
+		{
+			OpExecFunc(Operator.Get());
+
+			if (InputTester.GetNumMutableInputs() > 0)
+			{
+				InputTester.SetMutableInputsToDefault();
+				OpExecFunc(Operator.Get());
+
+				InputTester.SetMutableInputsToMin();
+				OpExecFunc(Operator.Get());
+
+				InputTester.SetMutableInputsToMax();
+				OpExecFunc(Operator.Get());
+
+				InputTester.SetMutableInputsToRandom();
+				OpExecFunc(Operator.Get());
+			}
+		}
+
+		return true;
+	};
+
+	bool bSuccess = true;
+	// Test entire operator lifecycle with different starting conditions if
+	// any of the inputs are mutable
+	InputTester.SetMutableInputsToDefault();
+	bSuccess &= RunTestIteration();
+	if (InputTester.GetNumMutableInputs() > 0)
+	{
+		InputTester.SetMutableInputsToMin();
+		bSuccess &= RunTestIteration();
+		InputTester.SetMutableInputsToMax();
+		bSuccess &= RunTestIteration();
+		InputTester.SetMutableInputsToRandom();
+		bSuccess &= RunTestIteration();
+	}
+
+	return bSuccess;
+}
+
+#endif // WITH_EDITORONLY_DATA
+
+#endif //WITH_DEV_AUTOMATION_TESTS
diff --git a/Engine/Plugins/Runtime/Metasound/Source/MetasoundFrontend/Public/MetasoundArrayRandomNode.h b/Engine/Plugins/Runtime/Metasound/Source/MetasoundFrontend/Public/MetasoundArrayRandomNode.h
--- a/Engine/Plugins/Runtime/Metasound/Source/MetasoundFrontend/Public/MetasoundArrayRandomNode.h
+++ b/Engine/Plugins/Runtime/Metasound/Source/MetasoundFrontend/Public/MetasoundArrayRandomNode.h
@@ -185,7 +185,10 @@ namespace Metasound
 			using namespace Frontend;
 
 #if WITH_METASOUND_DEBUG_ENVIRONMENT
-			GraphName = *InParams.Environment.GetValue<FString>(SourceInterface::Environment::GraphName);
+			if (InParams.Environment.Contains<FString>(SourceInterface::Environment::GraphName))
+			{
+				GraphName = *InParams.Environment.GetValue<FString>(SourceInterface::Environment::GraphName);
+			}
 #endif // WITH_METASOUND_DEBUG_ENVIRONMENT
 
 			// Check to see if this is a global shuffler or a local one. 
diff --git a/Engine/Plugins/Runtime/Metasound/Source/MetasoundFrontend/Public/MetasoundArrayShuffleNode.h b/Engine/Plugins/Runtime/Metasound/Source/MetasoundFrontend/Public/MetasoundArrayShuffleNode.h
--- a/Engine/Plugins/Runtime/Metasound/Source/MetasoundFrontend/Public/MetasoundArrayShuffleNode.h
+++ b/Engine/Plugins/Runtime/Metasound/Source/MetasoundFrontend/Public/MetasoundArrayShuffleNode.h
@@ -217,9 +217,7 @@ namespace Metasound
 			, SeedValue(InSeedValue)
 			, bAutoShuffle(bInAutoShuffle)
 			, bEnableSharedState(bInEnableSharedState)
-			, TriggerOnNext(FTriggerWriteRef::CreateNew(InParams.OperatorSettings))
 			, TriggerOnShuffle(FTriggerWriteRef::CreateNew(InParams.OperatorSettings))
-			, TriggerOnReset(FTriggerWriteRef::CreateNew(InParams.OperatorSettings))
 			, OutValue(TDataWriteReferenceFactory<ElementType>::CreateAny(InParams.OperatorSettings))
 		{
 			using namespace Frontend;
@@ -228,29 +226,13 @@ namespace Metasound
 			// Global shuffler will use a namespace to opt into it.
 			PrevSeedValue = *SeedValue;
 
-			const ArrayType& InputArrayRef = *InputArray;
-			int32 ArraySize = InputArrayRef.Num();
-
-			if (ArraySize > 0)
-			{
-				if (*bEnableSharedState)
-				{
-					// Get the environment variable for the unique ID of the sound
-					SharedStateUniqueId = InParams.Environment.GetValue<uint32>(SourceInterface::Environment::SoundUniqueID);
-					check(SharedStateUniqueId != INDEX_NONE);
-
-					FSharedStateShuffleManager& SM = FSharedStateShuffleManager::Get();
-					SM.InitSharedState(SharedStateUniqueId, PrevSeedValue, ArraySize);
-				}
-				else
-				{
-					ArrayIndexShuffler = MakeUnique<FArrayIndexShuffler>(PrevSeedValue, ArraySize);
-				}
-			}
-			else 
+			if (InParams.Environment.Contains<uint32>(SourceInterface::Environment::SoundUniqueID))
 			{
-				UE_LOG(LogMetaSound, Error, TEXT("Array Shuffle: Can't shuffle an empty array"));
+				// Get the environment variable for the unique ID of the sound
+				SharedStateUniqueId = InParams.Environment.GetValue<uint32>(SourceInterface::Environment::SoundUniqueID);
 			}
+
+			UpdateArraySize(InputArray->Num());
 		}
 
 		virtual ~TArrayShuffleOperator() = default;
@@ -276,9 +258,9 @@ namespace Metasound
 			using namespace ArrayNodeShuffleVertexNames;
 
 			FDataReferenceCollection Outputs;
-			Outputs.AddDataReadReference(METASOUND_GET_PARAM_NAME(OutputTriggerOnNext), TriggerOnNext);
+			Outputs.AddDataReadReference(METASOUND_GET_PARAM_NAME(OutputTriggerOnNext), TriggerNext);
 			Outputs.AddDataReadReference(METASOUND_GET_PARAM_NAME(OutputTriggerOnShuffle), TriggerOnShuffle);
-			Outputs.AddDataReadReference(METASOUND_GET_PARAM_NAME(OutputTriggerOnResetSeed), TriggerOnReset);
+			Outputs.AddDataReadReference(METASOUND_GET_PARAM_NAME(OutputTriggerOnResetSeed), TriggerReset);
 			Outputs.AddDataReadReference(METASOUND_GET_PARAM_NAME(OutputShuffledValue), OutValue);
 
 			return Outputs;
@@ -286,52 +268,33 @@ namespace Metasound
 
 		void Execute()
 		{
-			TriggerOnNext->AdvanceBlock();
 			TriggerOnShuffle->AdvanceBlock();
-			TriggerOnReset->AdvanceBlock();
 
 			const ArrayType& InputArrayRef = *InputArray;
- 
-			// Check for a seed change
-			if (PrevSeedValue != *SeedValue)
-			{
-				PrevSeedValue = *SeedValue;
-
-				if (SharedStateUniqueId != INDEX_NONE)
-				{
-					FSharedStateShuffleManager& SM = FSharedStateShuffleManager::Get();
-					SM.SetSeed(SharedStateUniqueId, PrevSeedValue);
-				}
-				else
-				{
-					check(ArrayIndexShuffler.IsValid());
-					ArrayIndexShuffler->SetSeed(PrevSeedValue);
-				}
-			}
-
 			// Don't do anything if our array is empty
 			if (InputArrayRef.Num() == 0)
 			{
 				return;
 			}
 
+			if (PrevArraySize != InputArrayRef.Num())
+			{
+				UpdateArraySize(InputArrayRef.Num());
+			}
+ 
+			// Check for a seed change
+			if (PrevSeedValue != *SeedValue)
+			{
+				Seed(*SeedValue);
+			}
+
 			TriggerReset->ExecuteBlock(
 				[&](int32 StartFrame, int32 EndFrame)
 				{
 				},
 				[this](int32 StartFrame, int32 EndFrame)
 				{
-					if (SharedStateUniqueId != INDEX_NONE)
-					{
-						FSharedStateShuffleManager& SM = FSharedStateShuffleManager::Get();
-						SM.ResetSeed(SharedStateUniqueId);
-					}
-					else
-					{
-						check(ArrayIndexShuffler.IsValid());
-						ArrayIndexShuffler->ResetSeed();
-					}
-					TriggerOnReset->TriggerFrame(StartFrame);
+					Reset();
 				}
 			);
 
@@ -341,16 +304,8 @@ namespace Metasound
 				},
 				[this](int32 StartFrame, int32 EndFrame)
 				{
-					if (SharedStateUniqueId != INDEX_NONE)
-					{
-						FSharedStateShuffleManager& SM = FSharedStateShuffleManager::Get();
-						SM.ShuffleArray(SharedStateUniqueId);
-					}
-					else
-					{
-						check(ArrayIndexShuffler.IsValid());
-						ArrayIndexShuffler->ShuffleArray();
-					}
+					Shuffle();
+
 					TriggerOnShuffle->TriggerFrame(StartFrame);
 				}
 			);
@@ -363,27 +318,12 @@ namespace Metasound
 				{
 					const ArrayType& InputArrayRef = *InputArray;
 
-					bool bShuffleTriggered = false;
 					int32 OutShuffleIndex = INDEX_NONE;
-
-					if (SharedStateUniqueId != INDEX_NONE)
-					{
-						FSharedStateShuffleManager& SM = FSharedStateShuffleManager::Get();
-						bShuffleTriggered = SM.NextValue(SharedStateUniqueId, *bAutoShuffle, OutShuffleIndex);
-					}
-					else
-					{
-						check(ArrayIndexShuffler.IsValid());
-						bShuffleTriggered = ArrayIndexShuffler->NextValue(*bAutoShuffle, OutShuffleIndex);
-					}
-
-					check(OutShuffleIndex != INDEX_NONE);
+					bool bShuffleTriggered = GetNextIndex(OutShuffleIndex);
 
 					// The input array size may have changed, so make sure it's wrapped into range of the input array
 					*OutValue = InputArrayRef[OutShuffleIndex % InputArrayRef.Num()];
 
-					TriggerOnNext->TriggerFrame(StartFrame);
-
 					// Trigger out if the array was auto-shuffled
 					if (bShuffleTriggered)
 					{
@@ -394,6 +334,84 @@ namespace Metasound
 		}
 
 	private:
+		void UpdateArraySize(int32 InSize)
+		{
+			PrevArraySize = InSize;
+
+			if (InSize > 0)
+			{
+				if (SharedStateUniqueId != INDEX_NONE)
+				{
+					FSharedStateShuffleManager& SM = FSharedStateShuffleManager::Get();
+					SM.InitSharedState(SharedStateUniqueId, PrevSeedValue, InSize);
+				}
+				ArrayIndexShuffler = MakeUnique<FArrayIndexShuffler>(PrevSeedValue, InSize);
+			}
+		}
+
+
+		void Seed(int32 InNewSeed)
+		{
+			PrevSeedValue = InNewSeed;
+
+			if (SharedStateUniqueId != INDEX_NONE)
+			{
+				FSharedStateShuffleManager& SM = FSharedStateShuffleManager::Get();
+				SM.SetSeed(SharedStateUniqueId, InNewSeed);
+			}
+
+			if (ArrayIndexShuffler.IsValid())
+			{
+				ArrayIndexShuffler->SetSeed(InNewSeed);
+			}
+		}
+
+		void Reset()
+		{
+			if (SharedStateUniqueId != INDEX_NONE)
+			{
+				FSharedStateShuffleManager& SM = FSharedStateShuffleManager::Get();
+				SM.ResetSeed(SharedStateUniqueId);
+			}
+		
+			if (ArrayIndexShuffler.IsValid())
+			{
+				ArrayIndexShuffler->ResetSeed();
+			}
+		}
+
+		bool GetNextIndex(int32& OutIndex) 
+		{
+			if (*bEnableSharedState && SharedStateUniqueId != INDEX_NONE)
+			{
+				FSharedStateShuffleManager& SM = FSharedStateShuffleManager::Get();
+				return SM.NextValue(SharedStateUniqueId, *bAutoShuffle, OutIndex);
+			}
+			else if (ArrayIndexShuffler.IsValid())
+			{
+				return ArrayIndexShuffler->NextValue(*bAutoShuffle, OutIndex);
+			}
+			else
+			{
+				checkNoEntry();
+				OutIndex = 0;
+				return false;
+			}
+		}
+
+		void Shuffle()
+		{
+			if (SharedStateUniqueId != INDEX_NONE)
+			{
+				FSharedStateShuffleManager& SM = FSharedStateShuffleManager::Get();
+				SM.ShuffleArray(SharedStateUniqueId);
+			}
+
+			if (ArrayIndexShuffler.IsValid())
+			{
+				ArrayIndexShuffler->ShuffleArray();
+			}
+		}
 
 		// Inputs
 		TDataReadReference<FTrigger> TriggerNext;
@@ -405,14 +423,13 @@ namespace Metasound
 		TDataReadReference<bool> bEnableSharedState;
 
 		// Outputs
-		TDataWriteReference<FTrigger> TriggerOnNext;
 		TDataWriteReference<FTrigger> TriggerOnShuffle;
-		TDataWriteReference<FTrigger> TriggerOnReset;
 		TDataWriteReference<ElementType> OutValue;
 
 		// Data
 		TUniquePtr<FArrayIndexShuffler> ArrayIndexShuffler;
 		int32 PrevSeedValue = INDEX_NONE;
+		int32 PrevArraySize = 0;
 		uint32 SharedStateUniqueId = INDEX_NONE;
 	};
 
diff --git a/Engine/Plugins/Runtime/Metasound/Source/MetasoundGraphCore/Public/MetasoundVariableNodes.h b/Engine/Plugins/Runtime/Metasound/Source/MetasoundGraphCore/Public/MetasoundVariableNodes.h
--- a/Engine/Plugins/Runtime/Metasound/Source/MetasoundGraphCore/Public/MetasoundVariableNodes.h
+++ b/Engine/Plugins/Runtime/Metasound/Source/MetasoundGraphCore/Public/MetasoundVariableNodes.h
@@ -327,6 +327,7 @@ namespace Metasound
 		else 
 		{
 			// Nothing to do if there's no input data.
+			UE_LOG(LogMetaSound, Verbose, TEXT("Missing internal variable connection. Failed to create valid \"SetVariable\" operator"));
 			return MakeUnique<FNoOpOperator>();
 		}
 	}
@@ -433,7 +434,7 @@ namespace Metasound
 
 		const FDataReferenceCollection& Inputs = InParams.InputDataReferences;
 
-		if (ensure(Inputs.ContainsDataWriteReference<FVariable>(METASOUND_GET_PARAM_NAME(InputVariable))))
+		if (Inputs.ContainsDataWriteReference<FVariable>(METASOUND_GET_PARAM_NAME(InputVariable)))
 		{
 			TDataWriteReference<FVariable> DelayedVariable = Inputs.GetDataWriteReference<FVariable>(METASOUND_GET_PARAM_NAME(InputVariable));
 
@@ -442,7 +443,7 @@ namespace Metasound
 		else 
 		{
 			// Nothing to do if there's no input data.
-			UE_LOG(LogMetaSound, Warning, TEXT("Missing internal variable connection. Failed to create valid \"GetDelayedVariable\" operator"));
+			UE_LOG(LogMetaSound, Verbose, TEXT("Missing internal variable connection. Failed to create valid \"GetDelayedVariable\" operator"));
 			return MakeUnique<FNoOpOperator>();
 		}
 	}
@@ -556,7 +557,7 @@ namespace Metasound
 		else 
 		{
 			// Nothing to do if there's no input data.
-			UE_LOG(LogMetaSound, Warning, TEXT("Missing internal variable connection. Failed to create valid \"GetVariable\" operator"));
+			UE_LOG(LogMetaSound, Verbose, TEXT("Missing internal variable connection. Failed to create valid \"GetVariable\" operator"));
 			return MakeUnique<FNoOpOperator>();
 		}
 	}
diff --git a/Engine/Plugins/Runtime/Metasound/Source/MetasoundStandardNodes/Private/MetasoundADSREnvelopeNode.cpp b/Engine/Plugins/Runtime/Metasound/Source/MetasoundStandardNodes/Private/MetasoundADSREnvelopeNode.cpp
--- a/Engine/Plugins/Runtime/Metasound/Source/MetasoundStandardNodes/Private/MetasoundADSREnvelopeNode.cpp
+++ b/Engine/Plugins/Runtime/Metasound/Source/MetasoundStandardNodes/Private/MetasoundADSREnvelopeNode.cpp
@@ -614,10 +614,10 @@ namespace Metasound
 			float AttackTimeSeconds = AttackTime->GetSeconds();
 			float DecayTimeSeconds = DecayTime->GetSeconds();
 			float ReleaseTimeSeconds = ReleaseTime->GetSeconds();
-			EnvState.AttackSampleCount = FMath::Max(1, SampleRate * FMath::Max(0.0f, AttackTimeSeconds));
-			EnvState.DecaySampleCount = SampleRate * FMath::Max(0.0f, DecayTimeSeconds);
+			EnvState.AttackSampleCount = FMath::Max(1, static_cast<int32>(SampleRate * FMath::Max(0.0f, AttackTimeSeconds)));
+			EnvState.DecaySampleCount = FMath::Max(0, static_cast<int32>(SampleRate * FMath::Max(0.0f, DecayTimeSeconds)));
 			EnvState.SustainLevel = FMath::Max(0.0f, *SustainLevel);
-			EnvState.ReleaseSampleCount = SampleRate * FMath::Max(0.0f, ReleaseTimeSeconds);
+			EnvState.ReleaseSampleCount = FMath::Max(0, static_cast<int32>(SampleRate * FMath::Max(0.0f, ReleaseTimeSeconds)));
 			EnvState.AttackCurveFactor = FMath::Max(KINDA_SMALL_NUMBER, *AttackCurveFactor);
 			EnvState.DecayCurveFactor = FMath::Max(KINDA_SMALL_NUMBER, *DecayCurveFactor);
 			EnvState.ReleaseCurveFactor = FMath::Max(KINDA_SMALL_NUMBER, *ReleaseCurveFactor);
diff --git a/Engine/Plugins/Runtime/Metasound/Source/MetasoundStandardNodes/Private/MetasoundDebugLogNode.cpp b/Engine/Plugins/Runtime/Metasound/Source/MetasoundStandardNodes/Private/MetasoundDebugLogNode.cpp
--- a/Engine/Plugins/Runtime/Metasound/Source/MetasoundStandardNodes/Private/MetasoundDebugLogNode.cpp
+++ b/Engine/Plugins/Runtime/Metasound/Source/MetasoundStandardNodes/Private/MetasoundDebugLogNode.cpp
@@ -99,11 +99,17 @@ namespace Metasound
 				TDataReadReference<FString> Label = InputCollection.GetDataReadReferenceOrConstructWithVertexDefault<FString>(InputInterface, METASOUND_GET_PARAM_NAME(InputLabel), InParams.OperatorSettings);
 				TDataReadReference<PrintLogType> ValueToLog = InputCollection.GetDataReadReferenceOrConstructWithVertexDefault<PrintLogType>(InputInterface, METASOUND_GET_PARAM_NAME(InputValueToLog), InParams.OperatorSettings);
 
-				FString GraphNameFull = InParams.Environment.GetValue<FString>(Frontend::SourceInterface::Environment::GraphName);
-				TArray<FString> ParsedString;
-				GraphNameFull.ParseIntoArray(ParsedString, TEXT("."), true);
-
-				const FString& GraphName = ParsedString.Last();
+				FString GraphName;
+				if (InParams.Environment.Contains<FString>(Frontend::SourceInterface::Environment::GraphName))
+				{
+					FString GraphNameFull = InParams.Environment.GetValue<FString>(Frontend::SourceInterface::Environment::GraphName);
+					TArray<FString> ParsedString;
+					GraphNameFull.ParseIntoArray(ParsedString, TEXT("."), true);
+					if (ParsedString.Num() > 0)
+					{
+						GraphName = ParsedString.Last();
+					}
+				}
 
 				return MakeUnique<TPrintLogOperator<PrintLogType>>(InParams.OperatorSettings, Trigger, Label, ValueToLog, GraphName);
 			}
diff --git a/Engine/Plugins/Runtime/Metasound/Source/MetasoundStandardNodes/Private/MetasoundOscillators.h b/Engine/Plugins/Runtime/Metasound/Source/MetasoundStandardNodes/Private/MetasoundOscillators.h
--- a/Engine/Plugins/Runtime/Metasound/Source/MetasoundStandardNodes/Private/MetasoundOscillators.h
+++ b/Engine/Plugins/Runtime/Metasound/Source/MetasoundStandardNodes/Private/MetasoundOscillators.h
@@ -43,14 +43,14 @@ namespace Metasound
 		{
 			FORCEINLINE void operator()(float& InPhase) const
 			{
-				while (InPhase >= 1.f)
+				if (InPhase >= 1.f)
 				{
-					InPhase -= 1.f;
+					InPhase -= FMath::TruncToFloat(InPhase);
 				}
 
-				while (InPhase < 0.0f)
+				if (InPhase < 0.0f)
 				{
-					InPhase += 1.f;
+					InPhase -= FMath::TruncToFloat(InPhase) + 1.f;
 				}
 			}
 		};
diff --git a/Engine/Plugins/Runtime/WaveTable/Source/WaveTable/Public/WaveTable.h b/Engine/Plugins/Runtime/WaveTable/Source/WaveTable/Public/WaveTable.h
--- a/Engine/Plugins/Runtime/WaveTable/Source/WaveTable/Public/WaveTable.h
+++ b/Engine/Plugins/Runtime/WaveTable/Source/WaveTable/Public/WaveTable.h
@@ -38,7 +38,7 @@ namespace WaveTable
 	public:
 		FWaveTable() = default;
 
-		FWaveTable(const TArray<float>& InSamples, const float InFinalValue)
+		FWaveTable(const TArray<float>& InSamples, const float InFinalValue=0.f)
 			: Samples(InSamples)
 			, FinalValue(InFinalValue)
 		{
diff --git a/Engine/Source/Runtime/SignalProcessing/Private/SineWaveTableOsc.cpp b/Engine/Source/Runtime/SignalProcessing/Private/SineWaveTableOsc.cpp
--- a/Engine/Source/Runtime/SignalProcessing/Private/SineWaveTableOsc.cpp
+++ b/Engine/Source/Runtime/SignalProcessing/Private/SineWaveTableOsc.cpp
@@ -89,8 +89,8 @@ namespace Audio
 		for (int32 SampleIndex = 0; SampleIndex < NumSamples; ++SampleIndex)
 		{
 			// Interpolate between two samples
-			const int32 ReadIndexPrev = (int32)ReadIndex;
-			const float Alpha = ReadIndex - (float)ReadIndexPrev;
+			const int32 ReadIndexPrev = FMath::Clamp((int32)ReadIndex, 0, WaveTableBuffer.Num() - 2);
+			const float Alpha = FMath::Clamp(ReadIndex - (float)ReadIndexPrev, 0.f, 1.f);
 
 			const int32 ReadIndexNext = (ReadIndexPrev + 1) % WaveTableBuffer.Num();
 			OutBuffer[SampleIndex] = FMath::Lerp(WaveTableBuffer[ReadIndexPrev], WaveTableBuffer[ReadIndexNext], Alpha);
