Fix for race condition in async metasound builder
#rnx

diff --git a/Engine/Plugins/Runtime/Metasound/Source/MetasoundGenerator/Private/MetasoundGenerator.cpp b/Engine/Plugins/Runtime/Metasound/Source/MetasoundGenerator/Private/MetasoundGenerator.cpp
--- a/Engine/Plugins/Runtime/Metasound/Source/MetasoundGenerator/Private/MetasoundGenerator.cpp
+++ b/Engine/Plugins/Runtime/Metasound/Source/MetasoundGenerator/Private/MetasoundGenerator.cpp
@@ -66,55 +66,11 @@ namespace Metasound
 		{
 			// Create graph analyzer
 			TUniquePtr<FGraphAnalyzer> GraphAnalyzer = BuildGraphAnalyzer(MoveTemp(BuildResults.InternalDataReferences));
+	
+			// Collect data for generator
+			FMetasoundGeneratorData GeneratorData = BuildGeneratorData(InitParams, MoveTemp(GraphOperator), MoveTemp(GraphAnalyzer));
 
-			// Gather relevant input and output references
-			FVertexInterfaceData VertexData(InitParams.Graph->GetVertexInterface());
-			GraphOperator->Bind(VertexData);
-
-			// Get inputs
-			FTriggerWriteRef PlayTrigger = VertexData.GetInputs().GetOrConstructDataWriteReference<FTrigger>(SourceInterface::Inputs::OnPlay, InitParams.OperatorSettings, false);
-
-			// Get outputs
-			TArray<FAudioBufferReadRef> OutputBuffers = FindOutputAudioBuffers(VertexData);
-			FTriggerReadRef FinishTrigger = TDataReadReferenceFactory<FTrigger>::CreateExplicitArgs(InitParams.OperatorSettings, false);
-
-			if (InitParams.Graph->GetVertexInterface().GetOutputInterface().Contains(SourceOneShotInterface::Outputs::OnFinished))
-			{
-				// Only attempt to retrieve the on finished trigger if it exists.
-				// Attempting to retrieve a data reference from a non-existent vertex 
-				// will log an error. 
-				FinishTrigger = VertexData.GetOutputs().GetOrConstructDataReadReference<FTrigger>(SourceOneShotInterface::Outputs::OnFinished, InitParams.OperatorSettings, false);
-			}
-
-			// Create the parameter setter map so parameter packs can be cracked
-			// open and distributed as appropriate...
-			TMap<FName, FParameterSetter> ParameterSetters;
-			FInputVertexInterfaceData& GraphInputs = VertexData.GetInputs();
-			for (auto InputIterator = GraphInputs.begin(); InputIterator != GraphInputs.end(); ++InputIterator)
-			{
-				const FInputDataVertex& InputVertex = (*InputIterator).GetVertex();
-				const Frontend::IParameterAssignmentFunction& Setter = IDataTypeRegistry::Get().GetRawAssignmentFunction(InputVertex.DataTypeName);
-				if (Setter)
-				{
-					FParameterSetter ParameterSetter(InputVertex.DataTypeName,
-						(*InputIterator).GetDataReference()->GetRaw(),
-						Setter);
-					ParameterSetters.Add(InputVertex.VertexName, ParameterSetter);
-				}
-			}
-
-			// Set data needed for graph
-			FMetasoundGeneratorData GeneratorData
-			{
-				InitParams.OperatorSettings,
-				MoveTemp(GraphOperator),
-				MoveTemp(ParameterSetters),
-				MoveTemp(GraphAnalyzer),
-				MoveTemp(OutputBuffers),
-				MoveTemp(PlayTrigger),
-				MoveTemp(FinishTrigger),
-			};
-
+			// Update generator with new metasound graph
 			Generator->SetPendingGraph(MoveTemp(GeneratorData), bTriggerGenerator);
 		}
 		else 
@@ -128,8 +84,67 @@ namespace Metasound
 		InitParams.Release();
 	}
 
+	FMetasoundGeneratorData FAsyncMetaSoundBuilder::BuildGeneratorData(const FMetasoundGeneratorInitParams& InInitParams, TUniquePtr<IOperator> InGraphOperator, TUniquePtr<Frontend::FGraphAnalyzer> InAnalyzer) const
+	{
+		using namespace Audio;
+		using namespace Frontend;
+
+		METASOUND_TRACE_CPUPROFILER_EVENT_SCOPE_TEXT(TEXT("AsyncMetaSoundBuilder::BuildGeneratorData"));
+
+		checkf(InGraphOperator.IsValid(), TEXT("Graph operator must be a valid object"));
+
+		// Gather relevant input and output references
+		FVertexInterfaceData VertexData(InInitParams.Graph->GetVertexInterface());
+		InGraphOperator->Bind(VertexData);
+
+		// Get inputs
+		FTriggerWriteRef PlayTrigger = VertexData.GetInputs().GetOrConstructDataWriteReference<FTrigger>(SourceInterface::Inputs::OnPlay, InInitParams.OperatorSettings, false);
+
+		// Get outputs
+		TArray<FAudioBufferReadRef> OutputBuffers = FindOutputAudioBuffers(VertexData);
+		FTriggerReadRef FinishTrigger = TDataReadReferenceFactory<FTrigger>::CreateExplicitArgs(InInitParams.OperatorSettings, false);
+
+		if (InInitParams.Graph->GetVertexInterface().GetOutputInterface().Contains(SourceOneShotInterface::Outputs::OnFinished))
+		{
+			// Only attempt to retrieve the on finished trigger if it exists.
+			// Attempting to retrieve a data reference from a non-existent vertex 
+			// will log an error. 
+			FinishTrigger = VertexData.GetOutputs().GetOrConstructDataReadReference<FTrigger>(SourceOneShotInterface::Outputs::OnFinished, InitParams.OperatorSettings, false);
+		}
+
+		// Create the parameter setter map so parameter packs can be cracked
+		// open and distributed as appropriate...
+		TMap<FName, FParameterSetter> ParameterSetters;
+		FInputVertexInterfaceData& GraphInputs = VertexData.GetInputs();
+		for (auto InputIterator = GraphInputs.begin(); InputIterator != GraphInputs.end(); ++InputIterator)
+		{
+			const FInputDataVertex& InputVertex = (*InputIterator).GetVertex();
+			const Frontend::IParameterAssignmentFunction& Setter = IDataTypeRegistry::Get().GetRawAssignmentFunction(InputVertex.DataTypeName);
+			if (Setter)
+			{
+				FParameterSetter ParameterSetter(InputVertex.DataTypeName,
+					(*InputIterator).GetDataReference()->GetRaw(),
+					Setter);
+				ParameterSetters.Add(InputVertex.VertexName, ParameterSetter);
+			}
+		}
+
+		// Set data needed for graph
+		return FMetasoundGeneratorData 
+		{
+			InInitParams.OperatorSettings,
+			MoveTemp(InGraphOperator),
+			MoveTemp(ParameterSetters),
+			MoveTemp(InAnalyzer),
+			MoveTemp(OutputBuffers),
+			MoveTemp(PlayTrigger),
+			MoveTemp(FinishTrigger),
+		};
+	}
+
 	TUniquePtr<IOperator> FAsyncMetaSoundBuilder::BuildGraphOperator(TArray<FAudioParameter>&& InParameters, FBuildResults& OutBuildResults) const
 	{
+		METASOUND_TRACE_CPUPROFILER_EVENT_SCOPE_TEXT(TEXT("AsyncMetaSoundBuilder::BuildGraphOperator"));
 		using namespace Frontend;
 
 		// Set input data based on the input parameters and the input interface
@@ -171,6 +186,7 @@ namespace Metasound
 
 	TUniquePtr<Frontend::FGraphAnalyzer> FAsyncMetaSoundBuilder::BuildGraphAnalyzer(TMap<FGuid, FDataReferenceCollection>&& InInternalDataReferences) const
 	{
+		METASOUND_TRACE_CPUPROFILER_EVENT_SCOPE_TEXT(TEXT("AsyncMetaSoundBuilder::BuildGraphAnalyzer"));
 		using namespace Frontend;
 
 		if (InitParams.BuilderSettings.bPopulateInternalDataReferences)
diff --git a/Engine/Plugins/Runtime/Metasound/Source/MetasoundGenerator/Public/MetasoundGenerator.h b/Engine/Plugins/Runtime/Metasound/Source/MetasoundGenerator/Public/MetasoundGenerator.h
--- a/Engine/Plugins/Runtime/Metasound/Source/MetasoundGenerator/Public/MetasoundGenerator.h
+++ b/Engine/Plugins/Runtime/Metasound/Source/MetasoundGenerator/Public/MetasoundGenerator.h
@@ -82,6 +82,7 @@ namespace Metasound
 		}
 
 	private:
+		FMetasoundGeneratorData BuildGeneratorData(const FMetasoundGeneratorInitParams& InInitParams, TUniquePtr<IOperator> InOperator, TUniquePtr<Frontend::FGraphAnalyzer> InAnalyzer) const;
 		TUniquePtr<IOperator> BuildGraphOperator(TArray<FAudioParameter>&& InParameters, FBuildResults& OutBuildResults) const;
 		TUniquePtr<Frontend::FGraphAnalyzer> BuildGraphAnalyzer(TMap<FGuid, FDataReferenceCollection>&& InInternalDataReferences) const;
 		void LogBuildErrors(const FBuildResults& InBuildResults) const;
