Adding ability to debug metasounds which are causing issues.

diff --git a/Engine/Plugins/Runtime/Metasound/Source/MetasoundGenerator/Private/MetasoundGenerator.cpp b/Engine/Plugins/Runtime/Metasound/Source/MetasoundGenerator/Private/MetasoundGenerator.cpp
--- a/Engine/Plugins/Runtime/Metasound/Source/MetasoundGenerator/Private/MetasoundGenerator.cpp
+++ b/Engine/Plugins/Runtime/Metasound/Source/MetasoundGenerator/Private/MetasoundGenerator.cpp
@@ -3,6 +3,7 @@
 
 #include "Analysis/MetasoundFrontendVertexAnalyzer.h"
 #include "AudioParameter.h"
+#include "DSP/FloatArrayMath.h"
 #include "Interfaces/MetasoundFrontendSourceInterface.h"
 #include "MetasoundOperatorBuilder.h"
 #include "MetasoundOperatorInterface.h"
@@ -14,13 +15,63 @@
 #include "MetasoundVertexData.h"
 #include "MetasoundFrontendDataTypeRegistry.h"
 
+#ifndef METASOUNDGENERATOR_ENABLE_INVALID_SAMPLE_VALUE_LOGGING 
+#define METASOUNDGENERATOR_ENABLE_INVALID_SAMPLE_VALUE_LOGGING !UE_BUILD_SHIPPING
+#endif
 
 namespace Metasound
 {
 	namespace ConsoleVariables
 	{
 		static bool bEnableAsyncMetaSoundGeneratorBuilder = true;
+#if METASOUNDGENERATOR_ENABLE_INVALID_SAMPLE_VALUE_LOGGING
+		static bool bEnableMetaSoundGeneratorNonFiniteLogging = false;
+		static bool bEnableMetaSoundGeneratorInvalidSampleValueLogging = false;
+		static float MetasoundGeneratorSampleValueThreshold = 2.f;
+#endif // #if METASOUNDGENERATOR_ENABLE_INVALID_SAMPLE_VALUE_LOGGING
 	}
+
+#if METASOUNDGENERATOR_ENABLE_INVALID_SAMPLE_VALUE_LOGGING
+	namespace MetasoundGeneratorPrivate
+	{
+		void LogInvalidAudioSampleValues(const FString& InMetaSoundName, const TArray<FAudioBufferReadRef>& InAudioBuffers)
+		{
+			if (ConsoleVariables::bEnableMetaSoundGeneratorNonFiniteLogging || ConsoleVariables::bEnableMetaSoundGeneratorInvalidSampleValueLogging)
+			{
+				const int32 NumChannels = InAudioBuffers.Num();
+
+				// Check outputs for non finite values if any sample value logging is enabled.
+				for (int32 ChannelIndex = 0; ChannelIndex < NumChannels; ChannelIndex++)
+				{
+					const FAudioBuffer& Buffer = *InAudioBuffers[ChannelIndex];
+					const float* Data = Buffer.GetData();
+					const int32 Num = Buffer.Num();
+
+					for (int32 i = 0; i < Num; i++)
+					{
+						if (!FMath::IsFinite(Data[i]))
+						{
+							UE_LOG(LogMetaSound, Error, TEXT("Found non-finite sample (%f) in channel %d of MetaSound %s"), Data[i], ChannelIndex, *InMetaSoundName);
+							break;
+						}
+					}
+
+					// Only check threshold if explicitly enabled
+					if (ConsoleVariables::bEnableMetaSoundGeneratorInvalidSampleValueLogging)
+					{
+						const float Threshold = FMath::Abs(ConsoleVariables::MetasoundGeneratorSampleValueThreshold);
+						const float MaxAbsValue = Audio::ArrayMaxAbsValue(Buffer);
+						if (MaxAbsValue > Threshold)
+						{
+							UE_LOG(LogMetaSound, Warning, TEXT("Found sample (absolute value: %f) exceeding threshold (%f) in channel %d of MetaSound %s"), MaxAbsValue, Threshold, ChannelIndex, *InMetaSoundName);
+
+						}
+					}
+				}
+			}
+		}
+	}
+#endif // #if METASOUNDGENERATOR_ENABLE_INVALID_SAMPLE_VALUE_LOGGING
 }
 
 FAutoConsoleVariableRef CVarMetaSoundEnableAsyncGeneratorBuilder(
@@ -30,6 +81,30 @@ FAutoConsoleVariableRef CVarMetaSoundEnableAsyncGeneratorBuilder(
 	TEXT("Default: true"),
 	ECVF_Default);
 
+#if !UE_BUILD_SHIPPING
+
+FAutoConsoleVariableRef CVarMetaSoundEnableGeneratorNonFiniteLogging(
+	TEXT("au.MetaSound.EnableGeneratorNonFiniteLogging"),
+	Metasound::ConsoleVariables::bEnableMetaSoundGeneratorNonFiniteLogging,
+	TEXT("Enables logging of non-finite (NaN/inf) audio samples values produced from a FMetaSoundGenerator\n")
+	TEXT("Default: false"),
+	ECVF_Default);
+
+FAutoConsoleVariableRef CVarMetaSoundEnableGeneratorInvalidSampleValueLogging(
+	TEXT("au.MetaSound.EnableGeneratorInvalidSampleValueLogging"),
+	Metasound::ConsoleVariables::bEnableMetaSoundGeneratorInvalidSampleValueLogging,
+	TEXT("Enables logging of audio samples values produced from a FMetaSoundGenerator which exceed the absolute sample value threshold\n")
+	TEXT("Default: false"),
+	ECVF_Default);
+
+FAutoConsoleVariableRef CVarMetaSoundGeneratorSampleValueThrehshold(
+	TEXT("au.MetaSound.GeneratorSampleValueThreshold"),
+	Metasound::ConsoleVariables::MetasoundGeneratorSampleValueThreshold,
+	TEXT("If invalid sample value logging is enabled, this sets the maximum abs value threshold for logging samples\n")
+	TEXT("Default: 2.0"),
+	ECVF_Default);
+
+#endif // #if !UE_BUILD_SHIPPING
 
 namespace Metasound
 {
@@ -519,6 +594,11 @@ namespace Metasound
 			InterleavedAudioBuffer.AddUninitialized(NumSamplesPerExecute);
 		}
 
+#if METASOUNDGENERATOR_ENABLE_INVALID_SAMPLE_VALUE_LOGGING
+		MetasoundGeneratorPrivate::LogInvalidAudioSampleValues(MetasoundName, GraphOutputAudio);
+#endif // #if METASOUNDGENERATOR_ENABLE_INVALID_SAMPLE_VALUE_LOGGING
+		
+
 		// Iterate over channels
 		for (int32 ChannelIndex = 0; ChannelIndex < NumChannels; ChannelIndex++)
 		{
