Fix for SetSeed not being determinstic with the same seed.

diff --git a/Engine/Plugins/Runtime/Metasound/Source/MetasoundStandardNodes/Private/MetasoundLowFrequencyNoise.cpp b/Engine/Plugins/Runtime/Metasound/Source/MetasoundStandardNodes/Private/MetasoundLowFrequencyNoise.cpp
--- a/Engine/Plugins/Runtime/Metasound/Source/MetasoundStandardNodes/Private/MetasoundLowFrequencyNoise.cpp
+++ b/Engine/Plugins/Runtime/Metasound/Source/MetasoundStandardNodes/Private/MetasoundLowFrequencyNoise.cpp
@@ -186,12 +186,23 @@ namespace Metasound
 			, Pins{ MoveTemp(InPins) }
 			, ScaledOutput{ FFloatWriteRef::CreateNew(0.f) }
 			, NormalizedOutput{ FFloatWriteRef::CreateNew(0.f) }
-		{
+		{	
+			SetSeedAndFillValues();
+		}
+
+		void SetSeedAndFillValues()
+		{		
 			SetSeed(*Pins.Seed);
-			for (int32 i = 0; i < Values.Num(); ++i)
+		
+			// Every time we change the seed, in order for the results to be deterministic 
+			// we must also also flush and fill the value/delay line. This will create a discontinuation in the 
+			// interpolation, but will be consistent.
+			for(int32 i = 0; i < Values.Num(); ++i)
 			{
 				Values.Push(GenerateRandomValue());
-			}	
+			}
+
+			// The random jitter should be chosen again following the seed change.
 			CurrentRandomRate = GenerateRandomRate();
 		}
 
@@ -234,10 +245,12 @@ namespace Metasound
 			if (InSeedValue != -1)
 			{
 				RandomStream.Initialize(InSeedValue);
+				RandomJitter.Initialize(InSeedValue);
 			}
 			else
 			{
 				RandomStream.GenerateNewSeed();
+				RandomJitter.GenerateNewSeed();
 			}
 			CurrentSeed = InSeedValue;
 		}
@@ -258,7 +271,7 @@ namespace Metasound
 		float GenerateRandomRate()
 		{
 			const float Nyquist = BlockRate / 2.f;
-			float RandomFraction = RandomStream.GetFraction();			// NOTE <-- Might need its own stream if we want it to remain deterministic.
+			float RandomFraction = RandomJitter.GetFraction();
 			float ClampedJitter =  FMath::Clamp(*Pins.RateJitter, 0.f, 1.f);
 			float RangeMin = FMath::Max(*Pins.Rate - (*Pins.Rate * ClampedJitter), 0.f);
 			float RangeMax = FMath::Min(*Pins.Rate + (*Pins.Rate * ClampedJitter), Nyquist);
@@ -292,7 +305,7 @@ namespace Metasound
 				}
 				case ELowFrequencyNoiseInterpType::Linear:
 				{
-					return FMath::Lerp(Values.GetValue(-1), Values.GetValue(0), Phase);
+					return FMath::Lerp(Values.GetValue(0), Values.GetValue(1), Phase);
 				}
 				default:
 				case ELowFrequencyNoiseInterpType::Cubic:
@@ -314,7 +327,7 @@ namespace Metasound
 		{
 			if (CurrentSeed != *Pins.Seed || Pins.ResetSeed->IsTriggeredInBlock())
 			{
-				SetSeed(*Pins.Seed);
+				SetSeedAndFillValues();
 			}
 		}
 
@@ -386,10 +399,10 @@ namespace Metasound
 		float Phase = 0.f;
 		float CurrentRandomRate = 0.f;
 		float CurrentRate = 0.f;
-		float LastOutputValue  = 0.f;
 		int32 CurrentSeed = 0;
 		FValueFIFO Values;
 		FRandomStream RandomStream;
+		FRandomStream RandomJitter;
 		FPinReadRefs Pins;
 		FFloatWriteRef ScaledOutput;		// Scaled output value (min/max)
 		FFloatWriteRef NormalizedOutput;	// Unit value (0..1)
