Fix Metasound SuperOscillator node discontinuities when LimitOutput option is checked

diff --git a/Engine/Plugins/Runtime/Metasound/Source/MetasoundStandardNodes/Private/MetasoundSuperOscillatorNode.cpp b/Engine/Plugins/Runtime/Metasound/Source/MetasoundStandardNodes/Private/MetasoundSuperOscillatorNode.cpp
--- a/Engine/Plugins/Runtime/Metasound/Source/MetasoundStandardNodes/Private/MetasoundSuperOscillatorNode.cpp
+++ b/Engine/Plugins/Runtime/Metasound/Source/MetasoundStandardNodes/Private/MetasoundSuperOscillatorNode.cpp
@@ -304,7 +304,7 @@ namespace Metasound
 		void SetLimiterSettings(Audio::FDynamicsProcessor& InLimiter) const
 		{
 			InLimiter.Init(SampleRate, 1);
-			InLimiter.SetLookaheadMsec(0.0f);
+			InLimiter.SetLookaheadMsec(5.0f);
 			InLimiter.SetAttackTime(20.0f);
 			InLimiter.SetReleaseTime(200.0f);
 			InLimiter.SetThreshold(-6.0f);
@@ -370,11 +370,11 @@ namespace Metasound
 					FillOutputs(Gain, PanRatio * Args.StereoWidth);
 
 					bFirstVoiceProcessed = true;
-
-					if (*bLimit)
-					{
-						LimitOutput();
-					}
+				}
+				
+				if (*bLimit)
+				{
+					LimitOutput();
 				}
 			}
 		}
