[Metasounds] Adding a debug node that outputs a trigger when a Nan is detected.

diff --git a/Engine/Plugins/Runtime/Metasound/Source/MetasoundStandardNodes/Private/MetasoundDebugTriggerOnNanNode.cpp b/Engine/Plugins/Runtime/Metasound/Source/MetasoundStandardNodes/Private/MetasoundDebugTriggerOnNanNode.cpp
new file mode 100644
--- /dev/null
+++ b/Engine/Plugins/Runtime/Metasound/Source/MetasoundStandardNodes/Private/MetasoundDebugTriggerOnNanNode.cpp
@@ -0,0 +1,171 @@
+// Copyright Epic Games, Inc. All Rights Reserved.
+
+#include "CoreMinimal.h"
+#include "MetasoundBuilderInterface.h"
+#include "MetasoundDataReferenceCollection.h"
+#include "MetasoundExecutableOperator.h"
+#include "MetasoundFacade.h"
+#include "MetasoundNode.h"
+#include "MetasoundNodeInterface.h"
+#include "MetasoundNodeRegistrationMacro.h"
+#include "MetasoundOperatorInterface.h"
+#include "MetasoundPrimitives.h"
+#include "MetasoundSampleCounter.h"
+#include "MetasoundStandardNodesNames.h"
+#include "MetasoundStandardNodesCategories.h"
+#include "MetasoundParamHelper.h"
+
+#define LOCTEXT_NAMESPACE "MetasoundStandardNodes_TriggerOnNanNode"
+
+namespace Metasound
+{
+
+	namespace TriggerOnNanVertexNames
+	{
+		METASOUND_PARAM(AudioInput, "Audio In", "The Audio Buffer to be evaluated for NANs.");
+		METASOUND_PARAM(TriggerOnce, "Trigger Once", "Whether we should trigger for every NAN, or just once for the lifecycle of the Metasound Source");
+		
+		METASOUND_PARAM(OutputTriggerOnNan, "Trigger Out", "Trigger output on the audio frame containing the Nan");
+	}
+
+	class FTriggerOnNanOperator : public TExecutableOperator<FTriggerOnNanOperator>
+	{
+	public:
+		static const FNodeClassMetadata& GetNodeInfo();
+		static const FVertexInterface& GetVertexInterface();
+		static TUniquePtr<IOperator> CreateOperator(const FCreateOperatorParams& InParams, FBuildErrorArray& OutErrors);
+
+		FTriggerOnNanOperator(const FOperatorSettings& InSettings, const FAudioBufferReadRef& InAudioInput, const FBoolReadRef& InTriggerOnce);
+
+		virtual FDataReferenceCollection GetInputs() const override;
+
+		virtual FDataReferenceCollection GetOutputs() const override;
+
+		void Execute();
+
+	private:
+		// Try to go through gate
+		FAudioBufferReadRef AudioInput;
+		FBoolReadRef bTriggerOnce;
+		FTriggerWriteRef TriggerOutput;
+
+		// Status of the gate
+		bool bHasTriggered;
+	};
+
+	FTriggerOnNanOperator::FTriggerOnNanOperator(
+		const FOperatorSettings& InSettings,
+		const FAudioBufferReadRef& InAudioInput,
+		const FBoolReadRef& InTriggerOnce)
+		: AudioInput(InAudioInput)
+		, bTriggerOnce(InTriggerOnce)
+		, TriggerOutput(FTriggerWriteRef::CreateNew(InSettings))
+		, bHasTriggered(false)
+	{
+	}
+
+	FDataReferenceCollection FTriggerOnNanOperator::GetInputs() const
+	{
+		using namespace TriggerOnNanVertexNames;
+
+		FDataReferenceCollection InputDataReferences;
+		InputDataReferences.AddDataReadReference(METASOUND_GET_PARAM_NAME(AudioInput), AudioInput);
+		InputDataReferences.AddDataReadReference(METASOUND_GET_PARAM_NAME(TriggerOnce), bTriggerOnce);
+		return InputDataReferences;
+	}
+
+	FDataReferenceCollection FTriggerOnNanOperator::GetOutputs() const
+	{
+		using namespace TriggerOnNanVertexNames;
+
+		FDataReferenceCollection OutputDataReferences;
+		OutputDataReferences.AddDataReadReference(METASOUND_GET_PARAM_NAME(OutputTriggerOnNan), TriggerOutput);
+		return OutputDataReferences;
+	}
+
+	void FTriggerOnNanOperator::Execute()
+	{
+		const bool bShouldTriggerOnce = *bTriggerOnce;
+		TriggerOutput->AdvanceBlock();
+
+		const float* AudioPtr = AudioInput->GetData();
+		const int32 NumSamples = AudioInput->Num();
+
+		for(int32 i = 0; i < NumSamples; ++i)
+		{
+			const bool bSuppressTriggers = (bShouldTriggerOnce && bHasTriggered);
+			if(!bSuppressTriggers && !FMath::IsFinite(AudioPtr[i]))
+			{
+				TriggerOutput->TriggerFrame(i);
+				bHasTriggered = true;
+			}
+		}
+	}
+
+
+	TUniquePtr<IOperator> FTriggerOnNanOperator::CreateOperator(const FCreateOperatorParams& InParams, FBuildErrorArray& OutErrors)
+	{
+		using namespace TriggerOnNanVertexNames;
+
+		const FInputVertexInterface& InputInterface = GetVertexInterface().GetInputInterface();
+
+		FAudioBufferReadRef AudioInput = InParams.InputDataReferences.GetDataReadReferenceOrConstruct<FAudioBuffer>(METASOUND_GET_PARAM_NAME(AudioInput), InParams.OperatorSettings);
+		FBoolReadRef bTriggerOnce = InParams.InputDataReferences.GetDataReadReferenceOrConstructWithVertexDefault<bool>(InputInterface, METASOUND_GET_PARAM_NAME(TriggerOnce), InParams.OperatorSettings);
+
+		return MakeUnique<FTriggerOnNanOperator>(InParams.OperatorSettings, AudioInput, bTriggerOnce);
+	}
+
+	const FVertexInterface& FTriggerOnNanOperator::GetVertexInterface()
+	{
+		using namespace TriggerOnNanVertexNames;
+
+		static const FVertexInterface Interface(
+			FInputVertexInterface(
+				TInputDataVertex<FAudioBuffer>(METASOUND_GET_PARAM_NAME_AND_METADATA(AudioInput)),
+				TInputDataVertex<bool>(METASOUND_GET_PARAM_NAME_AND_METADATA(TriggerOnce))
+			),
+			FOutputVertexInterface(
+				TOutputDataVertex<FTrigger>(METASOUND_GET_PARAM_NAME_AND_METADATA(OutputTriggerOnNan))
+			)
+		);
+
+		return Interface;
+	}
+
+	const FNodeClassMetadata& FTriggerOnNanOperator::GetNodeInfo()
+	{
+		auto InitNodeInfo = []() -> FNodeClassMetadata
+		{
+			FNodeClassMetadata Info;
+			Info.ClassName = { Metasound::StandardNodes::Namespace, "Trigger On Nan", FName() };
+			Info.MajorVersion = 1;
+			Info.MinorVersion = 0;
+			Info.DisplayName = METASOUND_LOCTEXT("Metasound_TriggerOnNanNodeDisplayName", "Trigger On Nan");
+			Info.Description = METASOUND_LOCTEXT("Metasound_TriggerOnNanNodeDescription", "Sends an output trigger the first time the node is triggered, and ignores all others (can be re-opened).");
+			Info.Author = PluginAuthor;
+			Info.PromptIfMissing = PluginNodeMissingPrompt;
+			Info.DefaultInterface = GetVertexInterface();
+			Info.CategoryHierarchy.Emplace(NodeCategories::Trigger);
+
+			return Info;
+		};
+
+		static const FNodeClassMetadata Info = InitNodeInfo();
+
+		return Info;
+	}
+
+	// Node Class
+	class FTriggerOnNanNode : public FNodeFacade
+	{
+	public:
+		FTriggerOnNanNode(const FNodeInitData& InitData)
+			: FNodeFacade(InitData.InstanceName, InitData.InstanceID, TFacadeOperatorClass<FTriggerOnNanOperator>())
+		{
+		}
+	};
+
+	METASOUND_REGISTER_NODE(FTriggerOnNanNode)
+}
+
+#undef LOCTEXT_NAMESPACE
