Support for arrays in MetasoundParameterPacks

diff --git a/Engine/Plugins/Runtime/Metasound/Source/MetasoundFrontend/Public/MetasoundDataTypeRegistrationMacro.h b/Engine/Plugins/Runtime/Metasound/Source/MetasoundFrontend/Public/MetasoundDataTypeRegistrationMacro.h
--- a/Engine/Plugins/Runtime/Metasound/Source/MetasoundFrontend/Public/MetasoundDataTypeRegistrationMacro.h
+++ b/Engine/Plugins/Runtime/Metasound/Source/MetasoundFrontend/Public/MetasoundDataTypeRegistrationMacro.h
@@ -30,6 +30,50 @@
 
 namespace Metasound
 {
+	template<typename ElementType, int32 NumInlineElements>
+	class TParamPackFixedArray
+	{
+	public:
+		TParamPackFixedArray(const TArray<ElementType>& Other)
+		{
+			*this = Other;
+		}
+
+		TParamPackFixedArray<ElementType, NumInlineElements>& operator=(const TArray<ElementType>& Other)
+		{
+			int32 NumToCopy = FMath::Min(Other.Num(), NumInlineElements);
+			for (int32 i = 0; i < NumToCopy; ++i)
+			{
+				GetData()[i] = Other[i];
+			}
+			NumValidElements = NumToCopy;
+			return *this;
+		}
+
+		void CopyToArray(TArray<ElementType>& Dest) const
+		{
+			Dest.Empty();
+			for (int32 i = 0; i < NumValidElements; ++i)
+			{
+				Dest.Add(GetData()[i]);
+			}
+		}
+
+		FORCEINLINE const ElementType* GetData() const
+		{
+			return (ElementType*)InlineData;
+		}
+
+		FORCEINLINE ElementType* GetData()
+		{
+			return (ElementType*)InlineData;
+		}
+
+	private:
+		int32 NumValidElements = 0;
+		TTypeCompatibleBytes<ElementType> InlineData[NumInlineElements];
+	};
+
 	namespace MetasoundDataTypeRegistrationPrivate
 	{
 		template<typename DataType>
@@ -327,7 +371,23 @@ namespace Metasound
 					}
 					else if constexpr (std::is_copy_assignable_v<TDataType>)
 					{
-						RawAssignmentFunction = [](const void* src, void* dest) { *(reinterpret_cast<TDataType*>(dest)) = *(reinterpret_cast<const TDataType*>(src)); };
+						if constexpr (!TIsArrayType<TDataType>::Value)
+						{
+							RawAssignmentFunction = [](const void* Src, void* Dest) 
+								{
+									*(reinterpret_cast<TDataType*>(Dest)) = *(reinterpret_cast<const TDataType*>(Src)); 
+								};
+						}
+						else
+						{
+							RawAssignmentFunction = [](const void* Src, void* Dest) 
+								{
+									TDataType& DestinationArray = *(reinterpret_cast<TDataType*>(Dest));
+									const TParamPackFixedArray<typename TDataType::ElementType, 1>& SourceArray = 
+										*(reinterpret_cast<const TParamPackFixedArray<typename TDataType::ElementType, 1>*>(Src));
+									SourceArray.CopyToArray(DestinationArray);
+								};
+						}
 					}
 
 					// Create class info using prototype node
diff --git a/Engine/Plugins/Runtime/Metasound/Source/MetasoundFrontend/Public/MetasoundParameterPack.h b/Engine/Plugins/Runtime/Metasound/Source/MetasoundFrontend/Public/MetasoundParameterPack.h
--- a/Engine/Plugins/Runtime/Metasound/Source/MetasoundFrontend/Public/MetasoundParameterPack.h
+++ b/Engine/Plugins/Runtime/Metasound/Source/MetasoundFrontend/Public/MetasoundParameterPack.h
@@ -5,6 +5,7 @@
 
 #include "MetasoundFrontendLiteral.h"
 #include "MetasoundDataReferenceMacro.h"
+#include "MetasoundDataTypeRegistrationMacro.h"
 #include "MetasoundRouter.h"
 #include "MetasoundFrontendDocument.h"
 
@@ -35,18 +36,38 @@ struct FMetasoundParameterPackItemBase
 	uint16 ValueOffset;
 };
 
+template<bool, typename T, int32 MAX_NUM_ITEMS = 0>
+struct TRootValueTypeHelper { using Type = T; };
+
+template<typename T, int32 MAX_NUM_ITEMS>
+struct TRootValueTypeHelper<true, T, MAX_NUM_ITEMS> { using Type = Metasound::TParamPackFixedArray<typename T::ElementType, MAX_NUM_ITEMS>; };
+
+template<typename T, int32 MAX_NUM_ITEMS = 0>
+struct TRootValueType : TRootValueTypeHelper<TIsTArray_V<T>, T, MAX_NUM_ITEMS> {};
+
 // A template that lets us stick an arbitrary data type into the raw 
 // 'bag-o-bytes' that is the parameter storage. Uses the base class 
 // above.
-template <typename T>
+template <typename T, int32 MAX_NUM = 0>
 struct FMetasoundParameterPackItem : public FMetasoundParameterPackItemBase
 {
+	using ThisType = FMetasoundParameterPackItem<T,MAX_NUM>;
 	FMetasoundParameterPackItem(const FName& InName, const FName& InTypeName, const T& InValue)
-		: FMetasoundParameterPackItemBase(InName, InTypeName, AlignArbitrary<uint16>((uint16)sizeof(FMetasoundParameterPackItem<T>), alignof(std::max_align_t)), offsetof(FMetasoundParameterPackItem<T>,Value))
+		: FMetasoundParameterPackItemBase(InName, InTypeName, SizeOfParam(InValue), offsetof(FMetasoundParameterPackItem<ThisType>, Value))
 		, Value(InValue)
-	{}
+	{
+		if constexpr (TIsTArray_V<T>)
+		{
+			check(InValue.Num() <= MAX_NUM);
+		}
+	}
 	FMetasoundParameterPackItem(const FMetasoundParameterPackItem& Other) = default;
-	T      Value;
+	typename TRootValueType<T, MAX_NUM>::Type Value;
+	
+	static uint16 SizeOfParam(const T& InValue)
+	{
+		return AlignArbitrary<uint16>((uint16)sizeof(ThisType), alignof(std::max_align_t));
+	}
 };
 
 // This next class owns the 'bag-o-bytes' that holds all of the parameters. 
@@ -135,8 +156,8 @@ struct FMetasoundParameterPackStorage
 		return ParameterIterator(nullptr, nullptr);
 	}
 
-	template<typename T>
-	FMetasoundParameterPackItem<T>* FindParameter(const FName& ParamName, const FName& InTypeName)
+	template<typename T, int32 MAX_NUM_ITEMS = 0>
+	FMetasoundParameterPackItem<T, MAX_NUM_ITEMS>* FindParameter(const FName& ParamName, const FName& InTypeName)
 	{
 		if (Storage.IsEmpty())
 		{
@@ -150,7 +171,7 @@ struct FMetasoundParameterPackStorage
 			{
 				if (ParamWalker->TypeName == InTypeName)
 				{
-					return static_cast<FMetasoundParameterPackItem<T>*>(ParamWalker);
+					return static_cast<FMetasoundParameterPackItem<T, MAX_NUM_ITEMS>*>(ParamWalker);
 				}
 				return nullptr;
 			}
@@ -167,34 +188,50 @@ struct FMetasoundParameterPackStorage
 		return nullptr;
 	}
 
+	// non-class version
 	template<typename T>
-	std::enable_if_t<!std::is_class_v<T>,T*> AddParameter(FMetasoundParameterPackItem<T>& NewParam)
+	std::enable_if_t<!std::is_class_v<T>,T*> AddParameter(const FName& Name, const FName& InTypeName, const T& InValue)
 	{
 		int32 NextParamLocation = AlignArbitrary<int32>(StorageBytesInUse, alignof(std::max_align_t));
-		int32 TotalStorageNeeded = NextParamLocation + NewParam.NextHeaderOffset;
+		int32 TotalStorageNeeded = NextParamLocation + FMetasoundParameterPackItem<T>::SizeOfParam(InValue);
 		if (TotalStorageNeeded > Storage.Num())
 		{
 			Storage.AddUninitialized(TotalStorageNeeded - Storage.Num());
 		}
-		FMetasoundParameterPackItem<T>* Destination = reinterpret_cast<FMetasoundParameterPackItem<T>*>(&Storage[NextParamLocation]);
-		*Destination = NewParam;
-		StorageBytesInUse = NextParamLocation + NewParam.NextHeaderOffset;
+		FMetasoundParameterPackItem<T>* Destination = new (reinterpret_cast<FMetasoundParameterPackItem<T>*>(&Storage[NextParamLocation])) FMetasoundParameterPackItem<T>(Name, InTypeName, InValue);
+		StorageBytesInUse = NextParamLocation + Destination->NextHeaderOffset;
 		return &Destination->Value;
 	}
 
+	// class version
 	template<typename T>
-	std::enable_if_t<std::is_class_v<T>, T*> AddParameter(FMetasoundParameterPackItem<T>& NewParam)
+	std::enable_if_t<std::is_class_v<T> && !TIsTArray_V<T>, T*> AddParameter(const FName& Name, const FName& InTypeName, const T& InValue)
 	{
 		int32 NextParamLocation = AlignArbitrary<int32>(StorageBytesInUse, alignof(std::max_align_t));
-		int32 TotalStorageNeeded = NextParamLocation + NewParam.NextHeaderOffset;
+		int32 TotalStorageNeeded = NextParamLocation + FMetasoundParameterPackItem<T>::SizeOfParam(InValue);
 		if (TotalStorageNeeded > Storage.Num())
 		{
 			Storage.AddUninitialized(TotalStorageNeeded - Storage.Num());
 		}
-		FMetasoundParameterPackItem<T>* Destination = new (reinterpret_cast<FMetasoundParameterPackItem<T>*>(&Storage[NextParamLocation])) FMetasoundParameterPackItem<T>(NewParam);
-		StorageBytesInUse = NextParamLocation + NewParam.NextHeaderOffset;
+		FMetasoundParameterPackItem<T>* Destination = new (reinterpret_cast<FMetasoundParameterPackItem<T>*>(&Storage[NextParamLocation])) FMetasoundParameterPackItem<T>(Name, InTypeName, InValue);
+		StorageBytesInUse = NextParamLocation + Destination->NextHeaderOffset;
 		return &Destination->Value;
 	}
+
+	// array version
+	template<typename T, int32 MAX_NUM_ITEMS>
+	std::enable_if_t<std::is_class_v<T> && TIsTArray_V<T>, void> AddArrayParameter(const FName& Name, const FName& InTypeName, const T& InValue)
+	{
+		check(InValue.Num() <= MAX_NUM_ITEMS);
+		int32 NextParamLocation = AlignArbitrary<int32>(StorageBytesInUse, alignof(std::max_align_t));
+		int32 TotalStorageNeeded = NextParamLocation + FMetasoundParameterPackItem<T, MAX_NUM_ITEMS>::SizeOfParam(InValue);
+		if (TotalStorageNeeded > Storage.Num())
+		{
+			Storage.AddUninitialized(TotalStorageNeeded - Storage.Num());
+		}
+		FMetasoundParameterPackItem<T, MAX_NUM_ITEMS>* Destination = new (reinterpret_cast<FMetasoundParameterPackItem<T>*>(&Storage[NextParamLocation])) FMetasoundParameterPackItem<T, MAX_NUM_ITEMS>(Name, InTypeName, InValue);
+		StorageBytesInUse = NextParamLocation + Destination->NextHeaderOffset;
+	}
 };
 
 using FSharedMetasoundParameterStoragePtr = TSharedPtr<FMetasoundParameterPackStorage, ESPMode::ThreadSafe>;
@@ -230,26 +267,26 @@ public:
 	ESetParamResult SetTrigger(FName ParameterName, bool OnlyIfExists = true);
 
 	UFUNCTION(BlueprintCallable, Category = "MetaSoundParameterPack", meta = (ExpandEnumAsExecs = "Result"))
-	bool GetBool(FName ParameterName, ESetParamResult& Result);
+	bool GetBool(FName ParameterName, ESetParamResult& Result) const;
 	UFUNCTION(BlueprintCallable, Category = "MetaSoundParameterPack", meta = (ExpandEnumAsExecs = "Result"))
-	int32 GetInt(FName ParameterName, ESetParamResult& Result);
+	int32 GetInt(FName ParameterName, ESetParamResult& Result) const;
 	UFUNCTION(BlueprintCallable, Category = "MetaSoundParameterPack", meta = (ExpandEnumAsExecs = "Result"))
-	float GetFloat(FName ParameterName, ESetParamResult& Result);
+	float GetFloat(FName ParameterName, ESetParamResult& Result) const;
 	UFUNCTION(BlueprintCallable, Category = "MetaSoundParameterPack", meta = (ExpandEnumAsExecs = "Result"))
-	FString GetString(FName ParameterName, ESetParamResult& Result);
+	FString GetString(FName ParameterName, ESetParamResult& Result) const;
 	UFUNCTION(BlueprintCallable, Category = "MetaSoundParameterPack", meta = (ExpandEnumAsExecs = "Result"))
-	bool GetTrigger(FName ParameterName, ESetParamResult& Result);
+	bool GetTrigger(FName ParameterName, ESetParamResult& Result) const;
 
 	UFUNCTION(BlueprintCallable, Category = "MetaSoundParameterPack", meta = (ExpandEnumAsExecs = ReturnValue))
-	bool HasBool(FName ParameterName);
+	bool HasBool(FName ParameterName) const;
 	UFUNCTION(BlueprintCallable, Category = "MetaSoundParameterPack", meta = (ExpandEnumAsExecs = ReturnValue))
-	bool HasInt(FName ParameterName);
+	bool HasInt(FName ParameterName) const;
 	UFUNCTION(BlueprintCallable, Category = "MetaSoundParameterPack", meta = (ExpandEnumAsExecs = ReturnValue))
-	bool HasFloat(FName ParameterName);
+	bool HasFloat(FName ParameterName) const;
 	UFUNCTION(BlueprintCallable, Category = "MetaSoundParameterPack", meta = (ExpandEnumAsExecs = ReturnValue))
-	bool HasString(FName ParameterName);
+	bool HasString(FName ParameterName) const;
 	UFUNCTION(BlueprintCallable, Category = "MetaSoundParameterPack", meta = (ExpandEnumAsExecs = ReturnValue))
-	bool HasTrigger(FName ParameterName);
+	bool HasTrigger(FName ParameterName) const;
 
 	void AddBoolParameter(FName Name, bool InValue);
 	void AddIntParameter(FName Name, int32 InValue);
@@ -258,16 +295,16 @@ public:
 	void AddTriggerParameter(FName Name, bool InValue = true);
 
 	// YIKES! BEWARE: If the returned pointers are only valid until another parameter is added!
-	bool*    GetBoolParameterPtr(FName Name);
-	int32*   GetIntParameterPtr(FName Name);
-	float*   GetFloatParameterPtr(FName Name);
-	FString* GetStringParameterPtr(FName Name);
-	bool*    GetTriggerParameterPtr(FName Name);
+	bool*    GetBoolParameterPtr(FName Name) const;
+	int32*   GetIntParameterPtr(FName Name) const;
+	float*   GetFloatParameterPtr(FName Name) const;
+	FString* GetStringParameterPtr(FName Name) const;
+	bool*    GetTriggerParameterPtr(FName Name) const;
 
 	template<typename T>
 	ESetParamResult SetParameter(const FName& Name, const FName& InTypeName, const T& InValue, bool OnlyIfExists = true)
 	{
-		T* TheParam = GetParameterPtr<T>(Name, InTypeName);
+		typename TRootValueType<T>::Type* TheParam = GetParameterPtr<T>(Name, InTypeName);
 		if (TheParam)
 		{
 			*TheParam = InValue;
@@ -282,14 +319,37 @@ public:
 		return ESetParamResult::Failed;
 	}
 
-	template<typename T>
-	bool HasParameter(const FName& Name, const FName& InTypeName)
+	template<typename T, int32 MAX_NUM_ITEMS>
+	ESetParamResult SetArrayParameter(const FName& Name, const FName& InTypeName, const T& InValue, bool OnlyIfExists = true)
+	{
+		static_assert(TIsTArray_V<T>);
+		if (!ensureMsgf(InValue.Num() <= MAX_NUM_ITEMS, TEXT("Too many items in source array \"%s\"! Max is %d."), *Name.ToString(), MAX_NUM_ITEMS))
+		{
+			return ESetParamResult::Failed;
+		}
+		typename TRootValueType<T, MAX_NUM_ITEMS>::Type* TheParam = GetParameterPtr<T,MAX_NUM_ITEMS>(Name, InTypeName);
+		if (TheParam)
+		{
+			*TheParam = InValue;
+			return ESetParamResult::Succeeded;
+		}
+		else if (!OnlyIfExists)
+		{
+			// TODO Add test to see if we are allowed to create it!
+			AddArrayParameter<T, MAX_NUM_ITEMS>(Name, InTypeName, InValue);
+			return ESetParamResult::Succeeded;
+		}
+		return ESetParamResult::Failed;
+	}
+
+	template<typename T, int32 MAX_NUM_ITEMS=0>
+	bool HasParameter(const FName& Name, const FName& InTypeName) const
 	{
 		if (!ParameterStorage.IsValid())
 		{
 			return false;
 		}
-		return ParameterStorage->FindParameter<T>(Name, InTypeName) != nullptr;
+		return ParameterStorage->FindParameter<T,MAX_NUM_ITEMS>(Name, InTypeName) != nullptr;
 	}
 
 	template<typename T>
@@ -299,18 +359,27 @@ public:
 		{
 			ParameterStorage = MakeShared<FMetasoundParameterPackStorage>();
 		}
-		FMetasoundParameterPackItem<T> NewParameter(Name, InTypeName, InValue);
-		ParameterStorage->AddParameter(NewParameter);
+		ParameterStorage->AddParameter(Name, InTypeName, InValue);
 	}
 
-	template<typename T>
-	T* GetParameterPtr(const FName& Name, const FName& InTypeName)
+	template<typename T, int32 MAX_NUM_ITEMS>
+	void AddArrayParameter(const FName& Name, const FName& InTypeName, const T& InValue)
+	{
+		if (!ParameterStorage.IsValid())
+		{
+			ParameterStorage = MakeShared<FMetasoundParameterPackStorage>();
+		}
+		ParameterStorage->AddArrayParameter<T,MAX_NUM_ITEMS>(Name, InTypeName, InValue);
+	}
+
+	template<typename T, int MAX_NUM_ITEMS = 0>
+	typename TRootValueType<T, MAX_NUM_ITEMS>::Type* GetParameterPtr(const FName& Name, const FName& InTypeName) const
 	{
 		if (!ParameterStorage.IsValid())
 		{
 			return nullptr;
 		}
-		auto TheParameter = ParameterStorage->FindParameter<T>(Name, InTypeName);
+		auto TheParameter = ParameterStorage->FindParameter<T, MAX_NUM_ITEMS>(Name, InTypeName);
 		if (!TheParameter)
 		{
 			return nullptr;
@@ -319,7 +388,7 @@ public:
 	}
 
 	template<typename T>
-	T GetParameter(const FName& Name, const FName& InTypeName, ESetParamResult& Result)
+	T GetParameter(const FName& Name, const FName& InTypeName, ESetParamResult& Result) const
 	{
 		T* TheParameter = GetParameterPtr<T>(Name, InTypeName);
 		if (!TheParameter)
@@ -331,6 +400,19 @@ public:
 		return *TheParameter;
 	}
 
+	template<typename T, int32 MAX_NUM_ITEMS>
+	void GetParameter(const FName& Name, const FName& InTypeName, TArray<typename T::ElementType>& ResultArray, ESetParamResult& Result) const
+	{
+		typename TRootValueType<T, MAX_NUM_ITEMS>::Type* TheParameter = GetParameterPtr<T,MAX_NUM_ITEMS>(Name, InTypeName);
+		if (!TheParameter)
+		{
+			Result = ESetParamResult::Failed;
+			return;
+		}
+		TheParameter->CopyToArray(ResultArray);
+		Result = ESetParamResult::Succeeded;
+	}
+
 	TSharedPtr<Audio::IProxyData> CreateProxyData(const Audio::FProxyDataInitParams& InitParams) override;
 	
 	// A couple of utilities for use by MetasoundAssetBase and MetasoundGenerator to set
