Fix missing file from check in

diff --git a/Engine/Plugins/Runtime/Metasound/Source/MetasoundFrontend/Private/MetasoundParameterPack.cpp b/Engine/Plugins/Runtime/Metasound/Source/MetasoundFrontend/Private/MetasoundParameterPack.cpp
--- a/Engine/Plugins/Runtime/Metasound/Source/MetasoundFrontend/Private/MetasoundParameterPack.cpp
+++ b/Engine/Plugins/Runtime/Metasound/Source/MetasoundFrontend/Private/MetasoundParameterPack.cpp
@@ -65,77 +65,77 @@ ESetParamResult UMetasoundParameterPack::SetTrigger(FName Name, bool OnlyIfExist
 	return SetParameter<bool>(Name, Metasound::GetMetasoundDataTypeName<Metasound::FTrigger>(), true, OnlyIfExists);
 }
 
-bool UMetasoundParameterPack::GetBool(FName Name, ESetParamResult& Result)
+bool UMetasoundParameterPack::GetBool(FName Name, ESetParamResult& Result) const
 {
 	return GetParameter<bool>(Name, Metasound::GetMetasoundDataTypeName<bool>(), Result);
 }
 
-int32 UMetasoundParameterPack::GetInt(FName Name, ESetParamResult& Result)
+int32 UMetasoundParameterPack::GetInt(FName Name, ESetParamResult& Result) const
 {
 	return GetParameter<int32>(Name, Metasound::GetMetasoundDataTypeName<int32>(), Result);
 }
 
-float UMetasoundParameterPack::GetFloat(FName Name, ESetParamResult& Result)
+float UMetasoundParameterPack::GetFloat(FName Name, ESetParamResult& Result) const
 {
 	return GetParameter<float>(Name, Metasound::GetMetasoundDataTypeName<float>(), Result);
 }
 
-FString UMetasoundParameterPack::GetString(FName Name, ESetParamResult& Result)
+FString UMetasoundParameterPack::GetString(FName Name, ESetParamResult& Result) const
 {
 	return GetParameter<FString>(Name, Metasound::GetMetasoundDataTypeName<FString>(), Result);
 }
 
-bool UMetasoundParameterPack::GetTrigger(FName Name, ESetParamResult& Result)
+bool UMetasoundParameterPack::GetTrigger(FName Name, ESetParamResult& Result) const
 {
 	return GetParameter<bool>(Name, Metasound::GetMetasoundDataTypeName<Metasound::FTrigger>(), Result);
 }
 
-bool* UMetasoundParameterPack::GetBoolParameterPtr(FName Name)
+bool* UMetasoundParameterPack::GetBoolParameterPtr(FName Name) const
 {
 	return GetParameterPtr<bool>(Name, Metasound::GetMetasoundDataTypeName<bool>());
 }
 
-int32* UMetasoundParameterPack::GetIntParameterPtr(FName Name)
+int32* UMetasoundParameterPack::GetIntParameterPtr(FName Name) const
 {
 	return GetParameterPtr<int32>(Name, Metasound::GetMetasoundDataTypeName<int32>());
 }
 
-float* UMetasoundParameterPack::GetFloatParameterPtr(FName Name)
+float* UMetasoundParameterPack::GetFloatParameterPtr(FName Name) const
 {
 	return GetParameterPtr<float>(Name, Metasound::GetMetasoundDataTypeName<float>());
 }
 
-FString* UMetasoundParameterPack::GetStringParameterPtr(FName Name)
+FString* UMetasoundParameterPack::GetStringParameterPtr(FName Name) const
 {
 	return GetParameterPtr<FString>(Name, Metasound::GetMetasoundDataTypeName<FString>());
 }
 
-bool* UMetasoundParameterPack::GetTriggerParameterPtr(FName Name)
+bool* UMetasoundParameterPack::GetTriggerParameterPtr(FName Name) const
 {
 	return GetParameterPtr<bool>(Name, Metasound::GetMetasoundDataTypeName<Metasound::FTrigger>());
 }
 
-bool UMetasoundParameterPack::HasBool(FName Name)
+bool UMetasoundParameterPack::HasBool(FName Name) const
 {
 	return HasParameter<bool>(Name, Metasound::GetMetasoundDataTypeName<bool>());
 }
 
-bool UMetasoundParameterPack::HasInt(FName Name)
+bool UMetasoundParameterPack::HasInt(FName Name) const
 {
 	return HasParameter<int32>(Name, Metasound::GetMetasoundDataTypeName<int32>());
 }
 
-bool UMetasoundParameterPack::HasFloat(FName Name)
+bool UMetasoundParameterPack::HasFloat(FName Name) const
 {
 	return HasParameter<float>(Name, Metasound::GetMetasoundDataTypeName<float>());
 }
 
-bool UMetasoundParameterPack::HasString(FName Name)
+bool UMetasoundParameterPack::HasString(FName Name) const
 {
 	return HasParameter<FString>(Name, Metasound::GetMetasoundDataTypeName<FString>());
 }
 
-bool UMetasoundParameterPack::HasTrigger(FName Name)
+bool UMetasoundParameterPack::HasTrigger(FName Name) const
 {
 	return HasParameter<bool>(Name, Metasound::GetMetasoundDataTypeName<Metasound::FTrigger>());
 }
