[Metasounds] Allow user to provide sample rate and block size to graph builder, so we can test with different settings.

diff --git a/Engine/Plugins/Runtime/Metasound/Source/MetasoundEngineTest/Private/NodeTestGraphBuilder.cpp b/Engine/Plugins/Runtime/Metasound/Source/MetasoundEngineTest/Private/NodeTestGraphBuilder.cpp
--- a/Engine/Plugins/Runtime/Metasound/Source/MetasoundEngineTest/Private/NodeTestGraphBuilder.cpp
+++ b/Engine/Plugins/Runtime/Metasound/Source/MetasoundEngineTest/Private/NodeTestGraphBuilder.cpp
@@ -56,7 +56,7 @@ namespace Metasound::Test
 		return RootGraph->AddOutputVertex(Output);
 	}
 
-	TUniquePtr<IOperator> FNodeTestGraphBuilder::BuildGraph()
+	TUniquePtr<IOperator> FNodeTestGraphBuilder::BuildGraph(FSampleRate SampleRate, int32 SamplesPerBlock)
 	{
 		TSet<FName> TransmittableInputNames;
 		const FString UnknownAsset = TEXT("UnknownAsset");
@@ -72,7 +72,7 @@ namespace Metasound::Test
 			BuilderSettings.bValidateVerticesExist = true;
 			FOperatorBuilder Builder{ BuilderSettings };
 
-			FOperatorSettings OperatorSettings = FOperatorSettings(48000.0f, 256);
+			FOperatorSettings OperatorSettings{SampleRate, static_cast<float>(SampleRate) / SamplesPerBlock};
 			FInputVertexInterfaceData InterfaceData;
 			FMetasoundEnvironment Environment;
 			FBuildGraphOperatorParams BuildParams{ *Graph, OperatorSettings, InterfaceData, Environment };
@@ -83,7 +83,11 @@ namespace Metasound::Test
 		return nullptr;
 	}
 
-	TUniquePtr<IOperator> FNodeTestGraphBuilder::MakeSingleNodeGraph(const FNodeClassName& ClassName, int32 MajorVersion)
+	TUniquePtr<IOperator> FNodeTestGraphBuilder::MakeSingleNodeGraph(
+		const FNodeClassName& ClassName,
+		int32 MajorVersion,
+		FSampleRate SampleRate,
+		int32 SamplesPerBlock)
 	{
 		FNodeTestGraphBuilder Builder;
 		FNodeHandle NodeHandle = Builder.AddNode(ClassName, MajorVersion);
@@ -132,6 +136,6 @@ namespace Metasound::Test
 		}
 
 		// build the graph
-		return Builder.BuildGraph();
+		return Builder.BuildGraph(SampleRate, SamplesPerBlock);
 	}
 }
diff --git a/Engine/Plugins/Runtime/Metasound/Source/MetasoundEngineTest/Public/NodeTestGraphBuilder.h b/Engine/Plugins/Runtime/Metasound/Source/MetasoundEngineTest/Public/NodeTestGraphBuilder.h
--- a/Engine/Plugins/Runtime/Metasound/Source/MetasoundEngineTest/Public/NodeTestGraphBuilder.h
+++ b/Engine/Plugins/Runtime/Metasound/Source/MetasoundEngineTest/Public/NodeTestGraphBuilder.h
@@ -21,14 +21,18 @@ namespace Metasound::Test
 		/** Add an input node to the graph */
 		Frontend::FNodeHandle AddInput(const FName& InputName, const FName& TypeName);
 
-		/** Add an output node to the grapn */
+		/** Add an output node to the graph */
 		Frontend::FNodeHandle AddOutput(const FName& OutputName, const FName& TypeName);
 
 		/** Build the graph and get the graph's operator to work with */
-		TUniquePtr<IOperator> BuildGraph();
+		TUniquePtr<IOperator> BuildGraph(FSampleRate SampleRate = 48000, int32 SamplesPerBlock = 256);
 
 		/** Helper that will add a single node, wire up the node's inputs and outputs, and hand back the graph's operator */
-		static TUniquePtr<IOperator> MakeSingleNodeGraph(const FNodeClassName& ClassName, int32 MajorVersion);
+		static TUniquePtr<IOperator> MakeSingleNodeGraph(
+			const FNodeClassName& ClassName,
+			int32 MajorVersion,
+			FSampleRate SampleRate = 48000,
+			int32 SamplesPerBlock = 256);
 
 	private:
 		FMetasoundFrontendDocument Document;
