Reduced code bloat from template instantiation

diff --git a/Engine/Plugins/Runtime/Metasound/Source/MetasoundFrontend/Public/MetasoundDataTypeRegistrationMacro.h b/Engine/Plugins/Runtime/Metasound/Source/MetasoundFrontend/Public/MetasoundDataTypeRegistrationMacro.h
--- a/Engine/Plugins/Runtime/Metasound/Source/MetasoundFrontend/Public/MetasoundDataTypeRegistrationMacro.h
+++ b/Engine/Plugins/Runtime/Metasound/Source/MetasoundFrontend/Public/MetasoundDataTypeRegistrationMacro.h
@@ -357,76 +357,15 @@
 			bAlreadyRegisteredThisDataType = true;
 
 			// Define registry entry class for this data type.
-			class FDataTypeRegistryEntry : public Frontend::IDataTypeRegistryEntry
+			class FDataTypeRegistryEntryBase : public Frontend::IDataTypeRegistryEntry
 			{
 			public:
-
-				FDataTypeRegistryEntry()
-					: Info(CreateDataTypeInfo<TDataType, PreferredArgType, UClassToUse>())
-					, EnumInterface(GetEnumDataTypeInterface<TDataType>())
+				FDataTypeRegistryEntryBase(const Frontend::FDataTypeRegistryInfo& Info, const TSharedPtr<Frontend::IEnumDataTypeInterface>& EnumInterface)
+					: Info(Info), EnumInterface(EnumInterface)
 				{
-					if constexpr (HasRawParameterAssignmentOp<TDataType>().value)
-					{
-						RawAssignmentFunction = [](const void* src, void* dest) { reinterpret_cast<TDataType*>(dest)->AssignRawParameter(src); };
-					}
-					else if constexpr (std::is_copy_assignable_v<TDataType>)
-					{
-						if constexpr (!TIsArrayType<TDataType>::Value)
-						{
-							RawAssignmentFunction = [](const void* Src, void* Dest) 
-								{
-									*(reinterpret_cast<TDataType*>(Dest)) = *(reinterpret_cast<const TDataType*>(Src)); 
-								};
-						}
-						else
-						{
-							RawAssignmentFunction = [](const void* Src, void* Dest) 
-								{
-									TDataType& DestinationArray = *(reinterpret_cast<TDataType*>(Dest));
-									const TParamPackFixedArray<typename TDataType::ElementType, 1>& SourceArray = 
-										*(reinterpret_cast<const TParamPackFixedArray<typename TDataType::ElementType, 1>*>(Src));
-									SourceArray.CopyToArray(DestinationArray);
-								};
-						}
-					}
-
-					// Create class info using prototype node
-					// TODO: register nodes with static class info instead of prototype instance.
-
-					if constexpr (bIsParsable)
-					{
-						TInputNode<TDataType, EVertexAccessType::Reference> InputPrototype(FInputNodeConstructorParams { });
-						InputClass = Metasound::Frontend::GenerateClass(InputPrototype.GetMetadata(), EMetasoundFrontendClassType::Input);
-
-						TOutputNode<TDataType, EVertexAccessType::Reference> OutputPrototype(TEXT(""), FGuid(), TEXT(""));
-						OutputClass = Metasound::Frontend::GenerateClass(OutputPrototype.GetMetadata(), EMetasoundFrontendClassType::Output);
-
-						TLiteralNode<TDataType> LiteralPrototype(TEXT(""), FGuid(), FLiteral());
-						LiteralClass = Metasound::Frontend::GenerateClass(LiteralPrototype.GetMetadata(), EMetasoundFrontendClassType::Literal);
-
-						TVariableNode<TDataType> VariablePrototype(TEXT(""), FGuid(), FLiteral());
-						VariableClass = Metasound::Frontend::GenerateClass(VariablePrototype.GetMetadata(), EMetasoundFrontendClassType::Variable);
-
-						TVariableMutatorNode<TDataType> VariableMutatorPrototype(TEXT(""), FGuid());
-						VariableMutatorClass = Metasound::Frontend::GenerateClass(VariableMutatorPrototype.GetMetadata(), EMetasoundFrontendClassType::VariableMutator);
-
-						TVariableAccessorNode<TDataType> VariableAccessorPrototype(TEXT(""), FGuid());
-						VariableAccessorClass = Metasound::Frontend::GenerateClass(VariableAccessorPrototype.GetMetadata(), EMetasoundFrontendClassType::VariableAccessor);
-
-						TVariableDeferredAccessorNode<TDataType> VariableDeferredAccessorPrototype(TEXT(""), FGuid());
-						VariableDeferredAccessorClass = Metasound::Frontend::GenerateClass(VariableDeferredAccessorPrototype.GetMetadata(), EMetasoundFrontendClassType::VariableDeferredAccessor);
-
-						if constexpr (bIsConstructorType)
-						{
-							TInputNode<TDataType, EVertexAccessType::Value> ConstructorInputPrototype(FInputNodeConstructorParams { });
-							ConstructorInputClass = Metasound::Frontend::GenerateClass(ConstructorInputPrototype.GetMetadata(), EMetasoundFrontendClassType::Input);
-							TOutputNode<TDataType, EVertexAccessType::Value> ConstructorOutputPrototype(TEXT(""), FGuid(), TEXT(""));
-							ConstructorOutputClass = Metasound::Frontend::GenerateClass(ConstructorOutputPrototype.GetMetadata(), EMetasoundFrontendClassType::Output);
-						}
-					}
 				}
 
-				virtual ~FDataTypeRegistryEntry() {}
+				virtual ~FDataTypeRegistryEntryBase() {}
 
 				virtual const Frontend::FDataTypeRegistryInfo& GetDataTypeInfo() const override
 				{
@@ -483,6 +422,96 @@
 					return VariableDeferredAccessorClass;
 				}
 
+				virtual TUniquePtr<IDataTypeRegistryEntry> Clone() const = 0;
+
+				virtual const Frontend::IParameterAssignmentFunction& GetRawAssignmentFunction() const override
+				{
+					return RawAssignmentFunction;
+				}
+
+			protected:
+				Frontend::FDataTypeRegistryInfo Info;
+				FMetasoundFrontendClass InputClass;
+				FMetasoundFrontendClass ConstructorInputClass;
+				FMetasoundFrontendClass OutputClass;
+				FMetasoundFrontendClass ConstructorOutputClass;
+				FMetasoundFrontendClass LiteralClass;
+				FMetasoundFrontendClass VariableClass;
+				FMetasoundFrontendClass VariableMutatorClass;
+				FMetasoundFrontendClass VariableAccessorClass;
+				FMetasoundFrontendClass VariableDeferredAccessorClass;
+				TSharedPtr<Frontend::IEnumDataTypeInterface> EnumInterface;
+				Frontend::IParameterAssignmentFunction RawAssignmentFunction;
+			};
+
+			class FDataTypeRegistryEntry : public FDataTypeRegistryEntryBase
+			{
+			public:
+
+				FDataTypeRegistryEntry()
+					: FDataTypeRegistryEntryBase(CreateDataTypeInfo<TDataType, PreferredArgType, UClassToUse>(), GetEnumDataTypeInterface<TDataType>())
+				{
+					if constexpr (HasRawParameterAssignmentOp<TDataType>().value)
+					{
+						this->RawAssignmentFunction = [](const void* src, void* dest) { reinterpret_cast<TDataType*>(dest)->AssignRawParameter(src); };
+					}
+					else if constexpr (std::is_copy_assignable_v<TDataType>)
+					{
+						if constexpr (!TIsArrayType<TDataType>::Value)
+						{
+							RawAssignmentFunction = [](const void* Src, void* Dest) 
+								{
+									*(reinterpret_cast<TDataType*>(Dest)) = *(reinterpret_cast<const TDataType*>(Src)); 
+								};
+						}
+						else
+						{
+							RawAssignmentFunction = [](const void* Src, void* Dest) 
+								{
+									TDataType& DestinationArray = *(reinterpret_cast<TDataType*>(Dest));
+									const TParamPackFixedArray<typename TDataType::ElementType, 1>& SourceArray = 
+										*(reinterpret_cast<const TParamPackFixedArray<typename TDataType::ElementType, 1>*>(Src));
+									SourceArray.CopyToArray(DestinationArray);
+								};
+						}
+					}
+
+					// Create class info using prototype node
+					// TODO: register nodes with static class info instead of prototype instance.
+
+					if constexpr (bIsParsable)
+					{
+						TInputNode<TDataType, EVertexAccessType::Reference> InputPrototype(FInputNodeConstructorParams{ });
+						this->InputClass = Metasound::Frontend::GenerateClass(InputPrototype.GetMetadata(), EMetasoundFrontendClassType::Input);
+
+						TOutputNode<TDataType, EVertexAccessType::Reference> OutputPrototype(TEXT(""), FGuid(), TEXT(""));
+						this->OutputClass = Metasound::Frontend::GenerateClass(OutputPrototype.GetMetadata(), EMetasoundFrontendClassType::Output);
+
+						TLiteralNode<TDataType> LiteralPrototype(TEXT(""), FGuid(), FLiteral());
+						this->LiteralClass = Metasound::Frontend::GenerateClass(LiteralPrototype.GetMetadata(), EMetasoundFrontendClassType::Literal);
+
+						TVariableNode<TDataType> VariablePrototype(TEXT(""), FGuid(), FLiteral());
+						this->VariableClass = Metasound::Frontend::GenerateClass(VariablePrototype.GetMetadata(), EMetasoundFrontendClassType::Variable);
+
+						TVariableMutatorNode<TDataType> VariableMutatorPrototype(TEXT(""), FGuid());
+						this->VariableMutatorClass = Metasound::Frontend::GenerateClass(VariableMutatorPrototype.GetMetadata(), EMetasoundFrontendClassType::VariableMutator);
+
+						TVariableAccessorNode<TDataType> VariableAccessorPrototype(TEXT(""), FGuid());
+						this->VariableAccessorClass = Metasound::Frontend::GenerateClass(VariableAccessorPrototype.GetMetadata(), EMetasoundFrontendClassType::VariableAccessor);
+
+						TVariableDeferredAccessorNode<TDataType> VariableDeferredAccessorPrototype(TEXT(""), FGuid());
+						this->VariableDeferredAccessorClass = Metasound::Frontend::GenerateClass(VariableDeferredAccessorPrototype.GetMetadata(), EMetasoundFrontendClassType::VariableDeferredAccessor);
+
+						if constexpr (bIsConstructorType)
+						{
+							TInputNode<TDataType, EVertexAccessType::Value> ConstructorInputPrototype(FInputNodeConstructorParams{ });
+							this->ConstructorInputClass = Metasound::Frontend::GenerateClass(ConstructorInputPrototype.GetMetadata(), EMetasoundFrontendClassType::Input);
+							TOutputNode<TDataType, EVertexAccessType::Value> ConstructorOutputPrototype(TEXT(""), FGuid(), TEXT(""));
+							this->ConstructorOutputClass = Metasound::Frontend::GenerateClass(ConstructorOutputPrototype.GetMetadata(), EMetasoundFrontendClassType::Output);
+						}
+					}
+				}
+
 				virtual TUniquePtr<INode> CreateInputNode(FInputNodeConstructorParams&& InParams) const override
 				{
 					if constexpr (bIsParsable)
@@ -628,21 +657,21 @@
 
 				virtual TOptional<FAnyDataReference> CreateDataReference(EDataReferenceAccessType InAccessType, const FLiteral& InLiteral, const FOperatorSettings& InOperatorSettings) const override
 				{
-					if constexpr(bIsParsable)
+					if constexpr (bIsParsable)
 					{
 						switch (InAccessType)
 						{
-							case EDataReferenceAccessType::Read:
-								return FAnyDataReference{TDataReadReferenceLiteralFactory<TDataType>::CreateExplicitArgs(InOperatorSettings, InLiteral)};
+						case EDataReferenceAccessType::Read:
+							return FAnyDataReference{ TDataReadReferenceLiteralFactory<TDataType>::CreateExplicitArgs(InOperatorSettings, InLiteral) };
 
-							case EDataReferenceAccessType::Write:
-								return FAnyDataReference{TDataWriteReferenceLiteralFactory<TDataType>::CreateExplicitArgs(InOperatorSettings, InLiteral)};
+						case EDataReferenceAccessType::Write:
+							return FAnyDataReference{ TDataWriteReferenceLiteralFactory<TDataType>::CreateExplicitArgs(InOperatorSettings, InLiteral) };
 
-							case EDataReferenceAccessType::Value:
-								return FAnyDataReference{TDataValueReferenceLiteralFactory<TDataType>::CreateExplicitArgs(InOperatorSettings, InLiteral)};
+						case EDataReferenceAccessType::Value:
+							return FAnyDataReference{ TDataValueReferenceLiteralFactory<TDataType>::CreateExplicitArgs(InOperatorSettings, InLiteral) };
 
-							default:
-								break;
+						default:
+							break;
 						}
 					}
 					return TOptional<FAnyDataReference>();
@@ -660,29 +689,10 @@
 					}
 				}
 
-				virtual TUniquePtr<IDataTypeRegistryEntry> Clone() const override
+				virtual TUniquePtr<Frontend::IDataTypeRegistryEntry> Clone() const override
 				{
 					return MakeUnique<FDataTypeRegistryEntry>();
 				}
-
-				virtual const Frontend::IParameterAssignmentFunction& GetRawAssignmentFunction() const override
-				{
-					return RawAssignmentFunction;
-				}
-
-			private:
-				Frontend::FDataTypeRegistryInfo Info;
-				FMetasoundFrontendClass InputClass;
-				FMetasoundFrontendClass ConstructorInputClass;
-				FMetasoundFrontendClass OutputClass;
-				FMetasoundFrontendClass ConstructorOutputClass;
-				FMetasoundFrontendClass LiteralClass;
-				FMetasoundFrontendClass VariableClass;
-				FMetasoundFrontendClass VariableMutatorClass;
-				FMetasoundFrontendClass VariableAccessorClass;
-				FMetasoundFrontendClass VariableDeferredAccessorClass;
-				TSharedPtr<Frontend::IEnumDataTypeInterface> EnumInterface;
-				Frontend::IParameterAssignmentFunction RawAssignmentFunction;
 			};
 
 			bool bSucceeded = Frontend::IDataTypeRegistry::Get().RegisterDataType(MakeUnique<FDataTypeRegistryEntry>());
diff --git a/Engine/Plugins/Runtime/Metasound/Source/MetasoundFrontend/Public/MetasoundNodeRegistrationMacro.h b/Engine/Plugins/Runtime/Metasound/Source/MetasoundFrontend/Public/MetasoundNodeRegistrationMacro.h
--- a/Engine/Plugins/Runtime/Metasound/Source/MetasoundFrontend/Public/MetasoundNodeRegistrationMacro.h
+++ b/Engine/Plugins/Runtime/Metasound/Source/MetasoundFrontend/Public/MetasoundNodeRegistrationMacro.h
@@ -41,26 +41,23 @@ namespace Metasound
 		bAlreadyRegisteredThisDataType = true;
 
 		// Node registry entry specialized to FNodeType.
-		class FNodeRegistryEntry : public Frontend::INodeRegistryEntry
+		class FNodeRegistryEntryBase : public Frontend::INodeRegistryEntry
 		{
 		public:
-			FNodeRegistryEntry(const Metasound::FNodeClassMetadata& InMetadata)
+			FNodeRegistryEntryBase(const Metasound::FNodeClassMetadata& InMetadata)
 			: ClassInfo(FMetasoundFrontendClassMetadata::GenerateClassMetadata(InMetadata, EMetasoundFrontendClassType::External))
 			, FrontendClass(Metasound::Frontend::GenerateClass(InMetadata))
 			{
 			}
 
-			virtual ~FNodeRegistryEntry() = default;
+			virtual ~FNodeRegistryEntryBase() = default;
 
 			virtual const Frontend::FNodeClassInfo& GetClassInfo() const override
 			{
 				return ClassInfo;
 			}
 
-			virtual TUniquePtr<INode> CreateNode(const FNodeInitData& InParams) const override
-			{
-				return MakeUnique<FNodeType>(InParams);
-			}
+			virtual TUniquePtr<INode> CreateNode(const FNodeInitData& InParams) const = 0;
 
 			virtual TUniquePtr<INode> CreateNode(FDefaultLiteralNodeConstructorParams&& InParams) const override
 			{
@@ -83,10 +80,7 @@ namespace Metasound
 				return FrontendClass;
 			}
 
-			virtual TUniquePtr<INodeRegistryEntry> Clone() const override
-			{
-				return MakeUnique<FNodeRegistryEntry>(*this);
-			}
+			virtual TUniquePtr<INodeRegistryEntry> Clone() const override = 0;
 
 			virtual bool IsNative() const override
 			{
@@ -99,6 +93,25 @@ namespace Metasound
 			FMetasoundFrontendClass FrontendClass;
 		};
 
+		class FNodeRegistryEntry : public FNodeRegistryEntryBase
+		{
+		public:
+			FNodeRegistryEntry(const Metasound::FNodeClassMetadata& InMetadata)
+				: FNodeRegistryEntryBase(InMetadata)
+			{
+			}
+
+			virtual TUniquePtr<INode> CreateNode(const FNodeInitData& InParams) const override
+			{
+				return MakeUnique<FNodeType>(InParams);
+			}
+
+			virtual TUniquePtr<Frontend::INodeRegistryEntry> Clone() const override
+			{
+				return MakeUnique<FNodeRegistryEntry>(*this);
+			}
+		};
+
 
 		Frontend::FNodeRegistryKey Key = FMetasoundFrontendRegistryContainer::Get()->RegisterNode(MakeUnique<FNodeRegistryEntry>(InMetadata));
 		const bool bSuccessfullyRegisteredNode = Frontend::NodeRegistryKey::IsValid(Key);
diff --git a/Engine/Source/Runtime/Analytics/Analytics/Public/AnalyticsEventAttribute.h b/Engine/Source/Runtime/Analytics/Analytics/Public/AnalyticsEventAttribute.h
--- a/Engine/Source/Runtime/Analytics/Analytics/Public/AnalyticsEventAttribute.h
+++ b/Engine/Source/Runtime/Analytics/Analytics/Public/AnalyticsEventAttribute.h
@@ -249,20 +249,20 @@ namespace ImplMakeAnalyticsEventAttributeArray
 {
 	/** Recursion terminator. Empty list. */
 	template <typename Allocator>
-	inline void MakeArray(TArray<FAnalyticsEventAttribute, Allocator>& Attrs)
+	FORCEINLINE void MakeArray(TArray<FAnalyticsEventAttribute, Allocator>& Attrs)
 	{
 	}
 
 	/** Recursion terminator. Convert the key/value pair to analytics strings. */
 	template <typename Allocator, typename KeyType, typename ValueType>
-	inline void MakeArray(TArray<FAnalyticsEventAttribute, Allocator>& Attrs, KeyType&& Key, ValueType&& Value)
+	FORCEINLINE void MakeArray(TArray<FAnalyticsEventAttribute, Allocator>& Attrs, KeyType&& Key, ValueType&& Value)
 	{
 		Attrs.Emplace(Forward<KeyType>(Key), Forward<ValueType>(Value));
 	}
 
 	/** recursively add the arguments to the array. */
 	template <typename Allocator, typename KeyType, typename ValueType, typename...ArgTypes>
-	inline void MakeArray(TArray<FAnalyticsEventAttribute, Allocator>& Attrs, KeyType&& Key, ValueType&& Value, ArgTypes&&...Args)
+	FORCEINLINE void MakeArray(TArray<FAnalyticsEventAttribute, Allocator>& Attrs, KeyType&& Key, ValueType&& Value, ArgTypes&&...Args)
 	{
 		// pop off the top two args and recursively apply the rest.
 		Attrs.Emplace(Forward<KeyType>(Key), Forward<ValueType>(Value));
@@ -272,7 +272,7 @@ namespace ImplMakeAnalyticsEventAttributeArray
 
 /** Helper to create an array of attributes using a single expression. Reserves the necessary space in advance. There must be an even number of arguments, one for each key/value pair. */
 template <typename Allocator = FDefaultAllocator, typename...ArgTypes>
-inline TArray<FAnalyticsEventAttribute, Allocator> MakeAnalyticsEventAttributeArray(ArgTypes&&...Args)
+FORCEINLINE TArray<FAnalyticsEventAttribute, Allocator> MakeAnalyticsEventAttributeArray(ArgTypes&&...Args)
 {
 	static_assert(sizeof...(Args) % 2 == 0, "Must pass an even number of arguments.");
 	TArray<FAnalyticsEventAttribute, Allocator> Attrs;
@@ -283,7 +283,7 @@ inline TArray<FAnalyticsEventAttribute, Allocator> MakeAnalyticsEventAttributeAr
 
 /** Helper to append to an array of attributes using a single expression. Reserves the necessary space in advance. There must be an even number of arguments, one for each key/value pair. */
 template <typename Allocator = FDefaultAllocator, typename...ArgTypes>
-inline TArray<FAnalyticsEventAttribute, Allocator>& AppendAnalyticsEventAttributeArray(TArray<FAnalyticsEventAttribute, Allocator>& Attrs, ArgTypes&&...Args)
+FORCEINLINE TArray<FAnalyticsEventAttribute, Allocator>& AppendAnalyticsEventAttributeArray(TArray<FAnalyticsEventAttribute, Allocator>& Attrs, ArgTypes&&...Args)
 {
 	static_assert(sizeof...(Args) % 2 == 0, "Must pass an even number of arguments.");
 	Attrs.Reserve(Attrs.Num() + (sizeof...(Args) / 2));
