Fixed compilation

diff --git a/Engine/Plugins/Runtime/Metasound/Source/MetasoundFrontend/Public/MetasoundDataTypeRegistrationMacro.h b/Engine/Plugins/Runtime/Metasound/Source/MetasoundFrontend/Public/MetasoundDataTypeRegistrationMacro.h
--- a/Engine/Plugins/Runtime/Metasound/Source/MetasoundFrontend/Public/MetasoundDataTypeRegistrationMacro.h
+++ b/Engine/Plugins/Runtime/Metasound/Source/MetasoundFrontend/Public/MetasoundDataTypeRegistrationMacro.h
@@ -459,14 +459,14 @@ namespace Metasound
 					{
 						if constexpr (!TIsArrayType<TDataType>::Value)
 						{
-							RawAssignmentFunction = [](const void* Src, void* Dest) 
+							this->RawAssignmentFunction = [](const void* Src, void* Dest)
 								{
 									*(reinterpret_cast<TDataType*>(Dest)) = *(reinterpret_cast<const TDataType*>(Src)); 
 								};
 						}
 						else
 						{
-							RawAssignmentFunction = [](const void* Src, void* Dest) 
+							this->RawAssignmentFunction = [](const void* Src, void* Dest)
 								{
 									TDataType& DestinationArray = *(reinterpret_cast<TDataType*>(Dest));
 									const TParamPackFixedArray<typename TDataType::ElementType, 1>& SourceArray = 
