Add Initialize function to metasound IOperator interface
Implement Initialize function for auto generated data type related operators.

diff --git a/Engine/Plugins/Runtime/Metasound/Source/MetasoundFrontend/Private/MetasoundInputNode.cpp b/Engine/Plugins/Runtime/Metasound/Source/MetasoundFrontend/Private/MetasoundInputNode.cpp
--- a/Engine/Plugins/Runtime/Metasound/Source/MetasoundFrontend/Private/MetasoundInputNode.cpp
+++ b/Engine/Plugins/Runtime/Metasound/Source/MetasoundFrontend/Private/MetasoundInputNode.cpp
@@ -41,5 +41,10 @@ namespace Metasound
 		{
 			return nullptr;
 		}
+
+		IOperator::FResetFunction FNonExecutableInputPassThroughOperator::GetResetFunction() 
+		{
+			return nullptr;
+		}
 	}
 }
diff --git a/Engine/Plugins/Runtime/Metasound/Source/MetasoundFrontend/Private/MetasoundOutputNode.cpp b/Engine/Plugins/Runtime/Metasound/Source/MetasoundFrontend/Private/MetasoundOutputNode.cpp
--- a/Engine/Plugins/Runtime/Metasound/Source/MetasoundFrontend/Private/MetasoundOutputNode.cpp
+++ b/Engine/Plugins/Runtime/Metasound/Source/MetasoundFrontend/Private/MetasoundOutputNode.cpp
@@ -34,6 +34,11 @@ namespace Metasound
 		{
 			return nullptr;
 		}
+
+		IOperator::FResetFunction FOutputOperator::GetResetFunction()
+		{
+			return nullptr;
+		}
 	}
 }
 
diff --git a/Engine/Plugins/Runtime/Metasound/Source/MetasoundFrontend/Public/MetasoundArrayNodes.h b/Engine/Plugins/Runtime/Metasound/Source/MetasoundFrontend/Public/MetasoundArrayNodes.h
--- a/Engine/Plugins/Runtime/Metasound/Source/MetasoundFrontend/Public/MetasoundArrayNodes.h
+++ b/Engine/Plugins/Runtime/Metasound/Source/MetasoundFrontend/Public/MetasoundArrayNodes.h
@@ -1,7 +1,6 @@
 // Copyright Epic Games, Inc. All Rights Reserved.
 #pragma once
 
-#include "CoreMinimal.h"
 #include "Interfaces/MetasoundFrontendSourceInterface.h"
 #include "Internationalization/Text.h"
 #include "MetasoundBuilderInterface.h"
@@ -166,6 +165,11 @@ namespace Metasound
 			*Num = Array->Num();
 		}
 
+		void Reset(const IOperator::FResetParams& InParams)
+		{
+			Execute();
+		}
+
 	private:
 
 		FArrayDataReadReference Array;
@@ -336,6 +340,21 @@ namespace Metasound
 			}
 		}
 
+		void Reset(const IOperator::FResetParams& InParams)
+		{
+			const int32 IndexValue = *Index;
+			const ArrayType& ArrayRef = *Array;
+
+			if ((IndexValue >= 0) && (IndexValue < ArrayRef.Num()))
+			{
+				*Value = ArrayRef[IndexValue];
+			}
+			else
+			{
+				*Value = TDataTypeFactory<ElementType>::CreateAny(InParams.OperatorSettings);
+			}
+		}
+
 	private:
 
 		TDataReadReference<FTrigger> Trigger;
@@ -514,6 +533,11 @@ namespace Metasound
 			}
 		}
 
+		void Reset(const IOperator::FResetParams& Inparams)
+		{
+			*Array = *InitArray;
+		}
+
 	private:
 		FOperatorSettings OperatorSettings;
 
@@ -598,17 +622,15 @@ namespace Metasound
 			FArrayDataReadReference LeftArray = InParams.InputDataReferences.GetDataReadReferenceOrConstructWithVertexDefault<ArrayType>(Inputs, METASOUND_GET_PARAM_NAME(InputLeftArray), InParams.OperatorSettings);
 			FArrayDataReadReference RightArray = InParams.InputDataReferences.GetDataReadReferenceOrConstructWithVertexDefault<ArrayType>(Inputs, METASOUND_GET_PARAM_NAME(InputRightArray), InParams.OperatorSettings);
 
-			FArrayDataWriteReference OutArray = TDataWriteReferenceFactory<ArrayType>::CreateExplicitArgs(InParams.OperatorSettings);
-
-			return MakeUnique<TArrayConcatOperator>(Trigger, LeftArray, RightArray, OutArray);
+			return MakeUnique<TArrayConcatOperator>(Trigger, LeftArray, RightArray);
 		}
 
 
-		TArrayConcatOperator(TDataReadReference<FTrigger> InTrigger, FArrayDataReadReference InLeftArray, FArrayDataReadReference InRightArray, FArrayDataWriteReference InOutArray)
+		TArrayConcatOperator(TDataReadReference<FTrigger> InTrigger, FArrayDataReadReference InLeftArray, FArrayDataReadReference InRightArray)
 		: Trigger(InTrigger)
 		, LeftArray(InLeftArray)
 		, RightArray(InRightArray)
-		, OutArray(InOutArray)
+		, OutArray(TDataWriteReference<ArrayType>::CreateNew())
 		{
 		}
 
@@ -645,6 +667,11 @@ namespace Metasound
 			}
 		}
 
+		void Reset(const IOperator::FResetParams& InParams)
+		{
+			OutArray->Reset();
+		}
+
 	private:
 
 		TDataReadReference<FTrigger> Trigger;
@@ -726,13 +753,13 @@ namespace Metasound
 			TDataReadReference<int32> StartIndex = InParams.InputDataReferences.GetDataReadReferenceOrConstructWithVertexDefault<int32>(Inputs, METASOUND_GET_PARAM_NAME(InputStartIndex), InParams.OperatorSettings);
 			TDataReadReference<int32> EndIndex = InParams.InputDataReferences.GetDataReadReferenceOrConstructWithVertexDefault<int32>(Inputs, METASOUND_GET_PARAM_NAME(InputEndIndex), InParams.OperatorSettings);
 
-			FArrayDataWriteReference OutArray = TDataWriteReferenceFactory<ArrayType>::CreateExplicitArgs(InParams.OperatorSettings);
+			TDataWriteReference<ArrayType> OutputArray = TDataWriteReferenceFactory<ArrayType>::CreateExplicitArgs(InParams.OperatorSettings);
 
-			return MakeUnique<TArraySubsetOperator>(Trigger, InArray, StartIndex, EndIndex, OutArray);
+			return MakeUnique<TArraySubsetOperator>(Trigger, InArray, StartIndex, EndIndex, OutputArray);
 		}
 
 
-		TArraySubsetOperator(TDataReadReference<FTrigger> InTrigger, FArrayDataReadReference InInputArray, TDataReadReference<int32> InStartIndex, TDataReadReference<int32> InEndIndex, FArrayDataWriteReference InOutputArray)
+		TArraySubsetOperator(TDataReadReference<FTrigger> InTrigger, FArrayDataReadReference InInputArray, TDataReadReference<int32> InStartIndex, TDataReadReference<int32> InEndIndex, TDataWriteReference<ArrayType> InOutputArray)
 		: Trigger(InTrigger)
 		, InputArray(InInputArray)
 		, StartIndex(InStartIndex)
@@ -786,6 +813,11 @@ namespace Metasound
 			}
 		}
 
+		void Reset(const IOperator::FResetParams& InParams)
+		{
+			OutputArray->Reset();
+		}
+
 	private:
 
 		TDataReadReference<FTrigger> Trigger;
diff --git a/Engine/Plugins/Runtime/Metasound/Source/MetasoundFrontend/Public/MetasoundArrayRandomNode.h b/Engine/Plugins/Runtime/Metasound/Source/MetasoundFrontend/Public/MetasoundArrayRandomNode.h
--- a/Engine/Plugins/Runtime/Metasound/Source/MetasoundFrontend/Public/MetasoundArrayRandomNode.h
+++ b/Engine/Plugins/Runtime/Metasound/Source/MetasoundFrontend/Public/MetasoundArrayRandomNode.h
@@ -184,12 +184,51 @@ namespace Metasound
 		{
 			using namespace Frontend;
 
+			SharedStateUniqueId = InParams.Node.GetInstanceID();
+
+			Reset(InParams);
+		}
+
+		virtual ~TArrayRandomGetOperator() = default;
+
+		virtual FDataReferenceCollection GetInputs() const override
+		{
+			using namespace ArrayNodeRandomGetVertexNames;
+
+			FDataReferenceCollection Inputs;
+			Inputs.AddDataReadReference(METASOUND_GET_PARAM_NAME(InputTriggerNextValue), TriggerNext);
+			Inputs.AddDataReadReference(METASOUND_GET_PARAM_NAME(InputTriggerResetSeed), TriggerReset);
+			Inputs.AddDataReadReference(METASOUND_GET_PARAM_NAME(InputRandomArray), InputArray);
+			Inputs.AddDataReadReference(METASOUND_GET_PARAM_NAME(InputWeights), InputWeightsArray);
+			Inputs.AddDataReadReference(METASOUND_GET_PARAM_NAME(InputSeed), SeedValue);
+			Inputs.AddDataReadReference(METASOUND_GET_PARAM_NAME(InputNoRepeatOrder), NoRepeatOrder);
+			Inputs.AddDataReadReference(METASOUND_GET_PARAM_NAME(InputEnableSharedState), bEnableSharedState);
+
+			return Inputs;
+		}
+
+		virtual FDataReferenceCollection GetOutputs() const override
+		{
+			using namespace ArrayNodeRandomGetVertexNames;
+
+			FDataReferenceCollection Outputs;
+			Outputs.AddDataReadReference(METASOUND_GET_PARAM_NAME(OutputTriggerOnNext), TriggerOnNext);
+			Outputs.AddDataReadReference(METASOUND_GET_PARAM_NAME(OutputTriggerOnReset), TriggerOnReset);
+			Outputs.AddDataReadReference(METASOUND_GET_PARAM_NAME(ShuffleOutputValue), OutValue);
+
+			return Outputs;
+		}
+
+		void Reset(const IOperator::FResetParams& InParams)
+		{
+			using namespace Frontend;
 #if WITH_METASOUND_DEBUG_ENVIRONMENT
 			if (InParams.Environment.Contains<FString>(SourceInterface::Environment::GraphName))
 			{
 				GraphName = *InParams.Environment.GetValue<FString>(SourceInterface::Environment::GraphName);
 			}
 #endif // WITH_METASOUND_DEBUG_ENVIRONMENT
+			bIsPreviewSound = InParams.Environment.GetValue<bool>(SourceInterface::Environment::IsPreview);
 
 			// Check to see if this is a global shuffler or a local one. 
 			// Global shuffler will use a namespace to opt into it.
@@ -206,11 +245,7 @@ namespace Metasound
 				if (*bEnableSharedState)
 				{
 					// Get the environment variable for the unique ID of the sound
-					SharedStateUniqueId = InParams.Node.GetInstanceID();
 					check(SharedStateUniqueId.IsValid());
-
-					bIsPreviewSound = InParams.Environment.GetValue<bool>(SourceInterface::Environment::IsPreview);
-
 					FSharedStateRandomGetManager& RGM = FSharedStateRandomGetManager::Get();
 
 					InitSharedStateArgs Args;
@@ -227,6 +262,7 @@ namespace Metasound
 				{
 					ArrayRandomGet = MakeUnique<FArrayRandomGet>(PrevSeedValue, PrevArraySize, WeightsArray, PrevNoRepeatOrder);
 				}
+
 			}
 #if WITH_METASOUND_DEBUG_ENVIRONMENT
 			else
@@ -234,36 +270,10 @@ namespace Metasound
  				UE_LOG(LogMetaSound, Verbose, TEXT("Array Random Get: Can't retrieve random elements from an empty array in graph '%s'"), *GraphName);
 			}
 #endif // WITH_METASOUND_DEBUG_ENVIRONMENT
-		}
 
-		virtual ~TArrayRandomGetOperator() = default;
-
-		virtual FDataReferenceCollection GetInputs() const override
-		{
-			using namespace ArrayNodeRandomGetVertexNames;
-
-			FDataReferenceCollection Inputs;
-			Inputs.AddDataReadReference(METASOUND_GET_PARAM_NAME(InputTriggerNextValue), TriggerNext);
-			Inputs.AddDataReadReference(METASOUND_GET_PARAM_NAME(InputTriggerResetSeed), TriggerReset);
-			Inputs.AddDataReadReference(METASOUND_GET_PARAM_NAME(InputRandomArray), InputArray);
-			Inputs.AddDataReadReference(METASOUND_GET_PARAM_NAME(InputWeights), InputWeightsArray);
-			Inputs.AddDataReadReference(METASOUND_GET_PARAM_NAME(InputSeed), SeedValue);
-			Inputs.AddDataReadReference(METASOUND_GET_PARAM_NAME(InputNoRepeatOrder), NoRepeatOrder);
-			Inputs.AddDataReadReference(METASOUND_GET_PARAM_NAME(InputEnableSharedState), bEnableSharedState);
-
-			return Inputs;
-		}
-
-		virtual FDataReferenceCollection GetOutputs() const override
-		{
-			using namespace ArrayNodeRandomGetVertexNames;
-
-			FDataReferenceCollection Outputs;
-			Outputs.AddDataReadReference(METASOUND_GET_PARAM_NAME(OutputTriggerOnNext), TriggerOnNext);
-			Outputs.AddDataReadReference(METASOUND_GET_PARAM_NAME(OutputTriggerOnReset), TriggerOnReset);
-			Outputs.AddDataReadReference(METASOUND_GET_PARAM_NAME(ShuffleOutputValue), OutValue);
-
-			return Outputs;
+			*OutValue = TDataTypeFactory<ElementType>::CreateAny(InParams.OperatorSettings);
+			TriggerOnNext->Reset();
+			TriggerOnReset->Reset();
 		}
 
 		void Execute()
diff --git a/Engine/Plugins/Runtime/Metasound/Source/MetasoundFrontend/Public/MetasoundArrayShuffleNode.h b/Engine/Plugins/Runtime/Metasound/Source/MetasoundFrontend/Public/MetasoundArrayShuffleNode.h
--- a/Engine/Plugins/Runtime/Metasound/Source/MetasoundFrontend/Public/MetasoundArrayShuffleNode.h
+++ b/Engine/Plugins/Runtime/Metasound/Source/MetasoundFrontend/Public/MetasoundArrayShuffleNode.h
@@ -220,19 +220,7 @@ namespace Metasound
 			, TriggerOnShuffle(FTriggerWriteRef::CreateNew(InParams.OperatorSettings))
 			, OutValue(TDataWriteReferenceFactory<ElementType>::CreateAny(InParams.OperatorSettings))
 		{
-			using namespace Frontend;
-
-			// Check to see if this is a global shuffler or a local one. 
-			// Global shuffler will use a namespace to opt into it.
-			PrevSeedValue = *SeedValue;
-
-			if (InParams.Environment.Contains<uint32>(SourceInterface::Environment::SoundUniqueID))
-			{
-				// Get the environment variable for the unique ID of the sound
-				SharedStateUniqueId = InParams.Environment.GetValue<uint32>(SourceInterface::Environment::SoundUniqueID);
-			}
-
-			UpdateArraySize(InputArray->Num());
+			Reset(InParams);
 		}
 
 		virtual ~TArrayShuffleOperator() = default;
@@ -266,6 +254,25 @@ namespace Metasound
 			return Outputs;
 		}
 
+		void Reset(const IOperator::FResetParams& InParams)
+		{
+			using namespace Frontend;
+
+			SharedStateUniqueId = INDEX_NONE;
+			if (InParams.Environment.Contains<uint32>(SourceInterface::Environment::SoundUniqueID))
+			{
+				// Get the environment variable for the unique ID of the sound
+				SharedStateUniqueId = InParams.Environment.GetValue<uint32>(SourceInterface::Environment::SoundUniqueID);
+			}
+
+			PrevSeedValue = *SeedValue;
+
+			InitializeShufflers(InputArray->Num());
+
+			TriggerOnShuffle->Reset();
+			*OutValue = TDataTypeFactory<ElementType>::CreateAny(InParams.OperatorSettings);
+		}
+
 		void Execute()
 		{
 			TriggerOnShuffle->AdvanceBlock();
@@ -279,7 +286,7 @@ namespace Metasound
 
 			if (PrevArraySize != InputArrayRef.Num())
 			{
-				UpdateArraySize(InputArrayRef.Num());
+				InitializeShufflers(InputArrayRef.Num());
 			}
  
 			// Check for a seed change
@@ -334,7 +341,7 @@ namespace Metasound
 		}
 
 	private:
-		void UpdateArraySize(int32 InSize)
+		void InitializeShufflers(int32 InSize)
 		{
 			PrevArraySize = InSize;
 
diff --git a/Engine/Plugins/Runtime/Metasound/Source/MetasoundFrontend/Public/MetasoundAutoConverterNode.h b/Engine/Plugins/Runtime/Metasound/Source/MetasoundFrontend/Public/MetasoundAutoConverterNode.h
--- a/Engine/Plugins/Runtime/Metasound/Source/MetasoundFrontend/Public/MetasoundAutoConverterNode.h
+++ b/Engine/Plugins/Runtime/Metasound/Source/MetasoundFrontend/Public/MetasoundAutoConverterNode.h
@@ -194,6 +194,13 @@ namespace Metasound
 				}
 			}
 
+			void Reset(const IOperator::FResetParams& InParams)
+			{
+				PreviousIntValueForEnumConversion = 0;
+				bHasLoggedInvalidEnum = false;
+				Execute();
+			}
+
 			private:
 				TDataReadReference<FromDataType> FromData;
 				TDataWriteReference<ToDataType> ToData;
diff --git a/Engine/Plugins/Runtime/Metasound/Source/MetasoundFrontend/Public/MetasoundInputNode.h b/Engine/Plugins/Runtime/Metasound/Source/MetasoundFrontend/Public/MetasoundInputNode.h
--- a/Engine/Plugins/Runtime/Metasound/Source/MetasoundFrontend/Public/MetasoundInputNode.h
+++ b/Engine/Plugins/Runtime/Metasound/Source/MetasoundFrontend/Public/MetasoundInputNode.h
@@ -32,7 +32,7 @@ namespace Metasound
 			virtual FDataReferenceCollection GetInputs() const override;
 			virtual FDataReferenceCollection GetOutputs() const override;
 		};
-
+		
 		class METASOUNDFRONTEND_API FNonExecutableInputOperatorBase : public FInputOperatorBase
 		{	
 		public:
@@ -47,6 +47,7 @@ namespace Metasound
 			FAnyDataReference DataRef;
 		};
 
+
 		class METASOUNDFRONTEND_API FNonExecutableInputPassThroughOperator : public FNonExecutableInputOperatorBase
 		{
 		public:
@@ -55,26 +56,73 @@ namespace Metasound
 			: FNonExecutableInputOperatorBase(InVertexName, FAnyDataReference{InDataRef})
 			{
 			}
-
-			template<typename DataType>
-			FNonExecutableInputPassThroughOperator(const FVertexName& InVertexName, const TDataWriteReference<DataType>& InDataRef)
-			: FNonExecutableInputOperatorBase(InVertexName, FAnyDataReference{InDataRef})
-			{
-			}
+			virtual IOperator::FResetFunction GetResetFunction() override;
 		};
 
-		/** FInputValueOperator provides an input for value references. */
-		class METASOUNDFRONTEND_API FInputValueOperator : public FNonExecutableInputOperatorBase
+
+		/** TInputValueOperator provides an input for value references. */
+		template<typename DataType>
+		class TInputValueOperator : public FNonExecutableInputOperatorBase
 		{
 		public:
-			/** Construct an FInputValueOperator with the name of the vertex and the 
+			/** Construct an TInputValueOperator with the name of the vertex and the 
 			 * value reference associated with input. 
 			 */
-			template<typename DataType>
-			explicit FInputValueOperator(const FName& InVertexName, const TDataValueReference<DataType>& InValueRef)
+			explicit TInputValueOperator(const FName& InVertexName, const TDataValueReference<DataType>& InValueRef)
 			: FNonExecutableInputOperatorBase(InVertexName, FAnyDataReference{InValueRef})
 			{
 			}
+
+			TInputValueOperator(const FVertexName& InVertexName, const FOperatorSettings& InSettings, const FLiteral& InLiteral)
+			: FNonExecutableInputOperatorBase(InVertexName, FAnyDataReference{TDataValueReferenceLiteralFactory<DataType>::CreateExplicitArgs(InSettings, InLiteral)})
+			{
+			}
+
+			virtual IOperator::FResetFunction GetResetFunction() { return nullptr; }
+		};
+
+
+		template<typename DataType> 
+		class TNonExecutableInputOperator : public FInputOperatorBase
+		{
+		public:
+			TNonExecutableInputOperator(const FVertexName& InVertexName, const FOperatorSettings& InSettings, const FLiteral& InLiteral)
+			: VertexName(InVertexName)
+			, DataRef(TDataWriteReferenceLiteralFactory<DataType>::CreateExplicitArgs(InSettings, InLiteral))
+			, Literal(InLiteral)
+			{
+			}
+
+			virtual void Bind(FVertexInterfaceData& InVertexData) const override
+			{
+				InVertexData.GetInputs().BindWriteVertex(VertexName, DataRef);
+				InVertexData.GetOutputs().BindReadVertex(VertexName, DataRef);
+			}
+
+			virtual IOperator::FExecuteFunction GetExecuteFunction() override
+			{
+				return nullptr;
+			}
+
+			virtual IOperator::FResetFunction GetResetFunction() override
+			{
+				return &Reset;
+			}
+
+		private:
+			static void Reset(IOperator* InOperator, const IOperator::FResetParams& InParams)
+			{
+				using FOperator = TNonExecutableInputOperator<DataType>;
+
+				FOperator* Operator = static_cast<FOperator*>(InOperator);
+				check(nullptr != Operator);
+
+				*Operator->DataRef = TDataTypeLiteralFactory<DataType>::CreateExplicitArgs(InParams.OperatorSettings, Operator->Literal);
+			}
+
+			FVertexName VertexName;
+			TDataWriteReference<DataType> DataRef;
+			FLiteral Literal;
 		};
 
 		/** A writable input and a readable output. */
@@ -84,13 +132,15 @@ namespace Metasound
 			static_assert(TExecutableDataType<DataType>::bIsExecutable, "TExecutableInputOperatorBase should only be used with executable data types");
 		public:
 			using FDataWriteReference = TDataWriteReference<DataType>;
+			using FDataWriteReferenceFactory = TDataWriteReferenceLiteralFactory<DataType>;
 
-			TExecutableInputOperator(const FVertexName& InDataReferenceName, FDataWriteReference InDataReference)
+			TExecutableInputOperator(const FVertexName& InDataReferenceName, const FOperatorSettings& InSettings, const FLiteral& InLiteral)
 				: DataReferenceName(InDataReferenceName)
 				// Executable DataTypes require a copy of the output to operate on whereas non-executable
 				// types do not. Avoid copy by assigning to reference for non-executable types.
-				, InputValue(InDataReference)
-				, OutputValue(FDataWriteReference::CreateNew(*InDataReference))
+				, InputValue(FDataWriteReferenceFactory::CreateExplicitArgs(InSettings, InLiteral))
+				, OutputValue(FDataWriteReferenceFactory::CreateExplicitArgs(InSettings, InLiteral))
+				, Literal(InLiteral)
 			{
 			}
 
@@ -105,8 +155,24 @@ namespace Metasound
 				return &Execute;
 			}
 
+			virtual FResetFunction GetResetFunction() override
+			{
+				return &Reset;
+			}
+
 		private:
 
+			static void Reset(IOperator* InOperator, const IOperator::FResetParams& InParams)
+			{
+				using FExecutableInputOperator = TExecutableInputOperator<DataType>;
+
+				FExecutableInputOperator* Operator = static_cast<FExecutableInputOperator*>(InOperator);
+				check(nullptr != Operator);
+
+				*Operator->InputValue = TDataTypeLiteralFactory<DataType>::CreateExplicitArgs(InParams.OperatorSettings, Operator->Literal);
+				*Operator->OutputValue = *Operator->InputValue;
+			}
+
 			static void Execute(IOperator* InOperator)
 			{
 				using FExecutableInputOperator = TExecutableInputOperator<DataType>;
@@ -121,10 +187,11 @@ namespace Metasound
 
 			FDataWriteReference InputValue;
 			FDataWriteReference OutputValue;
+			FLiteral Literal;
 		};
 
 
-		/** FInputValueOperator provides support for transmittable inputs. */
+		/** TInputReceiverOperator provides support for transmittable inputs. */
 		template<typename DataType>
 		class TInputReceiverOperator : public FInputOperatorBase
 		{
@@ -161,6 +228,11 @@ namespace Metasound
 				return &StaticExecute;
 			}
 
+			virtual FResetFunction GetResetFunction() override
+			{
+				return &StaticReset;
+			}
+
 		private:
 			void Execute()
 			{
@@ -194,6 +266,21 @@ namespace Metasound
 				Operator->Execute();
 			}
 
+			void Reset(const IOperator::FResetParams& InParams)
+			{
+				*OutputValue = *InputValue;
+				bHasNotReceivedData = true;
+			}
+
+			static void StaticReset(IOperator* InOperator, const IOperator::FResetParams& InParams)
+			{
+				using FOperator = TInputReceiverOperator<DataType>;
+
+				FOperator* Operator = static_cast<FOperator*>(InOperator);
+				check(nullptr != Operator);
+				Operator->Reset(InParams);
+			}
+
 			FVertexName DataReferenceName;
 			FDataReadReference InputValue;
 			FDataWriteReference OutputValue;
@@ -207,11 +294,11 @@ namespace Metasound
 	template<typename DataType, EVertexAccessType VertexAccess=EVertexAccessType::Reference>
 	using TInputOperator = std::conditional_t<
 		VertexAccess == EVertexAccessType::Value,
-		MetasoundInputNodePrivate::FInputValueOperator,
+		MetasoundInputNodePrivate::TInputValueOperator<DataType>,
 		std::conditional_t<
 			TExecutableDataType<DataType>::bIsExecutable,
 			MetasoundInputNodePrivate::TExecutableInputOperator<DataType>,
-			MetasoundInputNodePrivate::FNonExecutableInputPassThroughOperator
+			MetasoundInputNodePrivate::TNonExecutableInputOperator<DataType>
 		>
 	>;
 
@@ -219,78 +306,10 @@ namespace Metasound
 	template<typename DataType, EVertexAccessType VertexAccess=EVertexAccessType::Reference>
 	using TPassThroughOperator = std::conditional_t<
 		VertexAccess == EVertexAccessType::Value,
-		MetasoundInputNodePrivate::FInputValueOperator,
+		MetasoundInputNodePrivate::TInputValueOperator<DataType>,
 		MetasoundInputNodePrivate::FNonExecutableInputPassThroughOperator
 	>;
 
-	/** Data type creation policy to create by copy construction. */
-	template<typename DataType>
-	struct UE_DEPRECATED(5.1, "Moved to private implementation.") FCreateDataReferenceWithCopy
-	{
-		template<typename... ArgTypes>
-		FCreateDataReferenceWithCopy(ArgTypes&&... Args)
-		:	Data(Forward<ArgTypes>(Args)...)
-		{
-		}
-
-		TDataWriteReference<DataType> CreateDataReference(const FOperatorSettings& InOperatorSettings) const
-		{
-			return TDataWriteReferenceFactory<DataType>::CreateExplicitArgs(InOperatorSettings, Data);
-		}
-
-	private:
-		DataType Data;
-	};
-
-
-	/** Data type creation policy to create by literal construction. */
-	template<typename DataType>
-	struct UE_DEPRECATED(5.1, "Moved to private implementation.") FCreateDataReferenceWithLiteral
-	{
-		// If the data type is parsable from a literal type, then the data type 
-		// can be registered as an input type with the frontend.  To make a 
-		// DataType registrable, either create a constructor for the data type
-		// which accepts the one of the supported literal types with an optional 
-		// FOperatorSettings argument, or create a default constructor, or specialize
-		// this factory with an implementation for that specific data type.
-		static constexpr bool bCanCreateWithLiteral = TLiteralTraits<DataType>::bIsParsableFromAnyLiteralType;
-
-		FCreateDataReferenceWithLiteral(FLiteral&& InLiteral)
-		:	Literal(MoveTemp(InLiteral))
-		{
-		}
-
-		TDataWriteReference<DataType> CreateDataReference(const FOperatorSettings& InOperatorSettings) const
-		{
-			return TDataWriteReferenceLiteralFactory<DataType>::CreateExplicitArgs(InOperatorSettings, Literal);
-		}
-
-	private:
-		FLiteral Literal;
-	};
-
-
-	/** TInputOperatorFactory initializes the DataType at construction. It uses
-	 * the ReferenceCreatorType to create a data reference if one is not passed in.
-	 */
-	template<typename DataType, typename ReferenceCreatorType>
-	class UE_DEPRECATED(5.1, "Moved to private implementation") TInputOperatorFactory : public IOperatorFactory
-	{
-		public:
-
-			using FDataWriteReference = TDataWriteReference<DataType>;
-			using FDataReadReference = TDataReadReference<DataType>;
-
-			TInputOperatorFactory(ReferenceCreatorType&& InReferenceCreator)
-			:	ReferenceCreator(MoveTemp(InReferenceCreator))
-			{
-			}
-
-			virtual TUniquePtr<IOperator> CreateOperator(const FBuildOperatorParams& InParams, FBuildResults& OutResults) override;
-
-		private:
-			ReferenceCreatorType ReferenceCreator;
-	};
 
 	/** TInputNode represents an input to a metasound graph. */
 	template<typename DataType, EVertexAccessType VertexAccess=EVertexAccessType::Reference>
@@ -326,70 +345,30 @@ namespace Metasound
 			static_assert(bIsValueVertexAccess || bIsReferenceVertexAccess, "Unsupported EVertexAccessType");
 
 			// Choose which data reference type is created based on template parameters
-			using FDataReference = std::conditional_t<bIsReferenceVertexAccess, TDataWriteReference<DataType>, TDataValueReference<DataType>>;
+			using FDataReference = std::conditional_t<bIsReferenceVertexAccess, TDataReadReference<DataType>, TDataValueReference<DataType>>;
+			using FDataReferenceFactory = std::conditional_t<bIsReferenceVertexAccess, TDataReadReferenceLiteralFactory<DataType>, TDataValueReferenceLiteralFactory<DataType>>;
 			using FPassThroughDataReference = std::conditional_t<bIsReferenceVertexAccess, TDataReadReference<DataType>, TDataValueReference<DataType>>;
 
-			// Utility struct for creating data references for varying flavors of 
-			// runtime scenarios and vertex access types.
-			struct FDataReferenceCreatorBase
-			{
-				virtual ~FDataReferenceCreatorBase() = default;
-
-				// Create a a new data reference for constructing operators 
-				virtual FDataReference CreateDataReference(const FOperatorSettings& InOperatorSettings) const = 0;
-
-				// Create a data reference for constructing operators from a given data reference
-				virtual FPassThroughDataReference CreateDataReference(const FAnyDataReference& InRef) const
-				{
-					if constexpr (bIsReferenceVertexAccess)
-					{
-						return InRef.GetDataReadReference<DataType>();
-					}
-					else if constexpr (bIsValueVertexAccess)
-					{
-						return InRef.GetDataValueReference<DataType>();
-					}
-					else
-					{
-						static_assert("Unsupported EVertexAccessType");
-					}
-				}
-			};
-
-			// Create data references using a literal 
-			struct FCreateWithLiteral : FDataReferenceCreatorBase
+			// Return correct data reference type based on vertex access type for pass through scenario.
+			FPassThroughDataReference CreatePassThroughDataReference(const FAnyDataReference& InRef)
 			{
-				using FDataFactory = std::conditional_t<bIsReferenceVertexAccess, TDataWriteReferenceLiteralFactory<DataType>, TDataValueReferenceLiteralFactory<DataType>>;
-
-				FLiteral Literal;
-
-				FCreateWithLiteral(FLiteral&& InLiteral)
-				: Literal(MoveTemp(InLiteral))
+				if constexpr (bIsReferenceVertexAccess)
 				{
+					return InRef.GetDataReadReference<DataType>();
 				}
-
-				virtual FDataReference CreateDataReference(const FOperatorSettings& InOperatorSettings) const override
+				else if constexpr (bIsValueVertexAccess)
 				{
-					return FDataFactory::CreateExplicitArgs(InOperatorSettings, Literal);
+					return InRef.GetDataValueReference<DataType>();
 				}
-			};
-
-			// Create data references using a copy 
-			struct FCreateWithCopy : FDataReferenceCreatorBase
-			{
-				using FDataFactory = std::conditional_t<bIsReferenceVertexAccess, TDataWriteReferenceFactory<DataType>, TDataValueReferenceFactory<DataType>>;
-
-				DataType Value;
-
-				virtual FDataReference CreateDataReference(const FOperatorSettings& InOperatorSettings) const override
+				else
 				{
-					return FDataFactory::CreateExplicitArgs(InOperatorSettings, Value);
+					static_assert("Unsupported EVertexAccessType");
 				}
-			};
+			}
 
 		public:
 			explicit FInputNodeOperatorFactory(FLiteral&& InLiteral, bool bInEnableTransmission)
-			: ReferenceCreator(MakeUnique<FCreateWithLiteral>(MoveTemp(InLiteral)))
+			: Literal(MoveTemp(InLiteral))
 			, bEnableTransmission(bInEnableTransmission)
 			{
 			}
@@ -415,11 +394,11 @@ namespace Metasound
 						// Transmittable input value
 						if (const FAnyDataReference* Ref = InParams.InputData.FindDataReference(VertexKey))
 						{
-							return MakeUnique<TInputReceiverOperator<DataType>>(VertexKey, ReferenceCreator->CreateDataReference(*Ref), MoveTemp(SendAddress), MoveTemp(Receiver));
+							return MakeUnique<TInputReceiverOperator<DataType>>(VertexKey, CreatePassThroughDataReference(*Ref), MoveTemp(SendAddress), MoveTemp(Receiver));
 						}
 						else
 						{
-							FDataReference DataRef = ReferenceCreator->CreateDataReference(InParams.OperatorSettings);
+							FDataReference DataRef = FDataReferenceFactory::CreateExplicitArgs(InParams.OperatorSettings, Literal);
 							return MakeUnique<TInputReceiverOperator<DataType>>(VertexKey, DataRef, MoveTemp(SendAddress), MoveTemp(Receiver));
 						}
 					}
@@ -431,16 +410,15 @@ namespace Metasound
 				if (const FAnyDataReference* Ref = InParams.InputData.FindDataReference(VertexKey))
 				{
 					// Pass through input value
-					return MakeUnique<TPassThroughOperator<DataType, VertexAccess>>(VertexKey, ReferenceCreator->CreateDataReference(*Ref));
+					return MakeUnique<TPassThroughOperator<DataType, VertexAccess>>(VertexKey, CreatePassThroughDataReference(*Ref));
 				}
 
 				// Owned input value
-				FDataReference DataRef = ReferenceCreator->CreateDataReference(InParams.OperatorSettings);
-				return MakeUnique<TInputOperator<DataType, VertexAccess>>(VertexKey, DataRef);
+				return MakeUnique<TInputOperator<DataType, VertexAccess>>(VertexKey, InParams.OperatorSettings, Literal);
 			}
 
 		private:
-			TUniquePtr<FDataReferenceCreatorBase> ReferenceCreator;
+			FLiteral Literal;
 			bool bEnableTransmission = false;
 		};
 
@@ -476,37 +454,6 @@ namespace Metasound
 			return Info;
 		}
 
-		template<typename... ArgTypes>
-		UE_DEPRECATED(5.1, "Moved to internal implementation.")
-		static FOperatorFactorySharedRef CreateOperatorFactoryWithArgs(ArgTypes&&... Args)
-		{
-			using FCreatorType = FCreateDataReferenceWithCopy<DataType>;
-			using FFactoryType = TInputOperatorFactory<DataType, FCreatorType>;
-
-			return MakeOperatorFactoryRef<FFactoryType>(FCreatorType(Forward<ArgTypes>(Args)...));
-		}
-
-		UE_DEPRECATED(5.1, "Moved to internal implementation.")
-		static FOperatorFactorySharedRef CreateOperatorFactoryWithLiteral(FLiteral&& InLiteral)
-		{
-			using FCreatorType = FCreateDataReferenceWithLiteral<DataType>;
-			using FFactoryType = TInputOperatorFactory<DataType, FCreatorType>;
-
-			return MakeOperatorFactoryRef<FFactoryType>(FCreatorType(MoveTemp(InLiteral)));
-		}
-
-		/* Construct a TInputNode using the TInputOperatorFactory<> and forwarding 
-		 * Args to the TInputOperatorFactory constructor.*/
-		template<typename... ArgTypes>
-		UE_DEPRECATED(5.1, "Constructing an TInputNode with args will no longer be supported.")
-		TInputNode(const FVertexName& InInstanceName, const FGuid& InInstanceID, const FVertexName& InVertexName, ArgTypes&&... Args)
-		:	FNode(InInstanceName, InInstanceID, GetNodeInfo(InVertexName))
-		,	VertexName(InVertexName)
-		,	Interface(DeclareVertexInterface(InVertexName))
-		,	Factory(MakeShared<FInputNodeOperatorFactory>(Forward<ArgTypes>(Args)...))
-		{
-		}
-
 		/* Construct a TInputNode using the TInputOperatorLiteralFactory<> and moving
 		 * InParam to the TInputOperatorLiteralFactory constructor.*/
 		explicit TInputNode(FInputNodeConstructorParams&& InParams)
@@ -548,28 +495,5 @@ namespace Metasound
 		FVertexInterface Interface;
 		FOperatorFactorySharedRef Factory;
 	};
-
-
-	template<typename DataType, typename ReferenceCreatorType>
-	TUniquePtr<IOperator> TInputOperatorFactory<DataType, ReferenceCreatorType>::CreateOperator(const FBuildOperatorParams& InParams, FBuildResults& OutResults)
-	{
-		using FInputNodeType = TInputNode<DataType>;
-
-		const FInputNodeType& InputNode = static_cast<const FInputNodeType&>(InParams.Node);
-		const FVertexName& VertexKey = InputNode.GetVertexName();
-
-		if (InParams.InputData.IsVertexBound(VertexKey))
-		{
-			// Data is externally owned. Use pass through operator
-			TDataReadReference<DataType> DataRef = InParams.InputData.GetDataReadReference<DataType>(VertexKey);
-			return MakeUnique<TPassThroughOperator<DataType>>(InputNode.GetVertexName(), DataRef);
-		}
-		else
-		{
-			// Create write reference by calling compatible constructor with literal.
-			TDataWriteReference<DataType> DataRef = ReferenceCreator.CreateDataReference(InParams.OperatorSettings);
-			return MakeUnique<TInputOperator<DataType>>(InputNode.GetVertexName(), DataRef);
-		}
-	}
 } // namespace Metasound
 #undef LOCTEXT_NAMESPACE // MetasoundFrontend
diff --git a/Engine/Plugins/Runtime/Metasound/Source/MetasoundFrontend/Public/MetasoundOutputNode.h b/Engine/Plugins/Runtime/Metasound/Source/MetasoundFrontend/Public/MetasoundOutputNode.h
--- a/Engine/Plugins/Runtime/Metasound/Source/MetasoundFrontend/Public/MetasoundOutputNode.h
+++ b/Engine/Plugins/Runtime/Metasound/Source/MetasoundFrontend/Public/MetasoundOutputNode.h
@@ -49,6 +49,7 @@ namespace Metasound
 				virtual FDataReferenceCollection GetOutputs() const override;
 				virtual void Bind(FVertexInterfaceData& InVertexData) const;
 				virtual FExecuteFunction GetExecuteFunction() override;
+				virtual FResetFunction GetResetFunction() override;
 
 			private:
 				FVertexName VertexName;
diff --git a/Engine/Plugins/Runtime/Metasound/Source/MetasoundFrontend/Public/MetasoundReceiveNode.h b/Engine/Plugins/Runtime/Metasound/Source/MetasoundFrontend/Public/MetasoundReceiveNode.h
--- a/Engine/Plugins/Runtime/Metasound/Source/MetasoundFrontend/Public/MetasoundReceiveNode.h
+++ b/Engine/Plugins/Runtime/Metasound/Source/MetasoundFrontend/Public/MetasoundReceiveNode.h
@@ -170,6 +170,12 @@ namespace Metasound
 					}
 				}
 
+				void Reset(const IOperator::FResetParams& InParams)
+				{
+					*OutputData = *DefaultData;
+					bHasNotReceivedData = true;
+				}
+
 			private:
 
 				FSendAddress GetSendAddressWithDataType(const FSendAddress& InAddress) const 
diff --git a/Engine/Plugins/Runtime/Metasound/Source/MetasoundFrontend/Public/MetasoundSendNode.h b/Engine/Plugins/Runtime/Metasound/Source/MetasoundFrontend/Public/MetasoundSendNode.h
--- a/Engine/Plugins/Runtime/Metasound/Source/MetasoundFrontend/Public/MetasoundSendNode.h
+++ b/Engine/Plugins/Runtime/Metasound/Source/MetasoundFrontend/Public/MetasoundSendNode.h
@@ -133,6 +133,15 @@ namespace Metasound
 					Sender->Push(*InputData);
 				}
 
+				void Reset(const IOperator::FResetParams& InParams)
+				{
+					ResetSenderAndCleanupChannel();
+					CachedSendAddress = *SendAddress;
+					Sender = CreateNewSender();
+					check(Sender.IsValid());
+				}
+
+
 			private:
 				FSendAddress GetSendAddressWithDataType(const FSendAddress& InAddress) const 
 				{
diff --git a/Engine/Plugins/Runtime/Metasound/Source/MetasoundGraphCore/Private/MetasoundGraphOperator.cpp b/Engine/Plugins/Runtime/Metasound/Source/MetasoundGraphCore/Private/MetasoundGraphOperator.cpp
--- a/Engine/Plugins/Runtime/Metasound/Source/MetasoundGraphCore/Private/MetasoundGraphOperator.cpp
+++ b/Engine/Plugins/Runtime/Metasound/Source/MetasoundGraphCore/Private/MetasoundGraphOperator.cpp
@@ -10,18 +10,27 @@
 
 namespace Metasound
 {
-
-	FGraphOperator::~FGraphOperator()
-	{
-	}
-
 	void FGraphOperator::AppendOperator(FOperatorPtr InOperator)
 	{
 		if (InOperator.IsValid())
 		{
-			if (nullptr != InOperator->GetExecuteFunction())
+			bool bIsOperatorInAnyStack = false;
+
+			if (FExecuteFunction ExecuteFunc = InOperator->GetExecuteFunction())
 			{
-				OperatorStack.Emplace(MoveTemp(InOperator));
+				ExecuteStack.Emplace(*InOperator, ExecuteFunc);
+				bIsOperatorInAnyStack = true;
+			}
+
+			if (FResetFunction ResetFunc = InOperator->GetResetFunction())
+			{
+				ResetStack.Emplace(*InOperator, ResetFunc);
+				bIsOperatorInAnyStack = true;
+			}
+
+			if (bIsOperatorInAnyStack)
+			{
+				ActiveOperators.Add(MoveTemp(InOperator));
 			}
 		}
 	}
@@ -56,27 +65,47 @@ namespace Metasound
 		InVertexData.Bind(VertexData);
 	}
 
-	IOperator::FExecuteFunction FGraphOperator::GetExecuteFunction()
+	void FGraphOperator::Execute()
 	{
-		return &FGraphOperator::ExecuteFunction;
+		FExecuteEntry* StackPtr = ExecuteStack.GetData();
+		const int32 Num = ExecuteStack.Num();
+		for (int32 i = 0; i < Num; i++)
+		{
+			StackPtr[i].Execute();
+		}
 	}
 
-	void FGraphOperator::Execute()
+	void FGraphOperator::Reset(const FGraphOperator::FResetParams& InParams)
 	{
-		FExecuter* StackPtr = OperatorStack.GetData();
-		const int32 Num = OperatorStack.Num();
+		FResetEntry* StackPtr = ResetStack.GetData();
+		const int32 Num = ResetStack.Num();
 		for (int32 i = 0; i < Num; i++)
 		{
-			StackPtr[i].Execute();
+			StackPtr[i].Reset(InParams);
 		}
 	}
 
-	void FGraphOperator::ExecuteFunction(IOperator* InOperator)
+	FGraphOperator::FExecuteEntry::FExecuteEntry(IOperator& InOperator, FExecuteFunction InFunc)
+	: Operator(&InOperator)
+	, Function(InFunc)
+	{
+		check(Function);
+	}
+	void FGraphOperator::FExecuteEntry::Execute()
 	{
-		FGraphOperator* Operator = static_cast<FGraphOperator*>(InOperator);
+		Function(Operator);
+	}
 
-		check(nullptr != Operator);
 
-		Operator->Execute();
+	FGraphOperator::FResetEntry::FResetEntry(IOperator& InOperator, FResetFunction InFunc)
+	: Operator(&InOperator)
+	, Function(InFunc)
+	{
+		check(Function);
 	}
+	void FGraphOperator::FResetEntry::Reset(const FGraphOperator::FResetParams& InParams)
+	{
+		Function(Operator, InParams);
+	}
+
 }
diff --git a/Engine/Plugins/Runtime/Metasound/Source/MetasoundGraphCore/Public/MetasoundBuilderInterface.h b/Engine/Plugins/Runtime/Metasound/Source/MetasoundGraphCore/Public/MetasoundBuilderInterface.h
--- a/Engine/Plugins/Runtime/Metasound/Source/MetasoundGraphCore/Public/MetasoundBuilderInterface.h
+++ b/Engine/Plugins/Runtime/Metasound/Source/MetasoundGraphCore/Public/MetasoundBuilderInterface.h
@@ -8,8 +8,9 @@
 #include "MetasoundNodeInterface.h"
 #include "MetasoundOperatorInterface.h"
 #include "Templates/UniquePtr.h"
-#include "UObject/NameTypes.h"
 
+class FName;
+class FText;
 
 namespace Metasound
 {
@@ -76,6 +77,12 @@ namespace Metasound
 
 		/** Pointer to builder actively building graph. */
 		const IOperatorBuilder* Builder = nullptr;
+
+		/** Implicit conversion to FResetParams for convenience. */
+		operator IOperator::FResetParams() const
+		{
+			return IOperator::FResetParams{OperatorSettings, Environment};
+		}
 	};
 
 	/** FBuildOperatorParams holds the parameters provided to operator factories
@@ -97,6 +104,12 @@ namespace Metasound
 
 		/** Pointer to builder actively building graph. */
 		const IOperatorBuilder* Builder = nullptr;
+
+		/** Implicit conversion to FResetParams for convenience. */
+		operator IOperator::FResetParams() const
+		{
+			return IOperator::FResetParams{OperatorSettings, Environment};
+		}
 	};
 
 	/** Parameters for building an operator from a graph. 
diff --git a/Engine/Plugins/Runtime/Metasound/Source/MetasoundGraphCore/Public/MetasoundExecutableOperator.h b/Engine/Plugins/Runtime/Metasound/Source/MetasoundGraphCore/Public/MetasoundExecutableOperator.h
--- a/Engine/Plugins/Runtime/Metasound/Source/MetasoundGraphCore/Public/MetasoundExecutableOperator.h
+++ b/Engine/Plugins/Runtime/Metasound/Source/MetasoundGraphCore/Public/MetasoundExecutableOperator.h
@@ -2,13 +2,81 @@
 
 #pragma once
 
-#include "CoreMinimal.h"
-
 #include "MetasoundOperatorInterface.h"
 
+#include <type_traits>
+
 
 namespace Metasound
 {
+	namespace MetasoundExecutableOperatorPrivate
+	{
+		// Helper template to determine whether a member function is declared
+		// for a given template class.
+		template <typename U>
+		class TIsResetMethodDeclared 
+		{
+			private:
+				template<typename T, T> 
+				struct Helper;
+
+				template<typename T>
+				static uint8 Check(Helper<void(T::*)(const IOperator::FResetParams&), &T::Reset>*);
+
+				template<typename T> static uint16 Check(...);
+
+			public:
+				static constexpr bool Value = sizeof(Check<U>(0)) == sizeof(uint8);
+		};
+
+		// Helper to determine whether an Reset(...) function exists. 
+		//
+		// Some nodes have an Execute() call without an Reset(...) call because
+		// Reset(...) was introduced after several MetaSound releases. In order
+		// allow operators to compile, we make the Intiailize(...) function optional
+		// for subclasses of TExecutableOperator. This adapter determines whether
+		// it Reset(...) exists or not.
+		//
+		// Note: This helper must be instantiated inside the functions and cannot
+		// be part of the `TExecutableOperator<OperatorType>` class definition as
+		// classes are not completely defined at the point in the compilation process
+		// when static constexpr class members are evaluated. 
+		template<typename OperatorType>
+		struct TResetFunctionAdapter
+		{
+			// Returns an IOperator::FResetFunction if the OperatorType has a
+			// Reset(...) class member function.
+			static IOperator::FResetFunction GetResetFunction() 
+			{
+				if (TIsResetMethodDeclared<OperatorType>::Value)
+				{
+					return &TResetFunctionAdapter::ResetFunction;
+				}
+				else
+				{
+					return nullptr;
+				}
+			}
+
+		private:
+			static void ResetFunction(IOperator* InOperator, const IOperator::FResetParams& InParams)
+			{
+				if constexpr (TIsResetMethodDeclared<OperatorType>::Value)
+				{
+					OperatorType* DerivedOperator = static_cast<OperatorType*>(InOperator);
+
+					check(nullptr != DerivedOperator);
+
+					DerivedOperator->Reset(InParams);
+				}
+				else
+				{
+					checkNoEntry();
+				}
+			}
+		};
+	}
+
 	// As a general rule, ExecutableDataTypes should be avoided whenever possible
 	// as they incur an undesired cost and are generally not typically necessary.
 	// This is primarily for the special case of trigger types, where state management
@@ -48,25 +116,30 @@ namespace Metasound
 	template<class DerivedOperatorType>
 	class TExecutableOperator : public IOperator
 	{
-		public:
+	public:
 
-			virtual ~TExecutableOperator() {}
+		virtual ~TExecutableOperator() {}
 
-			virtual FExecuteFunction GetExecuteFunction() override
-			{
-				return &TExecutableOperator<DerivedOperatorType>::ExecuteFunction;
-			}
+		virtual FResetFunction GetResetFunction() override
+		{
+			return MetasoundExecutableOperatorPrivate::TResetFunctionAdapter<DerivedOperatorType>::GetResetFunction();
+		}
 
-		private:
+		virtual FExecuteFunction GetExecuteFunction() override
+		{
+			return &TExecutableOperator<DerivedOperatorType>::ExecuteFunction;
+		}
 
-			static void ExecuteFunction(IOperator* InOperator)
-			{
-				DerivedOperatorType* DerivedOperator = static_cast<DerivedOperatorType*>(InOperator);
+	private:
 
-				check(nullptr != DerivedOperator);
+		static void ExecuteFunction(IOperator* InOperator)
+		{
+			DerivedOperatorType* DerivedOperator = static_cast<DerivedOperatorType*>(InOperator);
 
-				DerivedOperator->Execute();
-			}
+			check(nullptr != DerivedOperator);
+
+			DerivedOperator->Execute();
+		}
 	};
 
 	class FNoOpOperator : public IOperator
@@ -79,6 +152,11 @@ namespace Metasound
 				return nullptr;
 			}
 
+			FResetFunction GetResetFunction() override
+			{
+				return nullptr;
+			}
+
 			FDataReferenceCollection GetInputs() const override
 			{
 				return FDataReferenceCollection{};
@@ -101,12 +179,14 @@ namespace Metasound
 			using FExecuteFunction = IOperator::FExecuteFunction;
 
 			FExecuter()
-			:	ExecuteFunction(&FExecuter::NoOp)
+			:	ExecuteFunction(&FExecuter::NoOpExecute)
+			,	ResetFunction(&FExecuter::NoOpReset)
 			{
 			}
 
 			FExecuter(FOperatorPtr InOperator)
-			:	ExecuteFunction(&FExecuter::NoOp)
+			:	ExecuteFunction(&FExecuter::NoOpExecute)
+			,	ResetFunction(&FExecuter::NoOpReset)
 			{
 				SetOperator(MoveTemp(InOperator));
 			}
@@ -114,30 +194,52 @@ namespace Metasound
 			void SetOperator(FOperatorPtr InOperator)
 			{
 				Operator = MoveTemp(InOperator);
+				ExecuteFunction = nullptr;
+				ResetFunction = nullptr;
 
 				if (Operator.IsValid())
 				{
-					FExecuteFunction Func = Operator->GetExecuteFunction();
-
-					if (nullptr != Func)
+					if (FExecuteFunction Func = Operator->GetExecuteFunction())
 					{
 						ExecuteFunction = Func;
 					}
+
+					if (FResetFunction Func = Operator->GetResetFunction())
+					{
+						ResetFunction = Func;
+					}
 				}
-				else
+
+				if (!ExecuteFunction)
+				{
+					ExecuteFunction = &FExecuter::NoOpExecute;
+				}
+
+				if (!ResetFunction)
 				{
-					ExecuteFunction = &FExecuter::NoOp;
+					ResetFunction = &FExecuter::NoOpReset;
 				}
 			}
 
 			void Execute()
 			{
+				// ExecuteFunction is always non-null. If the contained operator
+				// does not provide a valid execution function, then a No-Op function
+				// is utilized. 
 				ExecuteFunction(Operator.Get());
 			}
 
+			void Reset(const IOperator::FResetParams& InParams)
+			{
+				// ResetFunction is always non-null. If the contained operator
+				// does not provide a valid initialization function, then a No-Op function
+				// is utilized. 
+				ResetFunction(Operator.Get(), InParams);
+			}
+
 			bool IsNoOp()
 			{
-				return (ExecuteFunction == &FExecuter::NoOp);
+				return (ExecuteFunction == &FExecuter::NoOpExecute) && (ResetFunction == &FExecuter::NoOpReset);
 			}
 
 			bool IsValid() const
@@ -174,12 +276,22 @@ namespace Metasound
 				return ExecuteFunction;
 			}
 
+			virtual FResetFunction GetResetFunction() override
+			{
+				return ResetFunction;
+			}
+
 		private:
 			FOperatorPtr Operator;
 
 			FExecuteFunction ExecuteFunction;
+			FResetFunction ResetFunction;
+
+			static void NoOpExecute(IOperator*)
+			{
+			}
 
-			static void NoOp(IOperator*)
+			static void NoOpReset(IOperator*, const IOperator::FResetParams&)
 			{
 			}
 	};
diff --git a/Engine/Plugins/Runtime/Metasound/Source/MetasoundGraphCore/Public/MetasoundGraphOperator.h b/Engine/Plugins/Runtime/Metasound/Source/MetasoundGraphCore/Public/MetasoundGraphOperator.h
--- a/Engine/Plugins/Runtime/Metasound/Source/MetasoundGraphCore/Public/MetasoundGraphOperator.h
+++ b/Engine/Plugins/Runtime/Metasound/Source/MetasoundGraphCore/Public/MetasoundGraphOperator.h
@@ -10,15 +10,17 @@
 
 namespace Metasound
 {
-	class METASOUNDGRAPHCORE_API FGraphOperator : public IOperator
+	class METASOUNDGRAPHCORE_API FGraphOperator : public TExecutableOperator<FGraphOperator>
 	{
 		public:
 			using FOperatorPtr = TUniquePtr<IOperator>;
 			using FExecuteFunction = IOperator::FExecuteFunction;
+			using FResetFunction = IOperator::FResetFunction;
+			using FResetParams = IOperator::FResetParams;
 
 			FGraphOperator() = default;
 
-			virtual ~FGraphOperator();
+			virtual ~FGraphOperator() = default;
 
 			// Add an operator to the end of the executation stack.
 			void AppendOperator(FOperatorPtr InOperator);
@@ -40,18 +42,35 @@ namespace Metasound
 			// Bind the graph's interface data references to FVertexInterfaceData.
 			virtual void Bind(FVertexInterfaceData& InOutVertexData) const override;
 
-			virtual FExecuteFunction GetExecuteFunction() override;
-
 			void Execute();
+			void Reset(const FResetParams& InParams);
 
 		private:
 			// Delete copy operator because underlying types cannot be copied. 
 			FGraphOperator& operator=(const FGraphOperator&) = delete;
 			FGraphOperator(const FGraphOperator&) = delete;
 
-			static void ExecuteFunction(IOperator* InOperator);
+			struct FExecuteEntry
+			{
+				FExecuteEntry(IOperator& InOperator, FExecuteFunction InFunc);
+				void Execute();
+
+				IOperator* Operator;
+				FExecuteFunction Function;	
+			};
+
+			struct FResetEntry
+			{
+				FResetEntry(IOperator& InOperator, FResetFunction InFunc);
+				void Reset(const FResetParams& InParams);
+
+				IOperator* Operator;
+				FResetFunction Function;	
+			};
 
-			TArray<FExecuter> OperatorStack;
+			TArray<FExecuteEntry> ExecuteStack;
+			TArray<FResetEntry> ResetStack;
+			TArray<TUniquePtr<IOperator>> ActiveOperators;
 			FVertexInterfaceData VertexData;
 	};
 }
diff --git a/Engine/Plugins/Runtime/Metasound/Source/MetasoundGraphCore/Public/MetasoundLiteral.h b/Engine/Plugins/Runtime/Metasound/Source/MetasoundGraphCore/Public/MetasoundLiteral.h
--- a/Engine/Plugins/Runtime/Metasound/Source/MetasoundGraphCore/Public/MetasoundLiteral.h
+++ b/Engine/Plugins/Runtime/Metasound/Source/MetasoundGraphCore/Public/MetasoundLiteral.h
@@ -109,9 +109,6 @@ namespace Metasound
 #endif
 		}
 
-		FLiteral(FLiteral&& Other) = default;
-		FLiteral& operator=(FLiteral&& Other) = default;
-
 		FVariantType Value;
 
 		// builds an invalid FLiteral.
diff --git a/Engine/Plugins/Runtime/Metasound/Source/MetasoundGraphCore/Public/MetasoundLiteralNode.h b/Engine/Plugins/Runtime/Metasound/Source/MetasoundGraphCore/Public/MetasoundLiteralNode.h
--- a/Engine/Plugins/Runtime/Metasound/Source/MetasoundGraphCore/Public/MetasoundLiteralNode.h
+++ b/Engine/Plugins/Runtime/Metasound/Source/MetasoundGraphCore/Public/MetasoundLiteralNode.h
@@ -59,6 +59,11 @@ namespace Metasound
 			return nullptr;
 		}
 
+		virtual FResetFunction GetResetFunction() override
+		{
+			return nullptr;
+		}
+
 	private:
 		TDataValueReference<DataType> Value;
 	};
@@ -69,18 +74,19 @@ namespace Metasound
 	 * Note: this is supporting for FTrigger literals which are deprecated.
 	 */
 	template<typename DataType>
-	class TExecutableLiteralOperator : public IOperator
+	class TExecutableLiteralOperator : public TExecutableOperator<TExecutableLiteralOperator<DataType>>
 	{
 		static_assert(TExecutableDataType<DataType>::bIsExecutable, "TExecutableLiteralOperator is only suitable for executable data types");
 
 	public:
 		using FDataWriteReference = TDataWriteReference<DataType>;
 
-		TExecutableLiteralOperator(FDataWriteReference InDataReference)
+		TExecutableLiteralOperator(const FOperatorSettings& InSettings, const FLiteral& InLiteral)
 			// Executable DataTypes require a copy of the output to operate on whereas non-executable
 			// types do not. Avoid copy by assigning to reference for non-executable types.
-			: InputValue(InDataReference)
-			, OutputValue(FDataWriteReference::CreateNew(*InDataReference))
+			: Literal(InLiteral)
+			, InputValue(TDataWriteReferenceLiteralFactory<DataType>::CreateExplicitArgs(InSettings, InLiteral))
+			, OutputValue(TDataWriteReferenceLiteralFactory<DataType>::CreateExplicitArgs(InSettings, InLiteral))
 		{
 		}
 
@@ -111,17 +117,13 @@ namespace Metasound
 			TExecutableDataType<DataType>::Execute(*InputValue, *OutputValue);
 		}
 
-		static void ExecuteFunction(IOperator* InOperator)
-		{
-			static_cast<TExecutableLiteralOperator<DataType>*>(InOperator)->Execute();
-		}
-
-		virtual FExecuteFunction GetExecuteFunction() override
+		void Reset(const IOperator::FResetParams& InParams)
 		{
-			return &TExecutableLiteralOperator<DataType>::ExecuteFunction;
+			*OutputValue = *InputValue = TDataTypeLiteralFactory<DataType>::CreateExplicitArgs(InParams.OperatorSettings, Literal);
 		}
 
 	private:
+		FLiteral Literal;
 		FDataWriteReference InputValue;
 		FDataWriteReference OutputValue;
 	};
@@ -250,7 +252,7 @@ namespace Metasound
 		{
 			// Create write reference by calling compatible constructor with literal.
 			TDataWriteReference<DataType> DataRef = TDataWriteReferenceLiteralFactory<DataType>::CreateExplicitArgs(InParams.OperatorSettings, InitParam);
-			return MakeUnique<TExecutableLiteralOperator<DataType>>(DataRef);
+			return MakeUnique<TExecutableLiteralOperator<DataType>>(InParams.OperatorSettings, InitParam);
 		}
 		else
 		{
diff --git a/Engine/Plugins/Runtime/Metasound/Source/MetasoundGraphCore/Public/MetasoundOperatorInterface.h b/Engine/Plugins/Runtime/Metasound/Source/MetasoundGraphCore/Public/MetasoundOperatorInterface.h
--- a/Engine/Plugins/Runtime/Metasound/Source/MetasoundGraphCore/Public/MetasoundOperatorInterface.h
+++ b/Engine/Plugins/Runtime/Metasound/Source/MetasoundGraphCore/Public/MetasoundOperatorInterface.h
@@ -7,6 +7,7 @@
 
 namespace Metasound
 {
+
 	/** IOperator
 	 *
 	 *  IOperator defines the interface for render time operations.  IOperators are created using an INodeOperatorFactory.
@@ -14,11 +15,17 @@ namespace Metasound
 	class IOperator
 	{
 	public:
-		/** Pointer to execute function for an operator.
-		 *
-		 * @param IOperator* - The operator associated with the function pointer.
+		/** FResetOperatorParams holds the parameters provided to an IOperator's 
+		 * reset function.
 		 */
-		typedef void(*FExecuteFunction)(IOperator*);
+		struct FResetParams
+		{
+			/** General operator settings for the graph. */
+			const FOperatorSettings& OperatorSettings; 
+
+			/** Environment settings available. */
+			const FMetasoundEnvironment& Environment;
+		};
 
 		virtual ~IOperator() {}
 
@@ -52,6 +59,28 @@ namespace Metasound
 			InVertexData.GetOutputs().Bind(GetOutputs());
 		}
 
+		/** Pointer to intialize function for an operator.
+		 *
+		 * @param IOperator* - The operator associated with the function pointer.
+		 */
+		using FResetFunction = void(*)(IOperator*, const FResetParams& InParams);
+
+		/** Return the reset function to call during graph execution.
+		 *
+		 * The IOperator* argument to the FExecutionFunction will be the same IOperator instance
+		 * which returned the execution function.
+		 *
+		 * nullptr return values are valid and signal an IOperator which does not need to be
+		 * reset.
+		 */
+		virtual FResetFunction GetResetFunction() = 0;
+
+		/** Pointer to execute function for an operator.
+		 *
+		 * @param IOperator* - The operator associated with the function pointer.
+		 */
+		using FExecuteFunction = void(*)(IOperator*);
+
 		/** Return the execution function to call during graph execution.
 		 *
 		 * The IOperator* argument to the FExecutionFunction will be the same IOperator instance
diff --git a/Engine/Plugins/Runtime/Metasound/Source/MetasoundGraphCore/Public/MetasoundVariable.h b/Engine/Plugins/Runtime/Metasound/Source/MetasoundGraphCore/Public/MetasoundVariable.h
--- a/Engine/Plugins/Runtime/Metasound/Source/MetasoundGraphCore/Public/MetasoundVariable.h
+++ b/Engine/Plugins/Runtime/Metasound/Source/MetasoundGraphCore/Public/MetasoundVariable.h
@@ -42,6 +42,11 @@ namespace Metasound
 			DataReference = InDataReference;
 		}
 
+		void SetDelayedDataValue(const DataType& InDataType)
+		{
+			*DelayedDataReference = InDataType;
+		}
+
 		/** Get the current data reference. */
 		FReadReference GetDataReference() const
 		{
diff --git a/Engine/Plugins/Runtime/Metasound/Source/MetasoundGraphCore/Public/MetasoundVariableNodes.h b/Engine/Plugins/Runtime/Metasound/Source/MetasoundGraphCore/Public/MetasoundVariableNodes.h
--- a/Engine/Plugins/Runtime/Metasound/Source/MetasoundGraphCore/Public/MetasoundVariableNodes.h
+++ b/Engine/Plugins/Runtime/Metasound/Source/MetasoundGraphCore/Public/MetasoundVariableNodes.h
@@ -94,25 +94,23 @@ namespace Metasound
 
 	}
 
-	/** Variable nodes initialize variable values. The output of a VariableNode
-	 * is a TVariable.  */
-	template<typename DataType>
-	class TVariableNode : public FNode
+	namespace MetasoundVariableNodesPrivate
 	{
-		using FVariable = TVariable<DataType>;
-
-		class FOperator : public TExecutableOperator<FOperator>
+		template<typename DataType>
+		class TVariableOperator : public TExecutableOperator<TVariableOperator<DataType>>
 		{
-			using Super = TExecutableOperator<FOperator>;
+			using Super = TExecutableOperator<TVariableOperator<DataType>>;
+			using FVariable = TVariable<DataType>;
 
 		public:
-			FOperator(TDataWriteReference<FVariable> InVariable)
-			: Variable(InVariable)
+			TVariableOperator(const FOperatorSettings& InOperatorSettings, const FLiteral& InLiteral)
+			: Literal(InLiteral)
+			, Variable(TDataWriteReference<FVariable>::CreateNew(TDataWriteReferenceLiteralFactory<DataType>::CreateExplicitArgs(InOperatorSettings, Literal)))
 			, bCopyReferenceDataOnExecute(false)
 			{
 			}
 
-			virtual ~FOperator() = default;
+			virtual ~TVariableOperator() = default;
 
 			virtual FDataReferenceCollection GetInputs() const override
 			{
@@ -157,11 +155,26 @@ namespace Metasound
 				}
 			}
 
+			void Reset(const IOperator::FResetParams& InParams)
+			{
+				Variable->SetDelayedDataValue(TDataTypeLiteralFactory<DataType>::CreateExplicitArgs(InParams.OperatorSettings, Literal));
+			}
+
 		private:
 
+			FLiteral Literal;
 			TDataWriteReference<FVariable> Variable;
 			bool bCopyReferenceDataOnExecute;
 		};
+	}
+
+
+	/** Variable nodes initialize variable values. The output of a VariableNode
+	 * is a TVariable.  */
+	template<typename DataType>
+	class TVariableNode : public FNode
+	{
+		using FVariable = TVariable<DataType>;
 
 		class FFactory : public IOperatorFactory
 		{
@@ -175,8 +188,7 @@ namespace Metasound
 
 			virtual TUniquePtr<IOperator> CreateOperator(const FBuildOperatorParams& InParams, FBuildResults& OutResults) override
 			{
-				TDataWriteReference<DataType> Data = TDataWriteReferenceLiteralFactory<DataType>::CreateExplicitArgs(InParams.OperatorSettings, Literal);
-				return MakeUnique<FOperator>(TDataWriteReference<FVariable>::CreateNew(Data));
+				return MakeUnique<MetasoundVariableNodesPrivate::TVariableOperator<DataType>>(InParams.OperatorSettings, Literal);
 			}
 		private:
 			FLiteral Literal;
@@ -216,8 +228,8 @@ namespace Metasound
 			return Info;
 		}
 
-	public:
 
+	public:
 		TVariableNode(const FVertexName& InNodeName, const FGuid& InInstanceID, FLiteral&& InLiteral)
 		: FNode(InNodeName, InInstanceID, GetNodeMetadata())
 		, Interface(DeclareVertexInterface())
@@ -384,7 +396,7 @@ namespace Metasound
 	class TVariableDeferredAccessorNode : public FNodeFacade
 	{
 		using FVariable = TVariable<DataType>;
-
+	
 		class FOperator : public FNoOpOperator 
 		{
 		public:
diff --git a/Engine/Plugins/Runtime/Metasound/Source/MetasoundStandardNodes/Private/MetasoundAudioFormats.cpp b/Engine/Plugins/Runtime/Metasound/Source/MetasoundStandardNodes/Private/MetasoundAudioFormats.cpp
--- a/Engine/Plugins/Runtime/Metasound/Source/MetasoundStandardNodes/Private/MetasoundAudioFormats.cpp
+++ b/Engine/Plugins/Runtime/Metasound/Source/MetasoundStandardNodes/Private/MetasoundAudioFormats.cpp
@@ -137,6 +137,11 @@ namespace Metasound
 				return nullptr;
 			}
 
+			virtual FResetFunction GetResetFunction() override
+			{
+				return nullptr;
+			}
+
 		private:
 
 			FVertexName OutputName;
@@ -282,6 +287,11 @@ namespace Metasound
 				return nullptr;
 			}
 
+			virtual FResetFunction GetResetFunction() override
+			{
+				return nullptr;
+			}
+
 		private:
 
 			FVertexName InputName;
@@ -421,6 +431,11 @@ namespace Metasound
 				return nullptr;
 			}
 
+			virtual FResetFunction GetResetFunction() override
+			{
+				return nullptr;
+			}
+
 		private:
 
 			FVertexName OutputName;
@@ -555,6 +570,11 @@ namespace Metasound
 				return nullptr;
 			}
 
+			virtual FResetFunction GetResetFunction() override
+			{
+				return nullptr;
+			}
+
 		private:
 
 			FVertexName InputName;
diff --git a/Engine/Plugins/Runtime/Metasound/Source/MetasoundStandardNodes/Private/MetasoundNoiseGenerator.cpp b/Engine/Plugins/Runtime/Metasound/Source/MetasoundStandardNodes/Private/MetasoundNoiseGenerator.cpp
--- a/Engine/Plugins/Runtime/Metasound/Source/MetasoundStandardNodes/Private/MetasoundNoiseGenerator.cpp
+++ b/Engine/Plugins/Runtime/Metasound/Source/MetasoundStandardNodes/Private/MetasoundNoiseGenerator.cpp
@@ -54,6 +54,7 @@ namespace Metasound
 		virtual FDataReferenceCollection GetInputs() const override;
 		virtual FDataReferenceCollection GetOutputs() const override;
 
+
 	protected:
 		FInt32ReadRef Seed;
 		FAudioBufferWriteRef Out;
@@ -73,6 +74,13 @@ namespace Metasound
 			}
 		}
 
+		template<typename T>
+		void ResetNoiseOperator(T& InOutGenerator)
+		{
+			InOutGenerator = MakeGenerator<T>(*Seed);
+			OldSeed = *Seed;
+		}
+
 		template<typename T>
 		FORCEINLINE void CheckAndReseed(T& InOutGenerator)
 		{
@@ -107,6 +115,12 @@ namespace Metasound
 			, Generator{ MakeGenerator<Audio::FWhiteNoise>(*Seed) }
 		{}
 
+		void Reset(const IOperator::FResetParams& InParams)
+		{
+			ResetNoiseOperator(Generator);
+			Out->Zero();
+		}
+
 		void Execute()
 		{
 			// Reseed if necessary.
@@ -115,8 +129,13 @@ namespace Metasound
 			// Generate a block.
 			Generate(Generator);
 		}
+
+		virtual FResetFunction GetResetFunction() override { return &FNoiseOperator_White::ResetFunction; }
+		static void ResetFunction(IOperator* InOperator, const IOperator::FResetParams& InParams) { static_cast<FNoiseOperator_White*>(InOperator)->Reset(InParams); }
+
 		static void ExecuteFunction(IOperator* InOperator) { static_cast<FNoiseOperator_White*>(InOperator)->Execute(); }
 		FExecuteFunction GetExecuteFunction() override { return &FNoiseOperator_White::ExecuteFunction; }
+
 	};
 
 	struct FNoiseOperator_Pink final : public FNoiseOperator
@@ -132,6 +151,15 @@ namespace Metasound
 		static void ExecuteFunction(IOperator* InOperator) { static_cast<FNoiseOperator_Pink*>(InOperator)->Execute(); }
 		FExecuteFunction GetExecuteFunction() override { return &FNoiseOperator_Pink::ExecuteFunction; }
 
+		virtual FResetFunction GetResetFunction() override { return &FNoiseOperator_Pink::ResetFunction; }
+		static void ResetFunction(IOperator* InOperator, const IOperator::FResetParams& InParams) { static_cast<FNoiseOperator_Pink*>(InOperator)->Reset(InParams); }
+
+		void Reset(const IOperator::FResetParams& InParams)
+		{
+			ResetNoiseOperator(Generator);
+			Out->Zero();
+		}
+
 		void Execute()
 		{
 			// Reseed if necessary.
