fix crash in MetasoundRandomGet nodes if shared state is not enabled

diff --git a/Engine/Plugins/Runtime/Metasound/Source/MetasoundFrontend/Public/MetasoundArrayRandomNode.h b/Engine/Plugins/Runtime/Metasound/Source/MetasoundFrontend/Public/MetasoundArrayRandomNode.h
--- a/Engine/Plugins/Runtime/Metasound/Source/MetasoundFrontend/Public/MetasoundArrayRandomNode.h
+++ b/Engine/Plugins/Runtime/Metasound/Source/MetasoundFrontend/Public/MetasoundArrayRandomNode.h
@@ -257,6 +257,8 @@ namespace Metasound
 					Args.Weights = WeightsArray;
 
 					RGM.InitSharedState(Args);
+
+					bSharedStateInitialized = true;
 				}
 				else
 				{
@@ -276,6 +278,11 @@ namespace Metasound
 			TriggerOnReset->Reset();
 		}
 
+		bool UseSharedState() const
+		{
+			return bSharedStateInitialized && *bEnableSharedState;
+		}
+
 		void Execute()
 		{
 			TriggerOnNext->AdvanceBlock();
@@ -314,7 +321,7 @@ namespace Metasound
 			{
 				PrevSeedValue = *SeedValue;
 
-				if (SharedStateUniqueId.IsValid())
+				if (UseSharedState())
 				{
 					FSharedStateRandomGetManager& RGM = FSharedStateRandomGetManager::Get();
 					RGM.SetSeed(SharedStateUniqueId, PrevSeedValue);
@@ -329,7 +336,7 @@ namespace Metasound
 			if (PrevNoRepeatOrder != *NoRepeatOrder)
 			{
 				PrevNoRepeatOrder = *NoRepeatOrder;
-				if (SharedStateUniqueId.IsValid())
+				if (UseSharedState())
 				{
 					FSharedStateRandomGetManager& RGM = FSharedStateRandomGetManager::Get();
 					RGM.SetNoRepeatOrder(SharedStateUniqueId, PrevNoRepeatOrder);
@@ -342,7 +349,7 @@ namespace Metasound
 			}
 
 			WeightsArray = *InputWeightsArray;
-			if (SharedStateUniqueId.IsValid())
+			if (UseSharedState())
 			{
 				FSharedStateRandomGetManager& RGM = FSharedStateRandomGetManager::Get();
 				RGM.SetRandomWeights(SharedStateUniqueId, WeightsArray);
@@ -365,7 +372,7 @@ namespace Metasound
 				},
 				[this](int32 StartFrame, int32 EndFrame)
 				{
-					if (SharedStateUniqueId.IsValid())
+					if (UseSharedState())
 					{
 						FSharedStateRandomGetManager& RGM = FSharedStateRandomGetManager::Get();
 						RGM.ResetSeed(SharedStateUniqueId);
@@ -388,7 +395,7 @@ namespace Metasound
 					const ArrayType& InputArrayRef = *InputArray;
 					int32 OutRandomIndex = INDEX_NONE;
 
-					if (SharedStateUniqueId.IsValid())
+					if (UseSharedState())
 					{
 						FSharedStateRandomGetManager& RGM = FSharedStateRandomGetManager::Get();
 						OutRandomIndex = RGM.NextValue(SharedStateUniqueId);
@@ -438,6 +445,7 @@ namespace Metasound
 		int32 PrevArraySize = 0;
 		bool bIsPreviewSound = false;
 		bool bHasLoggedEmptyArrayWarning = false;
+		bool bSharedStateInitialized = false;
 	};
 
 	template<typename ArrayType>
