Speculative defense around Wave Player Node start time (Clamping around number of frames instead of just duration in seconds)

diff --git a/Engine/Plugins/Runtime/Metasound/Source/MetasoundEngine/Private/MetasoundWaveProxyReader.cpp b/Engine/Plugins/Runtime/Metasound/Source/MetasoundEngine/Private/MetasoundWaveProxyReader.cpp
--- a/Engine/Plugins/Runtime/Metasound/Source/MetasoundEngine/Private/MetasoundWaveProxyReader.cpp
+++ b/Engine/Plugins/Runtime/Metasound/Source/MetasoundEngine/Private/MetasoundWaveProxyReader.cpp
@@ -431,7 +431,7 @@ namespace Metasound
 					InStartTimeInSeconds);
 			}
 		}
-		CurrentFrameIndex = InStartTimeInSeconds * GetSampleRate();
+		CurrentFrameIndex = FMath::Clamp(static_cast<int32>(InStartTimeInSeconds * GetSampleRate()), 0, GetNumFramesInWave());
 
 		// Get codec ptr by reading the header info from the decoder input.
 		ICodecRegistry::FCodecPtr Codec = ICodecRegistry::Get().FindCodecByParsingInput(DecoderInput.Get());
