Fix robomerge conflict

diff --git a/Engine/Plugins/Runtime/Metasound/Source/MetasoundFrontend/Public/MetasoundDataTypeRegistrationMacro.h b/Engine/Plugins/Runtime/Metasound/Source/MetasoundFrontend/Public/MetasoundDataTypeRegistrationMacro.h
--- a/Engine/Plugins/Runtime/Metasound/Source/MetasoundFrontend/Public/MetasoundDataTypeRegistrationMacro.h
+++ b/Engine/Plugins/Runtime/Metasound/Source/MetasoundFrontend/Public/MetasoundDataTypeRegistrationMacro.h
@@ -25,55 +25,12 @@
 #include "MetasoundSendNode.h"
 #include "MetasoundTransmissionRegistration.h"
 #include "MetasoundVariableNodes.h"
+#include "MetasoundParameterPackFixedArray.h"
 
 #include <type_traits>
 
 namespace Metasound
 {
-	template<typename ElementType, int32 NumInlineElements>
-	class TParamPackFixedArray
-	{
-	public:
-		TParamPackFixedArray(const TArray<ElementType>& Other)
-		{
-			*this = Other;
-		}
-
-		TParamPackFixedArray<ElementType, NumInlineElements>& operator=(const TArray<ElementType>& Other)
-		{
-			int32 NumToCopy = FMath::Min(Other.Num(), NumInlineElements);
-			for (int32 i = 0; i < NumToCopy; ++i)
-			{
-				GetData()[i] = Other[i];
-			}
-			NumValidElements = NumToCopy;
-			return *this;
-		}
-
-		void CopyToArray(TArray<ElementType>& Dest) const
-		{
-			Dest.Empty();
-			for (int32 i = 0; i < NumValidElements; ++i)
-			{
-				Dest.Add(GetData()[i]);
-			}
-		}
-
-		FORCEINLINE const ElementType* GetData() const
-		{
-			return (ElementType*)InlineData;
-		}
-
-		FORCEINLINE ElementType* GetData()
-		{
-			return (ElementType*)InlineData;
-		}
-
-	private:
-		int32 NumValidElements = 0;
-		TTypeCompatibleBytes<ElementType> InlineData[NumInlineElements];
-	};
-
 	namespace MetasoundDataTypeRegistrationPrivate
 	{
 		template<typename DataType>
@@ -468,9 +425,15 @@ namespace Metasound
 						{
 							this->RawAssignmentFunction = [](const void* Src, void* Dest)
 								{
+									//***********************************************************************************************************
+									// sanity check to be sure memory layout is the same regardless of number of elements in the fixed array...
+									using FSmallFixedArray = TParamPackFixedArray<typename TDataType::ElementType, 1>;
+									using FLargeFixedArray = TParamPackFixedArray<typename TDataType::ElementType, 200>;
+									static_assert(offsetof(FSmallFixedArray,InlineData) > offsetof(FSmallFixedArray, NumValidElements));
+									static_assert(offsetof(FSmallFixedArray,InlineData) == offsetof(FLargeFixedArray,InlineData));
+									//***********************************************************************************************************
 									TDataType& DestinationArray = *(reinterpret_cast<TDataType*>(Dest));
-									const TParamPackFixedArray<typename TDataType::ElementType, 1>& SourceArray = 
-										*(reinterpret_cast<const TParamPackFixedArray<typename TDataType::ElementType, 1>*>(Src));
+									const FSmallFixedArray& SourceArray = *(reinterpret_cast<const FSmallFixedArray*>(Src));
 									SourceArray.CopyToArray(DestinationArray);
 								};
 						}
diff --git a/Engine/Plugins/Runtime/Metasound/Source/MetasoundFrontend/Public/MetasoundParameterPack.h b/Engine/Plugins/Runtime/Metasound/Source/MetasoundFrontend/Public/MetasoundParameterPack.h
--- a/Engine/Plugins/Runtime/Metasound/Source/MetasoundFrontend/Public/MetasoundParameterPack.h
+++ b/Engine/Plugins/Runtime/Metasound/Source/MetasoundFrontend/Public/MetasoundParameterPack.h
@@ -16,59 +16,62 @@ class UMetaSoundSource;
 // A structure that encapsulates a "parameter tag" in the raw 'bag-o-bytes'
 // that will be the parameter storage. It lets us walk through memory to 
 // find a parameter with a specific name and specific type. 
-struct FMetasoundParameterPackItemBase
+namespace MetasoundParameterPackPrivate
 {
-	FMetasoundParameterPackItemBase(const FName& InName, const FName& InTypeName, uint16 InNextOffset, uint16 InValueOffset)
-		: Name(InName)
-		, TypeName(InTypeName)
-		, NextHeaderOffset(InNextOffset)
-		, ValueOffset(InValueOffset)
-	{}
-	FMetasoundParameterPackItemBase(const FMetasoundParameterPackItemBase& Other) = default;
-	void* GetPayload()
+	struct FMetasoundParameterPackItemBase
 	{
-		uint8* ptr = reinterpret_cast<uint8*>(this);
-		return ptr + ValueOffset;
-	}
-	FName  Name;
-	FName  TypeName;
-	uint16 NextHeaderOffset;
-	uint16 ValueOffset;
-};
+		FMetasoundParameterPackItemBase(const FName& InName, const FName& InTypeName, uint16 InNextOffset, uint16 InValueOffset)
+			: Name(InName)
+			, TypeName(InTypeName)
+			, NextHeaderOffset(InNextOffset)
+			, ValueOffset(InValueOffset)
+		{}
+		FMetasoundParameterPackItemBase(const FMetasoundParameterPackItemBase& Other) = default;
+		void* GetPayload()
+		{
+			uint8* ptr = reinterpret_cast<uint8*>(this);
+			return ptr + ValueOffset;
+		}
+		FName  Name;
+		FName  TypeName;
+		uint16 NextHeaderOffset;
+		uint16 ValueOffset;
+	};
 
-template<bool, typename T, int32 MAX_NUM_ITEMS = 0>
-struct TRootValueTypeHelper { using Type = T; };
+	template<bool, typename T, int32 MaxNumItems = 0>
+	struct TRootValueTypeHelper { using Type = T; };
 
-template<typename T, int32 MAX_NUM_ITEMS>
-struct TRootValueTypeHelper<true, T, MAX_NUM_ITEMS> { using Type = Metasound::TParamPackFixedArray<typename T::ElementType, MAX_NUM_ITEMS>; };
+	template<typename T, int32 MaxNumItems>
+	struct TRootValueTypeHelper<true, T, MaxNumItems> { using Type = Metasound::TParamPackFixedArray<typename T::ElementType, MaxNumItems>; };
 
-template<typename T, int32 MAX_NUM_ITEMS = 0>
-struct TRootValueType : TRootValueTypeHelper<TIsTArray_V<T>, T, MAX_NUM_ITEMS> {};
+	template<typename T, int32 MaxNumItems = 0>
+	struct TRootValueType : TRootValueTypeHelper<TIsTArray_V<T>, T, MaxNumItems> {};
 
-// A template that lets us stick an arbitrary data type into the raw 
-// 'bag-o-bytes' that is the parameter storage. Uses the base class 
-// above.
-template <typename T, int32 MAX_NUM = 0>
-struct FMetasoundParameterPackItem : public FMetasoundParameterPackItemBase
-{
-	using ThisType = FMetasoundParameterPackItem<T,MAX_NUM>;
-	FMetasoundParameterPackItem(const FName& InName, const FName& InTypeName, const T& InValue)
-		: FMetasoundParameterPackItemBase(InName, InTypeName, SizeOfParam(InValue), offsetof(FMetasoundParameterPackItem<ThisType>, Value))
-		, Value(InValue)
+	// A template that lets us stick an arbitrary data type into the raw 
+	// 'bag-o-bytes' that is the parameter storage. Uses the base class 
+	// above.
+	template <typename T, int32 MaxNumItems = 0>
+	struct FMetasoundParameterPackItem : public FMetasoundParameterPackItemBase
 	{
-		if constexpr (TIsTArray_V<T>)
+		using ThisType = FMetasoundParameterPackItem<T,MaxNumItems>;
+		FMetasoundParameterPackItem(const FName& InName, const FName& InTypeName, const T& InValue)
+			: FMetasoundParameterPackItemBase(InName, InTypeName, SizeOfParam(InValue), offsetof(FMetasoundParameterPackItem<ThisType>, Value))
+			, Value(InValue)
 		{
-			check(InValue.Num() <= MAX_NUM);
+			if constexpr (TIsTArray_V<T>)
+			{
+				check(InValue.Num() <= MaxNumItems);
+			}
 		}
-	}
-	FMetasoundParameterPackItem(const FMetasoundParameterPackItem& Other) = default;
-	typename TRootValueType<T, MAX_NUM>::Type Value;
+		FMetasoundParameterPackItem(const FMetasoundParameterPackItem& Other) = default;
+		typename TRootValueType<T, MaxNumItems>::Type Value;
 	
-	static uint16 SizeOfParam(const T& InValue)
-	{
-		return AlignArbitrary<uint16>((uint16)sizeof(ThisType), alignof(std::max_align_t));
-	}
-};
+		static uint16 SizeOfParam(const T& InValue)
+		{
+			return AlignArbitrary<uint16>((uint16)sizeof(ThisType), alignof(std::max_align_t));
+		}
+	};
+}
 
 // This next class owns the 'bag-o-bytes' that holds all of the parameters. 
 struct FMetasoundParameterPackStorage
@@ -100,11 +103,13 @@ struct FMetasoundParameterPackStorage
 
 		ParameterIterator& operator++()
 		{
+			using namespace MetasoundParameterPackPrivate;
+
 			if (!CurrentNode)
 			{
 				return *this; 
 			}
-			FMetasoundParameterPackItemBase* ParamAddress = reinterpret_cast<FMetasoundParameterPackItemBase*>(CurrentNode);
+			MetasoundParameterPackPrivate::FMetasoundParameterPackItemBase * ParamAddress = reinterpret_cast<FMetasoundParameterPackItemBase*>(CurrentNode);
 			CurrentNode += ParamAddress->NextHeaderOffset;
 			if (CurrentNode >= EndAddress)
 			{
@@ -130,13 +135,15 @@ struct FMetasoundParameterPackStorage
 			return CurrentNode != Other.CurrentNode;
 		}
 
-		FMetasoundParameterPackItemBase* operator->()
+		MetasoundParameterPackPrivate::FMetasoundParameterPackItemBase* operator->()
 		{
+			using namespace MetasoundParameterPackPrivate;
 			return reinterpret_cast<FMetasoundParameterPackItemBase*>(CurrentNode);
 		}
 
-		FMetasoundParameterPackItemBase* GetParameterBase()
+		MetasoundParameterPackPrivate::FMetasoundParameterPackItemBase* GetParameterBase()
 		{
+			using namespace MetasoundParameterPackPrivate;
 			return reinterpret_cast<FMetasoundParameterPackItemBase*>(CurrentNode);
 		}
 
@@ -156,9 +163,10 @@ struct FMetasoundParameterPackStorage
 		return ParameterIterator(nullptr, nullptr);
 	}
 
-	template<typename T, int32 MAX_NUM_ITEMS = 0>
-	FMetasoundParameterPackItem<T, MAX_NUM_ITEMS>* FindParameter(const FName& ParamName, const FName& InTypeName)
+	template<typename T, int32 MaxNumItems = 0>
+	MetasoundParameterPackPrivate::FMetasoundParameterPackItem<T, MaxNumItems>* FindParameter(const FName& ParamName, const FName& InTypeName)
 	{
+		using namespace MetasoundParameterPackPrivate;
 		if (Storage.IsEmpty())
 		{
 			return nullptr;
@@ -171,7 +179,7 @@ struct FMetasoundParameterPackStorage
 			{
 				if (ParamWalker->TypeName == InTypeName)
 				{
-					return static_cast<FMetasoundParameterPackItem<T, MAX_NUM_ITEMS>*>(ParamWalker);
+					return static_cast<FMetasoundParameterPackItem<T, MaxNumItems>*>(ParamWalker);
 				}
 				return nullptr;
 			}
@@ -192,6 +200,7 @@ struct FMetasoundParameterPackStorage
 	template<typename T>
 	std::enable_if_t<!std::is_class_v<T>,T*> AddParameter(const FName& Name, const FName& InTypeName, const T& InValue)
 	{
+		using namespace MetasoundParameterPackPrivate;
 		int32 NextParamLocation = AlignArbitrary<int32>(StorageBytesInUse, alignof(std::max_align_t));
 		int32 TotalStorageNeeded = NextParamLocation + FMetasoundParameterPackItem<T>::SizeOfParam(InValue);
 		if (TotalStorageNeeded > Storage.Num())
@@ -207,6 +216,7 @@ struct FMetasoundParameterPackStorage
 	template<typename T>
 	std::enable_if_t<std::is_class_v<T> && !TIsTArray_V<T>, T*> AddParameter(const FName& Name, const FName& InTypeName, const T& InValue)
 	{
+		using namespace MetasoundParameterPackPrivate;
 		int32 NextParamLocation = AlignArbitrary<int32>(StorageBytesInUse, alignof(std::max_align_t));
 		int32 TotalStorageNeeded = NextParamLocation + FMetasoundParameterPackItem<T>::SizeOfParam(InValue);
 		if (TotalStorageNeeded > Storage.Num())
@@ -219,17 +229,18 @@ struct FMetasoundParameterPackStorage
 	}
 
 	// array version
-	template<typename T, int32 MAX_NUM_ITEMS>
+	template<typename T, int32 MaxNumItems>
 	std::enable_if_t<std::is_class_v<T> && TIsTArray_V<T>, void> AddArrayParameter(const FName& Name, const FName& InTypeName, const T& InValue)
 	{
-		check(InValue.Num() <= MAX_NUM_ITEMS);
+		using namespace MetasoundParameterPackPrivate;
+		check(InValue.Num() <= MaxNumItems);
 		int32 NextParamLocation = AlignArbitrary<int32>(StorageBytesInUse, alignof(std::max_align_t));
-		int32 TotalStorageNeeded = NextParamLocation + FMetasoundParameterPackItem<T, MAX_NUM_ITEMS>::SizeOfParam(InValue);
+		int32 TotalStorageNeeded = NextParamLocation + FMetasoundParameterPackItem<T, MaxNumItems>::SizeOfParam(InValue);
 		if (TotalStorageNeeded > Storage.Num())
 		{
 			Storage.AddUninitialized(TotalStorageNeeded - Storage.Num());
 		}
-		FMetasoundParameterPackItem<T, MAX_NUM_ITEMS>* Destination = new (reinterpret_cast<FMetasoundParameterPackItem<T>*>(&Storage[NextParamLocation])) FMetasoundParameterPackItem<T, MAX_NUM_ITEMS>(Name, InTypeName, InValue);
+		FMetasoundParameterPackItem<T, MaxNumItems>* Destination = new (reinterpret_cast<FMetasoundParameterPackItem<T>*>(&Storage[NextParamLocation])) FMetasoundParameterPackItem<T, MaxNumItems>(Name, InTypeName, InValue);
 		StorageBytesInUse = NextParamLocation + Destination->NextHeaderOffset;
 	}
 };
@@ -304,7 +315,7 @@ public:
 	template<typename T>
 	ESetParamResult SetParameter(const FName& Name, const FName& InTypeName, const T& InValue, bool OnlyIfExists = true)
 	{
-		typename TRootValueType<T>::Type* TheParam = GetParameterPtr<T>(Name, InTypeName);
+		typename MetasoundParameterPackPrivate::TRootValueType<T>::Type* TheParam = GetParameterPtr<T>(Name, InTypeName);
 		if (TheParam)
 		{
 			*TheParam = InValue;
@@ -319,15 +330,15 @@ public:
 		return ESetParamResult::Failed;
 	}
 
-	template<typename T, int32 MAX_NUM_ITEMS>
+	template<typename T, int32 MaxNumItems>
 	ESetParamResult SetArrayParameter(const FName& Name, const FName& InTypeName, const T& InValue, bool OnlyIfExists = true)
 	{
 		static_assert(TIsTArray_V<T>);
-		if (!ensureMsgf(InValue.Num() <= MAX_NUM_ITEMS, TEXT("Too many items in source array \"%s\"! Max is %d."), *Name.ToString(), MAX_NUM_ITEMS))
+		if (!ensureMsgf(InValue.Num() <= MaxNumItems, TEXT("Too many items in source array \"%s\"! Max is %d."), *Name.ToString(), MaxNumItems))
 		{
 			return ESetParamResult::Failed;
 		}
-		typename TRootValueType<T, MAX_NUM_ITEMS>::Type* TheParam = GetParameterPtr<T,MAX_NUM_ITEMS>(Name, InTypeName);
+		typename MetasoundParameterPackPrivate::TRootValueType<T, MaxNumItems>::Type* TheParam = GetParameterPtr<T,MaxNumItems>(Name, InTypeName);
 		if (TheParam)
 		{
 			*TheParam = InValue;
@@ -336,20 +347,20 @@ public:
 		else if (!OnlyIfExists)
 		{
 			// TODO Add test to see if we are allowed to create it!
-			AddArrayParameter<T, MAX_NUM_ITEMS>(Name, InTypeName, InValue);
+			AddArrayParameter<T, MaxNumItems>(Name, InTypeName, InValue);
 			return ESetParamResult::Succeeded;
 		}
 		return ESetParamResult::Failed;
 	}
 
-	template<typename T, int32 MAX_NUM_ITEMS=0>
+	template<typename T, int32 MaxNumItems=0>
 	bool HasParameter(const FName& Name, const FName& InTypeName) const
 	{
 		if (!ParameterStorage.IsValid())
 		{
 			return false;
 		}
-		return ParameterStorage->FindParameter<T,MAX_NUM_ITEMS>(Name, InTypeName) != nullptr;
+		return ParameterStorage->FindParameter<T,MaxNumItems>(Name, InTypeName) != nullptr;
 	}
 
 	template<typename T>
@@ -362,29 +373,14 @@ public:
 		ParameterStorage->AddParameter(Name, InTypeName, InValue);
 	}
 
-	template<typename T, int32 MAX_NUM_ITEMS>
+	template<typename T, int32 MaxNumItems>
 	void AddArrayParameter(const FName& Name, const FName& InTypeName, const T& InValue)
 	{
 		if (!ParameterStorage.IsValid())
 		{
 			ParameterStorage = MakeShared<FMetasoundParameterPackStorage>();
 		}
-		ParameterStorage->AddArrayParameter<T,MAX_NUM_ITEMS>(Name, InTypeName, InValue);
-	}
-
-	template<typename T, int MAX_NUM_ITEMS = 0>
-	typename TRootValueType<T, MAX_NUM_ITEMS>::Type* GetParameterPtr(const FName& Name, const FName& InTypeName) const
-	{
-		if (!ParameterStorage.IsValid())
-		{
-			return nullptr;
-		}
-		auto TheParameter = ParameterStorage->FindParameter<T, MAX_NUM_ITEMS>(Name, InTypeName);
-		if (!TheParameter)
-		{
-			return nullptr;
-		}
-		return &TheParameter->Value;
+		ParameterStorage->AddArrayParameter<T,MaxNumItems>(Name, InTypeName, InValue);
 	}
 
 	template<typename T>
@@ -400,10 +396,10 @@ public:
 		return *TheParameter;
 	}
 
-	template<typename T, int32 MAX_NUM_ITEMS>
+	template<typename T, int32 MaxNumItems>
 	void GetParameter(const FName& Name, const FName& InTypeName, TArray<typename T::ElementType>& ResultArray, ESetParamResult& Result) const
 	{
-		typename TRootValueType<T, MAX_NUM_ITEMS>::Type* TheParameter = GetParameterPtr<T,MAX_NUM_ITEMS>(Name, InTypeName);
+		typename MetasoundParameterPackPrivate::TRootValueType<T, MaxNumItems>::Type* TheParameter = GetParameterPtr<T,MaxNumItems>(Name, InTypeName);
 		if (!TheParameter)
 		{
 			Result = ESetParamResult::Failed;
@@ -424,7 +420,22 @@ public:
 	FSharedMetasoundParameterStoragePtr GetCopyOfParameterStorage() const;
 
 private:
-	TSharedPtr<FMetasoundParameterPackStorage> ParameterStorage;
+	template<typename T, int MaxNumItems = 0>
+	typename MetasoundParameterPackPrivate::TRootValueType<T, MaxNumItems>::Type* GetParameterPtr(const FName& Name, const FName& InTypeName) const
+	{
+		if (!ParameterStorage.IsValid())
+		{
+			return nullptr;
+		}
+		auto TheParameter = ParameterStorage->FindParameter<T, MaxNumItems>(Name, InTypeName);
+		if (!TheParameter)
+		{
+			return nullptr;
+		}
+		return &TheParameter->Value;
+	}
+
+	FSharedMetasoundParameterStoragePtr ParameterStorage;
 };
 
 // Here is the proxy that UMetasoundParameterPack creates when asked for a proxy by the audio system
diff --git a/Engine/Plugins/Runtime/Metasound/Source/MetasoundFrontend/Public/MetasoundParameterPackFixedArray.h b/Engine/Plugins/Runtime/Metasound/Source/MetasoundFrontend/Public/MetasoundParameterPackFixedArray.h
new file mode 100644
--- /dev/null
+++ b/Engine/Plugins/Runtime/Metasound/Source/MetasoundFrontend/Public/MetasoundParameterPackFixedArray.h
@@ -0,0 +1,54 @@
+// Copyright Epic Games, Inc. All Rights Reserved.
+
+#pragma once
+
+#include "Templates/TypeCompatibleBytes.h"
+
+namespace Metasound
+{
+
+	template<typename ElementType, int32 NumInlineElements>
+	class TParamPackFixedArray
+	{
+	public:
+		TParamPackFixedArray(const TArray<ElementType>& Other)
+		{
+			*this = Other;
+		}
+
+		TParamPackFixedArray<ElementType, NumInlineElements>& operator=(const TArray<ElementType>& Other)
+		{
+			int32 NumToCopy = FMath::Min(Other.Num(), NumInlineElements);
+			for (int32 i = 0; i < NumToCopy; ++i)
+			{
+				GetData()[i] = Other[i];
+			}
+			NumValidElements = NumToCopy;
+			return *this;
+		}
+
+		void CopyToArray(TArray<ElementType>& Dest) const
+		{
+			Dest.Empty();
+			for (int32 i = 0; i < NumValidElements; ++i)
+			{
+				Dest.Add(GetData()[i]);
+			}
+		}
+
+		FORCEINLINE const ElementType* GetData() const
+		{
+			return (ElementType*)InlineData;
+		}
+
+		FORCEINLINE ElementType* GetData()
+		{
+			return (ElementType*)InlineData;
+		}
+
+	public:
+		int32 NumValidElements = 0;
+		using FByteBucketType = TTypeCompatibleBytes<ElementType>;
+		FByteBucketType InlineData[NumInlineElements];
+	};
+}
