Spot fixup possible NaN sources and reading out of buffer boundaries

diff --git a/Engine/Plugins/Runtime/Metasound/Source/MetasoundStandardNodes/Private/MetasoundADSREnvelopeNode.cpp b/Engine/Plugins/Runtime/Metasound/Source/MetasoundStandardNodes/Private/MetasoundADSREnvelopeNode.cpp
--- a/Engine/Plugins/Runtime/Metasound/Source/MetasoundStandardNodes/Private/MetasoundADSREnvelopeNode.cpp
+++ b/Engine/Plugins/Runtime/Metasound/Source/MetasoundStandardNodes/Private/MetasoundADSREnvelopeNode.cpp
@@ -220,7 +220,7 @@ namespace Metasound
 				// We are in attack
 				if (AttackSamplesLeft > 0)
 				{
-					EndAttackFrame = FMath::Min(StartFrame + AttackSamplesLeft, EndFrame);
+					EndAttackFrame = FMath::Max(1, FMath::Min(StartFrame + AttackSamplesLeft, EndFrame));
 					for (int32 i = StartFrame; i < EndAttackFrame; ++i)
 					{
 						float AttackFraction = (float)++(InState.CurrentSampleIndex) / InState.AttackSampleCount;
@@ -562,10 +562,10 @@ namespace Metasound
 			, DecayCurveFactor(InArgs.DecayCurveFactor)
 			, ReleaseCurveFactor(InArgs.ReleaseCurveFactor)
 			, bHardReset(InArgs.bInHardReset)
-			, OnDecayTrigger(TDataWriteReferenceFactory<FTrigger>::CreateAny(InArgs.OperatorSettings))
-			, OnSustainTrigger(TDataWriteReferenceFactory<FTrigger>::CreateAny(InArgs.OperatorSettings))
-			, OnDone(TDataWriteReferenceFactory<FTrigger>::CreateAny(InArgs.OperatorSettings))
-			, OutputEnvelope(TDataWriteReferenceFactory<ValueType>::CreateAny(InArgs.OperatorSettings))
+			, OnDecayTrigger(TDataWriteReferenceFactory<FTrigger>::CreateExplicitArgs(InArgs.OperatorSettings))
+			, OnSustainTrigger(TDataWriteReferenceFactory<FTrigger>::CreateExplicitArgs(InArgs.OperatorSettings))
+			, OnDone(TDataWriteReferenceFactory<FTrigger>::CreateExplicitArgs(InArgs.OperatorSettings))
+			, OutputEnvelope(TDataWriteReferenceFactory<ValueType>::CreateExplicitArgs(InArgs.OperatorSettings))
 		{
 			NumFramesPerBlock = InArgs.OperatorSettings.GetNumFramesPerBlock();
 
diff --git a/Engine/Plugins/Runtime/Metasound/Source/MetasoundStandardNodes/Private/MetasoundMathNodes.cpp b/Engine/Plugins/Runtime/Metasound/Source/MetasoundStandardNodes/Private/MetasoundMathNodes.cpp
--- a/Engine/Plugins/Runtime/Metasound/Source/MetasoundStandardNodes/Private/MetasoundMathNodes.cpp
+++ b/Engine/Plugins/Runtime/Metasound/Source/MetasoundStandardNodes/Private/MetasoundMathNodes.cpp
@@ -1290,6 +1290,9 @@ namespace Metasound
 
 		static void Calculate(TMathOpMultiply& InInstanceData, const FAudioBufferReadRef& InPrimaryOperand, const TArray<FFloatReadRef>& InAdditionalOperands, FAudioBufferWriteRef& OutResult)
 		{
+			constexpr float MaxGain = TNumericLimits<float>::Max() / 1024.f;
+			constexpr float MinGain = -MaxGain;
+
 			if (!InInstanceData.bInit)
 			{
 				InInstanceData.bInit = true;
@@ -1297,7 +1300,7 @@ namespace Metasound
 
 				for (int32 i = 0; i < InAdditionalOperands.Num(); ++i)
 				{
-					InInstanceData.LastGain *= *InAdditionalOperands[i];
+					InInstanceData.LastGain *= FMath::Clamp(*InAdditionalOperands[i], MinGain, MaxGain);
 				}
 			}
 
@@ -1306,19 +1309,11 @@ namespace Metasound
 			float NewGain = 1.0f;
 			for (int32 i = 0; i < InAdditionalOperands.Num(); ++i)
 			{
-				const int32 SIMDRemainder = OutResult->Num() % AUDIO_NUM_FLOATS_PER_VECTOR_REGISTER;
-				const int32 SIMDCount = OutResult->Num() - SIMDRemainder;
-
-				TArrayView<float> OutResultView(OutResult->GetData(), SIMDCount);
-
-				Audio::ArrayFade(OutResultView, InInstanceData.LastGain, *InAdditionalOperands[i]);
+				const float Gain = FMath::Clamp(*InAdditionalOperands[i], MinGain, MaxGain);
 
-				for (int32 j = SIMDCount; j < OutResult->Num(); ++j)
-				{
-					OutResult->GetData()[j] *= *InAdditionalOperands[j];
-				}
+				Audio::ArrayFade(*OutResult, InInstanceData.LastGain, Gain);
 
-				NewGain *= *InAdditionalOperands[i];
+				NewGain *= Gain;
 			}
 
 			InInstanceData.LastGain = NewGain;
diff --git a/Engine/Plugins/Runtime/Metasound/Source/MetasoundStandardNodes/Private/MetasoundOscillators.h b/Engine/Plugins/Runtime/Metasound/Source/MetasoundStandardNodes/Private/MetasoundOscillators.h
--- a/Engine/Plugins/Runtime/Metasound/Source/MetasoundStandardNodes/Private/MetasoundOscillators.h
+++ b/Engine/Plugins/Runtime/Metasound/Source/MetasoundStandardNodes/Private/MetasoundOscillators.h
@@ -501,10 +501,6 @@ namespace Metasound
 		// Triangle.
 		struct FTriangleGenerator
 		{
-			float TriangleSign = -1.0f;
-			float PrevPhase = -1.0f;
-			float DPW_z1 = 0.0f;
-
 			FORCEINLINE float operator()(float InPhase, float InPhaseDelta, const FGeneratorArgs& InArgs)
 			{
 				constexpr float OneOverFastAsinHalfPi = 1.f / 1.5707963050f;
