Fix MetaSounds reroute nodes being internally disconnected when chained and created from node inputs

diff --git a/Engine/Plugins/Runtime/Metasound/Source/MetasoundFrontend/Private/NodeTemplates/MetasoundFrontendNodeTemplateReroute.cpp b/Engine/Plugins/Runtime/Metasound/Source/MetasoundFrontend/Private/NodeTemplates/MetasoundFrontendNodeTemplateReroute.cpp
--- a/Engine/Plugins/Runtime/Metasound/Source/MetasoundFrontend/Private/NodeTemplates/MetasoundFrontendNodeTemplateReroute.cpp
+++ b/Engine/Plugins/Runtime/Metasound/Source/MetasoundFrontend/Private/NodeTemplates/MetasoundFrontendNodeTemplateReroute.cpp
@@ -43,7 +43,8 @@ namespace Metasound
 			bool FRerouteNodeTemplatePreprocessTransform::Transform(FMetasoundFrontendNode& InOutNode) const
 			{
 				using namespace ReroutePrivate;
-
+				
+				// Find the input and output edges for this node
 				const FMetasoundFrontendEdge* InputEdge = nullptr;
 				{
 					if (!ensure(InOutNode.Interface.Inputs.Num() == 1))
@@ -64,14 +65,14 @@ namespace Metasound
 				}
 
 				TArray<FMetasoundFrontendEdge*>* OutputEdges = nullptr;
+				const FMetasoundFrontendVertex& OutputVertex = InOutNode.Interface.Outputs.Last();
+				FNodeVertexGuidPair OutputNodeVertexPair{ InOutNode.GetID(), OutputVertex.VertexID };
 				{
 					if (!ensure(InOutNode.Interface.Outputs.Num() == 1))
 					{
 						return false;
 					}
 
-					const FMetasoundFrontendVertex& OutputVertex = InOutNode.Interface.Outputs.Last();
-					FNodeVertexGuidPair OutputNodeVertexPair{ InOutNode.GetID(), OutputVertex.VertexID };
 					OutputEdges = OutputEdgeMap.Find(OutputNodeVertexPair);
 
 					// This can happen if the reroute node isn't provided any outputs to connect to, so its
@@ -82,12 +83,18 @@ namespace Metasound
 					}
 				}
 
+				// Update the output edges with the input edge 
+				FNodeVertexGuidPair NewOutputEdgeNodeVertexPair{ InputEdge->FromNodeID, InputEdge->FromVertexID };
+
 				for (FMetasoundFrontendEdge* OutputEdge : *OutputEdges)
 				{
-					OutputEdge->FromNodeID = InputEdge->FromNodeID;
-					OutputEdge->FromVertexID = InputEdge->FromVertexID;
+					OutputEdge->FromNodeID = NewOutputEdgeNodeVertexPair.Key;
+					OutputEdge->FromVertexID = NewOutputEdgeNodeVertexPair.Value;
 				}
 
+				OutputEdgeMap.FindOrAdd(NewOutputEdgeNodeVertexPair).Append(*OutputEdges);
+				OutputEdgeMap.Remove(OutputNodeVertexPair);
+
 				return true;
 			}
 		}
