Add metasound reset routine to external nodes

diff --git a/Engine/Plugins/Runtime/AudioModulation/Source/AudioModulation/Private/MetaSoundNodes/GetModulatorValueNode.cpp b/Engine/Plugins/Runtime/AudioModulation/Source/AudioModulation/Private/MetaSoundNodes/GetModulatorValueNode.cpp
--- a/Engine/Plugins/Runtime/AudioModulation/Source/AudioModulation/Private/MetaSoundNodes/GetModulatorValueNode.cpp
+++ b/Engine/Plugins/Runtime/AudioModulation/Source/AudioModulation/Private/MetaSoundNodes/GetModulatorValueNode.cpp
@@ -145,6 +145,27 @@ namespace AudioModulation
 			*OutValue = Value;
 		}
 
+		void Reset(const IOperator::FResetParams& InParams)
+		{
+			if (InParams.Environment.Contains<Audio::FDeviceId>(Metasound::Frontend::SourceInterface::Environment::DeviceID))
+			{
+				Audio::FDeviceId PriorDeviceId = DeviceId;
+				DeviceId = InParams.Environment.GetValue<Audio::FDeviceId>(Metasound::Frontend::SourceInterface::Environment::DeviceID);
+
+				if (PriorDeviceId != DeviceId)
+				{
+					// Reset mod handle if device ID is altered. Modulator handle IDs 
+					// are not unique across devices.
+					ModHandle = Audio::FModulatorHandle();
+				}
+			}
+			else
+			{
+				UE_LOG(LogMetaSound, Warning, TEXT("Missing audio device ID environment variable (%s) required to properly configure node (Modulation.GetModulatorValue)"), *Metasound::Frontend::SourceInterface::Environment::DeviceID.ToString());
+			}
+			Execute();
+		}
+
 	private:
 		Audio::FDeviceId DeviceId;
 
diff --git a/Engine/Plugins/Runtime/AudioModulation/Source/AudioModulation/Private/MetaSoundNodes/MixModulatorsNode.cpp b/Engine/Plugins/Runtime/AudioModulation/Source/AudioModulation/Private/MetaSoundNodes/MixModulatorsNode.cpp
--- a/Engine/Plugins/Runtime/AudioModulation/Source/AudioModulation/Private/MetaSoundNodes/MixModulatorsNode.cpp
+++ b/Engine/Plugins/Runtime/AudioModulation/Source/AudioModulation/Private/MetaSoundNodes/MixModulatorsNode.cpp
@@ -176,8 +176,30 @@ namespace AudioModulation
 			*OutValue = Value1;
 		}
 
+		void Reset(const IOperator::FResetParams& InParams)
+		{
+			if (InParams.Environment.Contains<Audio::FDeviceId>(Metasound::Frontend::SourceInterface::Environment::DeviceID))
+			{
+				Audio::FDeviceId PriorDeviceId = DeviceId;
+				DeviceId = InParams.Environment.GetValue<Audio::FDeviceId>(Metasound::Frontend::SourceInterface::Environment::DeviceID);
+
+				if (PriorDeviceId != DeviceId)
+				{
+					// Reset modulator handles because modulator ids are not unique across devices.
+					ModHandle1 = Audio::FModulatorHandle{};
+					ModHandle2 = Audio::FModulatorHandle{};
+				}
+			}
+			else
+			{
+				UE_LOG(LogMetaSound, Warning, TEXT("Missing audio device ID environment variable (%s) required to properly configure node (Modulation.MixModulators)"), *Metasound::Frontend::SourceInterface::Environment::DeviceID.ToString());
+			}
+			Execute();
+		}
+
 	private:
-		const Audio::FDeviceId DeviceId;
+
+		Audio::FDeviceId DeviceId;
 
 		Audio::FModulatorHandle ModHandle1;
 		Audio::FModulatorHandle ModHandle2;
diff --git a/Engine/Plugins/Runtime/Metasound/Source/MetasoundEngine/Private/MetasoundWaveInfoNode.cpp b/Engine/Plugins/Runtime/Metasound/Source/MetasoundEngine/Private/MetasoundWaveInfoNode.cpp
--- a/Engine/Plugins/Runtime/Metasound/Source/MetasoundEngine/Private/MetasoundWaveInfoNode.cpp
+++ b/Engine/Plugins/Runtime/Metasound/Source/MetasoundEngine/Private/MetasoundWaveInfoNode.cpp
@@ -50,6 +50,7 @@ namespace Metasound
 		virtual FDataReferenceCollection GetInputs() const override;
 		virtual FDataReferenceCollection GetOutputs() const override;
 		void Execute();
+		void Reset(const IOperator::FResetParams& InParams);
 
 	private: // members
 		// input pins
@@ -158,6 +159,11 @@ namespace Metasound
 		}
 	}
 
+	void FWaveInfoNodeOperator::Reset(const IOperator::FResetParams& InParams)
+	{
+		Execute();
+	}
+
 
 
 	class FWaveInfoNode : public FNodeFacade
diff --git a/Engine/Plugins/Runtime/Metasound/Source/MetasoundEngine/Private/MetasoundWavePlayerNode.cpp b/Engine/Plugins/Runtime/Metasound/Source/MetasoundEngine/Private/MetasoundWavePlayerNode.cpp
--- a/Engine/Plugins/Runtime/Metasound/Source/MetasoundEngine/Private/MetasoundWavePlayerNode.cpp
+++ b/Engine/Plugins/Runtime/Metasound/Source/MetasoundEngine/Private/MetasoundWavePlayerNode.cpp
@@ -434,6 +434,38 @@ namespace Metasound
 			UpdatePlaybackLocation();
 		}
 
+		void Reset(const IOperator::FResetParams& InParams)
+		{
+			TriggerOnDone->Reset();
+			TriggerOnNearlyDone->Reset();
+			TriggerOnLooped->Reset();
+			TriggerOnCuePoint->Reset();
+
+			*CuePointID = 0;
+			*CuePointLabel = TEXT("");
+			*LoopPercent = 0;
+			*PlaybackLocation = 0;
+			for (const FAudioBufferWriteRef& BufferRef : OutputAudioBuffers)
+			{
+				BufferRef->Zero();
+			}
+
+			WaveProxyReader.Reset();
+			ConvertDeinterleave.Reset();
+			Resampler.Reset();
+
+			SortedCuePoints.Reset();
+			for (Audio::TCircularAudioBuffer<float>& Buffer : SourceCircularBuffer)
+			{
+				Buffer.SetNum(0);
+			}
+
+			SourceState = WavePlayerNodePrivate::FSourceBufferState();
+			SampleRateFrameRatio = 1.f;
+			bOnNearlyDoneTriggeredForWave = false;
+			bIsPlaying = false;
+		}
+
 	private:
 
 		void ExecuteSubblocks()
diff --git a/Engine/Plugins/Runtime/Metasound/Source/MetasoundEngine/Private/MetasoundWaveTableEnvelopeNode.cpp b/Engine/Plugins/Runtime/Metasound/Source/MetasoundEngine/Private/MetasoundWaveTableEnvelopeNode.cpp
--- a/Engine/Plugins/Runtime/Metasound/Source/MetasoundEngine/Private/MetasoundWaveTableEnvelopeNode.cpp
+++ b/Engine/Plugins/Runtime/Metasound/Source/MetasoundEngine/Private/MetasoundWaveTableEnvelopeNode.cpp
@@ -134,23 +134,7 @@ namespace Metasound
 			, OutWriteRef(TDataWriteReferenceFactory<float>::CreateAny(InParams.OperatorSettings))
 			, SampleRate(InParams.OperatorSettings.GetSampleRate())
 		{
-			{
-				using namespace WaveTable;
-				FWaveTableSampler::FSettings Settings;
-				Settings.Freq = 0.0f; // Sampler phase is manually progressed via this node
-				Sampler = FWaveTableSampler(MoveTemp(Settings));
-			}
-
-			{
-				const float BlockRate = InParams.OperatorSettings.GetActualBlockRate();
-				check(BlockRate > 0.0f && !FMath::IsNearlyZero(BlockRate));
-				SecondsPerBlock = 1.0f / BlockRate;
-			}
-
-			{
-				BlockSize = InParams.OperatorSettings.GetNumFramesPerBlock();
-				check(BlockSize > 0);
-			}
+			Reset(InParams);
 		}
 
 		virtual ~FMetasoundWaveTableEnvelopeNodeOperator() = default;
@@ -311,6 +295,33 @@ namespace Metasound
 			*OutWriteRef = NextValue;
 		}
 
+		void Reset(const IOperator::FResetParams& InParams)
+		{
+			{
+				using namespace WaveTable;
+				FWaveTableSampler::FSettings Settings;
+				Settings.Freq = 0.0f; // Sampler phase is manually progressed via this node
+				Sampler = FWaveTableSampler(MoveTemp(Settings));
+			}
+
+			{
+				const float BlockRate = InParams.OperatorSettings.GetActualBlockRate();
+				check(BlockRate > 0.0f && !FMath::IsNearlyZero(BlockRate));
+				SecondsPerBlock = 1.0f / BlockRate;
+			}
+
+			{
+				BlockSize = InParams.OperatorSettings.GetNumFramesPerBlock();
+				check(BlockSize > 0);
+			}
+
+			OnFinishedWriteRef->Reset();
+			*OutWriteRef = 0.f;
+
+			Elapsed = -1.0f;
+			bPaused = false;
+		}
+
 	private:
 		FWaveTableReadRef WaveTableReadRef;
 		FTriggerReadRef PlayReadRef;
diff --git a/Engine/Plugins/Runtime/Metasound/Source/MetasoundEngine/Private/MetasoundWaveTableEvaluateNode.cpp b/Engine/Plugins/Runtime/Metasound/Source/MetasoundEngine/Private/MetasoundWaveTableEvaluateNode.cpp
--- a/Engine/Plugins/Runtime/Metasound/Source/MetasoundEngine/Private/MetasoundWaveTableEvaluateNode.cpp
+++ b/Engine/Plugins/Runtime/Metasound/Source/MetasoundEngine/Private/MetasoundWaveTableEvaluateNode.cpp
@@ -101,10 +101,7 @@ namespace Metasound
 			, InterpModeReadRef(InInterpModeReadRef)
 			, OutWriteRef(TDataWriteReferenceFactory<float>::CreateAny(InParams.OperatorSettings))
 		{
-			using namespace WaveTable;
-			FWaveTableSampler::FSettings Settings;
-			Settings.Freq = 0.0f; // Sampler phase is manually progressed via this node
-			Sampler = FWaveTableSampler(MoveTemp(Settings));
+			Reset(InParams);
 		}
 
 		virtual ~FMetasoundWaveTableEvaluateNodeOperator() = default;
@@ -161,6 +158,16 @@ namespace Metasound
 			*OutWriteRef = NextValue;
 		}
 
+		void Reset(const IOperator::FResetParams& InParams)
+		{
+			using namespace WaveTable;
+			FWaveTableSampler::FSettings Settings;
+			Settings.Freq = 0.0f; // Sampler phase is manually progressed via this node
+			Sampler = FWaveTableSampler(MoveTemp(Settings));
+
+			*OutWriteRef = 0.f;
+		}
+
 	private:
 		FWaveTableReadRef WaveTableReadRef;
 		FFloatReadRef InputReadRef;
diff --git a/Engine/Plugins/Runtime/Metasound/Source/MetasoundEngine/Private/MetasoundWaveTableGetNode.cpp b/Engine/Plugins/Runtime/Metasound/Source/MetasoundEngine/Private/MetasoundWaveTableGetNode.cpp
--- a/Engine/Plugins/Runtime/Metasound/Source/MetasoundEngine/Private/MetasoundWaveTableGetNode.cpp
+++ b/Engine/Plugins/Runtime/Metasound/Source/MetasoundEngine/Private/MetasoundWaveTableGetNode.cpp
@@ -84,6 +84,7 @@ namespace Metasound
 			, TableIndexReadRef(InTableIndexReadRef)
 			, OutTable(TDataWriteReferenceFactory<WaveTable::FWaveTable>::CreateAny(InParams.OperatorSettings))
 		{
+			Reset(InParams);
 		}
 
 		virtual ~FMetasoundWaveTableGetNodeOperator() = default;
@@ -204,6 +205,13 @@ namespace Metasound
 			LastTableIndex = NextTableIndex;
 		}
 
+		void Reset(const IOperator::FResetParams& InParams)
+		{
+			LastTableIndex = -1.0f;
+			OutTable->SetNum(0);
+			OutTable->SetFinalValue(0.f);
+		}
+
 	private:
 		FWaveTableBankAssetReadRef WaveTableBankReadRef;
 		FFloatReadRef TableIndexReadRef;
diff --git a/Engine/Plugins/Runtime/Metasound/Source/MetasoundEngine/Private/MetasoundWaveTableOscillatorNode.cpp b/Engine/Plugins/Runtime/Metasound/Source/MetasoundEngine/Private/MetasoundWaveTableOscillatorNode.cpp
--- a/Engine/Plugins/Runtime/Metasound/Source/MetasoundEngine/Private/MetasoundWaveTableOscillatorNode.cpp
+++ b/Engine/Plugins/Runtime/Metasound/Source/MetasoundEngine/Private/MetasoundWaveTableOscillatorNode.cpp
@@ -210,6 +210,23 @@ namespace Metasound
 			}
 		}
 
+		void Reset(const IOperator::FResetParams& InParams)
+		{
+			const float BlockRate = InParams.OperatorSettings.GetActualBlockRate();
+			if (BlockRate > 0.0f)
+			{
+				BlockPeriod = 1.0f / BlockRate;
+			}
+
+			bPlaying = false;
+			if (SyncBuffer.Num() > 0)
+			{
+				FMemory::Memset(SyncBuffer.GetData(), 0, sizeof(float) * SyncBuffer.Num());
+			}
+			Sampler = WaveTable::FWaveTableSampler{};
+			OutBufferWriteRef->Zero();
+		}
+
 	private:
 		float BlockPeriod = 0.0f;
 		bool bPlaying = false;
diff --git a/Engine/Plugins/Runtime/Metasound/Source/MetasoundStandardNodes/Private/MetasoundADEnvelopeNode.cpp b/Engine/Plugins/Runtime/Metasound/Source/MetasoundStandardNodes/Private/MetasoundADEnvelopeNode.cpp
--- a/Engine/Plugins/Runtime/Metasound/Source/MetasoundStandardNodes/Private/MetasoundADEnvelopeNode.cpp
+++ b/Engine/Plugins/Runtime/Metasound/Source/MetasoundStandardNodes/Private/MetasoundADEnvelopeNode.cpp
@@ -60,6 +60,23 @@ namespace Metasound
 
 			bool bLooping = false;
 			bool bHardReset = false;
+
+			void Reset()
+			{
+				CurrentSampleIndex = INDEX_NONE;
+				AttackSampleCount = 1;
+				DecaySampleCount = 1;
+
+				AttackCurveFactor = 0.0f;
+				DecayCurveFactor = 0.0f;
+
+				StartingEnvelopeValue = 0.0f;
+				CurrentEnvelopeValue = 0.0f;
+
+				bLooping = false;
+				bHardReset = false;
+				EnvEase.Init(0.f, 0.01f);
+			}
 		};
 
 		template<typename ValueType>
@@ -121,7 +138,12 @@ namespace Metasound
 				}
 			}
 
-			static bool IsAudio() { return false; }
+			static void GetInitialOutputEnvelope(float& OutputEnvelope)
+			{
+				OutputEnvelope = 0.f;
+			}
+
+			static constexpr bool IsAudio() { return false; }
 		};
 
 		template<>
@@ -194,7 +216,12 @@ namespace Metasound
 				}
 			}
 
-			static bool IsAudio() { return true; }
+			static void GetInitialOutputEnvelope(FAudioBuffer& OutputEnvelope)
+			{
+				OutputEnvelope.Zero();
+			}
+
+			static constexpr bool IsAudio() { return true; }
 		};
 	}
 
@@ -272,10 +299,10 @@ namespace Metasound
 			FBoolReadRef bLooping = InParams.InputDataReferences.GetDataReadReferenceOrConstructWithVertexDefault<bool>(InputInterface, METASOUND_GET_PARAM_NAME(InputLooping), InParams.OperatorSettings);
 			FBoolReadRef bHardReset = InParams.InputDataReferences.GetDataReadReferenceOrConstructWithVertexDefault<bool>(InputInterface, METASOUND_GET_PARAM_NAME(InputHardReset), InParams.OperatorSettings);
 
-			return MakeUnique<TADEnvelopeNodeOperator<ValueType>>(InParams.OperatorSettings, TriggerIn, AttackTime, DecayTime, AttackCurveFactor, DecayCurveFactor, bLooping, bHardReset);
+			return MakeUnique<TADEnvelopeNodeOperator<ValueType>>(InParams, TriggerIn, AttackTime, DecayTime, AttackCurveFactor, DecayCurveFactor, bLooping, bHardReset);
 		}
 
-		TADEnvelopeNodeOperator(const FOperatorSettings& InSettings,
+		TADEnvelopeNodeOperator(const FCreateOperatorParams& InParams,
 			const FTriggerReadRef& InTriggerIn,
 			const FTimeReadRef& InAttackTime,
 			const FTimeReadRef& InDecayTime,
@@ -290,22 +317,11 @@ namespace Metasound
 			, DecayCurveFactor(InDecayeCurveFactor)
 			, bLooping(bInLooping)
 			, bHardReset(bInHardReset)
-			, OnAttackTrigger(TDataWriteReferenceFactory<FTrigger>::CreateAny(InSettings))
-			, OnDone(TDataWriteReferenceFactory<FTrigger>::CreateAny(InSettings))
-			, OutputEnvelope(TDataWriteReferenceFactory<ValueType>::CreateAny(InSettings))
+			, OnAttackTrigger(TDataWriteReferenceFactory<FTrigger>::CreateExplicitArgs(InParams.OperatorSettings))
+			, OnDone(TDataWriteReferenceFactory<FTrigger>::CreateExplicitArgs(InParams.OperatorSettings))
+			, OutputEnvelope(TDataWriteReferenceFactory<ValueType>::CreateExplicitArgs(InParams.OperatorSettings))
 		{
-			NumFramesPerBlock = InSettings.GetNumFramesPerBlock();
-
-			EnvState.EnvEase.SetEaseFactor(0.01f);
-
-			if (ADEnvelopeNodePrivate::TADEnvelope<ValueType>::IsAudio())
-			{
-				SampleRate = InSettings.GetSampleRate();
-			}
-			else
-			{
-				SampleRate = InSettings.GetActualBlockRate();
-			}
+			Reset(InParams);
 		}
 
 		virtual ~TADEnvelopeNodeOperator() = default;
@@ -350,6 +366,25 @@ namespace Metasound
 			EnvState.bHardReset = *bHardReset;
 		}
 
+		void Reset(const IOperator::FResetParams& InParams)
+		{
+			NumFramesPerBlock = InParams.OperatorSettings.GetNumFramesPerBlock();
+
+
+			if constexpr (ADEnvelopeNodePrivate::TADEnvelope<ValueType>::IsAudio())
+			{
+				SampleRate = InParams.OperatorSettings.GetSampleRate();
+			}
+			else
+			{
+				SampleRate = InParams.OperatorSettings.GetActualBlockRate();
+			}
+
+			OnAttackTrigger->Reset();
+			OnDone->Reset();
+			ADEnvelopeNodePrivate::TADEnvelope<ValueType>::GetInitialOutputEnvelope(*OutputEnvelope);
+			EnvState.Reset();
+		}
 
 		void Execute()
 		{
diff --git a/Engine/Plugins/Runtime/Metasound/Source/MetasoundStandardNodes/Private/MetasoundADSREnvelopeNode.cpp b/Engine/Plugins/Runtime/Metasound/Source/MetasoundStandardNodes/Private/MetasoundADSREnvelopeNode.cpp
--- a/Engine/Plugins/Runtime/Metasound/Source/MetasoundStandardNodes/Private/MetasoundADSREnvelopeNode.cpp
+++ b/Engine/Plugins/Runtime/Metasound/Source/MetasoundStandardNodes/Private/MetasoundADSREnvelopeNode.cpp
@@ -83,11 +83,40 @@ namespace Metasound
 			{
 				return CurrentSampleIndex != INDEX_NONE; 
 			}
+
+			void Reset()
+			{
+				CurrentSampleIndex = INDEX_NONE;
+
+				AttackSampleCount = 0;
+				DecaySampleCount = 0;
+				ReleaseSampleCount = 0;
+				
+				SustainLevel = 0.0f;
+
+				AttackCurveFactor = 0.0f;
+				DecayCurveFactor = 0.0f;
+				ReleaseCurveFactor = 0.0f;
+
+				StartingEnvelopeValue = 0.0f;
+				CurrentEnvelopeValue = 0.0f;
+				EnvelopeValueAtReleaseStart = 0.0f;
+
+				bHardReset = false;
+				bIsInRelease = false;
+			
+				EnvEase.Init(0.f, 0.1f);
+			}
 		};
 
 
 		struct FFloatADSREnvelope
 		{
+			static void GetInitialOutputEnvelope(float& OutputEnvelope)
+			{
+				OutputEnvelope = 0.f;
+			}
+
 			static float GetSampleRate(const FOperatorSettings& InOperatorSettings)
 			{		
 				return InOperatorSettings.GetActualBlockRate();
@@ -183,6 +212,11 @@ namespace Metasound
 
 		struct FAudioADSREnvelope
 		{
+			static void GetInitialOutputEnvelope(FAudioBuffer& OutputEnvelope)
+			{
+				OutputEnvelope.Zero();
+			}
+
 			static float GetSampleRate(const FOperatorSettings& InOperatorSettings)
 			{
 					return InOperatorSettings.GetSampleRate();
@@ -514,7 +548,7 @@ namespace Metasound
 
 		struct FOperatorArgs
 		{
-			FOperatorSettings OperatorSettings;
+			const FBuildOperatorParams& BuildOperatorParams;
 			FTriggerReadRef TriggerAttackIn;
 			FTriggerReadRef TriggerReleaseIn;
 			FTimeReadRef AttackTime;
@@ -526,8 +560,7 @@ namespace Metasound
 			FFloatReadRef ReleaseCurveFactor;
 			FBoolReadRef bInHardReset;
 		};
-
-		static TUniquePtr<IOperator> CreateOperator(const FCreateOperatorParams& InParams, TArray<TUniquePtr<IOperatorBuildError>>& OutErrors)
+		static TUniquePtr<IOperator> CreateOperator(const FBuildOperatorParams& InParams, FBuildResults& OutResults)
 		{
 			using namespace ADSREnvelopeVertexNames;
 
@@ -535,17 +568,17 @@ namespace Metasound
 
 			FOperatorArgs Args 
 			{
-				InParams.OperatorSettings,
-				InParams.InputDataReferences.GetDataReadReferenceOrConstruct<FTrigger>(METASOUND_GET_PARAM_NAME(InputAttackTrigger), InParams.OperatorSettings),
-				InParams.InputDataReferences.GetDataReadReferenceOrConstruct<FTrigger>(METASOUND_GET_PARAM_NAME(InputReleaseTrigger), InParams.OperatorSettings),
-				InParams.InputDataReferences.GetDataReadReferenceOrConstructWithVertexDefault<FTime>(InputInterface, METASOUND_GET_PARAM_NAME(InputAttackTime), InParams.OperatorSettings),
-				InParams.InputDataReferences.GetDataReadReferenceOrConstructWithVertexDefault<FTime>(InputInterface, METASOUND_GET_PARAM_NAME(InputDecayTime), InParams.OperatorSettings),
-				InParams.InputDataReferences.GetDataReadReferenceOrConstructWithVertexDefault<float>(InputInterface, METASOUND_GET_PARAM_NAME(InputSustainLevel), InParams.OperatorSettings),
-				InParams.InputDataReferences.GetDataReadReferenceOrConstructWithVertexDefault<FTime>(InputInterface, METASOUND_GET_PARAM_NAME(InputReleaseTime), InParams.OperatorSettings),
-				InParams.InputDataReferences.GetDataReadReferenceOrConstructWithVertexDefault<float>(InputInterface, METASOUND_GET_PARAM_NAME(InputAttackCurve), InParams.OperatorSettings),
-				InParams.InputDataReferences.GetDataReadReferenceOrConstructWithVertexDefault<float>(InputInterface, METASOUND_GET_PARAM_NAME(InputDecayCurve), InParams.OperatorSettings),
-				InParams.InputDataReferences.GetDataReadReferenceOrConstructWithVertexDefault<float>(InputInterface, METASOUND_GET_PARAM_NAME(InputReleaseCurve), InParams.OperatorSettings),
-				InParams.InputDataReferences.GetDataReadReferenceOrConstructWithVertexDefault<bool>(InputInterface, METASOUND_GET_PARAM_NAME(InputHardReset), InParams.OperatorSettings)
+				InParams,
+				InParams.InputData.GetOrConstructDataReadReference<FTrigger>(METASOUND_GET_PARAM_NAME(InputAttackTrigger), InParams.OperatorSettings),
+				InParams.InputData.GetOrConstructDataReadReference<FTrigger>(METASOUND_GET_PARAM_NAME(InputReleaseTrigger), InParams.OperatorSettings),
+				InParams.InputData.GetOrCreateDefaultDataReadReference<FTime>(METASOUND_GET_PARAM_NAME(InputAttackTime), InParams.OperatorSettings),
+				InParams.InputData.GetOrCreateDefaultDataReadReference<FTime>(METASOUND_GET_PARAM_NAME(InputDecayTime), InParams.OperatorSettings),
+				InParams.InputData.GetOrCreateDefaultDataReadReference<float>(METASOUND_GET_PARAM_NAME(InputSustainLevel), InParams.OperatorSettings),
+				InParams.InputData.GetOrCreateDefaultDataReadReference<FTime>(METASOUND_GET_PARAM_NAME(InputReleaseTime), InParams.OperatorSettings),
+				InParams.InputData.GetOrCreateDefaultDataReadReference<float>(METASOUND_GET_PARAM_NAME(InputAttackCurve), InParams.OperatorSettings),
+				InParams.InputData.GetOrCreateDefaultDataReadReference<float>(METASOUND_GET_PARAM_NAME(InputDecayCurve), InParams.OperatorSettings),
+				InParams.InputData.GetOrCreateDefaultDataReadReference<float>(METASOUND_GET_PARAM_NAME(InputReleaseCurve), InParams.OperatorSettings),
+				InParams.InputData.GetOrCreateDefaultDataReadReference<bool>(METASOUND_GET_PARAM_NAME(InputHardReset), InParams.OperatorSettings)
 			};
 
 			return MakeUnique<TADSREnvelopeNodeOperator<EnvelopeClass, ValueType>>(Args);
@@ -562,15 +595,12 @@ namespace Metasound
 			, DecayCurveFactor(InArgs.DecayCurveFactor)
 			, ReleaseCurveFactor(InArgs.ReleaseCurveFactor)
 			, bHardReset(InArgs.bInHardReset)
-			, OnDecayTrigger(TDataWriteReferenceFactory<FTrigger>::CreateExplicitArgs(InArgs.OperatorSettings))
-			, OnSustainTrigger(TDataWriteReferenceFactory<FTrigger>::CreateExplicitArgs(InArgs.OperatorSettings))
-			, OnDone(TDataWriteReferenceFactory<FTrigger>::CreateExplicitArgs(InArgs.OperatorSettings))
-			, OutputEnvelope(TDataWriteReferenceFactory<ValueType>::CreateExplicitArgs(InArgs.OperatorSettings))
+			, OnDecayTrigger(TDataWriteReferenceFactory<FTrigger>::CreateExplicitArgs(InArgs.BuildOperatorParams.OperatorSettings))
+			, OnSustainTrigger(TDataWriteReferenceFactory<FTrigger>::CreateExplicitArgs(InArgs.BuildOperatorParams.OperatorSettings))
+			, OnDone(TDataWriteReferenceFactory<FTrigger>::CreateExplicitArgs(InArgs.BuildOperatorParams.OperatorSettings))
+			, OutputEnvelope(TDataWriteReferenceFactory<ValueType>::CreateExplicitArgs(InArgs.BuildOperatorParams.OperatorSettings))
 		{
-			NumFramesPerBlock = InArgs.OperatorSettings.GetNumFramesPerBlock();
-
-			EnvState.EnvEase.SetEaseFactor(0.1f);
-			SampleRate = EnvelopeClass::GetSampleRate(InArgs.OperatorSettings);
+			Reset(InArgs.BuildOperatorParams);
 		}
 
 		virtual ~TADSREnvelopeNodeOperator() = default;
@@ -647,6 +677,19 @@ namespace Metasound
 			}
 		}
 
+		void Reset(const IOperator::FResetParams& InParams)
+		{
+			NumFramesPerBlock = InParams.OperatorSettings.GetNumFramesPerBlock();
+
+			SampleRate = EnvelopeClass::GetSampleRate(InParams.OperatorSettings);
+			EnvelopeClass::GetInitialOutputEnvelope(*OutputEnvelope);
+			EnvState.Reset();
+
+			OnDecayTrigger->Reset();
+			OnSustainTrigger->Reset();
+			OnDone->Reset();
+		}
+
 		void Execute()
 		{
 			using namespace ADSREnvelopeNodePrivate;
@@ -655,8 +698,6 @@ namespace Metasound
 			OnSustainTrigger->AdvanceBlock();
 			OnDone->AdvanceBlock();
  
-			
-
 			FTriggerIterator TriggerIter(*TriggerAttackIn, *TriggerReleaseIn, NumFramesPerBlock);
 
 			FTriggerInfo NextTrigger = TriggerIter.NextTrigger(EnvState.IsTriggered());
diff --git a/Engine/Plugins/Runtime/Metasound/Source/MetasoundStandardNodes/Private/MetasoundAbsoluteValueNodes.cpp b/Engine/Plugins/Runtime/Metasound/Source/MetasoundStandardNodes/Private/MetasoundAbsoluteValueNodes.cpp
--- a/Engine/Plugins/Runtime/Metasound/Source/MetasoundStandardNodes/Private/MetasoundAbsoluteValueNodes.cpp
+++ b/Engine/Plugins/Runtime/Metasound/Source/MetasoundStandardNodes/Private/MetasoundAbsoluteValueNodes.cpp
@@ -212,6 +212,11 @@ namespace Metasound
 			GetAbs();
 		}
 
+		void Reset(const IOperator::FResetParams& InParams)
+		{
+			GetAbs();
+		}
+
 	private:
 
 		TDataReadReference<ValueType> Input;
diff --git a/Engine/Plugins/Runtime/Metasound/Source/MetasoundStandardNodes/Private/MetasoundAdditiveSynthNode.cpp b/Engine/Plugins/Runtime/Metasound/Source/MetasoundStandardNodes/Private/MetasoundAdditiveSynthNode.cpp
--- a/Engine/Plugins/Runtime/Metasound/Source/MetasoundStandardNodes/Private/MetasoundAdditiveSynthNode.cpp
+++ b/Engine/Plugins/Runtime/Metasound/Source/MetasoundStandardNodes/Private/MetasoundAdditiveSynthNode.cpp
@@ -139,6 +139,13 @@ namespace Metasound
 			return OutputDataReferences;
 		}
 
+		void Reset(const IOperator::FResetParams& InParams)
+		{
+			CreateSinusoids();
+			LeftAudioOut->Zero();
+			RightAudioOut->Zero();
+		}
+
 		void Execute()
 		{
 			int32 NumSamples = LeftAudioOut->Num();
diff --git a/Engine/Plugins/Runtime/Metasound/Source/MetasoundStandardNodes/Private/MetasoundBPMToSecondsNode.cpp b/Engine/Plugins/Runtime/Metasound/Source/MetasoundStandardNodes/Private/MetasoundBPMToSecondsNode.cpp
--- a/Engine/Plugins/Runtime/Metasound/Source/MetasoundStandardNodes/Private/MetasoundBPMToSecondsNode.cpp
+++ b/Engine/Plugins/Runtime/Metasound/Source/MetasoundStandardNodes/Private/MetasoundBPMToSecondsNode.cpp
@@ -30,10 +30,11 @@ namespace Metasound
 		static const FVertexInterface& GetVertexInterface();
 		static TUniquePtr<IOperator> CreateOperator(const FCreateOperatorParams& InParams, FBuildErrorArray& OutErrors);
 
-		FBPMToSecondsOperator(const FOperatorSettings& InSettings, const FFloatReadRef& InBPM, const FFloatReadRef& InBeatMultiplier, const FFloatReadRef& InDivOfWholeNote);
+		FBPMToSecondsOperator(const FCreateOperatorParams& InParams, const FFloatReadRef& InBPM, const FFloatReadRef& InBeatMultiplier, const FFloatReadRef& InDivOfWholeNote);
 
 		virtual FDataReferenceCollection GetInputs() const override;
 		virtual FDataReferenceCollection GetOutputs() const override;
+		void Reset(const IOperator::FResetParams& InParams);
 		void Execute();
 
 	private:
@@ -50,15 +51,22 @@ namespace Metasound
 		float PrevBPM = -1.0f;
 		float PrevBeatMultiplier = -1.0f;
 		float PrevDivOfWholeNote = -1.0f;
+
+		static constexpr float MaxBPM = UE_MAX_FLT / 4;
+		static constexpr float MinBPM = 1.f; 
+		static constexpr float MaxBeatMultiplier = UE_MAX_FLT / 4;
+		static constexpr float MinBeatMultiplier = KINDA_SMALL_NUMBER;
+		static constexpr float MaxDivOfWholeNote = UE_MAX_FLT / 4;
+		static constexpr float MinDivOfWholeNote = 1.f;
 	};
 
-	FBPMToSecondsOperator::FBPMToSecondsOperator(const FOperatorSettings& InSettings, const FFloatReadRef& InBPM, const FFloatReadRef& InBeatMultiplier, const FFloatReadRef& InDivOfWholeNote)
+	FBPMToSecondsOperator::FBPMToSecondsOperator(const FCreateOperatorParams& InParams, const FFloatReadRef& InBPM, const FFloatReadRef& InBeatMultiplier, const FFloatReadRef& InDivOfWholeNote)
 		: BPM(InBPM)
 		, BeatMultiplier(InBeatMultiplier)
 		, DivOfWholeNote(InDivOfWholeNote)
-		, TimeSeconds(TDataWriteReferenceFactory<FTime>::CreateAny(InSettings))
+		, TimeSeconds(TDataWriteReferenceFactory<FTime>::CreateExplicitArgs(InParams.OperatorSettings))
 	{
-		UpdateTime();
+		Reset(InParams);
 	}
 
 	FDataReferenceCollection FBPMToSecondsOperator::GetInputs() const
@@ -84,9 +92,9 @@ namespace Metasound
 
 	void FBPMToSecondsOperator::UpdateTime()
 	{
-		float CurrBPM = FMath::Max(*BPM, 1.0f);
-		float CurrBeatMultiplier = FMath::Max(*BeatMultiplier, KINDA_SMALL_NUMBER);
-		float CurrDivOfWholeNote = FMath::Max(*DivOfWholeNote, 1.0f);
+		float CurrBPM = FMath::Clamp(*BPM, MinBPM, MaxBPM);
+		float CurrBeatMultiplier = FMath::Clamp(*BeatMultiplier, MinBeatMultiplier, MaxBeatMultiplier);
+		float CurrDivOfWholeNote = FMath::Clamp(*DivOfWholeNote, MinDivOfWholeNote, MaxDivOfWholeNote);
 
 		if (!FMath::IsNearlyEqual(PrevBPM, CurrBPM) || 
 			!FMath::IsNearlyEqual(PrevBeatMultiplier, CurrBeatMultiplier) || 
@@ -105,6 +113,11 @@ namespace Metasound
 
 	}
 
+	void FBPMToSecondsOperator::Reset(const IOperator::FResetParams& InParams)
+	{
+		UpdateTime();
+	}
+
 	void FBPMToSecondsOperator::Execute()
 	{
 		UpdateTime();
@@ -162,7 +175,7 @@ namespace Metasound
 		FFloatReadRef InBeatMultiplier = InputCollection.GetDataReadReferenceOrConstructWithVertexDefault<float>(InputInterface, METASOUND_GET_PARAM_NAME(InputBeatMultiplier), InParams.OperatorSettings);
 		FFloatReadRef InDivOfWholeNote = InputCollection.GetDataReadReferenceOrConstructWithVertexDefault<float>(InputInterface, METASOUND_GET_PARAM_NAME(InputDivOfWholeNote), InParams.OperatorSettings);
 
-		return MakeUnique<FBPMToSecondsOperator>(InParams.OperatorSettings, InBPM, InBeatMultiplier, InDivOfWholeNote);
+		return MakeUnique<FBPMToSecondsOperator>(InParams, InBPM, InBeatMultiplier, InDivOfWholeNote);
 	}
 
 	FBPMToSecondsNode::FBPMToSecondsNode(const FNodeInitData& InitData)
diff --git a/Engine/Plugins/Runtime/Metasound/Source/MetasoundStandardNodes/Private/MetasoundBandsplitterNode.cpp b/Engine/Plugins/Runtime/Metasound/Source/MetasoundStandardNodes/Private/MetasoundBandsplitterNode.cpp
--- a/Engine/Plugins/Runtime/Metasound/Source/MetasoundStandardNodes/Private/MetasoundBandsplitterNode.cpp
+++ b/Engine/Plugins/Runtime/Metasound/Source/MetasoundStandardNodes/Private/MetasoundBandsplitterNode.cpp
@@ -57,13 +57,10 @@ namespace Metasound
 			, bPhaseCompensate(InPhaseCompensate)
 			, Settings(InSettings)
 		{
+			// Update internal crossover buffer from passed data references
 			Crossovers.Reset();
 			Crossovers.AddUninitialized(NumBands - 1);
-
-			for (int32 i = 0; i < InCrossoverFrequencies.Num(); ++i)
-			{
-				Crossovers[i] = *InCrossoverFrequencies[i];
-			}
+			CopyClampCrossovers();
 
 			bPrevPhaseCompensate = *bPhaseCompensate;
 			PrevFilterOrder = GetFilterOrder(*InFilterOrder);
@@ -261,16 +258,7 @@ namespace Metasound
 		void UpdateFilterSettings()
 		{
 			bool bNeedsReinitFilters = false;
-			bool bNeedsUpdateCrossovers = false;
-
-			for (int32 CrossoverIndex = 0; CrossoverIndex < Crossovers.Num(); CrossoverIndex++)
-			{
-				if (Crossovers[CrossoverIndex] != *CrossoverFrequencies[CrossoverIndex])
-				{
-					bNeedsUpdateCrossovers = true;
-					Crossovers[CrossoverIndex] = *CrossoverFrequencies[CrossoverIndex];
-				}
-			}
+			bool bNeedsUpdateCrossovers = CopyClampCrossovers();
 
 			if (*bPhaseCompensate != bPrevPhaseCompensate)
 			{
@@ -300,6 +288,25 @@ namespace Metasound
 			}
 		}
 
+		void Reset(const IOperator::FResetParams& InParams)
+		{
+			CopyClampCrossovers();
+
+			bPrevPhaseCompensate = *bPhaseCompensate;
+			PrevFilterOrder = GetFilterOrder(*FilterOrder);
+
+			for (uint32 ChannelIndex = 0; ChannelIndex < NumChannels; ++ChannelIndex)
+			{
+				MultiBandBuffers[ChannelIndex].Reset();
+				Filters[ChannelIndex].Init(1, InParams.OperatorSettings.GetSampleRate(), PrevFilterOrder, bPrevPhaseCompensate, Crossovers);
+			}
+
+			for (const TDataWriteReference<FAudioBuffer>& AudioOutput : AudioOutputs)
+			{
+				AudioOutput->Zero();
+			}
+		}
+
 		void Execute()
 		{
 			UpdateFilterSettings();
@@ -319,7 +326,33 @@ namespace Metasound
 			}
 		}
 
+
 	private:
+		bool CopyClampCrossovers()
+		{
+			check(Crossovers.Num() == CrossoverFrequencies.Num());
+
+			bool bDidUpdate = false;
+			
+			const float MaxFrequency = 0.95f * Settings.GetSampleRate() / 2.f;
+			const float MinFrequency = FMath::Min(20.f, MaxFrequency);
+
+			const TDataReadReference<float>* InputFreqs = CrossoverFrequencies.GetData();
+			float* CrossoverData = Crossovers.GetData();
+
+			for (int32 i = 0; i < CrossoverFrequencies.Num(); ++i)
+			{
+				const float NewCrossover = FMath::Clamp(*InputFreqs[i], MinFrequency, MaxFrequency);
+				if (NewCrossover != CrossoverData[i])
+				{
+					bDidUpdate = true;
+					CrossoverData[i] = NewCrossover;
+				}
+			}
+
+			return bDidUpdate;
+		}
+
 		TArray<FAudioBufferReadRef> AudioInputs;
 		TArray<FFloatReadRef> CrossoverFrequencies;
 
diff --git a/Engine/Plugins/Runtime/Metasound/Source/MetasoundStandardNodes/Private/MetasoundBasicFilters.cpp b/Engine/Plugins/Runtime/Metasound/Source/MetasoundStandardNodes/Private/MetasoundBasicFilters.cpp
--- a/Engine/Plugins/Runtime/Metasound/Source/MetasoundStandardNodes/Private/MetasoundBasicFilters.cpp
+++ b/Engine/Plugins/Runtime/Metasound/Source/MetasoundStandardNodes/Private/MetasoundBasicFilters.cpp
@@ -116,6 +116,7 @@ namespace Metasound
 			return OutputDataReferences;
 		}
 
+		void Reset(const IOperator::FResetParams& InParams);
 		void Execute();
 
 
@@ -224,6 +225,14 @@ namespace Metasound
 			);
 	}
 
+	void FLadderFilterOperator::Reset(const IOperator::FResetParams& InParams)
+	{
+		PreviousFrequency = InvalidValue;
+		PreviousResonance = InvalidValue;
+		AudioOutput->Zero();
+		LadderFilter.Init(SampleRate, 1 /* NumChannels */);
+	}
+
 	void FLadderFilterOperator::Execute()
 	{
 		const float CurrentFrequency = FMath::Clamp(*Frequency, 0.f, MaxCutoffFrequency);
@@ -286,6 +295,7 @@ namespace Metasound
 			return OutputDataReferences;
 		}
 
+		void Reset(const IOperator::FResetParams& InParams);
 		void Execute();
 
 
@@ -416,6 +426,21 @@ namespace Metasound
 			);
 	}
 
+	void FStateVariableFilterOperator::Reset(const IOperator::FResetParams& InParams)
+	{
+		PreviousFrequency = InvalidValue;
+		PreviousResonance = InvalidValue;
+		PreviousBandStopControl = InvalidValue;
+
+		// output pins
+		LowPassOutput->Zero();
+		HighPassOutput->Zero();
+		BandPassOutput->Zero();
+		BandStopOutput->Zero();
+
+		StateVariableFilter.Init(SampleRate, 1 /* NumChannels */);
+	}
+
 	void FStateVariableFilterOperator::Execute()
 	{
 		const float CurrentFrequency = FMath::Clamp(*Frequency, 0.f, MaxCutoffFrequency);
@@ -484,6 +509,7 @@ namespace Metasound
 			return OutputDataReferences;
 		}
 
+		void Reset(const IOperator::FResetParams& InParams);
 		void Execute();
 
 
@@ -580,6 +606,12 @@ namespace Metasound
 			);
 	}
 
+	void FOnePoleLowPassFilterOperator::Reset(const IOperator::FResetParams& InParams)
+	{
+		AudioOutput->Zero();
+		OnePoleLowPassFilter.Init(SampleRate, 1 /* NumChannels */);
+	}
+
 	void FOnePoleLowPassFilterOperator::Execute()
 	{
 		float ClampedFreq = FMath::Clamp(0.0f, *Frequency, SampleRate);
@@ -622,6 +654,7 @@ namespace Metasound
 			return OutputDataReferences;
 		}
 
+		void Reset(const IOperator::FResetParams& InParams);
 		void Execute();
 
 
@@ -718,6 +751,12 @@ namespace Metasound
 			);
 	}
 
+	void FOnePoleHighPassFilterOperator::Reset(const IOperator::FResetParams& InParams)
+	{
+		AudioOutput->Zero();
+		OnePoleHighPassFilter.Init(SampleRate, 1 /* NumChannels */);
+	}
+
 	void FOnePoleHighPassFilterOperator::Execute()
 	{
 		float ClampedFreq = FMath::Clamp(0.0f, *Frequency, SampleRate);
@@ -763,10 +802,13 @@ namespace Metasound
 			return OutputDataReferences;
 		}
 
+		void Reset(const IOperator::FResetParams& InParams);
 		void Execute();
 
 
-	private: // members
+	private: 
+		void ResetInternal();
+
 		// input pins
 		FAudioBufferReadRef AudioInput;
 		FFloatReadRef Frequency;
@@ -813,9 +855,7 @@ namespace Metasound
 			, MaxCutoffFrequency(0.5f * SampleRate)
 			, PreviousFilterType(*FilterType)
 		{
-			// verify our buffer sizes:
-			check(AudioOutput->Num() == BlockSize);
-			BiquadFilter.Init(SampleRate, 1, *FilterType);
+			ResetInternal();
 		}
 	};
 
@@ -899,6 +939,22 @@ namespace Metasound
 			);
 	}
 
+	void FBiquadFilterOperator::Reset(const IOperator::FResetParams& InParams)
+	{
+		ResetInternal();
+	}
+
+	void FBiquadFilterOperator::ResetInternal()
+	{
+		PreviousFrequency = InvalidValue;
+		PreviousBandwidth = InvalidValue;
+		PreviousFilterGainDb = InvalidValue;
+		PreviousFilterType = *FilterType;
+
+		AudioOutput->Zero();
+		BiquadFilter.Init(SampleRate, 1, *FilterType);
+	}
+
 	void FBiquadFilterOperator::Execute()
 	{
 		const float CurrentFrequency = FMath::Clamp(*Frequency, 0.f, MaxCutoffFrequency);
diff --git a/Engine/Plugins/Runtime/Metasound/Source/MetasoundStandardNodes/Private/MetasoundBitcrusherNode.cpp b/Engine/Plugins/Runtime/Metasound/Source/MetasoundStandardNodes/Private/MetasoundBitcrusherNode.cpp
--- a/Engine/Plugins/Runtime/Metasound/Source/MetasoundStandardNodes/Private/MetasoundBitcrusherNode.cpp
+++ b/Engine/Plugins/Runtime/Metasound/Source/MetasoundStandardNodes/Private/MetasoundBitcrusherNode.cpp
@@ -33,20 +33,17 @@ namespace Metasound
 	{
 	public:
 
-		FBitcrusherOperator(const FOperatorSettings& InSettings,
+		FBitcrusherOperator(const FCreateOperatorParams& InParams,
 			const FAudioBufferReadRef& InAudio,
 			const FFloatReadRef& InCrushedSampleRate,
 			const FFloatReadRef& InCrushedBitDepth)
 			: AudioInput(InAudio)
-			, AudioOutput(FAudioBufferWriteRef::CreateNew(InSettings))
+			, AudioOutput(FAudioBufferWriteRef::CreateNew(InParams.OperatorSettings))
 			, CrushedSampleRate(InCrushedSampleRate)
 			, CrushedBitDepth(InCrushedBitDepth)
-			, MaxSampleRate(InSettings.GetSampleRate())
+			, MaxSampleRate(InParams.OperatorSettings.GetSampleRate())
 		{
-			// Passing in 1 for NumChannels because the node takes one audio input
-			Bitcrusher.Init(InSettings.GetSampleRate(), 1);
-			Bitcrusher.SetSampleRateCrush(FMath::Clamp(*CrushedSampleRate, 1.0f, MaxSampleRate));
-			Bitcrusher.SetBitDepthCrush(FMath::Clamp(*CrushedBitDepth, 1.0f, Bitcrusher.GetMaxBitDepth()));
+			Reset(InParams);
 		}
 
 		static const FNodeClassMetadata& GetNodeInfo()
@@ -130,7 +127,24 @@ namespace Metasound
 			FFloatReadRef SampleRateIn = Inputs.GetDataReadReferenceOrConstructWithVertexDefault<float>(InputInterface, METASOUND_GET_PARAM_NAME(InputSampleRate), InParams.OperatorSettings);
 			FFloatReadRef BitDepthIn = Inputs.GetDataReadReferenceOrConstructWithVertexDefault<float>(InputInterface, METASOUND_GET_PARAM_NAME(InputBitDepth), InParams.OperatorSettings);
 
-			return MakeUnique<FBitcrusherOperator>(InParams.OperatorSettings, AudioIn, SampleRateIn, BitDepthIn);
+			return MakeUnique<FBitcrusherOperator>(InParams, AudioIn, SampleRateIn, BitDepthIn);
+		}
+
+		void Reset(const IOperator::FResetParams& InParams)
+		{
+			AudioOutput->Zero();
+
+			// Passing in 1 for NumChannels because the node takes one audio input
+			Bitcrusher.Init(InParams.OperatorSettings.GetSampleRate(), 1);
+
+			const float CurSampleRate = FMath::Clamp(*CrushedSampleRate, 1.0f, MaxSampleRate);
+			const float CurBitDepth = FMath::Clamp(*CrushedBitDepth, 1.0f, Bitcrusher.GetMaxBitDepth());
+
+			Bitcrusher.SetSampleRateCrush(CurSampleRate);
+			Bitcrusher.SetBitDepthCrush(CurBitDepth);
+
+			PrevSampleRate = CurSampleRate;
+			PrevBitDepth = CurBitDepth;
 		}
 
 		void Execute()
diff --git a/Engine/Plugins/Runtime/Metasound/Source/MetasoundStandardNodes/Private/MetasoundClampNode.cpp b/Engine/Plugins/Runtime/Metasound/Source/MetasoundStandardNodes/Private/MetasoundClampNode.cpp
--- a/Engine/Plugins/Runtime/Metasound/Source/MetasoundStandardNodes/Private/MetasoundClampNode.cpp
+++ b/Engine/Plugins/Runtime/Metasound/Source/MetasoundStandardNodes/Private/MetasoundClampNode.cpp
@@ -242,6 +242,11 @@ namespace Metasound
 			TClamp<ValueType>::GetClamped(*InputValue, *InputMinValue, *InputMaxValue, *OutputValue);
 		}
 
+		void Reset(const IOperator::FResetParams& InParams)
+		{
+			GetClamped();
+		}
+
 		void Execute()
 		{
 			GetClamped();
diff --git a/Engine/Plugins/Runtime/Metasound/Source/MetasoundStandardNodes/Private/MetasoundCompressorNode.cpp b/Engine/Plugins/Runtime/Metasound/Source/MetasoundStandardNodes/Private/MetasoundCompressorNode.cpp
--- a/Engine/Plugins/Runtime/Metasound/Source/MetasoundStandardNodes/Private/MetasoundCompressorNode.cpp
+++ b/Engine/Plugins/Runtime/Metasound/Source/MetasoundStandardNodes/Private/MetasoundCompressorNode.cpp
@@ -230,6 +230,53 @@ namespace Metasound
 			return MakeUnique<FCompressorOperator>(InParams.OperatorSettings, AudioIn, RatioIn, ThresholdDbIn, AttackTimeIn, ReleaseTimeIn, LookaheadTimeIn, KneeIn, bIsSidechainConnected, SidechainIn, EnvelopeModeIn, bIsAnalogIn, bIsUpwardsIn, WetDryMixIn);
 		}
 
+		void Reset(const IOperator::FResetParams& InParams)
+		{
+			// Flush audio buffers
+			AudioOutput->Zero();
+			EnvelopeOutput->Zero();
+			InputDelay.Reset();
+			DelayedInputSignal.Zero();
+
+			// Cache dynamics timing
+			PrevAttackTime = FMath::Max(FTime::ToMilliseconds(*AttackTimeInput), 0.0);
+			PrevReleaseTime = FMath::Max(FTime::ToMilliseconds(*ReleaseTimeInput), 0.0);
+			PrevLookaheadTime = FMath::Max(FTime::ToMilliseconds(*LookaheadTimeInput), 0.0);
+
+			// Initialize compressor
+			Compressor.Init(InParams.OperatorSettings.GetSampleRate(), 1);
+			Compressor.SetKeyNumChannels(1);
+			Compressor.SetRatio(FMath::Max(*RatioInput, 1.0f));
+			Compressor.SetThreshold(*ThresholdDbInput);
+			Compressor.SetAttackTime(PrevAttackTime);
+			Compressor.SetReleaseTime(PrevReleaseTime);
+			Compressor.SetLookaheadMsec(PrevLookaheadTime);
+			Compressor.SetKneeBandwidth(*KneeInput);
+
+			if (*bIsUpwardsInput)
+			{
+				Compressor.SetProcessingMode(Audio::EDynamicsProcessingMode::UpwardsCompressor);
+			}
+			else
+			{
+				Compressor.SetProcessingMode(Audio::EDynamicsProcessingMode::Compressor);
+			}
+
+			switch (*EnvelopeModeInput)
+			{
+			default:
+			case EEnvelopePeakMode::MeanSquared:
+				Compressor.SetPeakMode(Audio::EPeakMode::MeanSquared);
+				break;
+			case EEnvelopePeakMode::RootMeanSquared:
+				Compressor.SetPeakMode(Audio::EPeakMode::RootMeanSquared);
+				break;
+			case EEnvelopePeakMode::Peak:
+				Compressor.SetPeakMode(Audio::EPeakMode::Peak);
+				break;
+			}
+		}
+
 		void Execute()
 		{
 			/* Update parameters */
diff --git a/Engine/Plugins/Runtime/Metasound/Source/MetasoundStandardNodes/Private/MetasoundConversionNodes.cpp b/Engine/Plugins/Runtime/Metasound/Source/MetasoundStandardNodes/Private/MetasoundConversionNodes.cpp
--- a/Engine/Plugins/Runtime/Metasound/Source/MetasoundStandardNodes/Private/MetasoundConversionNodes.cpp
+++ b/Engine/Plugins/Runtime/Metasound/Source/MetasoundStandardNodes/Private/MetasoundConversionNodes.cpp
@@ -259,6 +259,12 @@ namespace Metasound
 			return Outputs;
 		}
 
+		void Reset(const IOperator::FResetParams& InParams)
+		{
+			using namespace MetasoundConversionNodePrivate;
+			TConversionNodeSpecialization<FromType, ToType>::GetConvertedValue(*InputValue, *OutputValue);
+		}
+
 		void Execute()
 		{
 			using namespace MetasoundConversionNodePrivate;
diff --git a/Engine/Plugins/Runtime/Metasound/Source/MetasoundStandardNodes/Private/MetasoundCrossfadeNode.cpp b/Engine/Plugins/Runtime/Metasound/Source/MetasoundStandardNodes/Private/MetasoundCrossfadeNode.cpp
--- a/Engine/Plugins/Runtime/Metasound/Source/MetasoundStandardNodes/Private/MetasoundCrossfadeNode.cpp
+++ b/Engine/Plugins/Runtime/Metasound/Source/MetasoundStandardNodes/Private/MetasoundCrossfadeNode.cpp
@@ -272,6 +272,11 @@ namespace Metasound
 			Crossfader.GetCrossfadeOutput(IndexA, IndexB, Alpha, InputValues, *OutputValue);
 		}
 
+		void Reset(const IOperator::FResetParams& InParams)
+		{
+			PerformCrossfadeOutput();
+		}
+
 		void Execute()
 		{
 			PerformCrossfadeOutput();
diff --git a/Engine/Plugins/Runtime/Metasound/Source/MetasoundStandardNodes/Private/MetasoundDebugLogNode.cpp b/Engine/Plugins/Runtime/Metasound/Source/MetasoundStandardNodes/Private/MetasoundDebugLogNode.cpp
--- a/Engine/Plugins/Runtime/Metasound/Source/MetasoundStandardNodes/Private/MetasoundDebugLogNode.cpp
+++ b/Engine/Plugins/Runtime/Metasound/Source/MetasoundStandardNodes/Private/MetasoundDebugLogNode.cpp
@@ -99,28 +99,16 @@ namespace Metasound
 				TDataReadReference<FString> Label = InputCollection.GetDataReadReferenceOrConstructWithVertexDefault<FString>(InputInterface, METASOUND_GET_PARAM_NAME(InputLabel), InParams.OperatorSettings);
 				TDataReadReference<PrintLogType> ValueToLog = InputCollection.GetDataReadReferenceOrConstructWithVertexDefault<PrintLogType>(InputInterface, METASOUND_GET_PARAM_NAME(InputValueToLog), InParams.OperatorSettings);
 
-				FString GraphName;
-				if (InParams.Environment.Contains<FString>(Frontend::SourceInterface::Environment::GraphName))
-				{
-					FString GraphNameFull = InParams.Environment.GetValue<FString>(Frontend::SourceInterface::Environment::GraphName);
-					TArray<FString> ParsedString;
-					GraphNameFull.ParseIntoArray(ParsedString, TEXT("."), true);
-					if (ParsedString.Num() > 0)
-					{
-						GraphName = ParsedString.Last();
-					}
-				}
-
-				return MakeUnique<TPrintLogOperator<PrintLogType>>(InParams.OperatorSettings, Trigger, Label, ValueToLog, GraphName);
+				return MakeUnique<TPrintLogOperator<PrintLogType>>(InParams, Trigger, Label, ValueToLog);
 			}
 
 
-			TPrintLogOperator(const FOperatorSettings& InSettings, TDataReadReference<FTrigger> InTrigger, TDataReadReference<FString> InLabelPrintLog, TDataReadReference<PrintLogType> InValueToLogPrintLog, const FString& InGraphName)
+			TPrintLogOperator(const FCreateOperatorParams& InParams, TDataReadReference<FTrigger> InTrigger, TDataReadReference<FString> InLabelPrintLog, TDataReadReference<PrintLogType> InValueToLogPrintLog)
 				: Trigger(InTrigger)
 				, Label(InLabelPrintLog)
 				, ValueToLog(InValueToLogPrintLog)
-				, GraphName(InGraphName)
 			{
+				Reset(InParams);
 			}
 
 			virtual ~TPrintLogOperator() = default;
@@ -145,6 +133,21 @@ namespace Metasound
 				return Outputs;
 			}
 
+			void Reset(const IOperator::FResetParams& InParams)
+			{
+				GraphName = TEXT("");
+				if (InParams.Environment.Contains<FString>(Frontend::SourceInterface::Environment::GraphName))
+				{
+					FString GraphNameFull = InParams.Environment.GetValue<FString>(Frontend::SourceInterface::Environment::GraphName);
+					TArray<FString> ParsedString;
+					GraphNameFull.ParseIntoArray(ParsedString, TEXT("."), true);
+					if (ParsedString.Num() > 0)
+					{
+						GraphName = ParsedString.Last();
+					}
+				}
+			}
+
 			void Execute()
 			{
 				if (*Trigger)
diff --git a/Engine/Plugins/Runtime/Metasound/Source/MetasoundStandardNodes/Private/MetasoundDecibelsToLinearGain.cpp b/Engine/Plugins/Runtime/Metasound/Source/MetasoundStandardNodes/Private/MetasoundDecibelsToLinearGain.cpp
--- a/Engine/Plugins/Runtime/Metasound/Source/MetasoundStandardNodes/Private/MetasoundDecibelsToLinearGain.cpp
+++ b/Engine/Plugins/Runtime/Metasound/Source/MetasoundStandardNodes/Private/MetasoundDecibelsToLinearGain.cpp
@@ -31,10 +31,11 @@ namespace Metasound
 		static const FVertexInterface& GetVertexInterface();
 		static TUniquePtr<IOperator> CreateOperator(const FCreateOperatorParams& InParams, FBuildErrorArray& OutErrors);
 
-		FDecibelsToLinearGainOperator(const FOperatorSettings& InSettings, const FFloatReadRef& InDecibelGain);
+		FDecibelsToLinearGainOperator(const FCreateOperatorParams& InParams, const FFloatReadRef& InDecibelGain);
 
 		virtual FDataReferenceCollection GetInputs() const override;
 		virtual FDataReferenceCollection GetOutputs() const override;
+		void Reset(const IOperator::FResetParams& InParams);
 		void Execute();
 
 	private:
@@ -45,12 +46,12 @@ namespace Metasound
 		FFloatWriteRef LinearGainOutput;
 	};
 
-	FDecibelsToLinearGainOperator::FDecibelsToLinearGainOperator(const FOperatorSettings& InSettings, const FFloatReadRef& InDecibelGain)
+	FDecibelsToLinearGainOperator::FDecibelsToLinearGainOperator(const FCreateOperatorParams& InParams, const FFloatReadRef& InDecibelGain)
 		: DecibelGainInput(InDecibelGain)
-		, LinearGainOutput(FFloatWriteRef::CreateNew(Audio::ConvertToLinear(*InDecibelGain)))
+		, LinearGainOutput(FFloatWriteRef::CreateNew())
 	{
+		Reset(InParams);
 	}
-
 	
 	FDataReferenceCollection FDecibelsToLinearGainOperator::GetInputs() const
 	{
@@ -71,6 +72,11 @@ namespace Metasound
 		return OutputDataReferences;
 	}
 
+	void FDecibelsToLinearGainOperator::Reset(const IOperator::FResetParams& InParams)
+	{
+		Execute();
+	}
+
 	void FDecibelsToLinearGainOperator::Execute()
 	{
 		*LinearGainOutput = Audio::ConvertToLinear(*DecibelGainInput);
@@ -129,7 +135,7 @@ namespace Metasound
 
 		FFloatReadRef InDecibelGain = InputCollection.GetDataReadReferenceOrConstructWithVertexDefault<float>(InputInterface, METASOUND_GET_PARAM_NAME(InputDecibelGain), InParams.OperatorSettings);
 
-		return MakeUnique<FDecibelsToLinearGainOperator>(InParams.OperatorSettings, InDecibelGain);
+		return MakeUnique<FDecibelsToLinearGainOperator>(InParams, InDecibelGain);
 	}
 
 	class METASOUNDSTANDARDNODES_API FDecibelsToLinearGainNode : public FNodeFacade
diff --git a/Engine/Plugins/Runtime/Metasound/Source/MetasoundStandardNodes/Private/MetasoundDelayNode.cpp b/Engine/Plugins/Runtime/Metasound/Source/MetasoundStandardNodes/Private/MetasoundDelayNode.cpp
--- a/Engine/Plugins/Runtime/Metasound/Source/MetasoundStandardNodes/Private/MetasoundDelayNode.cpp
+++ b/Engine/Plugins/Runtime/Metasound/Source/MetasoundStandardNodes/Private/MetasoundDelayNode.cpp
@@ -40,7 +40,7 @@ namespace Metasound
 		static const FVertexInterface& GetVertexInterface();
 		static TUniquePtr<IOperator> CreateOperator(const FBuildOperatorParams& InParams, FBuildResults& OutResults);
 
-		FDelayOperator(const FOperatorSettings& InSettings, 
+		FDelayOperator(const FBuildOperatorParams& InParams, 
 			const FAudioBufferReadRef& InAudioInput, 
 			const FTimeReadRef& InDelayTime, 
 			const FFloatReadRef& InDryLevel, 
@@ -51,6 +51,7 @@ namespace Metasound
 		virtual void Bind(FVertexInterfaceData& InVertexData) const override;
 		virtual FDataReferenceCollection GetInputs() const override;
 		virtual FDataReferenceCollection GetOutputs() const override;
+		void Reset(const IOperator::FResetParams& InParams);
 		void Execute();
 
 	private:
@@ -87,7 +88,7 @@ namespace Metasound
 		float MaxDelayTimeSeconds = Delay::DefaultMaxDelaySeconds;
 	};
 
-	FDelayOperator::FDelayOperator(const FOperatorSettings& InSettings,
+	FDelayOperator::FDelayOperator(const FBuildOperatorParams& InParams,
 		const FAudioBufferReadRef& InAudioInput,
 		const FTimeReadRef& InDelayTime,
 		const FFloatReadRef& InDryLevel,
@@ -100,13 +101,11 @@ namespace Metasound
 		, DryLevel(InDryLevel)
 		, WetLevel(InWetLevel)
 		, Feedback(InFeedback)
-		, AudioOutput(FAudioBufferWriteRef::CreateNew(InSettings))
+		, AudioOutput(FAudioBufferWriteRef::CreateNew(InParams.OperatorSettings))
 	{
 		MaxDelayTimeSeconds = FMath::Clamp(InMaxDelayTimeSeconds, Delay::MinMaxDelaySeconds, Delay::MaxMaxDelaySeconds);
-		PrevDelayTimeMsec = GetInputDelayTimeMsec();
 
-		DelayBuffer.Init(InSettings.GetSampleRate(), MaxDelayTimeSeconds);
-		DelayBuffer.SetDelayMsec(PrevDelayTimeMsec);
+		Reset(InParams);
 	}
 
 	void FDelayOperator::Bind(FVertexInterfaceData& InVertexData) const
@@ -154,6 +153,18 @@ namespace Metasound
 		return 1000.0f * FMath::Clamp((float)DelayTime->GetSeconds(), 0.0f, MaxDelayTimeSeconds);
 	}
 
+	void FDelayOperator::Reset(const IOperator::FResetParams& InParams)
+	{
+		FeedbackSample = 0.f;
+		PrevDelayTimeMsec = GetInputDelayTimeMsec();
+
+		DelayBuffer.Reset();
+		DelayBuffer.Init(InParams.OperatorSettings.GetSampleRate(), MaxDelayTimeSeconds);
+		DelayBuffer.SetDelayMsec(PrevDelayTimeMsec);
+
+		AudioOutput->Zero();
+	}
+
 	void FDelayOperator::Execute()
 	{
 		// Get clamped delay time
@@ -255,7 +266,7 @@ namespace Metasound
 		FFloatReadRef Feedback = InputData.GetOrCreateDefaultDataReadReference<float>(METASOUND_GET_PARAM_NAME(InParamFeedbackAmount), InParams.OperatorSettings);
 		FTime MaxDelayTime = InputData.GetOrCreateDefaultValue<FTime>(METASOUND_GET_PARAM_NAME(InParamMaxDelayTime), InParams.OperatorSettings);
 
-		return MakeUnique<FDelayOperator>(InParams.OperatorSettings, AudioIn, DelayTime, DryLevel, WetLevel, Feedback, MaxDelayTime.GetSeconds());
+		return MakeUnique<FDelayOperator>(InParams, AudioIn, DelayTime, DryLevel, WetLevel, Feedback, MaxDelayTime.GetSeconds());
 	}
 
 	class FDelayNode : public FNodeFacade
diff --git a/Engine/Plugins/Runtime/Metasound/Source/MetasoundStandardNodes/Private/MetasoundDiffuserNode.cpp b/Engine/Plugins/Runtime/Metasound/Source/MetasoundStandardNodes/Private/MetasoundDiffuserNode.cpp
--- a/Engine/Plugins/Runtime/Metasound/Source/MetasoundStandardNodes/Private/MetasoundDiffuserNode.cpp
+++ b/Engine/Plugins/Runtime/Metasound/Source/MetasoundStandardNodes/Private/MetasoundDiffuserNode.cpp
@@ -139,6 +139,17 @@ namespace Metasound
 			return OutputDataReferences;
 		}
 
+		void Reset(const IOperator::FResetParams& InParams)
+		{
+			using namespace DiffuserNode;
+
+			for (Audio::FLongDelayAPF* Delay : Delays)
+			{
+				Delay->Reset();
+			}
+			AudioOutput->Zero();
+		}
+
 		void Execute()
 		{
 			const int32 NumSamples = AudioInput->Num();
diff --git a/Engine/Plugins/Runtime/Metasound/Source/MetasoundStandardNodes/Private/MetasoundDopplerPitchShiftNode.cpp b/Engine/Plugins/Runtime/Metasound/Source/MetasoundStandardNodes/Private/MetasoundDopplerPitchShiftNode.cpp
--- a/Engine/Plugins/Runtime/Metasound/Source/MetasoundStandardNodes/Private/MetasoundDopplerPitchShiftNode.cpp
+++ b/Engine/Plugins/Runtime/Metasound/Source/MetasoundStandardNodes/Private/MetasoundDopplerPitchShiftNode.cpp
@@ -30,10 +30,11 @@ namespace Metasound
 		static const FVertexInterface& GetVertexInterface();
 		static TUniquePtr<IOperator> CreateOperator(const FCreateOperatorParams& InParams, FBuildErrorArray& OutErrors);
 
-		FDopplerPitchShiftOperator(const FOperatorSettings& InSettings, const FAudioBufferReadRef& InAudioInput, const FFloatReadRef& InPitchShift, const FFloatReadRef& InDelayLength);
+		FDopplerPitchShiftOperator(const FCreateOperatorParams& InParams, const FAudioBufferReadRef& InAudioInput, const FFloatReadRef& InPitchShift, const FFloatReadRef& InDelayLength);
 
 		virtual FDataReferenceCollection GetInputs() const override;
 		virtual FDataReferenceCollection GetOutputs() const override;
+		void Reset(const IOperator::FResetParams& InParams);
 		void Execute();
 
 	private:
@@ -45,7 +46,7 @@ namespace Metasound
 		Audio::FTapDelayPitchShifter DelayPitchShifter;
 	};
 
-	FDopplerPitchShiftOperator::FDopplerPitchShiftOperator(const FOperatorSettings& InSettings,
+	FDopplerPitchShiftOperator::FDopplerPitchShiftOperator(const FCreateOperatorParams& InParams,
 		const FAudioBufferReadRef& InAudioInput,
 		const FFloatReadRef& InPitchShift,
 		const FFloatReadRef& InDelayLength)
@@ -53,10 +54,9 @@ namespace Metasound
 		: AudioInput(InAudioInput)
 		, PitchShift(InPitchShift)
 		, DelayLength(InDelayLength)
-		, AudioOutput(FAudioBufferWriteRef::CreateNew(InSettings))
+		, AudioOutput(FAudioBufferWriteRef::CreateNew(InParams.OperatorSettings))
 	{
-		DelayBuffer.Init(InSettings.GetSampleRate(), 0.001f * Audio::FTapDelayPitchShifter::MaxDelayLength);
-		DelayPitchShifter.Init(InSettings.GetSampleRate(), *PitchShift, *DelayLength);
+		Reset(InParams);
 	}
 
 	FDataReferenceCollection FDopplerPitchShiftOperator::GetInputs() const
@@ -80,6 +80,13 @@ namespace Metasound
 		return OutputDataReferences;
 	}
 
+	void FDopplerPitchShiftOperator::Reset(const IOperator::FResetParams& InParams)
+	{
+		AudioOutput->Zero();
+		DelayBuffer.Init(InParams.OperatorSettings.GetSampleRate(), 0.001f * Audio::FTapDelayPitchShifter::MaxDelayLength);
+		DelayPitchShifter.Init(InParams.OperatorSettings.GetSampleRate(), *PitchShift, *DelayLength);
+	}
+
 	void FDopplerPitchShiftOperator::Execute()
 	{
 		DelayPitchShifter.SetDelayLength(*DelayLength);
@@ -138,7 +145,7 @@ namespace Metasound
 		FFloatReadRef PitchShift = InputCollection.GetDataReadReferenceOrConstructWithVertexDefault<float>(InputInterface, METASOUND_GET_PARAM_NAME(InParamPitchShift), InParams.OperatorSettings);
 		FFloatReadRef DelayLength = InputCollection.GetDataReadReferenceOrConstructWithVertexDefault<float>(InputInterface, METASOUND_GET_PARAM_NAME(InParamDelayLength), InParams.OperatorSettings);
 
-		return MakeUnique<FDopplerPitchShiftOperator>(InParams.OperatorSettings, AudioIn, PitchShift, DelayLength);
+		return MakeUnique<FDopplerPitchShiftOperator>(InParams, AudioIn, PitchShift, DelayLength);
 	}
 
 	class FDopperPitchShiftNode : public FNodeFacade
diff --git a/Engine/Plugins/Runtime/Metasound/Source/MetasoundStandardNodes/Private/MetasoundDynamicFilterNode.cpp b/Engine/Plugins/Runtime/Metasound/Source/MetasoundStandardNodes/Private/MetasoundDynamicFilterNode.cpp
--- a/Engine/Plugins/Runtime/Metasound/Source/MetasoundStandardNodes/Private/MetasoundDynamicFilterNode.cpp
+++ b/Engine/Plugins/Runtime/Metasound/Source/MetasoundStandardNodes/Private/MetasoundDynamicFilterNode.cpp
@@ -247,6 +247,13 @@ namespace Metasound
 			}
 		}
 
+		void Reset(const IOperator::FResetParams& InParams)
+		{
+			AudioOutput->Zero();
+			Filter.Reset();
+			UpdateFilterSettings();
+		}
+
 		void Execute()
 		{
 			const float* InputAudio = AudioInput->GetData();
diff --git a/Engine/Plugins/Runtime/Metasound/Source/MetasoundStandardNodes/Private/MetasoundEnvelopeFollowerNode.cpp b/Engine/Plugins/Runtime/Metasound/Source/MetasoundStandardNodes/Private/MetasoundEnvelopeFollowerNode.cpp
--- a/Engine/Plugins/Runtime/Metasound/Source/MetasoundStandardNodes/Private/MetasoundEnvelopeFollowerNode.cpp
+++ b/Engine/Plugins/Runtime/Metasound/Source/MetasoundStandardNodes/Private/MetasoundEnvelopeFollowerNode.cpp
@@ -42,7 +42,7 @@ namespace Metasound
 		static const FVertexInterface& GetVertexInterface();
 		static TUniquePtr<IOperator> CreateOperator(const FCreateOperatorParams& InParams, FBuildErrorArray& OutErrors);
 
-		FEnvelopeFollowerOperator(const FOperatorSettings& InSettings,
+		FEnvelopeFollowerOperator(const FCreateOperatorParams& InOperatorSettings,
 			const FAudioBufferReadRef& InAudioInput,
 			const FTimeReadRef& InAttackTime,
 			const FTimeReadRef& InReleaseTime,
@@ -50,6 +50,7 @@ namespace Metasound
 
 		virtual FDataReferenceCollection GetInputs() const override;
 		virtual FDataReferenceCollection GetOutputs() const override;
+		void Reset(const IOperator::FResetParams& InParams);
 		void Execute();
 
 	private:
@@ -77,7 +78,7 @@ namespace Metasound
 		EEnvelopePeakMode PrevFollowMode = EEnvelopePeakMode::Peak;
 	};
 
-	FEnvelopeFollowerOperator::FEnvelopeFollowerOperator(const FOperatorSettings& InSettings,
+	FEnvelopeFollowerOperator::FEnvelopeFollowerOperator(const FCreateOperatorParams& InParams,
 		const FAudioBufferReadRef& InAudioInput,
 		const FTimeReadRef& InAttackTime,
 		const FTimeReadRef& InReleaseTime,
@@ -87,19 +88,9 @@ namespace Metasound
 		, ReleaseTimeInput(InReleaseTime)
 		, FollowModeInput(InEnvelopeMode)
 		, EnvelopeFloatOutput(FFloatWriteRef::CreateNew())
-		, EnvelopeAudioOutput(FAudioBufferWriteRef::CreateNew(InSettings))
+		, EnvelopeAudioOutput(FAudioBufferWriteRef::CreateNew(InParams.OperatorSettings))
 	{
-		PrevAttackTime = FMath::Max(FTime::ToMilliseconds(*AttackTimeInput), 0.0);
-		PrevReleaseTime = FMath::Max(FTime::ToMilliseconds(*ReleaseTimeInput), 0.0);
-
-		Audio::FEnvelopeFollowerInitParams EnvelopeParamsInitParams;
-	
-		EnvelopeParamsInitParams.SampleRate = InSettings.GetSampleRate();
-		EnvelopeParamsInitParams.NumChannels = 1;
-		EnvelopeParamsInitParams.AttackTimeMsec = PrevAttackTime;
-		EnvelopeParamsInitParams.ReleaseTimeMsec = PrevReleaseTime;
-
-		EnvelopeFollower.Init(EnvelopeParamsInitParams);
+		Reset(InParams);
 	}
 
 	FDataReferenceCollection FEnvelopeFollowerOperator::GetInputs() const
@@ -125,6 +116,23 @@ namespace Metasound
 		return OutputDataReferences;
 	}
 
+	void FEnvelopeFollowerOperator::Reset(const IOperator::FResetParams& InParams)
+	{
+		PrevAttackTime = FMath::Max(FTime::ToMilliseconds(*AttackTimeInput), 0.0);
+		PrevReleaseTime = FMath::Max(FTime::ToMilliseconds(*ReleaseTimeInput), 0.0);
+
+		Audio::FEnvelopeFollowerInitParams EnvelopeParamsInitParams;
+	
+		EnvelopeParamsInitParams.SampleRate = InParams.OperatorSettings.GetSampleRate();
+		EnvelopeParamsInitParams.NumChannels = 1;
+		EnvelopeParamsInitParams.AttackTimeMsec = PrevAttackTime;
+		EnvelopeParamsInitParams.ReleaseTimeMsec = PrevReleaseTime;
+
+		EnvelopeFollower.Init(EnvelopeParamsInitParams);
+		*EnvelopeFloatOutput = 0.f;
+		EnvelopeAudioOutput->Zero();
+	}
+
 	void FEnvelopeFollowerOperator::Execute()
 	{
 		// Check for any input changes
@@ -228,7 +236,7 @@ namespace Metasound
 		FEnvelopePeakModeReadRef EnvelopeModeIn = InputCollection.GetDataReadReferenceOrConstruct<FEnumEnvelopePeakMode>(METASOUND_GET_PARAM_NAME(InParamFollowMode));
 
 
-		return MakeUnique<FEnvelopeFollowerOperator>(InParams.OperatorSettings, AudioIn, AttackTime, ReleaseTime, EnvelopeModeIn);
+		return MakeUnique<FEnvelopeFollowerOperator>(InParams, AudioIn, AttackTime, ReleaseTime, EnvelopeModeIn);
 	}
 
 	FEnvelopeFollowerNode::FEnvelopeFollowerNode(const FNodeInitData& InitData)
diff --git a/Engine/Plugins/Runtime/Metasound/Source/MetasoundStandardNodes/Private/MetasoundFilterQToBandwidthNode.cpp b/Engine/Plugins/Runtime/Metasound/Source/MetasoundStandardNodes/Private/MetasoundFilterQToBandwidthNode.cpp
--- a/Engine/Plugins/Runtime/Metasound/Source/MetasoundStandardNodes/Private/MetasoundFilterQToBandwidthNode.cpp
+++ b/Engine/Plugins/Runtime/Metasound/Source/MetasoundStandardNodes/Private/MetasoundFilterQToBandwidthNode.cpp
@@ -107,6 +107,11 @@ namespace Metasound
 			return Outputs;
 		}
 
+		void Reset(const IOperator::FResetParams& InParams)
+		{
+			*OutValue = Audio::GetBandwidthFromQ(FMath::Max(*InValue, KINDA_SMALL_NUMBER));
+		}
+
 		void Execute()
 		{
 			*OutValue = Audio::GetBandwidthFromQ(FMath::Max(*InValue, KINDA_SMALL_NUMBER));
diff --git a/Engine/Plugins/Runtime/Metasound/Source/MetasoundStandardNodes/Private/MetasoundFlangerNode.cpp b/Engine/Plugins/Runtime/Metasound/Source/MetasoundStandardNodes/Private/MetasoundFlangerNode.cpp
--- a/Engine/Plugins/Runtime/Metasound/Source/MetasoundStandardNodes/Private/MetasoundFlangerNode.cpp
+++ b/Engine/Plugins/Runtime/Metasound/Source/MetasoundStandardNodes/Private/MetasoundFlangerNode.cpp
@@ -88,10 +88,10 @@ namespace Metasound
 			FFloatReadRef CenterDelay = InParams.InputDataReferences.GetDataReadReferenceOrConstructWithVertexDefault<float>(InputInterface, METASOUND_GET_PARAM_NAME(InputCenterDelay), InParams.OperatorSettings);
 			FFloatReadRef MixLevel = InParams.InputDataReferences.GetDataReadReferenceOrConstructWithVertexDefault<float>(InputInterface, METASOUND_GET_PARAM_NAME(InputMixLevel), InParams.OperatorSettings);
 
-			return MakeUnique<FFlangerOperator>(InParams.OperatorSettings, AudioInput, ModulationRate, ModulationDepth, CenterDelay, MixLevel);
+			return MakeUnique<FFlangerOperator>(InParams, AudioInput, ModulationRate, ModulationDepth, CenterDelay, MixLevel);
 		}
 
-		FFlangerOperator(const FOperatorSettings& InSettings,
+		FFlangerOperator(const FCreateOperatorParams& InParams,
 			const FAudioBufferReadRef& InAudioInput,
 			const FFloatReadRef& InputModulationRate,
 			const FFloatReadRef& InputModulationDepth,
@@ -103,10 +103,10 @@ namespace Metasound
 			, ModulationDepth(InputModulationDepth)
 			, CenterDelay(InputCenterDelay)
 			, MixLevel(InMixLevel)
-			, AudioOut(FAudioBufferWriteRef::CreateNew(InSettings))
+			, AudioOut(FAudioBufferWriteRef::CreateNew(InParams.OperatorSettings))
 			, Flanger()
 		{
-			Flanger.Init(InSettings.GetSampleRate());
+			Reset(InParams);
 		}
 
 		FDataReferenceCollection GetInputs() const
@@ -132,6 +132,12 @@ namespace Metasound
 			return OutputDataReferences;
 		}
 
+		void Reset(const IOperator::FResetParams& InParams)
+		{
+			AudioOut->Zero();
+			Flanger.Init(InParams.OperatorSettings.GetSampleRate());
+		}
+
 		void Execute()
 		{
 			// Update flanger parameters if necessary 
@@ -140,19 +146,9 @@ namespace Metasound
 			Flanger.SetModulationDepth(*ModulationDepth);
 			Flanger.SetMixLevel(*MixLevel);
 
-			const float* InputAudio = AudioIn->GetData();
-			float* OutputAudio = AudioOut->GetData();
 			int32 NumFrames = AudioIn->Num();
 
-			// Copy input audio into aligned buffer
-			AlignedAudioIn.Reset();
-			AlignedAudioIn.Append(InputAudio, NumFrames);
-
-			AlignedAudioOut.Reset();
-			AlignedAudioOut.AddUninitialized(NumFrames);
-
-			Flanger.ProcessAudio(AlignedAudioIn, NumFrames, AlignedAudioOut);
-			FMemory::Memcpy(OutputAudio, AlignedAudioOut.GetData(), sizeof(float) * NumFrames);
+			Flanger.ProcessAudio(*AudioIn, NumFrames, *AudioOut);
 		}
 
 	private:
@@ -175,10 +171,6 @@ namespace Metasound
 
 		// Flanger DSP object
 		Audio::FFlanger Flanger;
-
-		// Input/output aligned buffers
-		Audio::FAlignedFloatBuffer AlignedAudioIn;
-		Audio::FAlignedFloatBuffer AlignedAudioOut;
 	};
 
 	class METASOUNDSTANDARDNODES_API FFlangerNode : public FNodeFacade
diff --git a/Engine/Plugins/Runtime/Metasound/Source/MetasoundStandardNodes/Private/MetasoundFreqMultiplierToSemitoneNode.cpp b/Engine/Plugins/Runtime/Metasound/Source/MetasoundStandardNodes/Private/MetasoundFreqMultiplierToSemitoneNode.cpp
--- a/Engine/Plugins/Runtime/Metasound/Source/MetasoundStandardNodes/Private/MetasoundFreqMultiplierToSemitoneNode.cpp
+++ b/Engine/Plugins/Runtime/Metasound/Source/MetasoundStandardNodes/Private/MetasoundFreqMultiplierToSemitoneNode.cpp
@@ -35,6 +35,7 @@ namespace Metasound
 
 		virtual FDataReferenceCollection GetInputs() const override;
 		virtual FDataReferenceCollection GetOutputs() const override;
+		void Reset(const IOperator::FResetParams& InParams);
 		void Execute();
 
 	private:
@@ -75,6 +76,12 @@ namespace Metasound
 		return OutputDataReferences;
 	}
 
+	void FFrequencyMultiplierToSemitoneOperator::Reset(const IOperator::FResetParams& InParams)
+	{
+		PrevFrequencyMultiplier = *FrequencyMultiplierInput;
+		*SemitoneOutput = Audio::GetSemitones(*FrequencyMultiplierInput);
+	}
+
 	void FFrequencyMultiplierToSemitoneOperator::Execute()
 	{
 		// Only do anything if the scale changes
diff --git a/Engine/Plugins/Runtime/Metasound/Source/MetasoundStandardNodes/Private/MetasoundFreqToMIDINode.cpp b/Engine/Plugins/Runtime/Metasound/Source/MetasoundStandardNodes/Private/MetasoundFreqToMIDINode.cpp
--- a/Engine/Plugins/Runtime/Metasound/Source/MetasoundStandardNodes/Private/MetasoundFreqToMIDINode.cpp
+++ b/Engine/Plugins/Runtime/Metasound/Source/MetasoundStandardNodes/Private/MetasoundFreqToMIDINode.cpp
@@ -14,15 +14,12 @@
 
 namespace Metasound
 {
-	
-		
 	namespace FrequencyToMidiVertexNames
 	{
 		METASOUND_PARAM(InputFreq, "Frequency In", "Input frequency value in Hz.");
 		METASOUND_PARAM(OutputMidi, "Out MIDI", "Output MIDI note value that corresponds to the input frequency value.");
 	}
 
-
 	class FFreqToMidiOperator : public TExecutableOperator<FFreqToMidiOperator>
 	{
 	public:
@@ -35,6 +32,7 @@ namespace Metasound
 
 		virtual FDataReferenceCollection GetInputs() const override;
 		virtual FDataReferenceCollection GetOutputs() const override;
+		void Reset(const IOperator::FResetParams& InParams);
 		void Execute();
 
 	private:
@@ -54,7 +52,6 @@ namespace Metasound
 		, PrevFreq(*InFreq)
 	{
 	}
-
 	
 	FDataReferenceCollection FFreqToMidiOperator::GetInputs() const
 	{
@@ -75,6 +72,12 @@ namespace Metasound
 		return OutputDataReferences;
 	}
 
+	void FFreqToMidiOperator::Reset(const IOperator::FResetParams& InParams)
+	{
+		PrevFreq = *FreqInput;
+		*MidiOutput = Audio::GetMidiFromFrequency(PrevFreq);
+	}
+
 	void FFreqToMidiOperator::Execute()
 	{
 		// Only do anything if the Midi note changes
diff --git a/Engine/Plugins/Runtime/Metasound/Source/MetasoundStandardNodes/Private/MetasoundGrainDelayNode.cpp b/Engine/Plugins/Runtime/Metasound/Source/MetasoundStandardNodes/Private/MetasoundGrainDelayNode.cpp
--- a/Engine/Plugins/Runtime/Metasound/Source/MetasoundStandardNodes/Private/MetasoundGrainDelayNode.cpp
+++ b/Engine/Plugins/Runtime/Metasound/Source/MetasoundStandardNodes/Private/MetasoundGrainDelayNode.cpp
@@ -69,6 +69,7 @@ namespace Metasound
 		static TUniquePtr<IOperator> CreateOperator(const FCreateOperatorParams& InParams, FBuildErrorArray& OutErrors);
 		virtual FDataReferenceCollection GetInputs() const override;
 		virtual FDataReferenceCollection GetOutputs() const override;
+		void Reset(const IOperator::FResetParams& InParams);
 		void Execute();
 
 	private:
@@ -247,6 +248,19 @@ namespace Metasound
 		return InputDataReferences;
 	}
 
+	void FGrainDelayOperator::Reset(const IOperator::FResetParams& InParams)
+	{
+		AudioOutput->Zero();
+
+		GrainDelayProcessor.Reset();
+		PreviousGrainDelayMsec = GrainDelayProcessor.GetGrainDelayClamped(*GrainDelay);
+		PreviousGrainEnvelopeType = *GrainDelayEnvelope;
+		PreviousPitchShift = GrainDelayProcessor.GetGrainPitchShiftClamped(*GrainPitchShift);
+		
+		GrainDelayProcessor.SetGrainEnvelope(*GrainDelayEnvelope);
+		GrainDelayProcessor.SetGrainBasePitchShiftRatio(GrainDelayProcessor.GetGrainPitchShiftFrameRatio(PreviousPitchShift));
+	}
+
 	void FGrainDelayOperator::Execute()
 	{
 		// Parse triggers and render audio correctly in the audio block
diff --git a/Engine/Plugins/Runtime/Metasound/Source/MetasoundStandardNodes/Private/MetasoundITDPannerNode.cpp b/Engine/Plugins/Runtime/Metasound/Source/MetasoundStandardNodes/Private/MetasoundITDPannerNode.cpp
--- a/Engine/Plugins/Runtime/Metasound/Source/MetasoundStandardNodes/Private/MetasoundITDPannerNode.cpp
+++ b/Engine/Plugins/Runtime/Metasound/Source/MetasoundStandardNodes/Private/MetasoundITDPannerNode.cpp
@@ -48,6 +48,7 @@ namespace Metasound
 
 		virtual FDataReferenceCollection GetInputs() const override;
 		virtual FDataReferenceCollection GetOutputs() const override;
+		void Reset(const IOperator::FResetParams& InParams);
 		void Execute();
 
 	private:
@@ -178,6 +179,24 @@ namespace Metasound
 	}
 	
 
+	void FITDPannerOperator::Reset(const IOperator::FResetParams& InParams)
+	{
+		AudioLeftOutput->Zero();
+		AudioRightOutput->Zero();
+
+		LeftDelay.Reset();
+		RightDelay.Reset();
+
+		CurrAngle = FMath::Clamp(*PanningAngle, 0.0f, 360.0f);
+		CurrDistanceFactor = FMath::Clamp(*DistanceFactor, 0.0f, 1.0f);
+		CurrHeadWidth = FMath::Max(*HeadWidth, 0.0f);
+
+		UpdateParams(true);
+
+		PrevLeftGain = CurrLeftGain;
+		PrevRightGain = CurrRightGain;
+	}
+
 	void FITDPannerOperator::Execute()
 	{
 		float NewHeadWidth = FMath::Max(*HeadWidth, 0.0f);
diff --git a/Engine/Plugins/Runtime/Metasound/Source/MetasoundStandardNodes/Private/MetasoundInterpToNode.cpp b/Engine/Plugins/Runtime/Metasound/Source/MetasoundStandardNodes/Private/MetasoundInterpToNode.cpp
--- a/Engine/Plugins/Runtime/Metasound/Source/MetasoundStandardNodes/Private/MetasoundInterpToNode.cpp
+++ b/Engine/Plugins/Runtime/Metasound/Source/MetasoundStandardNodes/Private/MetasoundInterpToNode.cpp
@@ -44,10 +44,11 @@ namespace Metasound
 		static const FVertexInterface& GetVertexInterface();
 		static TUniquePtr<IOperator> CreateOperator(const FCreateOperatorParams& InParams, FBuildErrorArray& OutErrors);
 
-		FInterpToOperator(const FOperatorSettings& InSettings, const FFloatReadRef& InTargetValue, const FTimeReadRef& InInterpTime);
+		FInterpToOperator(const FCreateOperatorParams& InSettings, const FFloatReadRef& InTargetValue, const FTimeReadRef& InInterpTime);
 
 		virtual FDataReferenceCollection GetInputs() const override;
 		virtual FDataReferenceCollection GetOutputs() const override;
+		void Reset(const IOperator::FResetParams& InParams);
 		void Execute();
 
 	private:
@@ -70,20 +71,12 @@ namespace Metasound
 		float PreviousTargetValue = 0.0f;
 	};
 
-	FInterpToOperator::FInterpToOperator(const FOperatorSettings& InSettings, const FFloatReadRef& InTargetValue, const FTimeReadRef& InInterpTime)
+	FInterpToOperator::FInterpToOperator(const FCreateOperatorParams& InParams, const FFloatReadRef& InTargetValue, const FTimeReadRef& InInterpTime)
 		: TargetValue(InTargetValue)
 		, InterpTime(InInterpTime)
 		, ValueOutput(FFloatWriteRef::CreateNew(*TargetValue))
 	{
-		// Set the fade to start at the value specified in the current value
-		VolumeFader.SetVolume(*TargetValue);
-
-		float BlockRate = InSettings.GetActualBlockRate();
-		BlockTimeDelta = 1.0f / BlockRate;
-
-		PreviousTargetValue = *TargetValue;
-
-		*ValueOutput = *TargetValue;
+		Reset(InParams);
 	}
 
 	FDataReferenceCollection FInterpToOperator::GetInputs() const
@@ -106,6 +99,20 @@ namespace Metasound
 		return OutputDataReferences;
 	}
 
+	void FInterpToOperator::Reset(const IOperator::FResetParams& InParams)
+	{
+		// Set the fade to start at the value specified in the current value
+		VolumeFader.SetVolume(*TargetValue);
+
+		float BlockRate = InParams.OperatorSettings.GetActualBlockRate();
+		BlockTimeDelta = 1.0f / BlockRate;
+
+		PreviousTargetValue = *TargetValue;
+
+		*ValueOutput = *TargetValue;
+
+	}
+
 	void FInterpToOperator::Execute()
 	{
 		// Update the value output with the current value in case it was changed
@@ -184,7 +191,7 @@ namespace Metasound
 		FFloatReadRef TargetValue = InputCollection.GetDataReadReferenceOrConstructWithVertexDefault<float>(InputInterface, METASOUND_GET_PARAM_NAME(InParamTarget), InParams.OperatorSettings);
 		FTimeReadRef InterpTime = InputCollection.GetDataReadReferenceOrConstructWithVertexDefault<FTime>(InputInterface, METASOUND_GET_PARAM_NAME(InParamInterpTime), InParams.OperatorSettings);
 
-		return MakeUnique<FInterpToOperator>(InParams.OperatorSettings, TargetValue, InterpTime);
+		return MakeUnique<FInterpToOperator>(InParams, TargetValue, InterpTime);
 	}
 
 
diff --git a/Engine/Plugins/Runtime/Metasound/Source/MetasoundStandardNodes/Private/MetasoundLimiterNode.cpp b/Engine/Plugins/Runtime/Metasound/Source/MetasoundStandardNodes/Private/MetasoundLimiterNode.cpp
--- a/Engine/Plugins/Runtime/Metasound/Source/MetasoundStandardNodes/Private/MetasoundLimiterNode.cpp
+++ b/Engine/Plugins/Runtime/Metasound/Source/MetasoundStandardNodes/Private/MetasoundLimiterNode.cpp
@@ -52,7 +52,7 @@
 		static constexpr float SoftKneeBandwitdh = 10.0f;
 		static constexpr float MaxInputGain = 100.0f;
 
-		FLimiterOperator(const FOperatorSettings& InSettings,
+		FLimiterOperator(const FCreateOperatorParams& InParams,
 			const FAudioBufferReadRef& InAudio,
 			const FFloatReadRef& InGainDb,
 			const FFloatReadRef& InThresholdDb,
@@ -63,31 +63,13 @@
 			, ThresholdDbInput(InThresholdDb)
 			, ReleaseTimeInput(InReleaseTime)
 			, KneeModeInput(InKneeMode)
-			, AudioOutput(FAudioBufferWriteRef::CreateNew(InSettings))
+			, AudioOutput(FAudioBufferWriteRef::CreateNew(InParams.OperatorSettings))
 			, Limiter()
 			, PrevInGainDb(*InGainDb)
 			, PrevThresholdDb(*InThresholdDb)
 			, PrevReleaseTime(FMath::Max(FTime::ToMilliseconds(*InReleaseTime), 0.0))
 		{
-			Limiter.Init(InSettings.GetSampleRate(), 1);
-			Limiter.SetProcessingMode(Audio::EDynamicsProcessingMode::Limiter);
-			Limiter.SetInputGain(FMath::Min(*InGainDbInput, MaxInputGain));
-			Limiter.SetThreshold(*ThresholdDbInput);
-			Limiter.SetAttackTime(0.0f);
-			Limiter.SetReleaseTime(PrevReleaseTime);
-			Limiter.SetPeakMode(Audio::EPeakMode::Peak);
-
-			switch (*KneeModeInput)
-			{
-			default:
-			case EKneeMode::Hard:
-				Limiter.SetKneeBandwidth(HardKneeBandwitdh);
-				break;
-			case EKneeMode::Soft:
-				Limiter.SetKneeBandwidth(SoftKneeBandwitdh);
-				break;
-			}
-
+			Reset(InParams);
 		}
 
 		static const FNodeClassMetadata& GetNodeInfo()
@@ -176,18 +158,55 @@
 			FTimeReadRef ReleaseTimeIn = Inputs.GetDataReadReferenceOrConstructWithVertexDefault<FTime>(InputInterface, METASOUND_GET_PARAM_NAME(InputReleaseTime), InParams.OperatorSettings);
 			FKneeModeReadRef KneeModeIn = Inputs.GetDataReadReferenceOrConstruct<FEnumKneeMode>(METASOUND_GET_PARAM_NAME(InputKneeMode));
 
-			return MakeUnique<FLimiterOperator>(InParams.OperatorSettings, AudioIn, InGainDbIn, ThresholdDbIn, ReleaseTimeIn, KneeModeIn);
+			return MakeUnique<FLimiterOperator>(InParams, AudioIn, InGainDbIn, ThresholdDbIn, ReleaseTimeIn, KneeModeIn);
 		}
+
+		void Reset(const IOperator::FResetParams& InParams)
+		{
+			AudioOutput->Zero();
+
+			const float ClampedInGainDb = FMath::Min(*InGainDbInput, MaxInputGain);
+			const float ClampedReleaseTime = FMath::Max(FTime::ToMilliseconds(*ReleaseTimeInput), 0.0);
+
+			Limiter.Init(InParams.OperatorSettings.GetSampleRate(), 1);
+			Limiter.SetProcessingMode(Audio::EDynamicsProcessingMode::Limiter);
+			Limiter.SetInputGain(ClampedInGainDb);
+			Limiter.SetThreshold(*ThresholdDbInput);
+			Limiter.SetAttackTime(0.0f);
+			Limiter.SetReleaseTime(ClampedReleaseTime);
+			Limiter.SetPeakMode(Audio::EPeakMode::Peak);
+
+			switch (*KneeModeInput)
+			{
+			default:
+			case EKneeMode::Hard:
+				Limiter.SetKneeBandwidth(HardKneeBandwitdh);
+				break;
+			case EKneeMode::Soft:
+				Limiter.SetKneeBandwidth(SoftKneeBandwitdh);
+				break;
+			}
+
+			PrevInGainDb = ClampedInGainDb;
+			PrevReleaseTime = ClampedReleaseTime;
+			PrevThresholdDb = *ThresholdDbInput;
+			PrevKneeMode = *KneeModeInput;
+		}
+
 		void Execute()
 		{
 			/* Update parameters */
-			if (!FMath::IsNearlyEqual(*InGainDbInput, PrevInGainDb))
+			float ClampedInGainDb = FMath::Min(*InGainDbInput, MaxInputGain);
+			if (!FMath::IsNearlyEqual(ClampedInGainDb, PrevInGainDb))
 			{
-				Limiter.SetInputGain(FMath::Min(*InGainDbInput, MaxInputGain));
+				Limiter.SetInputGain(ClampedInGainDb);
+				PrevInGainDb = ClampedInGainDb;
 			}
+
 			if (!FMath::IsNearlyEqual(*ThresholdDbInput, PrevThresholdDb))
 			{
 				Limiter.SetThreshold(*ThresholdDbInput);
+				PrevThresholdDb = *ThresholdDbInput;
 			}
 			// Release time cannot be negative
 			double CurrRelease = FMath::Max(FTime::ToMilliseconds(*ReleaseTimeInput), 0.0f);
@@ -197,16 +216,20 @@
 				PrevReleaseTime = CurrRelease;
 			}
 
-			switch (*KneeModeInput)
+			if (PrevKneeMode != *KneeModeInput)
 			{
-			default:
-			case EKneeMode::Hard:
-				Limiter.SetKneeBandwidth(HardKneeBandwitdh);
-				break;
-			case EKneeMode::Soft:
-				Limiter.SetKneeBandwidth(SoftKneeBandwitdh);
-				break;
-			}			
+				switch (*KneeModeInput)
+				{
+				default:
+				case EKneeMode::Hard:
+					Limiter.SetKneeBandwidth(HardKneeBandwitdh);
+					break;
+				case EKneeMode::Soft:
+					Limiter.SetKneeBandwidth(SoftKneeBandwitdh);
+					break;
+				}			
+				PrevKneeMode = *KneeModeInput;
+			}
 
 			Limiter.ProcessAudio(AudioInput->GetData(), AudioInput->Num(), AudioOutput->GetData());
 		}
@@ -227,7 +250,7 @@
 		float PrevInGainDb;
 		float PrevThresholdDb;
 		double PrevReleaseTime;
-		float PrevKnee;
+		EKneeMode PrevKneeMode;
 	};
 
 	// Node Class
diff --git a/Engine/Plugins/Runtime/Metasound/Source/MetasoundStandardNodes/Private/MetasoundLinearGainToDecibels.cpp b/Engine/Plugins/Runtime/Metasound/Source/MetasoundStandardNodes/Private/MetasoundLinearGainToDecibels.cpp
--- a/Engine/Plugins/Runtime/Metasound/Source/MetasoundStandardNodes/Private/MetasoundLinearGainToDecibels.cpp
+++ b/Engine/Plugins/Runtime/Metasound/Source/MetasoundStandardNodes/Private/MetasoundLinearGainToDecibels.cpp
@@ -14,7 +14,6 @@
 
 namespace Metasound
 {
-
 	namespace LinearGainToDecibelsVertexNames
 	{
 		METASOUND_PARAM(InputLinearGain, "Linear Gain", "Input linear gain.");
@@ -22,7 +21,6 @@ namespace Metasound
 		METASOUND_PARAM(OutputDecibels, "Decibels", "Output corresponding logarithmic (dB) gain.");
 	}
 
-
 	class FLinearGainToDecibelsOperator : public TExecutableOperator<FLinearGainToDecibelsOperator>
 	{
 	public:
@@ -31,10 +29,11 @@ namespace Metasound
 		static const FVertexInterface& GetVertexInterface();
 		static TUniquePtr<IOperator> CreateOperator(const FCreateOperatorParams& InParams, FBuildErrorArray& OutErrors);
 
-		FLinearGainToDecibelsOperator(const FOperatorSettings& InSettings, const FFloatReadRef& InLinearGain);
+		FLinearGainToDecibelsOperator(const FCreateOperatorParams& InParams, const FFloatReadRef& InLinearGain);
 
 		virtual FDataReferenceCollection GetInputs() const override;
 		virtual FDataReferenceCollection GetOutputs() const override;
+		void Reset(const IOperator::FResetParams& InParams);
 		void Execute();
 
 	private:
@@ -45,10 +44,11 @@ namespace Metasound
 		FFloatWriteRef DecibelOutput;
 	};
 
-	FLinearGainToDecibelsOperator::FLinearGainToDecibelsOperator(const FOperatorSettings& InSettings, const FFloatReadRef& InLinearGain)
+	FLinearGainToDecibelsOperator::FLinearGainToDecibelsOperator(const FCreateOperatorParams& InParams, const FFloatReadRef& InLinearGain)
 		: LinearGainInput(InLinearGain)
-		, DecibelOutput(FFloatWriteRef::CreateNew(FMath::Max(Audio::ConvertToDecibels(*LinearGainInput, TNumericLimits<float>::Min()), -240.0f)))
+		, DecibelOutput(FFloatWriteRef::CreateNew(0))
 	{
+		Reset(InParams);
 	}
 
 	
@@ -71,6 +71,11 @@ namespace Metasound
 		return OutputDataReferences;
 	}
 
+	void FLinearGainToDecibelsOperator::Reset(const IOperator::FResetParams& InParams)
+	{
+		Execute();
+	}
+
 	void FLinearGainToDecibelsOperator::Execute()
 	{
 		// Linear gain value must be above 0
@@ -134,7 +139,7 @@ namespace Metasound
 
 		FFloatReadRef InLinearGain = InputCollection.GetDataReadReferenceOrConstructWithVertexDefault<float>(InputInterface, METASOUND_GET_PARAM_NAME(InputLinearGain), InParams.OperatorSettings);
 		
-		return MakeUnique<FLinearGainToDecibelsOperator>(InParams.OperatorSettings, InLinearGain);
+		return MakeUnique<FLinearGainToDecibelsOperator>(InParams, InLinearGain);
 	}
 
 	class METASOUNDSTANDARDNODES_API FLinearGainToDecibelsNode : public FNodeFacade
diff --git a/Engine/Plugins/Runtime/Metasound/Source/MetasoundStandardNodes/Private/MetasoundLinearToLogFrequencyNode.cpp b/Engine/Plugins/Runtime/Metasound/Source/MetasoundStandardNodes/Private/MetasoundLinearToLogFrequencyNode.cpp
--- a/Engine/Plugins/Runtime/Metasound/Source/MetasoundStandardNodes/Private/MetasoundLinearToLogFrequencyNode.cpp
+++ b/Engine/Plugins/Runtime/Metasound/Source/MetasoundStandardNodes/Private/MetasoundLinearToLogFrequencyNode.cpp
@@ -132,6 +132,12 @@ namespace Metasound
 			return OutputDataReferences;
 		}
 
+		void Reset(const IOperator::FResetParams& InParams)
+		{
+			PrevInputValue = *Value;
+			*FreqOutput = GetOutputValue();
+		}
+
 		void Execute()
 		{
 			using namespace LinearToLogFrequencyVertexNames;
diff --git a/Engine/Plugins/Runtime/Metasound/Source/MetasoundStandardNodes/Private/MetasoundMIDIToFreqNode.cpp b/Engine/Plugins/Runtime/Metasound/Source/MetasoundStandardNodes/Private/MetasoundMIDIToFreqNode.cpp
--- a/Engine/Plugins/Runtime/Metasound/Source/MetasoundStandardNodes/Private/MetasoundMIDIToFreqNode.cpp
+++ b/Engine/Plugins/Runtime/Metasound/Source/MetasoundStandardNodes/Private/MetasoundMIDIToFreqNode.cpp
@@ -66,10 +66,11 @@ namespace Metasound
 		static const FVertexInterface& GetVertexInterface();
 		static TUniquePtr<IOperator> CreateOperator(const FCreateOperatorParams& InParams, FBuildErrorArray& OutErrors);
 
-		TMidiToFreqOperator(const FOperatorSettings& InSettings, const TDataReadReference<ValueType>& InMidiNote);
+		TMidiToFreqOperator(const FCreateOperatorParams& InParams, const TDataReadReference<ValueType>& InMidiNote);
 
 		virtual FDataReferenceCollection GetInputs() const override;
 		virtual FDataReferenceCollection GetOutputs() const override;
+		void Reset(const IOperator::FResetParams& InParams);
 		void Execute();
 
 	private:
@@ -84,11 +85,12 @@ namespace Metasound
 	};
 
 	template<typename ValueType>
-	TMidiToFreqOperator<ValueType>::TMidiToFreqOperator(const FOperatorSettings& InSettings, const TDataReadReference<ValueType>& InMidiNote)
+	TMidiToFreqOperator<ValueType>::TMidiToFreqOperator(const FCreateOperatorParams& InParams, const TDataReadReference<ValueType>& InMidiNote)
 		: MidiNote(InMidiNote)
 		, FreqOutput(FFloatWriteRef::CreateNew(Audio::GetFrequencyFromMidi(*InMidiNote)))
 		, PrevMidiNote(*InMidiNote)
 	{
+		Reset(InParams);
 	}
 
 	template<typename ValueType>
@@ -112,6 +114,14 @@ namespace Metasound
 		return OutputDataReferences;
 	}
 
+	template<typename ValueType>
+	void TMidiToFreqOperator<ValueType>::Reset(const IOperator::FResetParams& InParams)
+	{
+		using namespace MidiToFrequencyPrivate;
+		PrevMidiNote = *MidiNote;
+		*FreqOutput = TMidiToFreqNodeSpecialization<ValueType>::GetFreqValue(PrevMidiNote);
+	}
+
 	template<typename ValueType>
 	void TMidiToFreqOperator<ValueType>::Execute()
 	{
@@ -183,7 +193,7 @@ namespace Metasound
 
 		TDataReadReference<ValueType> InMidiNote = InputCollection.GetDataReadReferenceOrConstructWithVertexDefault<ValueType>(InputInterface, METASOUND_GET_PARAM_NAME(InputMidi), InParams.OperatorSettings);
 
-		return MakeUnique<TMidiToFreqOperator>(InParams.OperatorSettings, InMidiNote);
+		return MakeUnique<TMidiToFreqOperator>(InParams, InMidiNote);
 	}
 
 	template<typename ValueType>
diff --git a/Engine/Plugins/Runtime/Metasound/Source/MetasoundStandardNodes/Private/MetasoundMapRangeNode.cpp b/Engine/Plugins/Runtime/Metasound/Source/MetasoundStandardNodes/Private/MetasoundMapRangeNode.cpp
--- a/Engine/Plugins/Runtime/Metasound/Source/MetasoundStandardNodes/Private/MetasoundMapRangeNode.cpp
+++ b/Engine/Plugins/Runtime/Metasound/Source/MetasoundStandardNodes/Private/MetasoundMapRangeNode.cpp
@@ -354,6 +354,11 @@ namespace Metasound
 			return MappingObject.GetOutputs();
 		}
 
+		void Reset(const IOperator::FResetParams& InParams)
+		{
+			MappingObject.DoMapping();
+		}
+
 		void Execute()
 		{
 			MappingObject.DoMapping();
diff --git a/Engine/Plugins/Runtime/Metasound/Source/MetasoundStandardNodes/Private/MetasoundMathNodes.cpp b/Engine/Plugins/Runtime/Metasound/Source/MetasoundStandardNodes/Private/MetasoundMathNodes.cpp
--- a/Engine/Plugins/Runtime/Metasound/Source/MetasoundStandardNodes/Private/MetasoundMathNodes.cpp
+++ b/Engine/Plugins/Runtime/Metasound/Source/MetasoundStandardNodes/Private/MetasoundMathNodes.cpp
@@ -235,6 +235,11 @@ namespace Metasound
 				return Info;
 			}
 
+			void Reset(const IOperator::FResetParams& InParams)
+			{
+				TMathOpClass::Calculate(InstanceData, PrimaryOperandRef, AdditionalOperandRefs, ValueRef);
+			}
+
 			void Execute()
 			{
 				TMathOpClass::Calculate(InstanceData, PrimaryOperandRef, AdditionalOperandRefs, ValueRef);
@@ -264,16 +269,15 @@ namespace Metasound
 				TOperandDataClassReadRef Op1 = InParams.InputDataReferences.GetDataReadReferenceOrConstruct<TOperandDataClass>(OpName, TMathOpClass::GetDefaultOp(InParams.OperatorSettings, DefaultValue));
 				TArray<TOperandDataClassReadRef> AdditionalOperandRefs = { Op1 };
 
-				return MakeUnique<TMathOperator>(InParams.OperatorSettings, PrimaryOperand, AdditionalOperandRefs);
+				return MakeUnique<TMathOperator>(InParams, PrimaryOperand, AdditionalOperandRefs);
 			}
 
-			TMathOperator(const FOperatorSettings& InSettings, const TDataClassReadRef& InPrimaryOperand, const TArray<TOperandDataClassReadRef>& InAdditionalOperands)
+			TMathOperator(const FCreateOperatorParams& InParams, const TDataClassReadRef& InPrimaryOperand, const TArray<TOperandDataClassReadRef>& InAdditionalOperands)
 				: PrimaryOperandRef(InPrimaryOperand)
 				, AdditionalOperandRefs(InAdditionalOperands)
-				, ValueRef(TDataWriteReferenceFactory<TDataClass>::CreateAny(InSettings))
+				, ValueRef(TDataWriteReferenceFactory<TDataClass>::CreateAny(InParams.OperatorSettings))
 			{
-				// Set initial value.
-				TMathOpClass::Calculate(InstanceData, PrimaryOperandRef, AdditionalOperandRefs, ValueRef);
+				Reset(InParams);
 			}
 
 			virtual FDataReferenceCollection GetInputs() const override
diff --git a/Engine/Plugins/Runtime/Metasound/Source/MetasoundStandardNodes/Private/MetasoundMaxNode.cpp b/Engine/Plugins/Runtime/Metasound/Source/MetasoundStandardNodes/Private/MetasoundMaxNode.cpp
--- a/Engine/Plugins/Runtime/Metasound/Source/MetasoundStandardNodes/Private/MetasoundMaxNode.cpp
+++ b/Engine/Plugins/Runtime/Metasound/Source/MetasoundStandardNodes/Private/MetasoundMaxNode.cpp
@@ -219,6 +219,11 @@ namespace Metasound
 			TMax<ValueType>::GetMax(*InputA, *InputB, *OutputValue);
 		}
 
+		void Reset(const IOperator::FResetParams& InParams)
+		{
+			GetMax();
+		}
+
 		void Execute()
 		{
 			GetMax();
diff --git a/Engine/Plugins/Runtime/Metasound/Source/MetasoundStandardNodes/Private/MetasoundMidSideNodes.cpp b/Engine/Plugins/Runtime/Metasound/Source/MetasoundStandardNodes/Private/MetasoundMidSideNodes.cpp
--- a/Engine/Plugins/Runtime/Metasound/Source/MetasoundStandardNodes/Private/MetasoundMidSideNodes.cpp
+++ b/Engine/Plugins/Runtime/Metasound/Source/MetasoundStandardNodes/Private/MetasoundMidSideNodes.cpp
@@ -34,7 +34,7 @@ namespace Metasound
 	{
 	public:
 
-		FMidSideEncodeOperator(const FOperatorSettings& InSettings,
+		FMidSideEncodeOperator(const FCreateOperatorParams& InParams,
 			const FAudioBufferReadRef& InLeftAudioInput,
 			const FAudioBufferReadRef& InRightAudioInput,
 			const FFloatReadRef& InSpreadAmount,
@@ -43,14 +43,15 @@ namespace Metasound
 			, AudioInputRight(InRightAudioInput)
 			, SpreadAmount(InSpreadAmount)
 			, bEqualPower(InEqualPower)
-			, AudioOutputMid(FAudioBufferWriteRef::CreateNew(InSettings))
-			, AudioOutputSide(FAudioBufferWriteRef::CreateNew(InSettings))
+			, AudioOutputMid(FAudioBufferWriteRef::CreateNew(InParams.OperatorSettings))
+			, AudioOutputSide(FAudioBufferWriteRef::CreateNew(InParams.OperatorSettings))
 			, MidScale(0.0f)
 			, SideScale(0.0f)
 			, PrevMidScale(0.0f)
 			, PrevSideScale(0.0f)
 			, PrevSpreadAmount(*SpreadAmount)
 		{
+			Reset(InParams);
 		}
 
 		static const FNodeClassMetadata& GetNodeInfo()
@@ -113,7 +114,6 @@ namespace Metasound
 			InputDataReferences.AddDataReadReference(METASOUND_GET_PARAM_NAME(InputEqualPower), bEqualPower);
 			
 			return InputDataReferences;
-
 		}
 
 		virtual FDataReferenceCollection GetOutputs() const override
@@ -139,39 +139,20 @@ namespace Metasound
 			FFloatReadRef SpreadAmountIn = Inputs.GetDataReadReferenceOrConstructWithVertexDefault<float>(InputInterface, METASOUND_GET_PARAM_NAME(InputSpreadAmount), InParams.OperatorSettings);
 			FBoolReadRef bEqualPowerIn = Inputs.GetDataReadReferenceOrConstructWithVertexDefault<bool>(InputInterface, METASOUND_GET_PARAM_NAME(InputEqualPower), InParams.OperatorSettings);
 
-			return MakeUnique<FMidSideEncodeOperator>(InParams.OperatorSettings, LeftAudioIn, RightAudioIn, SpreadAmountIn, bEqualPowerIn);
+			return MakeUnique<FMidSideEncodeOperator>(InParams, LeftAudioIn, RightAudioIn, SpreadAmountIn, bEqualPowerIn);
 		}
 
 		void Execute()
 		{
 			/* Update internal variables, if necessary */
-			
 			const float ClampedSpread = FMath::Clamp(*SpreadAmount, 0.0f, 1.0f);
 			bool bNeedsUpdate = FMath::IsNearlyEqual(ClampedSpread, PrevSpreadAmount);
 			if (bNeedsUpdate)
 			{
-				// Convert to radians between 0.0 and PI/2
-				const float SpreadScale = ClampedSpread * 0.5f * PI;
-
-				// Compute equal power relationship between Mid and Side
-				FMath::SinCos(&SideScale, &MidScale, SpreadScale);
-
-				// Adjust gain so 0.5f Spread results in a 1.0f to 1.0f gain ratio between Mid and Side
-				const float NormalizingFactor = UE_SQRT_2;
-				MidScale *= NormalizingFactor;
-				SideScale *= NormalizingFactor;
-
-				// Clamp values if not Equal Power
-				if (!*bEqualPower)
-				{
-					MidScale = FMath::Clamp(MidScale, 0.0f, 1.0f);
-					SideScale = FMath::Clamp(SideScale, 0.0f, 1.0f);
-				}
-
+				UpdateScale(ClampedSpread);
+				PrevSpreadAmount = ClampedSpread;
 			}
 
-			PrevSpreadAmount = ClampedSpread;
-
 			/* Generate Mid-Side Output */
 
 			// Encode the signal to to MS
@@ -193,12 +174,46 @@ namespace Metasound
 				PrevMidScale = MidScale;
 				PrevSideScale = SideScale;
 			}
-			
+		}
+
+		void Reset(const IOperator::FResetParams& InParams)
+		{
+			const float ClampedSpread = FMath::Clamp(*SpreadAmount, 0.0f, 1.0f);
+			UpdateScale(ClampedSpread);
+
+			PrevSpreadAmount = ClampedSpread;
+			PrevMidScale = MidScale;
+			PrevSideScale = SideScale;
 
+			// Encode the signal to to MS
+			Audio::EncodeMidSide(*AudioInputLeft, *AudioInputRight, *AudioOutputMid, *AudioOutputSide);			
+			Audio::ArrayMultiplyByConstantInPlace(*AudioOutputMid, MidScale);
+			Audio::ArrayMultiplyByConstantInPlace(*AudioOutputSide, SideScale);
 		}
 
 	private:
 
+		void UpdateScale(float InSpread)
+		{
+			// Convert to radians between 0.0 and PI/2
+			const float SpreadScale = InSpread * 0.5f * PI;
+
+			// Compute equal power relationship between Mid and Side
+			FMath::SinCos(&SideScale, &MidScale, SpreadScale);
+
+			// Adjust gain so 0.5f Spread results in a 1.0f to 1.0f gain ratio between Mid and Side
+			const float NormalizingFactor = UE_SQRT_2;
+			MidScale *= NormalizingFactor;
+			SideScale *= NormalizingFactor;
+
+			// Clamp values if not Equal Power
+			if (!*bEqualPower)
+			{
+				MidScale = FMath::Clamp(MidScale, 0.0f, 1.0f);
+				SideScale = FMath::Clamp(SideScale, 0.0f, 1.0f);
+			}
+		}
+
 		// The input audio buffer
 		FAudioBufferReadRef AudioInputLeft;
 		FAudioBufferReadRef AudioInputRight;
@@ -223,7 +238,6 @@ namespace Metasound
 
 		// Previous Spread variable; need to update MidScale and SideScale if this changes between Execute() calls
 		float PrevSpreadAmount;
-
 	};
 
 	// Node Class
@@ -275,6 +289,11 @@ namespace Metasound
 			, PrevSideScale(0.0f)
 			, PrevSpreadAmount(*SpreadAmount)
 		{
+			MidInBuffer.AddUninitialized(InSettings.GetNumFramesPerBlock());
+			FMemory::Memset(MidInBuffer.GetData(), 0, sizeof(float) * MidInBuffer.Num());
+
+			SideInBuffer.AddUninitialized(InSettings.GetNumFramesPerBlock());
+			FMemory::Memset(SideInBuffer.GetData(), 0, sizeof(float) * SideInBuffer.Num());
 		}
 
 		static const FNodeClassMetadata& GetNodeInfo()
@@ -370,57 +389,29 @@ namespace Metasound
 		{
 			/* Update internal variables, if necessary */
 			const float ClampedSpread = FMath::Clamp(*SpreadAmount, 0.0f, 1.0f);
-			bool bNeedsUpdate = FMath::IsNearlyEqual(ClampedSpread, PrevSpreadAmount);
-			if (bNeedsUpdate)
+			
+			if (!FMath::IsNearlyEqual(ClampedSpread, PrevSpreadAmount))
 			{
-				// Convert to radians between 0.0 and PI/2
-				const float SpreadScale = ClampedSpread * 0.5f * PI;
-
-				// Compute equal power relationship between Mid and Side
-				FMath::SinCos(&SideScale, &MidScale, SpreadScale);
-
-
-				// Adjust gain so 0.5f Spread results in a 1.0f to 1.0f gain ratio between Mid and Side
-				const float NormalizingFactor = UE_SQRT_2;
-				MidScale *= NormalizingFactor;
-				SideScale *= NormalizingFactor;
-
-				// Clamp values if not Equal Power
-				if (!*bEqualPower)
-				{
-					MidScale = FMath::Clamp(MidScale, 0.0f, 1.0f);
-					SideScale = FMath::Clamp(SideScale, 0.0f, 1.0f);
-				}
+				UpdateScale(ClampedSpread);
+				PrevSpreadAmount = ClampedSpread;
 			}
 
-			PrevSpreadAmount = ClampedSpread;
-
 			/* Generate Mid-Side Output */
 
-			// We need to make a copy of the input data so we can scale it before decoding
-			const float* MidInput = AudioInputMid->GetData();
-			const float* SideInput = AudioInputSide->GetData();
-
-			const int32 NumFrames = AudioInputMid->Num();
-
-			MidInBuffer.Reset();
-			MidInBuffer.Append(MidInput, NumFrames);
-			SideInBuffer.Reset();
-			SideInBuffer.Append(SideInput, NumFrames);
-
-
 			// SideScale never changes quickly without MidScale also changing quickly
 			if (FMath::IsNearlyEqual(PrevMidScale, MidScale))
 			{
-				Audio::ArrayMultiplyByConstantInPlace(MidInBuffer, MidScale);
-				Audio::ArrayMultiplyByConstantInPlace(SideInBuffer, SideScale);
+				Audio::ArrayMultiplyByConstant(*AudioInputMid, MidScale, MidInBuffer);
+				Audio::ArrayMultiplyByConstant(*AudioInputSide, SideScale, SideInBuffer);
 			}
 			else
 			{
+				MidInBuffer = *AudioInputMid;
 				Audio::ArrayFade(MidInBuffer, PrevMidScale, MidScale);
-				Audio::ArrayFade(SideInBuffer, PrevSideScale, SideScale);
-
 				PrevMidScale = MidScale;
+
+				SideInBuffer = *AudioInputSide;
+				Audio::ArrayFade(SideInBuffer, PrevSideScale, SideScale);
 				PrevSideScale = SideScale;
 			}
 
@@ -428,8 +419,42 @@ namespace Metasound
 			Audio::DecodeMidSide(MidInBuffer, SideInBuffer, *AudioOutputLeft, *AudioOutputRight);
 		}
 
+		void Reset(const IOperator::FResetParams& InParams)
+		{
+			AudioOutputLeft->Zero();
+			AudioOutputRight->Zero();
+
+			const float ClampedSpread = FMath::Clamp(*SpreadAmount, 0.0f, 1.0f);
+			UpdateScale(ClampedSpread);
+			PrevSpreadAmount = ClampedSpread;
+
+			PrevMidScale = MidScale;
+			PrevSideScale = SideScale;
+		}
+
 	private:
 
+		void UpdateScale(float InSpread)
+		{
+			// Convert to radians between 0.0 and PI/2
+			const float SpreadScale = InSpread * 0.5f * PI;
+
+			// Compute equal power relationship between Mid and Side
+			FMath::SinCos(&SideScale, &MidScale, SpreadScale);
+
+			// Adjust gain so 0.5f Spread results in a 1.0f to 1.0f gain ratio between Mid and Side
+			constexpr float NormalizingFactor = UE_SQRT_2;
+			MidScale *= NormalizingFactor;
+			SideScale *= NormalizingFactor;
+
+			// Clamp values if not Equal Power
+			if (!*bEqualPower)
+			{
+				MidScale = FMath::Clamp(MidScale, 0.0f, 1.0f);
+				SideScale = FMath::Clamp(SideScale, 0.0f, 1.0f);
+			}
+		}
+
 		// The input audio buffer
 		FAudioBufferReadRef AudioInputMid;
 		FAudioBufferReadRef AudioInputSide;
diff --git a/Engine/Plugins/Runtime/Metasound/Source/MetasoundStandardNodes/Private/MetasoundMidiNoteQuantizerNode.cpp b/Engine/Plugins/Runtime/Metasound/Source/MetasoundStandardNodes/Private/MetasoundMidiNoteQuantizerNode.cpp
--- a/Engine/Plugins/Runtime/Metasound/Source/MetasoundStandardNodes/Private/MetasoundMidiNoteQuantizerNode.cpp
+++ b/Engine/Plugins/Runtime/Metasound/Source/MetasoundStandardNodes/Private/MetasoundMidiNoteQuantizerNode.cpp
@@ -50,7 +50,7 @@ namespace Metasound
 
 		// ctor
 		FMidiNoteQuantizerOperator(
-			  const FOperatorSettings& InSettings
+			  const FCreateOperatorParams& InParams
 			, const FFloatReadRef& InMidiNoteIn
 			, const FFloatReadRef& InRootNote
 			, const FArrayScaleDegreeReadRef& InScale
@@ -62,6 +62,7 @@ namespace Metasound
 		static TUniquePtr<IOperator> CreateOperator(const FCreateOperatorParams& InParams, FBuildErrorArray& OutErrors);
 		virtual FDataReferenceCollection GetInputs() const override;
 		virtual FDataReferenceCollection GetOutputs() const override;
+		void Reset(const IOperator::FResetParams& InParams);
 		void Execute();
 
 	private: // members
@@ -79,8 +80,10 @@ namespace Metasound
 		TArray<float> PreviousScale;
 		float PreviousNoteOut = 0.0f;
 
-		// other
-		FOperatorSettings Settings;
+		static constexpr float MaxNote = 1e12f;
+		static constexpr float MinNote = -1e12f;
+		static constexpr float MaxRoot = 1e12f;
+		static constexpr float MinRoot = 0.f;
 
 	}; // class FMidiNoteQuantizerOperator
 
@@ -91,7 +94,7 @@ namespace Metasound
 
 	// ctor
 	FMidiNoteQuantizerOperator::FMidiNoteQuantizerOperator(
-		  const FOperatorSettings& InSettings
+		  const FCreateOperatorParams& InParams
 		, const FFloatReadRef& InMidiNoteIn
 		, const FFloatReadRef& InRootNote
 		, const FArrayScaleDegreeReadRef& InScale
@@ -100,8 +103,8 @@ namespace Metasound
 		, RootNote(InRootNote)
 		, Scale(InScale)
 		, MidiNoteOut(FFloatWriteRef::CreateNew())
-		, Settings(InSettings)
 	{
+		Reset(InParams);
 	}
 
 
@@ -157,7 +160,7 @@ namespace Metasound
 		FArrayScaleDegreeReadRef InScaleArray = InParams.InputDataReferences.GetDataReadReferenceOrConstruct<ScaleDegreeArrayType>(METASOUND_GET_PARAM_NAME(ParamScaleDegrees));
 
 		return MakeUnique <FMidiNoteQuantizerOperator>(
-			  InParams.OperatorSettings
+			  InParams
 			, MidiNoteIn
 			, RootNoteIn
 			, InScaleArray
@@ -183,6 +186,17 @@ namespace Metasound
 		return OutputDataReferences;
 	}
 
+	void FMidiNoteQuantizerOperator::Reset(const IOperator::FResetParams& InParams)
+	{
+		PreviousNoteIn = FMath::Clamp(*MidiNoteIn, MinNote, MaxNote);
+		PreviousRoot = FMath::Clamp(*RootNote, MinRoot, MaxRoot);
+		PreviousScale = *Scale;
+
+		PreviousScale.Sort();
+		PreviousNoteOut = Audio::FMidiNoteQuantizer::QuantizeMidiNote(PreviousNoteIn, PreviousRoot, PreviousScale);
+		*MidiNoteOut = PreviousNoteOut;
+	}
+
 	void FMidiNoteQuantizerOperator::Execute()
 	{
 		if ((*Scale).IsEmpty())
@@ -192,15 +206,16 @@ namespace Metasound
 		}
 
 		// calculate new output and cache values if needed
-		const float CurrentRoot = FMath::Max(*RootNote, 0.0f);
-		if (!FMath::IsNearlyEqual(PreviousNoteIn, *MidiNoteIn)
+		const float CurrentNote = FMath::Clamp(*MidiNoteIn, MinNote, MaxNote);
+		const float CurrentRoot = FMath::Clamp(*RootNote, MinRoot, MaxRoot);
+		if (!FMath::IsNearlyEqual(PreviousNoteIn, CurrentNote)
 			|| !FMath::IsNearlyEqual(PreviousRoot, CurrentRoot)
 			|| PreviousScale !=  *Scale
 			)
 		{
 			// cache values
-			PreviousNoteIn = *MidiNoteIn;
-			PreviousRoot = *RootNote;
+			PreviousNoteIn = CurrentNote;
+			PreviousRoot = CurrentRoot;
 			PreviousScale = *Scale;
 
 			PreviousScale.Sort();
diff --git a/Engine/Plugins/Runtime/Metasound/Source/MetasoundStandardNodes/Private/MetasoundMidiScaleToArrayNode.cpp b/Engine/Plugins/Runtime/Metasound/Source/MetasoundStandardNodes/Private/MetasoundMidiScaleToArrayNode.cpp
--- a/Engine/Plugins/Runtime/Metasound/Source/MetasoundStandardNodes/Private/MetasoundMidiScaleToArrayNode.cpp
+++ b/Engine/Plugins/Runtime/Metasound/Source/MetasoundStandardNodes/Private/MetasoundMidiScaleToArrayNode.cpp
@@ -47,7 +47,7 @@ namespace Metasound
 
 		// ctor
 		FMusicalScaleToNoteArrayOperator(
-			  const FOperatorSettings& InSettings
+			  const FCreateOperatorParams& InParams
 			, const FEnumMusicalScaleReadRef& InScale
 			, const FBoolReadRef& InChordTonesOnly
 		);
@@ -58,6 +58,7 @@ namespace Metasound
 		static TUniquePtr<IOperator> CreateOperator(const FCreateOperatorParams& InParams, FBuildErrorArray& OutErrors);
 		virtual FDataReferenceCollection GetInputs() const override;
 		virtual FDataReferenceCollection GetOutputs() const override;
+		void Reset(const IOperator::FResetParams& InParams);
 		void Execute();
 
 	private: // members
@@ -71,10 +72,6 @@ namespace Metasound
 		// cached values
 		Audio::EMusicalScale::Scale PreviousScale = Audio::EMusicalScale::Scale::Count;
 		bool bPreviousChordTones = false;
-		TArray<float> PreviousNoteOut;
-
-		// other
-		FOperatorSettings Settings;
 
 	}; // class FMusicalScaleToNoteArrayOperator
 
@@ -85,17 +82,16 @@ namespace Metasound
 
 	// ctor
 	FMusicalScaleToNoteArrayOperator::FMusicalScaleToNoteArrayOperator(
-		  const FOperatorSettings& InSettings
+		  const FCreateOperatorParams& InParams
 		, const FEnumMusicalScaleReadRef& InScale
 		, const FBoolReadRef& InChordTonesOnly
 		)
 		: Scale(InScale)
 		, bChordTonesOnly(InChordTonesOnly)
 		, ScaleArrayOutput(FArrayScaleDegreeWriteRef::CreateNew())
-		, Settings(InSettings)
 	{
 		// prime our output array
-		Execute();
+		Reset(InParams);
 	}
 
 
@@ -150,7 +146,7 @@ namespace Metasound
 		FBoolReadRef ChordTonesOnly = InputDataRefs.GetDataReadReferenceOrConstruct<bool>(METASOUND_GET_PARAM_NAME(ParamChordTonesOnly));
 
 		return MakeUnique <FMusicalScaleToNoteArrayOperator>(
-			  InParams.OperatorSettings
+			  InParams
 			, Scale
 			, ChordTonesOnly
 			);
@@ -174,6 +170,14 @@ namespace Metasound
 		return OutputDataReferences;
 	}
 
+	void FMusicalScaleToNoteArrayOperator::Reset(const IOperator::FResetParams& InParams)
+	{
+		PreviousScale = *Scale;
+		bPreviousChordTones = *bChordTonesOnly;
+
+		*ScaleArrayOutput = Audio::FMidiNoteQuantizer::ScaleDegreeSetMap[PreviousScale].GetScaleDegreeSet(bPreviousChordTones);
+	}
+
 	void FMusicalScaleToNoteArrayOperator::Execute()
 	{
 		// calculate new output and cache values if needed
@@ -183,7 +187,7 @@ namespace Metasound
 			PreviousScale = *Scale;
 			bPreviousChordTones = *bChordTonesOnly;
 
-			PreviousNoteOut = *ScaleArrayOutput = Audio::FMidiNoteQuantizer::ScaleDegreeSetMap[PreviousScale].GetScaleDegreeSet(bPreviousChordTones);
+			*ScaleArrayOutput = Audio::FMidiNoteQuantizer::ScaleDegreeSetMap[PreviousScale].GetScaleDegreeSet(bPreviousChordTones);
 		}
 	}
 
diff --git a/Engine/Plugins/Runtime/Metasound/Source/MetasoundStandardNodes/Private/MetasoundMinNode.cpp b/Engine/Plugins/Runtime/Metasound/Source/MetasoundStandardNodes/Private/MetasoundMinNode.cpp
--- a/Engine/Plugins/Runtime/Metasound/Source/MetasoundStandardNodes/Private/MetasoundMinNode.cpp
+++ b/Engine/Plugins/Runtime/Metasound/Source/MetasoundStandardNodes/Private/MetasoundMinNode.cpp
@@ -215,6 +215,11 @@ namespace Metasound
 			TMin<ValueType>::GetMin(*InputA, *InputB, *OutputValue);
 		}
 
+		void Reset(const IOperator::FResetParams& InParams)
+		{
+			GetMin();
+		}
+
 		void Execute()
 		{
 			GetMin();
diff --git a/Engine/Plugins/Runtime/Metasound/Source/MetasoundStandardNodes/Private/MetasoundMixerNode.cpp b/Engine/Plugins/Runtime/Metasound/Source/MetasoundStandardNodes/Private/MetasoundMixerNode.cpp
--- a/Engine/Plugins/Runtime/Metasound/Source/MetasoundStandardNodes/Private/MetasoundMixerNode.cpp
+++ b/Engine/Plugins/Runtime/Metasound/Source/MetasoundStandardNodes/Private/MetasoundMixerNode.cpp
@@ -25,24 +25,21 @@ namespace Metasound
 	{
 	public:
 		// ctor
-		TAudioMixerNodeOperator(const FOperatorSettings& InSettings, const TArray<FAudioBufferReadRef>&& InInputBuffers, const TArray<FFloatReadRef>&& InGainValues)
+		TAudioMixerNodeOperator(const FCreateOperatorParams& InParams, const TArray<FAudioBufferReadRef>&& InInputBuffers, const TArray<FFloatReadRef>&& InGainValues)
 			: Gains(InGainValues)
 			, Inputs (InInputBuffers)
-			, Settings(InSettings)
 		{
 			// create write refs
 			for (uint32 i = 0; i < NumChannels; ++i)
 			{
-				Outputs.Add(FAudioBufferWriteRef::CreateNew(InSettings));
+				Outputs.Add(FAudioBufferWriteRef::CreateNew(InParams.OperatorSettings));
 			}
 
 			// init previous gains to current values
 			PrevGains.Reset();
 			PrevGains.AddUninitialized(NumInputs);
-			for (uint32 i = 0; i < NumInputs; ++i)
-			{
-				PrevGains[i] = *Gains[i];
-			}
+
+			Reset(InParams);
 		}
 
 		// dtor
@@ -167,7 +164,7 @@ namespace Metasound
 				InputGains.Add(InputCollection.GetDataReadReferenceOrConstructWithVertexDefault<float>(InputInterface, GetGainInputName(i), InParams.OperatorSettings));
 			}
 
-			return MakeUnique<TAudioMixerNodeOperator<NumInputs, NumChannels>>(InParams.OperatorSettings, MoveTemp(InputBuffers), MoveTemp(InputGains));
+			return MakeUnique<TAudioMixerNodeOperator<NumInputs, NumChannels>>(InParams, MoveTemp(InputBuffers), MoveTemp(InputGains));
 		}
 
 		virtual FDataReferenceCollection GetInputs() const override
@@ -198,12 +195,22 @@ namespace Metasound
 			return OutputPins;
 		}
 
+		void Reset(const IOperator::FResetParams& InParams)
+		{
+			for (uint32 i = 0; i < NumInputs; ++i)
+			{
+				PrevGains[i] = *Gains[i];
+			}
+
+			Execute();
+		}
+
 		void Execute()
 		{
 			// zero the outputs
 			for (uint32 i = 0; i < NumChannels; ++i)
 			{
-				FMemory::Memzero(Outputs[i]->GetData(), sizeof(float) * Outputs[i]->Num());
+				Outputs[i]->Zero();
 			}
 
 			// for each input
@@ -216,11 +223,7 @@ namespace Metasound
 				for (uint32 ChanIndex = 0; ChanIndex < NumChannels; ++ChanIndex)
 				{
 					// Outputs[Chan] += Gains[i] * Inputs[i][Chan]
-					TArrayView<const float> InputView(Inputs[InputIndex * NumChannels + ChanIndex]->GetData(), Settings.GetNumFramesPerBlock());
-					TArrayView<float> OutputView(Outputs[ChanIndex]->GetData(), Settings.GetNumFramesPerBlock());
-
-					Audio::ArrayMixIn(InputView, OutputView, PrevGain, NextGain);
-
+					Audio::ArrayMixIn(*Inputs[InputIndex * NumChannels + ChanIndex], *Outputs[ChanIndex], PrevGain, NextGain);
 				}
 
 				PrevGains[InputIndex] = NextGain;
@@ -235,8 +238,6 @@ namespace Metasound
 
 		TArray<float> PrevGains;
 
-		FOperatorSettings Settings;
-
 		static FNodeClassMetadata CreateNodeClassMetadata(const FName& InOperatorName, const FText& InDisplayName, const FText& InDescription, const FVertexInterface& InDefaultInterface)
 		{
 			FNodeClassMetadata Metadata
diff --git a/Engine/Plugins/Runtime/Metasound/Source/MetasoundStandardNodes/Private/MetasoundOscillatorNodes.cpp b/Engine/Plugins/Runtime/Metasound/Source/MetasoundStandardNodes/Private/MetasoundOscillatorNodes.cpp
--- a/Engine/Plugins/Runtime/Metasound/Source/MetasoundStandardNodes/Private/MetasoundOscillatorNodes.cpp
+++ b/Engine/Plugins/Runtime/Metasound/Source/MetasoundStandardNodes/Private/MetasoundOscillatorNodes.cpp
@@ -80,7 +80,7 @@ namespace Metasound
 	// Base Oscillator Operator CRTP.
 	// Expects your Derived class to implement a Generate(int32 InStartFrame, int32 InEndFrame, float InFreq)
 	template<typename GeneratorPolicy, typename Derived>
-	class TOscillatorOperatorBase : public TExecutableOperator<Derived>
+	class TOscillatorOperatorBase : public TExecutableOperator<TOscillatorOperatorBase<GeneratorPolicy, Derived>>
 	{		
 	public:
 		TOscillatorOperatorBase(const FOscillatorOperatorConstructParams& InConstructParams)
@@ -129,6 +129,12 @@ namespace Metasound
 			Generator = GeneratorPolicy{ LinearPhase };
 		}
 
+		void Reset(const IOperator::FResetParams& InParams)
+		{
+			ResetPhase(*PhaseOffset);
+			AudioBuffer->Zero();
+		}
+
 		void Execute()
 		{
 			// Clamp frequencies into Nyquist range.
@@ -961,6 +967,12 @@ namespace Metasound
 			Phase = ClampedDegrees / 360.f;
 		}
 
+		void Reset(const IOperator::FResetParams& InParams)
+		{
+			ResetPhase();
+			*Output = 0.f;
+		}
+
 		void Execute()
 		{
 			using namespace Generators;
diff --git a/Engine/Plugins/Runtime/Metasound/Source/MetasoundStandardNodes/Private/MetasoundRingModNode.cpp b/Engine/Plugins/Runtime/Metasound/Source/MetasoundStandardNodes/Private/MetasoundRingModNode.cpp
--- a/Engine/Plugins/Runtime/Metasound/Source/MetasoundStandardNodes/Private/MetasoundRingModNode.cpp
+++ b/Engine/Plugins/Runtime/Metasound/Source/MetasoundStandardNodes/Private/MetasoundRingModNode.cpp
@@ -36,6 +36,7 @@ namespace Metasound
 			, AudioInputModulator(InputAudioModulator)
 			, AudioOutput(FAudioBufferWriteRef::CreateNew(InSettings))
 		{
+			AudioOutput->Zero();
 		}
 
 		static const FNodeClassMetadata& GetNodeInfo()
@@ -125,6 +126,11 @@ namespace Metasound
 			Audio::ArrayMultiplyInPlace(*AudioInputModulator, *AudioOutput);
 		}
 
+		void Reset(const IOperator::FResetParams& InParams)
+		{
+			AudioOutput->Zero();
+		}
+
 	private:
 
 		// The input audio buffer
diff --git a/Engine/Plugins/Runtime/Metasound/Source/MetasoundStandardNodes/Private/MetasoundSampleAndHoldNode.cpp b/Engine/Plugins/Runtime/Metasound/Source/MetasoundStandardNodes/Private/MetasoundSampleAndHoldNode.cpp
--- a/Engine/Plugins/Runtime/Metasound/Source/MetasoundStandardNodes/Private/MetasoundSampleAndHoldNode.cpp
+++ b/Engine/Plugins/Runtime/Metasound/Source/MetasoundStandardNodes/Private/MetasoundSampleAndHoldNode.cpp
@@ -84,16 +84,15 @@ namespace Metasound
 			FTriggerReadRef InputTriggerSampleAndHold = Inputs.GetDataReadReferenceOrConstruct<FTrigger>(METASOUND_GET_PARAM_NAME(InputTriggerSampleAndHold), InParams.OperatorSettings);
 			FAudioBufferReadRef InputAudio = Inputs.GetDataReadReferenceOrConstruct<FAudioBuffer>(METASOUND_GET_PARAM_NAME(InputAudio), InParams.OperatorSettings);
 
-			return MakeUnique<FSampleAndHoldOperator>(InParams.OperatorSettings, InputTriggerSampleAndHold, InputAudio);
+			return MakeUnique<FSampleAndHoldOperator>(InParams, InputTriggerSampleAndHold, InputAudio);
 		}
 
-		FSampleAndHoldOperator(const FOperatorSettings& InSettings, const FTriggerReadRef& InTriggerSampleAndHold, const FAudioBufferReadRef& InAudioInput)
+		FSampleAndHoldOperator(const FCreateOperatorParams& InParams, const FTriggerReadRef& InTriggerSampleAndHold, const FAudioBufferReadRef& InAudioInput)
 			: AudioInput(InAudioInput)
 			, TriggerSampleAndHold(InTriggerSampleAndHold)
-			, AudioOutput(FAudioBufferWriteRef::CreateNew(InSettings))
+			, AudioOutput(FAudioBufferWriteRef::CreateNew(InParams.OperatorSettings))
 		{
-			// Init the hold value to the first sample in the audio buffer
-			HoldValue = AudioInput->GetData()[0];
+			Reset(InParams);
 		}
 
 		virtual FDataReferenceCollection GetInputs() const override
@@ -119,37 +118,45 @@ namespace Metasound
 		
 		void Execute()
 		{
-			// Lambda is used for both pre and post trigger frames
-			auto SetToHoldLambda = [this](int32 StartFrame, int32 EndFrame)
-			{
-				check(AudioOutput->Num() >= EndFrame);
-				float* OutputBufferPtr = AudioOutput->GetData();
-				float* CurrOutputPtr = &OutputBufferPtr[StartFrame];
-
-				const int32 NumSamples = EndFrame - StartFrame;
-
-				// non-SIMD version
-				for (int32 i = 0; i < NumSamples; ++i)
-				{
-					CurrOutputPtr[i] = HoldValue;
-				}
-			};
-
 			TriggerSampleAndHold->ExecuteBlock(
 				// On Pre-Trigger, continue outputting the held value
-				SetToHoldLambda, 
+				[this](int32 StartFrame, int32 EndFrame)
+				{
+					SetToHold(StartFrame, EndFrame);
+				},
 				// On Trigger, get a new hold value, then output for the rest of the block
-				[this, &SetToHoldLambda](int32 StartFrame, int32 EndFrame)
+				[this](int32 StartFrame, int32 EndFrame)
 				{
 					// Get a new value to sample and hold
 					HoldValue = AudioInput->GetData()[StartFrame];
-					SetToHoldLambda(StartFrame, EndFrame);
+					SetToHold(StartFrame, EndFrame);
 				}
 			);
 		}
 
+		void Reset(const IOperator::FResetParams& InParams)
+		{
+			HoldValue = AudioInput->GetData()[0];
+			SetToHold(0, AudioOutput->Num());
+		}
+
 	private:
 
+		void SetToHold(int32 StartFrame, int32 EndFrame)
+		{
+			check(AudioOutput->Num() >= EndFrame);
+			float* OutputBufferPtr = AudioOutput->GetData();
+			float* CurrOutputPtr = &OutputBufferPtr[StartFrame];
+
+			const int32 NumSamples = EndFrame - StartFrame;
+
+			// non-SIMD version
+			for (int32 i = 0; i < NumSamples; ++i)
+			{
+				CurrOutputPtr[i] = HoldValue;
+			}
+		}
+
 		FAudioBufferReadRef AudioInput;
 		FTriggerReadRef TriggerSampleAndHold;
 
diff --git a/Engine/Plugins/Runtime/Metasound/Source/MetasoundStandardNodes/Private/MetasoundSemitoneToFreqMultiplierNode.cpp b/Engine/Plugins/Runtime/Metasound/Source/MetasoundStandardNodes/Private/MetasoundSemitoneToFreqMultiplierNode.cpp
--- a/Engine/Plugins/Runtime/Metasound/Source/MetasoundStandardNodes/Private/MetasoundSemitoneToFreqMultiplierNode.cpp
+++ b/Engine/Plugins/Runtime/Metasound/Source/MetasoundStandardNodes/Private/MetasoundSemitoneToFreqMultiplierNode.cpp
@@ -36,6 +36,7 @@ namespace Metasound
 		virtual FDataReferenceCollection GetInputs() const override;
 		virtual FDataReferenceCollection GetOutputs() const override;
 		void Execute();
+		void Reset(const IOperator::FResetParams& InParams);
 
 	private:
 		// The input difference in semitones
@@ -87,6 +88,12 @@ namespace Metasound
 		}
 	}
 
+	void FSemitoneToFrequencyMultiplierOperator::Reset(const IOperator::FResetParams& InParams)
+	{
+		*FrequencyMultiplierOutput = Audio::GetFrequencyMultiplier(*SemitoneInput);
+		PrevSemitone = *SemitoneInput;
+	}
+
 	const FVertexInterface& FSemitoneToFrequencyMultiplierOperator::GetVertexInterface()
 	{
 		using namespace SemitoneToFrequencyMultiplierVertexNames;
diff --git a/Engine/Plugins/Runtime/Metasound/Source/MetasoundStandardNodes/Private/MetasoundStereoDelayNode.cpp b/Engine/Plugins/Runtime/Metasound/Source/MetasoundStandardNodes/Private/MetasoundStereoDelayNode.cpp
--- a/Engine/Plugins/Runtime/Metasound/Source/MetasoundStandardNodes/Private/MetasoundStereoDelayNode.cpp
+++ b/Engine/Plugins/Runtime/Metasound/Source/MetasoundStandardNodes/Private/MetasoundStereoDelayNode.cpp
@@ -72,6 +72,7 @@ namespace Metasound
 		virtual FDataReferenceCollection GetInputs() const override;
 		virtual FDataReferenceCollection GetOutputs() const override;
 		void Execute();
+		void Reset(const IOperator::FResetParams& InParams);
 
 	private:
 		float GetInputDelayTimeMsecClamped() const;
@@ -285,6 +286,20 @@ namespace Metasound
 		}
 	}
 
+	void FStereoDelayOperator::Reset(const IOperator::FResetParams& InParams)
+	{
+		LeftAudioOutput->Zero();
+		RightAudioOutput->Zero();
+
+		PrevDelayTimeMsec = GetInputDelayTimeMsecClamped();
+		PrevDelayRatio = GetInputDelayRatioClamped();
+
+		LeftDelayBuffer.Reset();
+		LeftDelayBuffer.SetDelayMsec(PrevDelayTimeMsec * (1.0f + PrevDelayRatio));
+		RightDelayBuffer.Reset();
+		RightDelayBuffer.SetDelayMsec(PrevDelayTimeMsec * (1.0f - PrevDelayRatio));
+	}
+
 	const FVertexInterface& FStereoDelayOperator::GetVertexInterface()
 	{
 		using namespace StereoDelay; 
diff --git a/Engine/Plugins/Runtime/Metasound/Source/MetasoundStandardNodes/Private/MetasoundStereoPannerNode.cpp b/Engine/Plugins/Runtime/Metasound/Source/MetasoundStandardNodes/Private/MetasoundStereoPannerNode.cpp
--- a/Engine/Plugins/Runtime/Metasound/Source/MetasoundStandardNodes/Private/MetasoundStereoPannerNode.cpp
+++ b/Engine/Plugins/Runtime/Metasound/Source/MetasoundStandardNodes/Private/MetasoundStereoPannerNode.cpp
@@ -60,6 +60,7 @@ namespace Metasound
 		virtual FDataReferenceCollection GetInputs() const override;
 		virtual FDataReferenceCollection GetOutputs() const override;
 		void Execute();
+		void Reset(const IOperator::FResetParams& InParams);
 
 	private:
 		float GetInputDelayTimeMsec() const;
@@ -179,6 +180,15 @@ namespace Metasound
 		}
 	}
 
+	void FStereoPannerOperator::Reset(const IOperator::FResetParams& InParams)
+	{
+		AudioLeftOutput->Zero();
+		AudioRightOutput->Zero();
+
+		PrevPanningAmount = FMath::Clamp(*PanningAmount, -1.0f, 1.0f);
+		ComputePanGains(PrevPanningAmount, PrevLeftPan, PrevRightPan);
+	}
+
 	const FVertexInterface& FStereoPannerOperator::GetVertexInterface()
 	{
 		using namespace StereoPannerVertexNames;
diff --git a/Engine/Plugins/Runtime/Metasound/Source/MetasoundStandardNodes/Private/MetasoundSuperOscillatorNode.cpp b/Engine/Plugins/Runtime/Metasound/Source/MetasoundStandardNodes/Private/MetasoundSuperOscillatorNode.cpp
--- a/Engine/Plugins/Runtime/Metasound/Source/MetasoundStandardNodes/Private/MetasoundSuperOscillatorNode.cpp
+++ b/Engine/Plugins/Runtime/Metasound/Source/MetasoundStandardNodes/Private/MetasoundSuperOscillatorNode.cpp
@@ -207,59 +207,14 @@
 			, DetuneDb(InConstructParams.DetunedVoiceVolumeDb)
 			, GlideFactor(InConstructParams.GlideFactor)
 			, PulseWidth(InConstructParams.PulseWidth)
+			, bUseModulation(InConstructParams.bUseModulation)
 			, Fm(InConstructParams.Fm)
 			, OscType(InConstructParams.OscType)
 		{
-			using namespace Generators;
 
-			const int32 NumVoicesToInit = NumVoices;
 			ScratchBuffer.SetNumZeroed(InConstructParams.Settings.GetNumFramesPerBlock());
 			
-			// init generators
-			if (InConstructParams.bUseModulation)
-			{
-				// fm generators
-				switch (OscType)
-				{
-				case ELfoWaveshapeType::Sine:
-					FillGeneratorsArray<TVoiceGenerator<FSineWaveTableWithFm>>(Generators, NumVoicesToInit, SampleRate);
-					break;
-				case ELfoWaveshapeType::Triangle:
-					FillGeneratorsArray<TVoiceGenerator<FTrianglePolysmoothWithFm>>(Generators, NumVoicesToInit, SampleRate);
-					break;
-				case ELfoWaveshapeType::Square:
-					FillGeneratorsArray<TVoiceGenerator<FSquarePolysmoothWithFm>>(Generators, NumVoicesToInit, SampleRate);
-					break;
-				case ELfoWaveshapeType::Saw:
-					FillGeneratorsArray<TVoiceGenerator<FSawPolysmoothWithFm>>(Generators, NumVoicesToInit, SampleRate);
-					break;
-				default: 
-					check(!"Unsupported Wave Shape in FSuperOscillatorOperator");
-					return;
-				}
-			}
-			else
-			{
-				//non-fm generators
-				switch (OscType)
-				{
-				case ELfoWaveshapeType::Sine:
-					FillGeneratorsArray<TVoiceGenerator<FSineWaveTable>>(Generators, NumVoicesToInit, SampleRate);
-					break;
-				case ELfoWaveshapeType::Triangle:
-					FillGeneratorsArray<TVoiceGenerator<FTrianglePolysmooth>>(Generators, NumVoicesToInit, SampleRate);
-					break;
-				case ELfoWaveshapeType::Square:
-					FillGeneratorsArray<TVoiceGenerator<FSquarePolysmooth>>(Generators, NumVoicesToInit, SampleRate);
-					break;
-				case ELfoWaveshapeType::Saw:
-					FillGeneratorsArray<TVoiceGenerator<FSawPolysmooth>>(Generators, NumVoicesToInit, SampleRate);
-					break;
-				default: 
-					check(!"Unsupported Wave Shape in FSuperOscillatorOperator");
-					return;
-				}
-			}
+			InitializeGenerators();
 		}
 		
 		virtual void Bind(FVertexInterfaceData& InVertexData) const override
@@ -300,6 +255,7 @@
 		virtual void FillOutputs(const float Gain = 1.f, const float LinearPan = 0.f) {}
 		virtual void ZeroOutputs() {}
 		virtual float GetStereoWidth() const { return 0.f; }
+		virtual void ResetLimiters() = 0;
 		
 		void SetLimiterSettings(Audio::FDynamicsProcessor& InLimiter) const
 		{
@@ -379,11 +335,71 @@
 			}
 		}
 
+		void Reset(const IOperator::FResetParams& InParams)
+		{
+			ZeroOutputs();
+			ResetLimiters();
+
+			InitializeGenerators();
+		}
+
+	protected:
+
+		void InitializeGenerators()
+		{
+			using namespace Generators;
+			// init generators
+			if (bUseModulation)
+			{
+				// fm generators
+				switch (OscType)
+				{
+				case ELfoWaveshapeType::Sine:
+					FillGeneratorsArray<TVoiceGenerator<FSineWaveTableWithFm>>(Generators, NumVoices, SampleRate);
+					break;
+				case ELfoWaveshapeType::Triangle:
+					FillGeneratorsArray<TVoiceGenerator<FTrianglePolysmoothWithFm>>(Generators, NumVoices, SampleRate);
+					break;
+				case ELfoWaveshapeType::Square:
+					FillGeneratorsArray<TVoiceGenerator<FSquarePolysmoothWithFm>>(Generators, NumVoices, SampleRate);
+					break;
+				case ELfoWaveshapeType::Saw:
+					FillGeneratorsArray<TVoiceGenerator<FSawPolysmoothWithFm>>(Generators, NumVoices, SampleRate);
+					break;
+				default: 
+					check(!"Unsupported Wave Shape in FSuperOscillatorOperator");
+					return;
+				}
+			}
+			else
+			{
+				//non-fm generators
+				switch (OscType)
+				{
+				case ELfoWaveshapeType::Sine:
+					FillGeneratorsArray<TVoiceGenerator<FSineWaveTable>>(Generators, NumVoices, SampleRate);
+					break;
+				case ELfoWaveshapeType::Triangle:
+					FillGeneratorsArray<TVoiceGenerator<FTrianglePolysmooth>>(Generators, NumVoices, SampleRate);
+					break;
+				case ELfoWaveshapeType::Square:
+					FillGeneratorsArray<TVoiceGenerator<FSquarePolysmooth>>(Generators, NumVoices, SampleRate);
+					break;
+				case ELfoWaveshapeType::Saw:
+					FillGeneratorsArray<TVoiceGenerator<FSawPolysmooth>>(Generators, NumVoices, SampleRate);
+					break;
+				default: 
+					check(!"Unsupported Wave Shape in FSuperOscillatorOperator");
+					return;
+				}
+			}
+		}
+
 		float SampleRate = 1.f;
 		float Nyquist = 1.f;
 
 		// apply this gain whether we are limiting or not, to avoid large volume jumps when toggling the limiter
-		const float PreLimiterGain = 0.707f;
+		static constexpr float PreLimiterGain = 0.707f;
 
 		FBoolReadRef bEnabled;
 		FBoolReadRef bLimit;
@@ -395,6 +411,7 @@
 		FFloatReadRef GlideFactor;
 		FFloatReadRef PulseWidth;
 
+		bool bUseModulation = false;
 		FAudioBufferReadRef Fm;
 		FEnumLfoWaveshapeType OscType;
 
@@ -504,6 +521,11 @@
 			Limiter.ProcessAudio(AudioBuffer->GetData(), AudioBuffer->Num(), AudioBuffer->GetData());
 		}
 
+		virtual void ResetLimiters() override
+		{
+			SetLimiterSettings(Limiter);
+		}
+
 		virtual void Bind(FVertexInterfaceData& InVertexData) const override
 		{
 			using namespace SuperOscillatorVertexNames;
@@ -630,6 +652,12 @@
 			Audio::ArrayMixIn(ScratchBuffer, *AudioRight, Gain * PanRightGain * PreLimiterGain);
 		}
 
+		virtual void ResetLimiters() override
+		{
+			SetLimiterSettings(LimiterLeft);
+			SetLimiterSettings(LimiterRight);
+		}
+
 		virtual void LimitOutput() override
 		{
 			LimiterLeft.ProcessAudio(AudioLeft->GetData(), AudioLeft->Num(), AudioLeft->GetData());
diff --git a/Engine/Plugins/Runtime/Metasound/Source/MetasoundStandardNodes/Private/MetasoundSwitchAndRampNode.cpp b/Engine/Plugins/Runtime/Metasound/Source/MetasoundStandardNodes/Private/MetasoundSwitchAndRampNode.cpp
--- a/Engine/Plugins/Runtime/Metasound/Source/MetasoundStandardNodes/Private/MetasoundSwitchAndRampNode.cpp
+++ b/Engine/Plugins/Runtime/Metasound/Source/MetasoundStandardNodes/Private/MetasoundSwitchAndRampNode.cpp
@@ -168,6 +168,15 @@ namespace Metasound
 			);
 		}
 
+		void Reset(const IOperator::FResetParams& InParams)
+		{
+			AudioOut->Zero();
+			MostRecentOutputValue = 0.0f;
+			DiscontinuityAmount = 0.0f;
+			NumRampSamples = 0;
+			RampSampleIndex = 0;
+		}
+
 	private:
 		FTriggerReadRef TriggerIn;
 		
diff --git a/Engine/Plugins/Runtime/Metasound/Source/MetasoundStandardNodes/Private/MetasoundTriggerAnyNode.cpp b/Engine/Plugins/Runtime/Metasound/Source/MetasoundStandardNodes/Private/MetasoundTriggerAnyNode.cpp
--- a/Engine/Plugins/Runtime/Metasound/Source/MetasoundStandardNodes/Private/MetasoundTriggerAnyNode.cpp
+++ b/Engine/Plugins/Runtime/Metasound/Source/MetasoundStandardNodes/Private/MetasoundTriggerAnyNode.cpp
@@ -178,6 +178,11 @@ namespace Metasound
 			}
 		}
 
+		void Reset(const IOperator::FResetParams& InParams)
+		{
+			OutputTrigger->Reset();
+		}
+
 	private:
 
 		TArray<FTriggerReadRef> InputTriggers;
diff --git a/Engine/Plugins/Runtime/Metasound/Source/MetasoundStandardNodes/Private/MetasoundTriggerCoinNode.cpp b/Engine/Plugins/Runtime/Metasound/Source/MetasoundStandardNodes/Private/MetasoundTriggerCoinNode.cpp
--- a/Engine/Plugins/Runtime/Metasound/Source/MetasoundStandardNodes/Private/MetasoundTriggerCoinNode.cpp
+++ b/Engine/Plugins/Runtime/Metasound/Source/MetasoundStandardNodes/Private/MetasoundTriggerCoinNode.cpp
@@ -48,6 +48,7 @@ namespace Metasound
 		virtual FDataReferenceCollection GetOutputs() const override;
 
 		void Execute();
+		void Reset(const IOperator::FResetParams& InParams);
 
 	private:
 		// Generate new random number
@@ -156,6 +157,17 @@ namespace Metasound
 
 	}
 
+	void FTriggerCoinOperator::Reset(const IOperator::FResetParams& InParams)
+	{
+		// Trigger output
+		TriggerTrueOutput->Reset();
+		TriggerFalseOutput->Reset();
+
+		bIsDefaultSeeded = false;
+		bIsRandomStreamInitialized = false;
+		EvaluateSeedChanges();
+		RandomStream.Reset();
+	}
 
 	TUniquePtr<IOperator> FTriggerCoinOperator::CreateOperator(const FCreateOperatorParams& InParams, FBuildErrorArray& OutErrors)
 	{
diff --git a/Engine/Plugins/Runtime/Metasound/Source/MetasoundStandardNodes/Private/MetasoundTriggerControlNode.cpp b/Engine/Plugins/Runtime/Metasound/Source/MetasoundStandardNodes/Private/MetasoundTriggerControlNode.cpp
--- a/Engine/Plugins/Runtime/Metasound/Source/MetasoundStandardNodes/Private/MetasoundTriggerControlNode.cpp
+++ b/Engine/Plugins/Runtime/Metasound/Source/MetasoundStandardNodes/Private/MetasoundTriggerControlNode.cpp
@@ -50,6 +50,7 @@ namespace Metasound
 		virtual FDataReferenceCollection GetOutputs() const override;
 
 		void Execute();
+		void Reset(const IOperator::FResetParams& InParams);
 
 	private:
 		// Try to go through gate
@@ -157,9 +158,13 @@ namespace Metasound
 				}
 			}
 		);		
-
 	}
 
+	void FTriggerControlOperator::Reset(const IOperator::FResetParams& InParams)
+	{
+		TriggerExitOutput->Reset();
+		bIsGateOpen = !(*bStartClosedInput);
+	}
 
 	TUniquePtr<IOperator> FTriggerControlOperator::CreateOperator(const FCreateOperatorParams& InParams, FBuildErrorArray& OutErrors)
 	{
diff --git a/Engine/Plugins/Runtime/Metasound/Source/MetasoundStandardNodes/Private/MetasoundTriggerCounterNode.cpp b/Engine/Plugins/Runtime/Metasound/Source/MetasoundStandardNodes/Private/MetasoundTriggerCounterNode.cpp
--- a/Engine/Plugins/Runtime/Metasound/Source/MetasoundStandardNodes/Private/MetasoundTriggerCounterNode.cpp
+++ b/Engine/Plugins/Runtime/Metasound/Source/MetasoundStandardNodes/Private/MetasoundTriggerCounterNode.cpp
@@ -50,6 +50,7 @@ namespace Metasound
 			virtual FDataReferenceCollection GetOutputs() const override;
 
 			void Execute();
+			void Reset(const IOperator::FResetParams& InParams);
 
 		private:
 			FTriggerReadRef TriggerIn;
@@ -171,6 +172,19 @@ namespace Metasound
 
 	}
 
+	void FTriggerCounterOperator::Reset(const IOperator::FResetParams& InParams)
+	{
+		TriggerOut->Reset();
+		TriggerOnReset->Reset();
+		*OutValue = *StartValue;
+		*OutCount = 0;
+
+		CurrentAutoResetCount = FMath::Max(0, *AutoResetCount);
+		CurrentTriggerCount = 0;
+		CurrentValue = 0.0f;
+		bIsFirstTrigger = true;
+	}
+
 	TUniquePtr<IOperator> FTriggerCounterOperator::CreateOperator(const FCreateOperatorParams& InParams, FBuildErrorArray& OutErrors)
 	{
 		using namespace TriggerCounterVertexNames;
diff --git a/Engine/Plugins/Runtime/Metasound/Source/MetasoundStandardNodes/Private/MetasoundTriggerDelayNode.cpp b/Engine/Plugins/Runtime/Metasound/Source/MetasoundStandardNodes/Private/MetasoundTriggerDelayNode.cpp
--- a/Engine/Plugins/Runtime/Metasound/Source/MetasoundStandardNodes/Private/MetasoundTriggerDelayNode.cpp
+++ b/Engine/Plugins/Runtime/Metasound/Source/MetasoundStandardNodes/Private/MetasoundTriggerDelayNode.cpp
@@ -39,6 +39,7 @@ namespace Metasound
 			virtual FDataReferenceCollection GetOutputs() const override;
 
 			void Execute();
+			void Reset(const IOperator::FResetParams& InParams);
 
 		private:
 			float SampleRate;
@@ -106,6 +107,11 @@ namespace Metasound
 		);
 	}
 
+	void FTriggerDelayOperator::Reset(const IOperator::FResetParams& InParams)
+	{
+		TriggerOut->Reset();
+	}
+
 	TUniquePtr<IOperator> FTriggerDelayOperator::CreateOperator(const FCreateOperatorParams& InParams, FBuildErrorArray& OutErrors)
 	{
 		using namespace TriggerDelayVertexNames;
diff --git a/Engine/Plugins/Runtime/Metasound/Source/MetasoundStandardNodes/Private/MetasoundTriggerOnThresholdNode.cpp b/Engine/Plugins/Runtime/Metasound/Source/MetasoundStandardNodes/Private/MetasoundTriggerOnThresholdNode.cpp
--- a/Engine/Plugins/Runtime/Metasound/Source/MetasoundStandardNodes/Private/MetasoundTriggerOnThresholdNode.cpp
+++ b/Engine/Plugins/Runtime/Metasound/Source/MetasoundStandardNodes/Private/MetasoundTriggerOnThresholdNode.cpp
@@ -235,6 +235,13 @@ namespace Metasound
 			}
 		}
 
+		void Reset(const IOperator::FResetParams& InParams)
+		{
+			Out->Reset();
+			bTriggered = false;
+			LastSample = 0;
+		}
+
 		FDataReferenceCollection GetInputs() const
 		{
 			using namespace TriggerOnThresholdVertexNames;
diff --git a/Engine/Plugins/Runtime/Metasound/Source/MetasoundStandardNodes/Private/MetasoundTriggerOnValueChangeNode.cpp b/Engine/Plugins/Runtime/Metasound/Source/MetasoundStandardNodes/Private/MetasoundTriggerOnValueChangeNode.cpp
--- a/Engine/Plugins/Runtime/Metasound/Source/MetasoundStandardNodes/Private/MetasoundTriggerOnValueChangeNode.cpp
+++ b/Engine/Plugins/Runtime/Metasound/Source/MetasoundStandardNodes/Private/MetasoundTriggerOnValueChangeNode.cpp
@@ -56,6 +56,7 @@ namespace Metasound
 		virtual FDataReferenceCollection GetOutputs() const override;
 
 		void Execute();
+		void Reset(const IOperator::FResetParams& InParams);
 
 	private:
 		// Parameter to watch for changes
@@ -112,6 +113,12 @@ namespace Metasound
 		}
 
 	}
+	template <typename ValueType>
+	void TTriggerOnValueChangeOperator<ValueType>::Reset(const IOperator::FResetParams& InParams)
+	{
+		TriggerOnChangeOutput->Reset();
+		PrevValue = *ValueInput;
+	}
 
 	template <typename ValueType>
 	TUniquePtr<IOperator> TTriggerOnValueChangeOperator<ValueType>::CreateOperator(const FCreateOperatorParams& InParams, FBuildErrorArray& OutErrors)
diff --git a/Engine/Plugins/Runtime/Metasound/Source/MetasoundStandardNodes/Private/MetasoundTriggerOnceNode.cpp b/Engine/Plugins/Runtime/Metasound/Source/MetasoundStandardNodes/Private/MetasoundTriggerOnceNode.cpp
--- a/Engine/Plugins/Runtime/Metasound/Source/MetasoundStandardNodes/Private/MetasoundTriggerOnceNode.cpp
+++ b/Engine/Plugins/Runtime/Metasound/Source/MetasoundStandardNodes/Private/MetasoundTriggerOnceNode.cpp
@@ -46,6 +46,7 @@ namespace Metasound
 		virtual FDataReferenceCollection GetOutputs() const override;
 
 		void Execute();
+		void Reset(const IOperator::FResetParams& InParams);
 
 	private:
 		// Try to go through gate
@@ -126,6 +127,11 @@ namespace Metasound
 
 	}
 
+	void FTriggerOnceOperator::Reset(const IOperator::FResetParams& InParams)
+	{
+		TriggerExitOutput->Reset();
+		bIsGateOpen = !(*bStartClosedInput);
+	}
 
 	TUniquePtr<IOperator> FTriggerOnceOperator::CreateOperator(const FCreateOperatorParams& InParams, FBuildErrorArray& OutErrors)
 	{
diff --git a/Engine/Plugins/Runtime/Metasound/Source/MetasoundStandardNodes/Private/MetasoundTriggerPipeNode.cpp b/Engine/Plugins/Runtime/Metasound/Source/MetasoundStandardNodes/Private/MetasoundTriggerPipeNode.cpp
--- a/Engine/Plugins/Runtime/Metasound/Source/MetasoundStandardNodes/Private/MetasoundTriggerPipeNode.cpp
+++ b/Engine/Plugins/Runtime/Metasound/Source/MetasoundStandardNodes/Private/MetasoundTriggerPipeNode.cpp
@@ -41,6 +41,7 @@ namespace Metasound
 		virtual FDataReferenceCollection GetInputs() const override;
 		virtual FDataReferenceCollection GetOutputs() const override;
 		void Execute();
+		void Reset(const IOperator::FResetParams& InParams);
 
 	private:
 		TArray<FSampleCount> SamplesUntilTrigger;
@@ -135,6 +136,12 @@ namespace Metasound
 		}
 	}
 
+	void FTriggerPipeOperator::Reset(const IOperator::FResetParams& InParams)
+	{
+		SamplesUntilTrigger.Reset();
+		TriggerOut->Reset();
+	}
+
 	TUniquePtr<IOperator> FTriggerPipeOperator::CreateOperator(const FCreateOperatorParams& InParams, FBuildErrorArray& OutErrors)
 	{
 		using namespace TriggerPipeVertexNames;
diff --git a/Engine/Plugins/Runtime/Metasound/Source/MetasoundStandardNodes/Private/MetasoundTriggerRepeatNode.cpp b/Engine/Plugins/Runtime/Metasound/Source/MetasoundStandardNodes/Private/MetasoundTriggerRepeatNode.cpp
--- a/Engine/Plugins/Runtime/Metasound/Source/MetasoundStandardNodes/Private/MetasoundTriggerRepeatNode.cpp
+++ b/Engine/Plugins/Runtime/Metasound/Source/MetasoundStandardNodes/Private/MetasoundTriggerRepeatNode.cpp
@@ -46,6 +46,7 @@ namespace Metasound
 			virtual FDataReferenceCollection GetOutputs() const override;
 
 			void Execute();
+			void Reset(const IOperator::FResetParams& InParams);
 
 		private:
 			bool bEnabled;
@@ -131,6 +132,13 @@ namespace Metasound
 		}
 	}
 
+	void FTriggerRepeatOperator::Reset(const IOperator::FResetParams& InParams)
+	{
+		bEnabled = false;
+
+		TriggerOut->Reset();
+		SampleCounter.SetNumSamples(0);
+	}
 
 	TUniquePtr<IOperator> FTriggerRepeatOperator::CreateOperator(const FCreateOperatorParams& InParams, FBuildErrorArray& OutErrors)
 	{
diff --git a/Engine/Plugins/Runtime/Metasound/Source/MetasoundStandardNodes/Private/MetasoundTriggerRouteNode.cpp b/Engine/Plugins/Runtime/Metasound/Source/MetasoundStandardNodes/Private/MetasoundTriggerRouteNode.cpp
--- a/Engine/Plugins/Runtime/Metasound/Source/MetasoundStandardNodes/Private/MetasoundTriggerRouteNode.cpp
+++ b/Engine/Plugins/Runtime/Metasound/Source/MetasoundStandardNodes/Private/MetasoundTriggerRouteNode.cpp
@@ -160,33 +160,20 @@ namespace Metasound
 				InputValues.Add(InputCollection.GetDataReadReferenceOrConstructWithVertexDefault<ValueType>(InputInterface, GetInputValueName(i), InParams.OperatorSettings));
 			}
 
-			return MakeUnique<TTriggerRouteOperator<ValueType, NumInputs>>(InParams.OperatorSettings, MoveTemp(InputTriggers), MoveTemp(InputValues));
+			return MakeUnique<TTriggerRouteOperator<ValueType, NumInputs>>(InParams, MoveTemp(InputTriggers), MoveTemp(InputValues));
 		}
 
-		TTriggerRouteOperator(const FOperatorSettings& InSettings, TArray<FTriggerReadRef>&& InInputTriggers, TArray<TDataReadReference<ValueType>>&& InInputValues)
+		TTriggerRouteOperator(const FCreateOperatorParams& InParams, TArray<FTriggerReadRef>&& InInputTriggers, TArray<TDataReadReference<ValueType>>&& InInputValues)
 			: InputTriggers(MoveTemp(InInputTriggers))
 			, InputValues(MoveTemp(InInputValues))
-			, OutputTrigger(FTriggerWriteRef::CreateNew(InSettings))
-			, OutputValue(TDataWriteReferenceFactory<ValueType>::CreateAny(InSettings))
+			, OutputTrigger(FTriggerWriteRef::CreateNew(InParams.OperatorSettings))
+			, OutputValue(TDataWriteReferenceFactory<ValueType>::CreateAny(InParams.OperatorSettings))
 		{
-			check(InputValues.Num() > 0)
-			CurrentIndex = 0;
+			check(InputValues.Num() > 0);
 			
-			// Do the initial trigger input to determine the init index
-			for (uint32 i = 0; i < NumInputs; ++i)
-			{
-				InputTriggers[i]->ExecuteBlock(
-					[&](int32 StartFrame, int32 EndFrame)
-					{
-					},
-					[this, i](int32 StartFrame, int32 EndFrame)
-					{
-						CurrentIndex = i;
-					}
-				);
-			}
-			
-			*OutputValue = *InputValues[CurrentIndex];
+			// Call Update() to initialize output value and determine initial
+			// trigger index.
+			Update();
 		}
 
 		virtual ~TTriggerRouteOperator() = default;
@@ -220,7 +207,20 @@ namespace Metasound
 		void Execute()
 		{
 			OutputTrigger->AdvanceBlock();
+			Update();
+		}
+
+		void Reset(const IOperator::FResetParams& InParams)
+		{
+			OutputTrigger->Reset();
+			CurrentIndex = 0;
+			Update();
+		}
+
+	private:
 
+		void Update()
+		{
 			for (uint32 i = 0; i < NumInputs; ++i)
 			{
 				InputTriggers[i]->ExecuteBlock(
@@ -238,8 +238,6 @@ namespace Metasound
 			*OutputValue = *InputValues[CurrentIndex];
 		}
 
-	private:
-
 		TArray<FTriggerReadRef> InputTriggers;
 		TArray<TDataReadReference<ValueType>> InputValues;
 
diff --git a/Engine/Plugins/Runtime/Metasound/Source/MetasoundStandardNodes/Private/MetasoundTriggerSelectNode.cpp b/Engine/Plugins/Runtime/Metasound/Source/MetasoundStandardNodes/Private/MetasoundTriggerSelectNode.cpp
--- a/Engine/Plugins/Runtime/Metasound/Source/MetasoundStandardNodes/Private/MetasoundTriggerSelectNode.cpp
+++ b/Engine/Plugins/Runtime/Metasound/Source/MetasoundStandardNodes/Private/MetasoundTriggerSelectNode.cpp
@@ -49,13 +49,13 @@ namespace Metasound
 		virtual FDataReferenceCollection GetOutputs() const override;
 
 		void Execute();
+		void Reset(const IOperator::FResetParams& InParams);
 
 	private:
 		FTriggerReadRef TriggerInput;
 		FInt32ReadRef IndexInput;
 
 		TArray<FTriggerWriteRef> TriggerOutputs;
-
 	};
 
 	template<uint32 NumOutputs>
@@ -65,7 +65,6 @@ namespace Metasound
 		: TriggerInput(InTrigger)
 		, IndexInput(InIndex)
 	{
-
 		for (uint32 i = 0; i < NumOutputs; ++i)
 		{
 			TriggerOutputs.Add(FTriggerWriteRef::CreateNew(InSettings));
@@ -122,6 +121,15 @@ namespace Metasound
 			);
 	}
 
+	template<uint32 NumOutputs>
+	void TTriggerSelectOperator<NumOutputs>::Reset(const IOperator::FResetParams& InParams)
+	{
+		for (uint32 i = 0; i < NumOutputs; ++i)
+		{
+			TriggerOutputs[i]->Reset();
+		}
+	}
+
 	template<uint32 NumOutputs>
 	TUniquePtr<IOperator> TTriggerSelectOperator<NumOutputs>::CreateOperator(const FCreateOperatorParams& InParams, FBuildErrorArray& OutErrors)
 	{
diff --git a/Engine/Plugins/Runtime/Metasound/Source/MetasoundStandardNodes/Private/MetasoundTriggerSequenceNode.cpp b/Engine/Plugins/Runtime/Metasound/Source/MetasoundStandardNodes/Private/MetasoundTriggerSequenceNode.cpp
--- a/Engine/Plugins/Runtime/Metasound/Source/MetasoundStandardNodes/Private/MetasoundTriggerSequenceNode.cpp
+++ b/Engine/Plugins/Runtime/Metasound/Source/MetasoundStandardNodes/Private/MetasoundTriggerSequenceNode.cpp
@@ -52,6 +52,7 @@ namespace Metasound
 		virtual FDataReferenceCollection GetOutputs() const override;
 
 		void Execute();
+		void Reset(const IOperator::FResetParams& InParams);
 
 	private:
 		FTriggerReadRef TriggerInput;
@@ -149,6 +150,16 @@ namespace Metasound
 		);
 	}
 
+	template<uint32 NumOutputs>
+	void TTriggerSequenceOperator<NumOutputs>::Reset(const IOperator::FResetParams& InParams)
+	{
+		for (uint32 i = 0; i < NumOutputs; ++i)
+		{
+			TriggerOutputs[i]->Reset();
+		}
+		CurrentIndex = 0;
+	}
+
 	template<uint32 NumOutputs>
 	TUniquePtr<IOperator> TTriggerSequenceOperator<NumOutputs>::CreateOperator(const FCreateOperatorParams& InParams, FBuildErrorArray& OutErrors)
 	{
diff --git a/Engine/Plugins/Runtime/Metasound/Source/MetasoundStandardNodes/Private/MetasoundTriggerToggleNode.cpp b/Engine/Plugins/Runtime/Metasound/Source/MetasoundStandardNodes/Private/MetasoundTriggerToggleNode.cpp
--- a/Engine/Plugins/Runtime/Metasound/Source/MetasoundStandardNodes/Private/MetasoundTriggerToggleNode.cpp
+++ b/Engine/Plugins/Runtime/Metasound/Source/MetasoundStandardNodes/Private/MetasoundTriggerToggleNode.cpp
@@ -43,6 +43,7 @@ namespace Metasound
 			virtual FDataReferenceCollection GetOutputs() const override;
 
 			void Execute();
+			void Reset(const IOperator::FResetParams& InParams);
 
 		private:
 
@@ -111,6 +112,12 @@ namespace Metasound
 		);
 	}
 
+	void FTriggerToggleOperator::Reset(const IOperator::FResetParams& InParams)
+	{
+		TriggerOutput->Reset();
+		*ValueOutput = *InitValue;
+	}
+
 	TUniquePtr<IOperator> FTriggerToggleOperator::CreateOperator(const FCreateOperatorParams& InParams, FBuildErrorArray& OutErrors)
 	{
 		using namespace TriggerToggle;
diff --git a/Engine/Plugins/Runtime/Metasound/Source/MetasoundStandardNodes/Private/MetasoundWaveShaperNode.cpp b/Engine/Plugins/Runtime/Metasound/Source/MetasoundStandardNodes/Private/MetasoundWaveShaperNode.cpp
--- a/Engine/Plugins/Runtime/Metasound/Source/MetasoundStandardNodes/Private/MetasoundWaveShaperNode.cpp
+++ b/Engine/Plugins/Runtime/Metasound/Source/MetasoundStandardNodes/Private/MetasoundWaveShaperNode.cpp
@@ -55,6 +55,7 @@ namespace Metasound
 		virtual FDataReferenceCollection GetInputs() const override;
 		virtual FDataReferenceCollection GetOutputs() const override;
 		void Execute();
+		void Reset(const IOperator::FResetParams& InParams);
 
 	private:
 		// The input audio buffer
@@ -133,6 +134,12 @@ namespace Metasound
 		WaveShaper.ProcessAudioBuffer(InputAudio, OutputAudio, NumFrames);
 	}
 
+	void FWaveShaperOperator::Reset(const IOperator::FResetParams& InParams)
+	{
+		AudioOutput->Zero();
+		WaveShaper.Init(InParams.OperatorSettings.GetSampleRate());
+	}
+
 	const FVertexInterface& FWaveShaperOperator::GetVertexInterface()
 	{
 		static const FVertexInterface Interface(
diff --git a/Engine/Plugins/Runtime/Metasound/Source/MetasoundStandardNodes/Private/MetasoundWaveWriterNode.cpp b/Engine/Plugins/Runtime/Metasound/Source/MetasoundStandardNodes/Private/MetasoundWaveWriterNode.cpp
--- a/Engine/Plugins/Runtime/Metasound/Source/MetasoundStandardNodes/Private/MetasoundWaveWriterNode.cpp
+++ b/Engine/Plugins/Runtime/Metasound/Source/MetasoundStandardNodes/Private/MetasoundWaveWriterNode.cpp
@@ -371,6 +371,14 @@ namespace Metasound
 			}
 		}
 
+		void Reset(const IOperator::FResetParams& InParams)
+		{
+			if (bIsEnabled)
+			{
+				Disable();
+			}
+		}
+
 	protected:
 		static const FVertexName GetAudioInputName(int32 InInputIndex)
 		{
diff --git a/Engine/Plugins/Runtime/Metasound/Source/MetasoundStandardNodes/Public/MetasoundRandomNode.h b/Engine/Plugins/Runtime/Metasound/Source/MetasoundStandardNodes/Public/MetasoundRandomNode.h
--- a/Engine/Plugins/Runtime/Metasound/Source/MetasoundStandardNodes/Public/MetasoundRandomNode.h
+++ b/Engine/Plugins/Runtime/Metasound/Source/MetasoundStandardNodes/Public/MetasoundRandomNode.h
@@ -361,11 +361,7 @@ namespace Metasound
 			, bIsDefaultSeeded(*SeedValue == DefaultSeed)
 			, bIsRandomStreamInitialized(false)
 		{
-			EvaluateSeedChanges();
-			RandomStream.Reset();
-
-			// We need to initialize the output value to *something*
-			*OutputValue = TRandomNodeSpecialization<ValueType>::GetNextValue(RandomStream, *MinValue, *MaxValue);
+			ResetInternal();
 		}
 
 		TRandomNodeOperator(const FOperatorSettings& InSettings,
@@ -383,11 +379,7 @@ namespace Metasound
 			, bIsDefaultSeeded(*SeedValue == DefaultSeed)
 			, bIsRandomStreamInitialized(false)
 		{
-			EvaluateSeedChanges();
-			RandomStream.Reset();
-
-			// We need to initialize the output value to *something*
-			*OutputValue = TRandomNodeSpecialization<ValueType>::GetNextValue(RandomStream, *MinValue, *MaxValue);
+			ResetInternal();
 		}
 
 		virtual ~TRandomNodeOperator() = default;
@@ -451,7 +443,27 @@ namespace Metasound
 			);
 		}
 
+		// Externally visible initialize function
+		void Reset(const IOperator::FResetParams& InParams)
+		{
+			ResetInternal();
+		}
+
 	private:
+		void ResetInternal()
+		{
+			TriggerOutOnNext->Reset();
+			TriggerOutOnReset->Reset();
+			bIsDefaultSeeded = (DefaultSeed == *SeedValue);
+			bIsRandomStreamInitialized = false;
+
+			EvaluateSeedChanges();
+			RandomStream.Reset();
+
+			// We need to initialize the output value to *something*
+			*OutputValue = TRandomNodeSpecialization<ValueType>::GetNextValue(RandomStream, *MinValue, *MaxValue);
+		}
+
 		void EvaluateSeedChanges()
 		{
 			// if we have a non-zero seed
diff --git a/Engine/Plugins/Runtime/Metasound/Source/MetasoundStandardNodes/Public/MetasoundTriggerAccumulatorNode.h b/Engine/Plugins/Runtime/Metasound/Source/MetasoundStandardNodes/Public/MetasoundTriggerAccumulatorNode.h
--- a/Engine/Plugins/Runtime/Metasound/Source/MetasoundStandardNodes/Public/MetasoundTriggerAccumulatorNode.h
+++ b/Engine/Plugins/Runtime/Metasound/Source/MetasoundStandardNodes/Public/MetasoundTriggerAccumulatorNode.h
@@ -180,6 +180,12 @@ namespace Metasound
 			}
 		}
 
+		void Reset(const IOperator::FResetParams& InParams)
+		{
+			OutputTrigger->Reset();
+			ResetTriggerState();
+		}
+
 	private:
 
 		FBoolReadRef bAutoReset;
diff --git a/Engine/Plugins/Runtime/Metasound/Source/MetasoundStandardNodes/Public/MetasoundTriggerCompareNode.h b/Engine/Plugins/Runtime/Metasound/Source/MetasoundStandardNodes/Public/MetasoundTriggerCompareNode.h
--- a/Engine/Plugins/Runtime/Metasound/Source/MetasoundStandardNodes/Public/MetasoundTriggerCompareNode.h
+++ b/Engine/Plugins/Runtime/Metasound/Source/MetasoundStandardNodes/Public/MetasoundTriggerCompareNode.h
@@ -209,6 +209,12 @@ namespace Metasound
 				);
 		}
 
+		void Reset(const IOperator::FResetParams& InParams)
+		{
+			TriggerOutOnTrue->Reset();
+			TriggerOutOnFalse->Reset();
+		}
+
 	private:
 
 		FTriggerReadRef OnCompareTrigger;
diff --git a/Engine/Plugins/Runtime/Metasound/Source/MetasoundStandardNodes/Public/MetasoundValueNode.h b/Engine/Plugins/Runtime/Metasound/Source/MetasoundStandardNodes/Public/MetasoundValueNode.h
--- a/Engine/Plugins/Runtime/Metasound/Source/MetasoundStandardNodes/Public/MetasoundValueNode.h
+++ b/Engine/Plugins/Runtime/Metasound/Source/MetasoundStandardNodes/Public/MetasoundValueNode.h
@@ -167,6 +167,13 @@ namespace Metasound
 
 		}
 
+		void Reset(const IOperator::FResetParams& InParams)
+		{
+			TriggerOnSet->Reset();
+			TriggerOnReset->Reset();
+			*OutputValue = *InitValue;
+		}
+
 	private:
 
 		TDataReadReference<FTrigger> SetTrigger;
diff --git a/Engine/Source/Runtime/SignalProcessing/Private/GrainDelay.cpp b/Engine/Source/Runtime/SignalProcessing/Private/GrainDelay.cpp
--- a/Engine/Source/Runtime/SignalProcessing/Private/GrainDelay.cpp
+++ b/Engine/Source/Runtime/SignalProcessing/Private/GrainDelay.cpp
@@ -19,24 +19,24 @@ namespace Audio
 			// Generate the grain data
 			GenerateEnvelopeData(GrainEnvelope, GrainEnvelopeSampleCount, GrainEnvelopeType);
 			
-			// Setup the dynamics processor
-			DynamicsProcessor.Init(InSampleRate, 1);
-			DynamicsProcessor.SetLookaheadMsec(3.0f);
-			DynamicsProcessor.SetAttackTime(0.0f);
-			DynamicsProcessor.SetReleaseTime(10.0f);
-			DynamicsProcessor.SetThreshold(-6.0f);
-			DynamicsProcessor.SetKneeBandwidth(3.0f);
-			DynamicsProcessor.SetInputGain(0.0f);
-			DynamicsProcessor.SetOutputGain(0.0f);
-			DynamicsProcessor.SetAnalogMode(false);
-			DynamicsProcessor.SetPeakMode(EPeakMode::Peak);
-			DynamicsProcessor.SetProcessingMode(EDynamicsProcessingMode::Limiter);
+			InitDynamicsProcessor();
 		}
 
 		FGrainDelay::~FGrainDelay()
 		{
 		}
 
+		void FGrainDelay::Reset()
+		{
+			GrainDelayLine.Reset();
+
+			InitDynamicsProcessor();
+
+			// Move any active grains into free grains
+			FreeGrains.Append(ActiveGrains);
+			ActiveGrains.Reset();
+		}
+
 		float FGrainDelay::GetGrainDelayClamped(const float InDelay) const
 		{
 			return FMath::Clamp(InDelay, 0.0f, 1000.0f * MaxDelaySeconds);
@@ -182,6 +182,22 @@ namespace Audio
 				OutAudioBuffer[FrameIndex] = LastFrame;
 			}
 		}
+
+		void FGrainDelay::InitDynamicsProcessor()
+		{
+			// Setup the dynamics processor
+			DynamicsProcessor.Init(SampleRate, 1);
+			DynamicsProcessor.SetLookaheadMsec(3.0f);
+			DynamicsProcessor.SetAttackTime(0.0f);
+			DynamicsProcessor.SetReleaseTime(10.0f);
+			DynamicsProcessor.SetThreshold(-6.0f);
+			DynamicsProcessor.SetKneeBandwidth(3.0f);
+			DynamicsProcessor.SetInputGain(0.0f);
+			DynamicsProcessor.SetOutputGain(0.0f);
+			DynamicsProcessor.SetAnalogMode(false);
+			DynamicsProcessor.SetPeakMode(EPeakMode::Peak);
+			DynamicsProcessor.SetProcessingMode(EDynamicsProcessingMode::Limiter);
+		}
 	};
 
 	
diff --git a/Engine/Source/Runtime/SignalProcessing/Private/MidiNoteQuantizer.cpp b/Engine/Source/Runtime/SignalProcessing/Private/MidiNoteQuantizer.cpp
--- a/Engine/Source/Runtime/SignalProcessing/Private/MidiNoteQuantizer.cpp
+++ b/Engine/Source/Runtime/SignalProcessing/Private/MidiNoteQuantizer.cpp
@@ -53,7 +53,7 @@ namespace Audio
 		// QuantizeValueToScaleDegree() works within a single octave, so compensate for the octave and root
 		const float SetRange = InScaleDegrees[InScaleDegrees.Num() - 1]; // (SetRange will == 12.0 for a normal octave)
 		const float NoteRootDelta = (InNote - InRoot);
-		const float InNoteOctave = FMath::FloorToFloat(NoteRootDelta / SetRange);
+		const float InNoteOctave = FMath::FloorToFloat(NoteRootDelta / FMath::Max(SetRange, UE_SMALL_NUMBER));
 		const float ValueToQuant = NoteRootDelta - (SetRange * InNoteOctave);
 
 		const float QuantizedValue = QuantizeValueToScaleDegree(ValueToQuant, InScaleDegrees);
@@ -206,4 +206,4 @@ namespace Audio
 		, { EMusicalScale::Scale::HalfDiminished_LocrianNumber2,
 			{{ n1, b2, b3, n4, b5, b6, b7, n8 }, { n1, b3, b5, b7, n8 }}}
 	};
-} // namespace Audio
\ No newline at end of file
+} // namespace Audio
diff --git a/Engine/Source/Runtime/SignalProcessing/Public/DSP/GrainDelay.h b/Engine/Source/Runtime/SignalProcessing/Public/DSP/GrainDelay.h
--- a/Engine/Source/Runtime/SignalProcessing/Public/DSP/GrainDelay.h
+++ b/Engine/Source/Runtime/SignalProcessing/Public/DSP/GrainDelay.h
@@ -16,6 +16,8 @@ namespace Audio
 			FGrainDelay(const float InSampleRate);
 			~FGrainDelay();
 
+			void Reset();
+
 			// Helpers to clamp audio
 			float GetGrainDelayClamped(const float InDelay) const;
 			float GetGrainDurationClamped(const float InDuration) const;
@@ -44,6 +46,8 @@ namespace Audio
 		private:
 			// Synthesize a single frame of audio
 			float SynthesizeFrame(const Audio::FDelay& InDelayLine);
+
+			void InitDynamicsProcessor();
 			
 			// Simple grain data structure for tracking grain state
 			struct FGrain
