Reduce metasound binary size of output nodes.

diff --git a/Engine/Plugins/Runtime/Metasound/Source/MetasoundFrontend/Private/MetasoundFrontendDataTypeRegistry.cpp b/Engine/Plugins/Runtime/Metasound/Source/MetasoundFrontend/Private/MetasoundFrontendDataTypeRegistry.cpp
--- a/Engine/Plugins/Runtime/Metasound/Source/MetasoundFrontend/Private/MetasoundFrontendDataTypeRegistry.cpp
+++ b/Engine/Plugins/Runtime/Metasound/Source/MetasoundFrontend/Private/MetasoundFrontendDataTypeRegistry.cpp
@@ -2,6 +2,7 @@
 #include "MetasoundFrontendDataTypeRegistry.h"
 
 #include "MetasoundTrace.h"
+#include "Misc/ScopeLock.h"
 
 namespace Metasound
 {
@@ -494,6 +495,7 @@ namespace Metasound
 
 				virtual void GetRegisteredDataTypeNames(TArray<FName>& OutNames) const override;
 
+				virtual const IDataTypeRegistryEntry* FindDataTypeRegistryEntry(const FName& InDataTypeName) const override;
 				virtual bool GetDataTypeInfo(const UObject* InObject, FDataTypeRegistryInfo& OutInfo) const override;
 				virtual bool GetDataTypeInfo(const FName& InDataType, FDataTypeRegistryInfo& OutInfo) const override;
 
@@ -554,6 +556,9 @@ namespace Metasound
 
 				// UObject type names to DataTypeNames
 				TMap<const UClass*, FName> RegisteredObjectClasses;
+
+				mutable FCriticalSection RegistryMapMutex;
+				mutable FCriticalSection RegistryObjectMapMutex;
 			};
 
 			bool FDataTypeRegistry::RegisterDataType(TUniquePtr<IDataTypeRegistryEntry>&& InEntry)
@@ -566,12 +571,25 @@ namespace Metasound
 
 					const FName Name = Entry->GetDataTypeInfo().DataTypeName;
 
-					if (!ensureAlwaysMsgf(!RegisteredDataTypes.Contains(Name),
-						TEXT("Name collision when trying to register Metasound Data Type [Name:%s]. DataType must have "
-							"unique name and REGISTER_METASOUND_DATATYPE cannot be called in a public header."),
-							*Name.ToString()))
 					{
-						return false;
+						FScopeLock RegistryLock(&RegistryMapMutex);
+						FScopeLock ObjectRegistryLock(&RegistryObjectMapMutex);
+
+						if (!ensureAlwaysMsgf(!RegisteredDataTypes.Contains(Name),
+							TEXT("Name collision when trying to register Metasound Data Type [Name:%s]. DataType must have "
+								"unique name and REGISTER_METASOUND_DATATYPE cannot be called in a public header."),
+								*Name.ToString()))
+						{
+							return false;
+						}
+
+						const FDataTypeRegistryInfo& RegistryInfo = Entry->GetDataTypeInfo();
+						if (const UClass* Class = RegistryInfo.ProxyGeneratorClass)
+						{
+							RegisteredObjectClasses.Add(Class, Name);
+						}
+
+						RegisteredDataTypes.Add(Name, Entry);
 					}
 
 					// Register nodes associated with data type.
@@ -598,13 +616,6 @@ namespace Metasound
 						}
 					}
 
-					const FDataTypeRegistryInfo& RegistryInfo = Entry->GetDataTypeInfo();
-					if (const UClass* Class = RegistryInfo.ProxyGeneratorClass)
-					{
-						RegisteredObjectClasses.Add(Class, Name);
-					}
-
-					RegisteredDataTypes.Add(Name, Entry);
 
 					UE_LOG(LogMetaSound, Verbose, TEXT("Registered Metasound Datatype [Name:%s]."), *Name.ToString());
 					return true;
@@ -615,13 +626,20 @@ namespace Metasound
 
 			void FDataTypeRegistry::GetRegisteredDataTypeNames(TArray<FName>& OutNames) const
 			{
+				FScopeLock Lock(&RegistryMapMutex);
 				RegisteredDataTypes.GetKeys(OutNames);
 			}
 
+			const IDataTypeRegistryEntry* FDataTypeRegistry::FindDataTypeRegistryEntry(const FName& InDataTypeName) const
+			{
+				return FindDataTypeEntry(InDataTypeName);
+			}
+
 			bool FDataTypeRegistry::GetDataTypeInfo(const UObject* InObject, FDataTypeRegistryInfo& OutInfo) const
 			{
 				if (InObject)
 				{
+					FScopeLock Lock(&RegistryObjectMapMutex);
 					if (const FName* DataTypeName = RegisteredObjectClasses.Find(InObject->GetClass()))
 					{
 						if (const IDataTypeRegistryEntry* Entry = FindDataTypeEntry(*DataTypeName))
@@ -647,6 +665,7 @@ namespace Metasound
 
 			void FDataTypeRegistry::IterateDataTypeInfo(TFunctionRef<void(const FDataTypeRegistryInfo&)> InFunction) const
 			{
+				FScopeLock Lock(&RegistryMapMutex);
 				for (const TPair<FName, TSharedRef<IDataTypeRegistryEntry, ESPMode::ThreadSafe>>& Entry : RegisteredDataTypes)
 				{
 					InFunction(Entry.Value->GetDataTypeInfo());
@@ -655,6 +674,7 @@ namespace Metasound
 
 			bool FDataTypeRegistry::IsRegistered(const FName& InDataType) const
 			{
+				FScopeLock Lock(&RegistryMapMutex);
 				return RegisteredDataTypes.Contains(InDataType);
 			}
 
@@ -800,6 +820,7 @@ namespace Metasound
 					return false;
 				}
 
+				FScopeLock Lock(&RegistryObjectMapMutex);
 				UClass* ObjectClass = InObject->GetClass();
 				while (ObjectClass != UObject::StaticClass())
 				{
@@ -1092,11 +1113,14 @@ namespace Metasound
 
 			const IDataTypeRegistryEntry* FDataTypeRegistry::FindDataTypeEntry(const FName& InDataTypeName) const
 			{
-				const TSharedRef<IDataTypeRegistryEntry, ESPMode::ThreadSafe>* Entry = RegisteredDataTypes.Find(InDataTypeName);
-
-				if (ensureMsgf(nullptr != Entry, TEXT("Data type not registered [Name:%s]"), *InDataTypeName.ToString()))
+				FScopeLock Lock(&RegistryMapMutex);
 				{
-					return &Entry->Get();
+					const TSharedRef<IDataTypeRegistryEntry, ESPMode::ThreadSafe>* Entry = RegisteredDataTypes.Find(InDataTypeName);
+
+					if (ensureMsgf(nullptr != Entry, TEXT("Data type not registered [Name:%s]"), *InDataTypeName.ToString()))
+					{
+						return &Entry->Get();
+					}
 				}
 
 				return nullptr;
diff --git a/Engine/Plugins/Runtime/Metasound/Source/MetasoundFrontend/Private/MetasoundOutputNode.cpp b/Engine/Plugins/Runtime/Metasound/Source/MetasoundFrontend/Private/MetasoundOutputNode.cpp
--- a/Engine/Plugins/Runtime/Metasound/Source/MetasoundFrontend/Private/MetasoundOutputNode.cpp
+++ b/Engine/Plugins/Runtime/Metasound/Source/MetasoundFrontend/Private/MetasoundOutputNode.cpp
@@ -1,11 +1,17 @@
 // Copyright Epic Games, Inc. All Rights Reserved.
 
 #include "MetasoundOutputNode.h"
+#include "MetasoundFrontendDataTypeRegistry.h"
+
+#define LOCTEXT_NAMESPACE "MetasoundGraphCore"
 
 namespace Metasound
 {
 	namespace OutputNodePrivate
 	{
+		static const FLazyName ValueVertexAccessVariantName("Constructor");
+		static const FLazyName ReferenceVertexAccessVariantName("");
+
 		FOutputOperator::FOutputOperator(const FVertexName& InVertexName, const FAnyDataReference& InDataReference)
 		: VertexName(InVertexName)
 		, DataReference(InDataReference)
@@ -39,6 +45,113 @@ namespace Metasound
 		{
 			return nullptr;
 		}
+
+		FOutputNode::FOutputOperatorFactory::FOutputOperatorFactory(const FVertexName& InVertexName, const FName& InDataTypeName, EVertexAccessType InVertexAccessType)
+		: VertexName(InVertexName)
+		, DataTypeName(InDataTypeName)
+		, VertexAccessType(InVertexAccessType)
+		{
+		}
+
+		TUniquePtr<IOperator> FOutputNode::FOutputOperatorFactory::CreateOperator(const FBuildOperatorParams& InParams, FBuildResults& OutResults)
+		{
+			using namespace OutputNodePrivate;
+
+			const FOutputNode& OutputNode = static_cast<const FOutputNode&>(InParams.Node);
+
+			// Use data reference if it is passed in. 
+			if (const FAnyDataReference* Ref = InParams.InputData.FindDataReference(VertexName))
+			{
+				checkf(Ref->GetDataTypeName() == DataTypeName, TEXT("Mismatched data type names on output node (%s). Expected (%s), received (%s)."), *VertexName.ToString(), *DataTypeName.ToString(), *(Ref->GetDataTypeName().ToString()));
+				return MakeUnique<FOutputOperator>(VertexName, *Ref);
+			}
+			else 
+			{
+				// Make data reference if none are passed in. 
+				Frontend::IDataTypeRegistry& DataTypeRegistry = Frontend::IDataTypeRegistry::Get();
+				if (const Frontend::IDataTypeRegistryEntry* Entry = DataTypeRegistry.FindDataTypeRegistryEntry(DataTypeName))
+				{
+					FLiteral DefaultLiteral = DataTypeRegistry.CreateDefaultLiteral(DataTypeName);
+					TOptional<FAnyDataReference> DataReference = Entry->CreateDataReference(EDataReferenceAccessType::Value, DefaultLiteral, InParams.OperatorSettings);
+					if (DataReference.IsSet())
+					{
+						return MakeUnique<FOutputOperator>(VertexName, *Ref);
+					}
+				}
+			}
+
+			// Do not make output operator if no data reference is available. 
+			OutResults.Errors.Emplace(MakeUnique<FMissingOutputNodeInputReferenceError>(InParams.Node));
+			return TUniquePtr<IOperator>(nullptr);
+		}
+
+		FName FOutputNode::GetVariantName(EVertexAccessType InVertexAccessType)
+		{
+			if (EVertexAccessType::Value == InVertexAccessType)
+			{
+				return ValueVertexAccessVariantName;
+			}
+			else
+			{
+				return ReferenceVertexAccessVariantName;
+			}
+		}
+
+		FVertexInterface FOutputNode::GetVertexInterface(const FVertexName& InVertexName, const FName& InDataTypeName, EVertexAccessType InVertexAccessType)
+		{
+			static const FText VertexDescription = METASOUND_LOCTEXT("Metasound_OutputVertexDescription", "Output from the parent Metasound graph.");
+
+			return FVertexInterface(
+				FInputVertexInterface(
+					FInputDataVertex(InVertexName, InDataTypeName, FDataVertexMetadata{VertexDescription}, InVertexAccessType)
+				),
+				FOutputVertexInterface(
+					FOutputDataVertex(InVertexName, InDataTypeName, FDataVertexMetadata{VertexDescription}, InVertexAccessType)
+				)
+			);
+		}
+
+		FNodeClassMetadata FOutputNode::GetNodeInfo(const FVertexName& InVertexName, const FName& InDataTypeName, EVertexAccessType InVertexAccessType)
+		{
+			FNodeClassMetadata Info;
+
+			Info.ClassName = { "Output", InDataTypeName, GetVariantName(InVertexAccessType)};
+			Info.MajorVersion = 1;
+			Info.MinorVersion = 0;
+			Info.Description = METASOUND_LOCTEXT("Metasound_OutputNodeDescription", "Output from the parent Metasound graph.");
+			Info.Author = PluginAuthor;
+			Info.PromptIfMissing = PluginNodeMissingPrompt;
+			Info.DefaultInterface = GetVertexInterface(InVertexName, InDataTypeName, InVertexAccessType);
+
+			return Info;
+		}
+
+		FOutputNode::FOutputNode(const FName& InInstanceName, const FGuid& InInstanceID, const FVertexName& InVertexName, const FName& InDataTypeName, EVertexAccessType InVertexAccessType)
+		:	FNode(InInstanceName, InInstanceID, GetNodeInfo(InVertexName, InDataTypeName, InVertexAccessType))
+		,	Factory(MakeShared<FOutputOperatorFactory, ESPMode::ThreadSafe>(InVertexName, InDataTypeName, InVertexAccessType))
+		{
+		}
+
+		const FVertexInterface& FOutputNode::GetVertexInterface() const
+		{
+			return GetMetadata().DefaultInterface;
+		}
+
+		bool FOutputNode::SetVertexInterface(const FVertexInterface& InInterface)
+		{
+			return GetVertexInterface() == InInterface;
+		}
+
+		bool FOutputNode::IsVertexInterfaceSupported(const FVertexInterface& InInterface) const 
+		{
+			return GetVertexInterface() == InInterface;
+		}
+
+		TSharedRef<IOperatorFactory, ESPMode::ThreadSafe> FOutputNode::GetDefaultOperatorFactory() const
+		{
+			return Factory;
+		}
 	}
 }
 
+#undef LOCTEXT_NAMESPACE // MetasoundOutputNode
diff --git a/Engine/Plugins/Runtime/Metasound/Source/MetasoundFrontend/Public/MetasoundFrontendDataTypeRegistry.h b/Engine/Plugins/Runtime/Metasound/Source/MetasoundFrontend/Public/MetasoundFrontendDataTypeRegistry.h
--- a/Engine/Plugins/Runtime/Metasound/Source/MetasoundFrontend/Public/MetasoundFrontendDataTypeRegistry.h
+++ b/Engine/Plugins/Runtime/Metasound/Source/MetasoundFrontend/Public/MetasoundFrontendDataTypeRegistry.h
@@ -237,6 +237,8 @@ namespace Metasound
 			/** Provides all names of registered DataTypes. */
 			virtual void GetRegisteredDataTypeNames(TArray<FName>& OutNames) const = 0;
 
+			virtual const IDataTypeRegistryEntry* FindDataTypeRegistryEntry(const FName& InDataTypeName) const = 0;
+
 			/** Returns DataType info associated with the provided object. */
 			virtual bool GetDataTypeInfo(const UObject* InObject, FDataTypeRegistryInfo& OutInfo) const = 0;
 
diff --git a/Engine/Plugins/Runtime/Metasound/Source/MetasoundFrontend/Public/MetasoundOutputNode.h b/Engine/Plugins/Runtime/Metasound/Source/MetasoundFrontend/Public/MetasoundOutputNode.h
--- a/Engine/Plugins/Runtime/Metasound/Source/MetasoundFrontend/Public/MetasoundOutputNode.h
+++ b/Engine/Plugins/Runtime/Metasound/Source/MetasoundFrontend/Public/MetasoundOutputNode.h
@@ -33,6 +33,14 @@ namespace Metasound
 			AddNode(InNode);
 		}
 
+		FMissingOutputNodeInputReferenceError(const INode& InNode)
+			: FBuildErrorBase(
+				"MetasoundMissingOutputDataReferenceError",
+				METASOUND_LOCTEXT("MissingOutputNodeInputReferenceError", "Missing required output node input reference."))
+		{
+			AddNode(InNode);
+		}
+
 		virtual ~FMissingOutputNodeInputReferenceError() = default;
 	};
 
@@ -42,7 +50,6 @@ namespace Metasound
 		{
 			public:
 				FOutputOperator(const FVertexName& InVertexName, const FAnyDataReference& InDataReference);
-
 				virtual ~FOutputOperator() = default;
 
 				virtual FDataReferenceCollection GetInputs() const override;
@@ -55,140 +62,66 @@ namespace Metasound
 				FVertexName VertexName;
 				FAnyDataReference DataReference;
 		};
-	}
 
-	template<typename DataType, EVertexAccessType VertexAccess=EVertexAccessType::Reference>
-	class TOutputNode : public FNode
-	{
-	private:
-		static FName GetVariantName()
+		class METASOUNDFRONTEND_API FOutputNode : public FNode
 		{
-			if constexpr (EVertexAccessType::Value == VertexAccess)
-			{
-				return FName("Constructor");
-			}
-			else
-			{
-				return FName();
-			}
-		}
-
-		class FOutputOperatorFactory : public IOperatorFactory
-		{
-		public:
-			FOutputOperatorFactory(const FVertexName& InDataReferenceName)
-			:	DataReferenceName(InDataReferenceName)
-			{
-			}
-
-			virtual TUniquePtr<IOperator> CreateOperator(const FBuildOperatorParams& InParams, FBuildResults& OutResults) override
+		private:
+			class FOutputOperatorFactory : public IOperatorFactory
 			{
-				using namespace OutputNodePrivate;
-
-				if (const FAnyDataReference* Ref = InParams.InputData.FindDataReference(DataReferenceName))
-				{
-					return MakeUnique<FOutputOperator>(DataReferenceName, *Ref);
-				}
-
-				// Only construct default if default construction is supported
-				if constexpr (TIsParsable<DataType>::Value)
-				{
-					if constexpr (TIsConstructorVertexSupported<DataType>::Value)
-					{
-						// Prefer to make a constant value
-						TDataValueReference<DataType> DefaultValueRef = TDataValueReferenceFactory<DataType>::CreateAny(InParams.OperatorSettings);
-						return MakeUnique<FOutputOperator>(DataReferenceName, DefaultValueRef);
-					}
-					else
-					{
-						TDataReadReference<DataType> DefaultReadRef = TDataReadReferenceFactory<DataType>::CreateAny(InParams.OperatorSettings);
-						return MakeUnique<FOutputOperator>(DataReferenceName, DefaultReadRef);
-					}
-				}
-
-				OutResults.Errors.Emplace(MakeUnique<FMissingOutputNodeInputReferenceError>(InParams.Node, GetMetasoundDataTypeDisplayText<DataType>()));
-				return TUniquePtr<IOperator>(nullptr);
-			}
+			public:
 
-		private:
-			FVertexName DataReferenceName;
-		};
+				FOutputOperatorFactory(const FVertexName& InVertexName, const FName& InDataTypeName, EVertexAccessType InVertexAccessType);
+				virtual TUniquePtr<IOperator> CreateOperator(const FBuildOperatorParams& InParams, FBuildResults& OutResults) override;
 
-		static FVertexInterface GetVertexInterface(const FVertexName& InVertexName)
-		{
-			static const FText VertexDescription = METASOUND_LOCTEXT("Metasound_OutputVertexDescription", "Output from the parent Metasound graph.");
-
-			return FVertexInterface(
-				FInputVertexInterface(
-					FInputDataVertex(InVertexName, GetMetasoundDataTypeName<DataType>(), FDataVertexMetadata{VertexDescription}, VertexAccess)
-				),
-				FOutputVertexInterface(
-					FOutputDataVertex(InVertexName, GetMetasoundDataTypeName<DataType>(), FDataVertexMetadata{VertexDescription}, VertexAccess)
-				)
-			);
-		}
+			private:
 
-		static FNodeClassMetadata GetNodeInfo(const FVertexName& InVertexName)
-		{
-			FNodeClassMetadata Info;
+				FVertexName VertexName;
+				FName DataTypeName;
+				EVertexAccessType VertexAccessType;
+			};
 
-			Info.ClassName = { "Output", GetMetasoundDataTypeName<DataType>(), GetVariantName() };
-			Info.MajorVersion = 1;
-			Info.MinorVersion = 0;
-			Info.Description = METASOUND_LOCTEXT("Metasound_OutputNodeDescription", "Output from the parent Metasound graph.");
-			Info.Author = PluginAuthor;
-			Info.PromptIfMissing = PluginNodeMissingPrompt;
-			Info.DefaultInterface = GetVertexInterface(InVertexName);
+			static FName GetVariantName(EVertexAccessType InVertexAccessType);
+			static FVertexInterface GetVertexInterface(const FVertexName& InVertexName, const FName& InDataTypeName, EVertexAccessType InVertexAccessType);
+			static FNodeClassMetadata GetNodeInfo(const FVertexName& InVertexName, const FName& InDataTypeName, EVertexAccessType InVertexAccessType);
 
-			return Info;
-		};
+			public:
+				FOutputNode(const FName& InInstanceName, const FGuid& InInstanceID, const FVertexName& InVertexName, const FName& InDataType, EVertexAccessType InVertexAccessType);
 
-		public:
-			TOutputNode(const FVertexName& InInstanceName, const FGuid& InInstanceID, const FVertexName& InVertexName)
-			:	FNode(InInstanceName, InInstanceID, GetNodeInfo(InVertexName))
-			,	VertexInterface(GetVertexInterface(InVertexName))
-			,	Factory(MakeShared<FOutputOperatorFactory, ESPMode::ThreadSafe>(InVertexName))
-			{
-			}
+				/** Return the current vertex interface. */
+				virtual const FVertexInterface& GetVertexInterface() const override;
 
-			/** Return the current vertex interface. */
-			virtual const FVertexInterface& GetVertexInterface() const override
-			{
-				return VertexInterface;
-			}
-
-			/** Set the vertex interface. If the vertex was successfully changed, returns true. 
-			 *
-			 * @param InInterface - New interface for node. 
-			 *
-			 * @return True on success, false otherwise.
-			 */
-			virtual bool SetVertexInterface(const FVertexInterface& InInterface) override
-			{
-				return VertexInterface == InInterface;
-			}
-
-			/** Expresses whether a specific vertex interface is supported.
-			 *
-			 * @param InInterface - New interface. 
-			 *
-			 * @return True if the interface is supported, false otherwise. 
-			 */
-			virtual bool IsVertexInterfaceSupported(const FVertexInterface& InInterface) const 
-			{
-				return VertexInterface == InInterface;
-			}
+				/** Set the vertex interface. If the vertex was successfully changed, returns true. 
+				 *
+				 * @param InInterface - New interface for node. 
+				 *
+				 * @return True on success, false otherwise.
+				 */
+				virtual bool SetVertexInterface(const FVertexInterface& InInterface) override;
 
-			virtual TSharedRef<IOperatorFactory, ESPMode::ThreadSafe> GetDefaultOperatorFactory() const override
-			{
-				return Factory;
-			}
+				/** Expresses whether a specific vertex interface is supported.
+				 *
+				 * @param InInterface - New interface. 
+				 *
+				 * @return True if the interface is supported, false otherwise. 
+				 */
+				virtual bool IsVertexInterfaceSupported(const FVertexInterface& InInterface) const;
 
-		private:
-			FVertexInterface VertexInterface;
+				virtual TSharedRef<IOperatorFactory, ESPMode::ThreadSafe> GetDefaultOperatorFactory() const override;
 
-			TSharedRef<FOutputOperatorFactory, ESPMode::ThreadSafe> Factory;
+			private:
+				TSharedRef<FOutputOperatorFactory, ESPMode::ThreadSafe> Factory;
+		};
+	}
 
+	/** Output nodes are used to expose graph data to external entities. */
+	template<typename DataType, EVertexAccessType VertexAccess=EVertexAccessType::Reference>
+	class TOutputNode : public OutputNodePrivate::FOutputNode
+	{
+	public:
+		TOutputNode(const FVertexName& InInstanceName, const FGuid& InInstanceID, const FVertexName& InVertexName)
+		:	FOutputNode(InInstanceName, InInstanceID, InVertexName, GetMetasoundDataTypeName<DataType>(), VertexAccess)
+		{
+		}
 	};
 }
 #undef LOCTEXT_NAMESPACE // MetasoundOutputNode
