- Add Runtime MetaSound Builder support
- Add MetaSound Patch Interface support
- Add MetaSound Interface Bindings

diff --git a/Engine/Plugins/Runtime/Metasound/Config/BaseMetasound.ini b/Engine/Plugins/Runtime/Metasound/Config/BaseMetasound.ini
--- a/Engine/Plugins/Runtime/Metasound/Config/BaseMetasound.ini
+++ b/Engine/Plugins/Runtime/Metasound/Config/BaseMetasound.ini
@@ -22,4 +22,7 @@
 
 ; 5.1
 +ClassRedirects=(OldName="MetaSound",NewName="/Script/MetasoundEngine.MetaSoundPatch")
-+PropertyRedirects=(OldName="MetasoundEditorSettings.bPinMetaSoundInAssetMenu",NewName="bPinMetaSoundPatchInAssetMenu")
\ No newline at end of file
++PropertyRedirects=(OldName="MetasoundEditorSettings.bPinMetaSoundInAssetMenu",NewName="bPinMetaSoundPatchInAssetMenu")
+
+; 5.3
++EnumRedirects=(OldName="MetasoundSourceAudioFormat",NewName="/Script/MetasoundEngine.MetaSoundOutputAudioFormat")
diff --git a/Engine/Plugins/Runtime/Metasound/Source/MetasoundEditor/MetasoundEditor.Build.cs b/Engine/Plugins/Runtime/Metasound/Source/MetasoundEditor/MetasoundEditor.Build.cs
--- a/Engine/Plugins/Runtime/Metasound/Source/MetasoundEditor/MetasoundEditor.Build.cs
+++ b/Engine/Plugins/Runtime/Metasound/Source/MetasoundEditor/MetasoundEditor.Build.cs
@@ -56,6 +56,7 @@ namespace UnrealBuildTool.Rules
 					"DetailCustomizations",
 					"EditorFramework",
 					"EditorStyle",
+					"EditorSubsystem",
 					"Engine",
 					"GraphEditor",
 					"InputCore",
diff --git a/Engine/Plugins/Runtime/Metasound/Source/MetasoundEditor/Private/MetasoundDetailCustomization.cpp b/Engine/Plugins/Runtime/Metasound/Source/MetasoundEditor/Private/MetasoundDetailCustomization.cpp
--- a/Engine/Plugins/Runtime/Metasound/Source/MetasoundEditor/Private/MetasoundDetailCustomization.cpp
+++ b/Engine/Plugins/Runtime/Metasound/Source/MetasoundEditor/Private/MetasoundDetailCustomization.cpp
@@ -534,17 +534,19 @@ namespace Metasound
 			AddableInterfaceNames.Reset();
 			ImplementedInterfaceNames.Reset();
 
-			if (const FMetasoundAssetBase* MetaSoundAsset = IMetasoundUObjectRegistry::Get().GetObjectAsAssetBase(MetaSound.Get()))
+			const UObject* MetaSoundObject = MetaSound.Get();
+			if (const FMetasoundAssetBase* MetaSoundAsset = IMetasoundUObjectRegistry::Get().GetObjectAsAssetBase(MetaSoundObject))
 			{
 				auto GetVersionName = [](const FMetasoundFrontendVersion& Version) { return Version.Name; };
-				auto CanAddOrRemoveInterface = [](const FMetasoundFrontendVersion& Version)
+				const UClass* MetaSoundClass = MetaSoundObject->GetClass();
+				auto CanAddOrRemoveInterface = [ClassName = MetaSoundClass->GetFName()](const FMetasoundFrontendVersion& Version)
 				{
 					using namespace Metasound::Frontend;
 
 					const FInterfaceRegistryKey Key = GetInterfaceRegistryKey(Version);
 					if (const IInterfaceRegistryEntry* Entry = IInterfaceRegistry::Get().FindInterfaceRegistryEntry(Key))
 					{
-						return Entry->EditorCanAddOrRemove();
+						return Entry->EditorCanAddOrRemove(ClassName);
 					}
 
 					return false;
@@ -565,6 +567,8 @@ namespace Metasound
 						}
 					}
 				}
+
+				AddableInterfaceNames.Sort([](const TSharedPtr<FString>& A, const TSharedPtr<FString>& B) { return A->Compare(*B) < 0; });
 			}
 		}
 	} // namespace Editor
diff --git a/Engine/Plugins/Runtime/Metasound/Source/MetasoundEditor/Private/MetasoundEditor.cpp b/Engine/Plugins/Runtime/Metasound/Source/MetasoundEditor/Private/MetasoundEditor.cpp
--- a/Engine/Plugins/Runtime/Metasound/Source/MetasoundEditor/Private/MetasoundEditor.cpp
+++ b/Engine/Plugins/Runtime/Metasound/Source/MetasoundEditor/Private/MetasoundEditor.cpp
@@ -29,6 +29,8 @@
 #include "Kismet2/KismetEditorUtilities.h"
 #include "Logging/TokenizedMessage.h"
 #include "Metasound.h"
+#include "MetasoundBuilderSubsystem.h"
+#include "MetasoundDocumentInterface.h"
 #include "MetasoundEditorCommands.h"
 #include "MetasoundEditorGraph.h"
 #include "MetasoundEditorGraphBuilder.h"
@@ -40,6 +42,7 @@
 #include "MetasoundEditorTabFactory.h"
 #include "MetasoundFrontend.h"
 #include "MetasoundFrontendDocument.h"
+#include "MetasoundFrontendDocumentBuilder.h"
 #include "MetasoundFrontendRegistries.h"
 #include "MetasoundFrontendSearchEngine.h"
 #include "MetasoundFrontendTransform.h"
@@ -58,6 +61,7 @@
 #include "Styling/AppStyle.h"
 #include "Templates/Function.h"
 #include "Templates/SharedPointer.h"
+#include "UObject/ScriptInterface.h"
 #include "Widgets/Docking/SDockTab.h"
 #include "Widgets/Images/SImage.h"
 #include "Widgets/Input/SButton.h"
@@ -1359,11 +1363,7 @@ namespace Metasound
 			{
 				TSet<UClass*> ImportClasses;
 
-				for (const FMetasoundFrontendVersion& InterfaceVersion : MetasoundDoc.Interfaces)
-				{
-					TArray<UClass*> InterfaceClasses = IMetasoundUObjectRegistry::Get().FindSupportedInterfaceClasses(InterfaceVersion);
-					ImportClasses.Append(MoveTemp(InterfaceClasses));
-				}
+				// TODO: Update importing to support interfaces
 
 				if (ImportClasses.Num() < 1)
 				{
@@ -1386,7 +1386,7 @@ namespace Metasound
 						}
 					}
 
-					IMetasoundUObjectRegistry::Get().NewObject(AnyClass, MetasoundDoc, OutputPath);
+					// TODO: Update to just use simple UObject NewObject
 				}
 			}
 			else
@@ -1722,13 +1722,18 @@ namespace Metasound
 
 		void FEditor::ConvertFromPreset()
 		{
+			using namespace Frontend;
+
 			check(GEditor);
 
 			if (Metasound)
 			{
 				FMetasoundAssetBase* MetasoundAsset = IMetasoundUObjectRegistry::Get().GetObjectAsAssetBase(Metasound);
 				check(MetasoundAsset);
-				MetasoundAsset->ConvertFromPreset();
+
+				TScriptInterface<IMetaSoundDocumentInterface> DocInterface = MetasoundAsset->GetOwningAsset();
+				FMetaSoundFrontendDocumentBuilder Builder(DocInterface);
+				Builder.ConvertFromPreset();
 
 				// Hack until toolbar is polished up & corner text properly dynamically updates
 				if (UAssetEditorSubsystem* AssetEditorSubsystem = GEditor->GetEditorSubsystem<UAssetEditorSubsystem>())
diff --git a/Engine/Plugins/Runtime/Metasound/Source/MetasoundEditor/Private/MetasoundEditorGraphBuilder.cpp b/Engine/Plugins/Runtime/Metasound/Source/MetasoundEditor/Private/MetasoundEditorGraphBuilder.cpp
--- a/Engine/Plugins/Runtime/Metasound/Source/MetasoundEditor/Private/MetasoundEditorGraphBuilder.cpp
+++ b/Engine/Plugins/Runtime/Metasound/Source/MetasoundEditor/Private/MetasoundEditorGraphBuilder.cpp
@@ -1542,25 +1542,33 @@ namespace Metasound
 			using namespace Frontend;
 			using namespace GraphBuilderPrivate;
 
-			FMetasoundFrontendClassMetadata Metadata;
+			FMetasoundAssetBase* MetaSoundAsset = IMetasoundUObjectRegistry::Get().GetObjectAsAssetBase(&InMetaSound);
+			check(MetaSoundAsset);
+			FMetasoundFrontendDocument& Doc = MetaSoundAsset->GetDocumentChecked();
 
 			// 1. Set default class Metadata
-			Metadata.SetClassName(FMetasoundFrontendClassName(FName(), *FGuid::NewGuid().ToString(), FName()));
-			Metadata.SetVersion({ 1, 0 });
-			Metadata.SetType(EMetasoundFrontendClassType::Graph);
-			Metadata.SetAuthor(InAuthor);
+			Doc.RootGraph.Metadata.SetClassName(FMetasoundFrontendClassName(FName(), *FGuid::NewGuid().ToString(), FName()));
+			Doc.RootGraph.Metadata.SetVersion({ 1, 0 });
+			Doc.RootGraph.Metadata.SetType(EMetasoundFrontendClassType::Graph);
+			Doc.RootGraph.Metadata.SetAuthor(InAuthor);
 
-			FMetasoundAssetBase* MetaSoundAsset = IMetasoundUObjectRegistry::Get().GetObjectAsAssetBase(&InMetaSound);
-			check(MetaSoundAsset);
-			MetaSoundAsset->SetMetadata(Metadata);
+			if (Doc.RootGraph.Metadata.GetType() != EMetasoundFrontendClassType::Graph)
+			{
+				UE_LOG(LogMetaSound, Display, TEXT("Forcing class type to EMetasoundFrontendClassType::Graph on root graph metadata"));
+				Doc.RootGraph.Metadata.SetType(EMetasoundFrontendClassType::Graph);
+			}
 
 			// 2. Set default doc version Metadata
-			FDocumentHandle DocumentHandle = MetaSoundAsset->GetDocumentHandle();
-			FMetasoundFrontendDocumentMetadata* DocMetadata = DocumentHandle->GetMetadata();
-			check(DocMetadata);
-			DocMetadata->Version.Number = FMetasoundFrontendDocument::GetMaxVersion();
+			Doc.Metadata.Version.Number = FMetasoundFrontendDocument::GetMaxVersion();
+
+			// 3. Add default interfaces
+			{
+				UClass* AssetClass = InMetaSound.GetClass();
+				check(AssetClass);
 
-			MetaSoundAsset->AddDefaultInterfaces();
+				TArray<FMetasoundFrontendInterface> InitInterfaces = ISearchEngine::Get().FindUClassDefaultInterfaces(AssetClass->GetFName());
+				FModifyRootGraphInterfaces({ }, InitInterfaces).Transform(Doc);
+			}
 
 			FGraphHandle GraphHandle = MetaSoundAsset->GetRootGraphHandle();
 			FVector2D InputNodeLocation = FVector2D::ZeroVector;
diff --git a/Engine/Plugins/Runtime/Metasound/Source/MetasoundEditor/Private/MetasoundEditorGraphInputNode.cpp b/Engine/Plugins/Runtime/Metasound/Source/MetasoundEditor/Private/MetasoundEditorGraphInputNode.cpp
--- a/Engine/Plugins/Runtime/Metasound/Source/MetasoundEditor/Private/MetasoundEditorGraphInputNode.cpp
+++ b/Engine/Plugins/Runtime/Metasound/Source/MetasoundEditor/Private/MetasoundEditorGraphInputNode.cpp
@@ -85,7 +85,7 @@ void UMetasoundEditorGraphInputNode::Validate(Metasound::Editor::FGraphNodeValid
 
 	Super::Validate(OutResult);
 
-	FNodeHandle NodeHandle = GetNodeHandle();
+	FConstNodeHandle NodeHandle = GetConstNodeHandle();
 	const FMetasoundFrontendClassMetadata& Metadata = NodeHandle->GetClassMetadata();
 
 	const FMetasoundFrontendVersion& MetasoundFrontendVersion = NodeHandle->GetInterfaceVersion();
@@ -98,11 +98,11 @@ void UMetasoundEditorGraphInputNode::Validate(Metasound::Editor::FGraphNodeValid
 		FText RequiredText;
 		if (InterfaceToValidate.IsMemberInputRequired(NodeName, RequiredText))
 		{
-			TArray<FConstInputHandle> InputHandles = NodeHandle->GetConstInputs();
-			if (ensure(!InputHandles.IsEmpty()))
+			TArray<FConstOutputHandle> OutputHandles = NodeHandle->GetConstOutputs();
+			if (ensure(!OutputHandles.IsEmpty()))
 			{
-				const FConstInputHandle& InputHandle = InputHandles.Last();
-				if (!InputHandle->IsConnected())
+				const FConstOutputHandle& OutputHandle = OutputHandles.Last();
+				if (!OutputHandle->IsConnected())
 				{
 					OutResult.SetMessage(EMessageSeverity::Warning, *RequiredText.ToString());
 				}
diff --git a/Engine/Plugins/Runtime/Metasound/Source/MetasoundEditor/Private/MetasoundEditorSubsystem.cpp b/Engine/Plugins/Runtime/Metasound/Source/MetasoundEditor/Private/MetasoundEditorSubsystem.cpp
new file mode 100644
--- /dev/null
+++ b/Engine/Plugins/Runtime/Metasound/Source/MetasoundEditor/Private/MetasoundEditorSubsystem.cpp
@@ -0,0 +1,45 @@
+// Copyright Epic Games, Inc. All Rights Reserved.
+#include "MetasoundEditorSubsystem.h"
+
+#include "IAssetTools.h"
+#include "MetasoundBuilderSubsystem.h"
+#include "MetasoundEditorGraph.h"
+#include "MetasoundEditorGraphSchema.h"
+#include "MetasoundUObjectRegistry.h"
+
+#include UE_INLINE_GENERATED_CPP_BY_NAME(MetasoundEditorSubsystem)
+
+#define LOCTEXT_NAMESPACE "MetaSoundEditor"
+
+
+TScriptInterface<IMetaSoundDocumentInterface> UMetaSoundEditorSubsystem::BuildToAsset(UMetaSoundBuilderBase* InBuilder, const FString& Author, const FString& AssetName, const FString& PackagePath, EMetaSoundBuilderResult& OutResult)
+{
+	using namespace Metasound;
+	using namespace Metasound::Frontend;
+
+	if (InBuilder)
+	{
+		constexpr UObject* Parent = nullptr;
+		constexpr UFactory* Factory = nullptr;
+		if (UObject* NewMetaSound = IAssetTools::Get().CreateAsset(AssetName, PackagePath, &InBuilder->GetBuilderUClass(), Factory))
+		{
+			InBuilder->InitNodeLocations();
+			InBuilder->SetAuthor(Author);
+
+			FMetaSoundBuilderOptions BuilderOptions { FName(*AssetName) };
+			BuilderOptions.ExistingAsset = NewMetaSound;
+			TScriptInterface<IMetaSoundDocumentInterface> DocInterface = InBuilder->Build(Parent, BuilderOptions);
+
+			FMetasoundAssetBase* MetaSoundAsset = IMetasoundUObjectRegistry::Get().GetObjectAsAssetBase(NewMetaSound);
+			UMetasoundEditorGraph* Graph = NewObject<UMetasoundEditorGraph>(NewMetaSound, FName(), RF_Transactional);
+			Graph->Schema = UMetasoundEditorGraphSchema::StaticClass();
+			MetaSoundAsset->SetGraph(Graph);
+
+			OutResult = EMetaSoundBuilderResult::Succeeded;
+			return NewMetaSound;
+		}
+	}
+
+	OutResult = EMetaSoundBuilderResult::Failed;
+	return nullptr;
+}
diff --git a/Engine/Plugins/Runtime/Metasound/Source/MetasoundEditor/Public/IMetasoundEditor.h b/Engine/Plugins/Runtime/Metasound/Source/MetasoundEditor/Public/IMetasoundEditor.h
--- a/Engine/Plugins/Runtime/Metasound/Source/MetasoundEditor/Public/IMetasoundEditor.h
+++ b/Engine/Plugins/Runtime/Metasound/Source/MetasoundEditor/Public/IMetasoundEditor.h
@@ -8,7 +8,7 @@
 #include "UObject/Object.h"
 
 
-class IMetasoundEditor : public FAssetEditorToolkit
+class METASOUNDEDITOR_API IMetasoundEditor : public FAssetEditorToolkit
 {
 public:
 	// Returns the editor's assigned MetaSound UObject.
diff --git a/Engine/Plugins/Runtime/Metasound/Source/MetasoundEditor/Public/MetasoundEditorSubsystem.h b/Engine/Plugins/Runtime/Metasound/Source/MetasoundEditor/Public/MetasoundEditorSubsystem.h
new file mode 100644
--- /dev/null
+++ b/Engine/Plugins/Runtime/Metasound/Source/MetasoundEditor/Public/MetasoundEditorSubsystem.h
@@ -0,0 +1,24 @@
+// Copyright Epic Games, Inc. All Rights Reserved.
+#pragma once
+
+#include "AssetRegistry/AssetData.h"
+#include "EditorSubsystem.h"
+#include "MetasoundDocumentInterface.h"
+#include "UObject/ScriptInterface.h"
+
+#include "MetasoundEditorSubsystem.generated.h"
+
+
+// Forward Declarations
+class UMetaSoundBuilderBase;
+
+/** The subsystem in charge of editor MetaSound functionality */
+UCLASS()
+class METASOUNDEDITOR_API UMetaSoundEditorSubsystem : public UEditorSubsystem
+{
+	GENERATED_BODY()
+
+public:
+	UFUNCTION(BlueprintCallable, Category = "Audio|MetaSound|Builder|Editor", meta = (WorldContext = "Parent", ExpandEnumAsExecs = "OutResult"))
+	UPARAM(DisplayName = "MetaSound Asset") TScriptInterface<IMetaSoundDocumentInterface> BuildToAsset(UMetaSoundBuilderBase* InBuilder, const FString& Author, const FString& AssetName, const FString& PackagePath, EMetaSoundBuilderResult& OutResult);
+};
diff --git a/Engine/Plugins/Runtime/Metasound/Source/MetasoundEngine/Private/MetasoundBuilderSubsystem.cpp b/Engine/Plugins/Runtime/Metasound/Source/MetasoundEngine/Private/MetasoundBuilderSubsystem.cpp
new file mode 100644
--- /dev/null
+++ b/Engine/Plugins/Runtime/Metasound/Source/MetasoundEngine/Private/MetasoundBuilderSubsystem.cpp
@@ -0,0 +1,959 @@
+// Copyright Epic Games, Inc. All Rights Reserved.
+#include "MetasoundBuilderSubsystem.h"
+
+#include "Algo/Transform.h"
+#include "Components/AudioComponent.h"
+#include "HAL/IConsoleManager.h"
+#include "Interfaces/MetasoundOutputFormatInterfaces.h"
+#include "Interfaces/MetasoundFrontendSourceInterface.h"
+#include "Metasound.h"
+#include "MetasoundFrontendDocument.h"
+#include "MetasoundFrontendSearchEngine.h"
+#include "MetasoundFrontendTransform.h"
+#include "MetasoundGeneratorHandle.h"
+#include "MetasoundLog.h"
+#include "MetasoundParameterTransmitter.h"
+#include "MetasoundSource.h"
+#include "MetasoundTrace.h"
+#include "MetasoundUObjectRegistry.h"
+#include "Algo/Find.h"
+
+#include UE_INLINE_GENERATED_CPP_BY_NAME(MetasoundBuilderSubsystem)
+
+
+namespace Metasound::Engine
+{
+	namespace BuilderSubsystemPrivate
+	{
+		static int32 DocumentBuilderEnabledCVar = 0;
+		FAutoConsoleVariableRef CVarEnableDocumentBuilder(
+			TEXT("au.MetaSound.BuilderSubsystem.Enabled"),
+			DocumentBuilderEnabledCVar,
+			TEXT("Enables the the Runtime MetaSoundBuilderSubsystem, which (as of 5.3) is in beta (Default: Disabled).\n")
+			TEXT("0: Disabled, 1: Enabled"),
+			ECVF_Default);
+
+		template <typename TLiteralType>
+		FMetasoundFrontendLiteral CreatePODMetaSoundLiteral(const TLiteralType& Value, FName& OutDataType)
+		{
+			OutDataType = GetMetasoundDataTypeName<TLiteralType>();
+
+			FMetasoundFrontendLiteral Literal;
+			Literal.Set(Value);
+			return Literal;
+		}
+	} // namespace BuilderSubsystemPrivate
+} // namespace Metasound::Engine
+
+FMetaSoundBuilderNodeOutputHandle UMetaSoundBuilderBase::AddGraphInputNode(FName Name, FName DataType, FMetasoundFrontendLiteral DefaultValue, EMetaSoundBuilderResult& OutResult, bool bIsConstructorInput)
+{
+	FMetasoundFrontendClassInput Description;
+	Description.Name = Name;
+	Description.TypeName = DataType;
+	Description.NodeID = FGuid::NewGuid();
+	Description.VertexID = FGuid::NewGuid();
+	Description.DefaultLiteral = static_cast<FMetasoundFrontendLiteral>(DefaultValue);
+	Description.AccessType = bIsConstructorInput ? EMetasoundFrontendVertexAccessType::Value : EMetasoundFrontendVertexAccessType::Reference;
+
+	FMetaSoundBuilderNodeOutputHandle NewHandle;
+
+	const FMetasoundFrontendNode* Node = Builder.FindGraphInputNode(Name);
+	if (Node)
+	{
+		UE_LOG(LogMetaSound, Warning, TEXT("AddGraphInputNode Failed: Input Node already exists with name '%s'; returning handle to existing node which may or may not match requested DataType '%s'"), *Name.ToString(), *DataType.ToString());
+	}
+	else
+	{
+		Node = Builder.AddGraphInput(Description);
+	}
+
+	if (Node)
+	{
+		const TArray<FMetasoundFrontendVertex>& Outputs = Node->Interface.Outputs;
+		checkf(!Outputs.IsEmpty(), TEXT("Node should be initialized and have one output."));
+
+		NewHandle.NodeID = Node->GetID();
+		NewHandle.VertexID = Outputs.Last().VertexID;
+	}
+
+	OutResult = NewHandle.IsSet() ? EMetaSoundBuilderResult::Succeeded : EMetaSoundBuilderResult::Failed;
+	return NewHandle;
+}
+
+FMetaSoundBuilderNodeInputHandle UMetaSoundBuilderBase::AddGraphOutputNode(FName Name, FName DataType, FMetasoundFrontendLiteral DefaultValue, EMetaSoundBuilderResult& OutResult, bool bIsConstructorInput)
+{
+	FMetasoundFrontendClassOutput Description;
+	Description.Name = Name;
+	Description.TypeName = DataType;
+	Description.NodeID = FGuid::NewGuid();
+	Description.VertexID = FGuid::NewGuid();
+	Description.AccessType = bIsConstructorInput ? EMetasoundFrontendVertexAccessType::Value : EMetasoundFrontendVertexAccessType::Reference;
+
+	FMetaSoundBuilderNodeInputHandle NewHandle;
+	const FMetasoundFrontendNode* Node = Builder.FindGraphOutputNode(Name);
+	if (Node)
+	{
+		UE_LOG(LogMetaSound, Warning, TEXT("AddGraphOutputNode Failed: Output Node already exists with name '%s'; returning handle to existing node which may or may not match requested DataType '%s'"), *Name.ToString(), *DataType.ToString());
+	}
+	else
+	{
+		Node = Builder.AddGraphOutput(Description);
+	}
+
+	if (Node)
+	{
+		const TArray<FMetasoundFrontendVertex>& Inputs = Node->Interface.Inputs;
+		checkf(!Inputs.IsEmpty(), TEXT("Node should be initialized and have one input."));
+
+		const FGuid& VertexID = Inputs.Last().VertexID;
+		if (Builder.SetNodeInputDefault(Node->GetID(), VertexID, DefaultValue))
+		{
+			NewHandle.NodeID = Node->GetID();
+			NewHandle.VertexID = VertexID;
+		}
+	}
+
+	OutResult = NewHandle.IsSet() ? EMetaSoundBuilderResult::Succeeded : EMetaSoundBuilderResult::Failed;
+	return NewHandle;
+}
+
+void UMetaSoundBuilderBase::AddInterface(FName InterfaceName, EMetaSoundBuilderResult& OutResult)
+{
+	const bool bInterfaceAdded = Builder.AddInterface(InterfaceName);
+	OutResult = bInterfaceAdded ? EMetaSoundBuilderResult::Succeeded : EMetaSoundBuilderResult::Failed;
+}
+
+FMetaSoundNodeHandle UMetaSoundBuilderBase::AddNode(TScriptInterface<IMetaSoundDocumentInterface> NodeClass, EMetaSoundBuilderResult& OutResult)
+{
+	using namespace Metasound;
+
+	FMetaSoundNodeHandle NewHandle;
+
+	if (NodeClass)
+	{
+		if (FMetasoundAssetBase* MetaSoundAsset = IMetasoundUObjectRegistry::Get().GetObjectAsAssetBase(NodeClass.GetObject()))
+		{
+			MetaSoundAsset->RegisterGraphWithFrontend();
+
+			const IMetaSoundDocumentInterface* Interface = NodeClass.GetInterface();
+			check(Interface);
+			const FMetasoundFrontendDocument& NodeClassDoc = Interface->GetDocument();
+			const FMetasoundFrontendGraphClass& NodeClassGraph = NodeClassDoc.RootGraph;
+			if (const FMetasoundFrontendNode* NewNode = Builder.AddGraphNode(NodeClassGraph))
+			{
+				FMetasoundFrontendClassMetadata NodeClassMetadata = NodeClassGraph.Metadata;
+				NewHandle.NodeID = NewNode->GetID();
+			}
+		}
+	}
+
+	OutResult = NewHandle.IsSet() ? EMetaSoundBuilderResult::Succeeded : EMetaSoundBuilderResult::Failed;
+	return NewHandle;
+}
+
+bool UMetaSoundBuilderBase::ContainsNode(const FMetaSoundNodeHandle& NodeHandle) const
+{
+	return Builder.ContainsNode(NodeHandle.NodeID);
+}
+
+bool UMetaSoundBuilderBase::ContainsNodeInput(const FMetaSoundBuilderNodeInputHandle& InputHandle) const
+{
+	return Builder.FindNodeInput(InputHandle.NodeID, InputHandle.VertexID) != nullptr;
+}
+
+bool UMetaSoundBuilderBase::ContainsNodeOutput(const FMetaSoundBuilderNodeOutputHandle& OutputHandle) const
+{
+	return Builder.FindNodeOutput(OutputHandle.NodeID, OutputHandle.VertexID) != nullptr;
+}
+
+void UMetaSoundBuilderBase::ConnectNodes(const FMetaSoundBuilderNodeOutputHandle& NodeOutput, const FMetaSoundBuilderNodeInputHandle& NodeInput, EMetaSoundBuilderResult& OutResult)
+{
+	const FMetasoundFrontendEdge* NewEdge = Builder.AddEdge(FMetasoundFrontendEdge
+	{
+		NodeOutput.NodeID,
+		NodeOutput.VertexID,
+		NodeInput.NodeID,
+		NodeInput.VertexID,
+	});
+	OutResult = NewEdge ? EMetaSoundBuilderResult::Succeeded : EMetaSoundBuilderResult::Failed;
+}
+
+void UMetaSoundBuilderBase::ConnectNodesByInterfaceBindings(const FMetaSoundNodeHandle& FromNode, const FMetaSoundNodeHandle& ToNode, EMetaSoundBuilderResult& OutResult)
+{
+	const bool bEdgesAdded = Builder.AddEdgesByNodeClassInterfaceBindings(FromNode.NodeID, ToNode.NodeID);
+	OutResult = bEdgesAdded ? EMetaSoundBuilderResult::Succeeded : EMetaSoundBuilderResult::Failed;
+}
+
+TArray<FMetaSoundBuilderNodeInputHandle> UMetaSoundBuilderBase::ConnectNodeOutputsToMatchingGraphInterfaceOutputs(const FMetaSoundNodeHandle& NodeHandle, EMetaSoundBuilderResult& OutResult)
+{
+	TArray<const FMetasoundFrontendEdge*> NewEdges;
+	const bool bEdgesAdded = Builder.AddEdgesFromMatchingInterfaceNodeOutputsToGraphOutputs(NodeHandle.NodeID, NewEdges);
+	OutResult = bEdgesAdded ? EMetaSoundBuilderResult::Succeeded : EMetaSoundBuilderResult::Failed;
+
+	TArray<FMetaSoundBuilderNodeInputHandle> ConnectedVertices;
+	Algo::Transform(NewEdges, ConnectedVertices, [this](const FMetasoundFrontendEdge* NewEdge)
+	{
+		const FMetasoundFrontendVertex* Vertex = Builder.FindNodeInput(NewEdge->ToNodeID, NewEdge->ToVertexID);
+		checkf(Vertex, TEXT("Edge connection reported success but vertex not found."));
+		return FMetaSoundBuilderNodeInputHandle { NewEdge->ToNodeID, Vertex->VertexID };
+	});
+
+	return ConnectedVertices;
+}
+
+TArray<FMetaSoundBuilderNodeOutputHandle> UMetaSoundBuilderBase::ConnectNodeInputsToMatchingGraphInterfaceInputs(const FMetaSoundNodeHandle& NodeHandle, EMetaSoundBuilderResult& OutResult)
+{
+	TArray<const FMetasoundFrontendEdge*> NewEdges;
+	const bool bEdgesAdded = Builder.AddEdgesFromMatchingInterfaceNodeInputsToGraphInputs(NodeHandle.NodeID, NewEdges);
+	OutResult = bEdgesAdded ? EMetaSoundBuilderResult::Succeeded : EMetaSoundBuilderResult::Failed;
+
+	TArray<FMetaSoundBuilderNodeOutputHandle> ConnectedVertices;
+	Algo::Transform(NewEdges, ConnectedVertices, [this](const FMetasoundFrontendEdge* NewEdge)
+	{
+		const FMetasoundFrontendVertex* Vertex = Builder.FindNodeInput(NewEdge->FromNodeID, NewEdge->FromVertexID);
+		checkf(Vertex, TEXT("Edge connection reported success but vertex not found."));
+		return FMetaSoundBuilderNodeOutputHandle { NewEdge->ToNodeID, Vertex->VertexID };
+	});
+
+	return ConnectedVertices;
+}
+
+void UMetaSoundBuilderBase::ConnectNodeOutputToGraphOutput(FName GraphOutputName, const FMetaSoundBuilderNodeOutputHandle& NodeOutput, EMetaSoundBuilderResult& OutResult)
+{
+	OutResult = EMetaSoundBuilderResult::Failed;
+
+	if (const FMetasoundFrontendNode* GraphOutputNode = Builder.FindGraphOutputNode(GraphOutputName))
+	{
+		const FMetasoundFrontendVertex& InputVertex = GraphOutputNode->Interface.Inputs.Last();
+		const FMetasoundFrontendEdge* EdgeAdded = Builder.AddEdge(FMetasoundFrontendEdge
+		{
+			NodeOutput.NodeID,
+			NodeOutput.VertexID,
+			GraphOutputNode->GetID(),
+			InputVertex.VertexID
+		});
+		OutResult = EdgeAdded ? EMetaSoundBuilderResult::Succeeded : EMetaSoundBuilderResult::Failed;
+	}
+}
+
+void UMetaSoundBuilderBase::ConnectNodeInputToGraphInput(FName GraphInputName, const FMetaSoundBuilderNodeInputHandle& NodeInput, EMetaSoundBuilderResult& OutResult)
+{
+	const FMetasoundFrontendNode* GraphInputNode = Builder.FindGraphInputNode(GraphInputName);
+	if (!GraphInputNode)
+	{
+		OutResult = EMetaSoundBuilderResult::Failed;
+		return;
+	}
+
+	const FMetasoundFrontendVertex& OutputVertex = GraphInputNode->Interface.Outputs.Last();
+	const FMetasoundFrontendEdge* EdgeAdded = Builder.AddEdge(FMetasoundFrontendEdge
+	{
+		GraphInputNode->GetID(),
+		OutputVertex.VertexID,
+		NodeInput.NodeID,
+		NodeInput.VertexID
+	});
+	OutResult = EdgeAdded ? EMetaSoundBuilderResult::Succeeded : EMetaSoundBuilderResult::Failed;
+}
+
+void UMetaSoundBuilderBase::DisconnectNodes(const FMetaSoundBuilderNodeOutputHandle& NodeOutput, const FMetaSoundBuilderNodeInputHandle& NodeInput, EMetaSoundBuilderResult& OutResult)
+{
+	const bool bEdgeRemoved = Builder.RemoveEdge(FMetasoundFrontendEdge
+	{
+		NodeOutput.NodeID,
+		NodeOutput.VertexID,
+		NodeInput.NodeID,
+		NodeInput.VertexID,
+	});
+	OutResult = bEdgeRemoved ? EMetaSoundBuilderResult::Succeeded : EMetaSoundBuilderResult::Failed;
+}
+
+void UMetaSoundBuilderBase::DisconnectNodeInput(const FMetaSoundBuilderNodeInputHandle& NodeInput, EMetaSoundBuilderResult& OutResult)
+{
+	const bool bEdgeRemoved = Builder.RemoveEdgeToNodeInput(NodeInput.NodeID, NodeInput.VertexID);
+	OutResult = bEdgeRemoved ? EMetaSoundBuilderResult::Succeeded : EMetaSoundBuilderResult::Failed;
+}
+
+void UMetaSoundBuilderBase::DisconnectNodeOutput(const FMetaSoundBuilderNodeOutputHandle& NodeOutput, EMetaSoundBuilderResult& OutResult)
+{
+	const bool bEdgeRemoved = Builder.RemoveEdgesFromNodeOutput(NodeOutput.NodeID, NodeOutput.VertexID);
+	OutResult = bEdgeRemoved ? EMetaSoundBuilderResult::Succeeded : EMetaSoundBuilderResult::Failed;
+}
+
+FMetaSoundBuilderNodeInputHandle UMetaSoundBuilderBase::FindNodeInputByName(const FMetaSoundNodeHandle& NodeHandle, FName InputName, EMetaSoundBuilderResult& OutResult) const
+{
+	if (const FMetasoundFrontendNode* Node = Builder.FindNode(NodeHandle.NodeID))
+	{
+		const TArray<FMetasoundFrontendVertex>& InputVertices = Node->Interface.Inputs;
+
+		auto FindByNamePredicate = [&InputName](const FMetasoundFrontendVertex& Vertex) { return Vertex.Name == InputName; };
+		if (const FMetasoundFrontendVertex* Input = InputVertices.FindByPredicate(FindByNamePredicate))
+		{
+			OutResult = EMetaSoundBuilderResult::Succeeded;
+			return FMetaSoundBuilderNodeInputHandle { Node->GetID(), Input->VertexID };
+		}
+
+		UE_LOG(LogMetaSound, Display, TEXT("Failed to find node input by name with builder '%s'. Node of class '%s' has no input named '%s'"), *GetName(), *InputName.ToString());
+	}
+
+	UE_LOG(LogMetaSound, Display, TEXT("Failed to find node input by name with builder '%s'. Node of with ID '%s' not found"), *GetName(), *NodeHandle.NodeID.ToString());
+	OutResult = EMetaSoundBuilderResult::Failed;
+	return { };
+}
+
+TArray<FMetaSoundBuilderNodeInputHandle> UMetaSoundBuilderBase::FindNodeInputs(const FMetaSoundNodeHandle& NodeHandle, EMetaSoundBuilderResult& OutResult) const
+{
+	return FindNodeInputsByDataType(NodeHandle, OutResult, { });
+}
+
+TArray<FMetaSoundBuilderNodeInputHandle> UMetaSoundBuilderBase::FindNodeInputsByDataType(const FMetaSoundNodeHandle& NodeHandle, EMetaSoundBuilderResult& OutResult, FName DataType) const
+{
+	TArray<FMetaSoundBuilderNodeInputHandle> FoundVertices;
+	if (Builder.ContainsNode(NodeHandle.NodeID))
+	{
+		TArray<const FMetasoundFrontendVertex*> Vertices = Builder.FindNodeInputs(NodeHandle.NodeID, DataType);
+		Algo::Transform(Vertices, FoundVertices, [&NodeHandle](const FMetasoundFrontendVertex* Vertex)
+		{
+			return FMetaSoundBuilderNodeInputHandle { NodeHandle.NodeID, Vertex->VertexID };
+		});
+		OutResult = EMetaSoundBuilderResult::Succeeded;
+	}
+	else
+	{
+		UE_LOG(LogMetaSound, Display, TEXT("Failed to find node inputs by data type with builder '%s'. Node of with ID '%s' not found"), *GetName(), *NodeHandle.NodeID.ToString());
+		OutResult = EMetaSoundBuilderResult::Failed;
+	}
+
+	return FoundVertices;
+}
+
+FMetaSoundBuilderNodeOutputHandle UMetaSoundBuilderBase::FindNodeOutputByName(const FMetaSoundNodeHandle& NodeHandle, FName OutputName, EMetaSoundBuilderResult& OutResult) const
+{
+	if (const FMetasoundFrontendNode* Node = Builder.FindNode(NodeHandle.NodeID))
+	{
+		const TArray<FMetasoundFrontendVertex>& OutputVertices = Node->Interface.Outputs;
+
+		auto FindByNamePredicate = [&OutputName](const FMetasoundFrontendVertex& Vertex) { return Vertex.Name == OutputName; };
+		if (const FMetasoundFrontendVertex* Output = OutputVertices.FindByPredicate(FindByNamePredicate))
+		{
+			OutResult = EMetaSoundBuilderResult::Succeeded;
+			return FMetaSoundBuilderNodeOutputHandle { Node->GetID(), Output->VertexID };
+		}
+		UE_LOG(LogMetaSound, Display, TEXT("Failed to find node output by name with builder '%s'. Node of class '%s' has no output named '%s'"), *GetName(), *OutputName.ToString());
+	}
+
+	UE_LOG(LogMetaSound, Display, TEXT("Failed to find node output by name with builder '%s'. Node of with ID '%s' not found"), *GetName(), *NodeHandle.NodeID.ToString());
+	OutResult = EMetaSoundBuilderResult::Failed;
+	return { };
+}
+
+TArray<FMetaSoundBuilderNodeOutputHandle> UMetaSoundBuilderBase::FindNodeOutputs(const FMetaSoundNodeHandle& NodeHandle, EMetaSoundBuilderResult& OutResult) const
+{
+	return FindNodeOutputsByDataType(NodeHandle, OutResult, { });
+}
+
+TArray<FMetaSoundBuilderNodeOutputHandle> UMetaSoundBuilderBase::FindNodeOutputsByDataType(const FMetaSoundNodeHandle& NodeHandle, EMetaSoundBuilderResult& OutResult, FName DataType) const
+{
+	TArray<FMetaSoundBuilderNodeOutputHandle> FoundVertices;
+	if (Builder.ContainsNode(NodeHandle.NodeID))
+	{
+		TArray<const FMetasoundFrontendVertex*> Vertices = Builder.FindNodeOutputs(NodeHandle.NodeID, DataType);
+		Algo::Transform(Vertices, FoundVertices, [&NodeHandle](const FMetasoundFrontendVertex* Vertex)
+		{
+			return FMetaSoundBuilderNodeOutputHandle { NodeHandle.NodeID, Vertex->VertexID };
+		});
+		OutResult = EMetaSoundBuilderResult::Succeeded;
+	}
+	else
+	{
+		UE_LOG(LogMetaSound, Display, TEXT("Failed to find node outputs by data type with builder '%s'. Node of with ID '%s' not found"), *GetName(), *NodeHandle.NodeID.ToString());
+		OutResult = EMetaSoundBuilderResult::Failed;
+	}
+
+	return FoundVertices;
+}
+
+TArray<FMetaSoundNodeHandle> UMetaSoundBuilderBase::FindInterfaceInputNodes(FName InterfaceName, EMetaSoundBuilderResult& OutResult) const
+{
+	TArray<FMetaSoundNodeHandle> NodeHandles;
+
+	TArray<const FMetasoundFrontendNode*> Nodes;
+	if (Builder.FindInterfaceInputNodes(InterfaceName, Nodes))
+	{
+		Algo::Transform(Nodes, NodeHandles, [this](const FMetasoundFrontendNode* Node)
+		{
+			check(Node);
+			return FMetaSoundNodeHandle { Node->GetID() };
+		});
+		OutResult = EMetaSoundBuilderResult::Succeeded;
+	}
+	else
+	{
+		UE_LOG(LogMetaSound, Display, TEXT("'%s' interface not found on builder '%s'. No input nodes returned"), *InterfaceName.ToString(), *GetName());
+		OutResult = EMetaSoundBuilderResult::Failed;
+	}
+
+	return NodeHandles;
+}
+
+TArray<FMetaSoundNodeHandle> UMetaSoundBuilderBase::FindInterfaceOutputNodes(FName InterfaceName, EMetaSoundBuilderResult& OutResult) const
+{
+	TArray<FMetaSoundNodeHandle> NodeHandles;
+
+	TArray<const FMetasoundFrontendNode*> Nodes;
+	if (Builder.FindInterfaceOutputNodes(InterfaceName, Nodes))
+	{
+		Algo::Transform(Nodes, NodeHandles, [this](const FMetasoundFrontendNode* Node)
+		{
+			check(Node);
+			return FMetaSoundNodeHandle { Node->GetID() };
+		});
+		OutResult = EMetaSoundBuilderResult::Succeeded;
+	}
+	else
+	{
+		OutResult = EMetaSoundBuilderResult::Failed;
+	}
+
+	return NodeHandles;
+}
+
+FMetaSoundNodeHandle UMetaSoundBuilderBase::FindGraphInputNode(FName InputName, EMetaSoundBuilderResult& OutResult) const
+{
+	if (const FMetasoundFrontendNode* GraphInputNode = Builder.FindGraphInputNode(InputName))
+	{
+		return FMetaSoundNodeHandle { GraphInputNode->GetID() };
+	}
+
+	UE_LOG(LogMetaSound, Display, TEXT("Failed to find graph input by name '%s' with builder '%s'"), *InputName.ToString(), *GetName());
+	OutResult = EMetaSoundBuilderResult::Failed;
+	return { };
+}
+
+FMetaSoundNodeHandle UMetaSoundBuilderBase::FindGraphOutputNode(FName OutputName, EMetaSoundBuilderResult& OutResult) const
+{
+	if (const FMetasoundFrontendNode* GraphOutputNode = Builder.FindGraphOutputNode(OutputName))
+	{
+		OutResult = EMetaSoundBuilderResult::Succeeded;
+		return FMetaSoundNodeHandle{ GraphOutputNode->GetID() };
+	}
+
+	UE_LOG(LogMetaSound, Display, TEXT("Failed to find graph output by name '%s' with builder '%s'"), *OutputName.ToString(), *GetName());
+	OutResult = EMetaSoundBuilderResult::Failed;
+	return { };
+}
+
+void UMetaSoundBuilderBase::InitNodeLocations()
+{
+	Builder.InitNodeLocations();
+}
+
+bool UMetaSoundBuilderBase::InterfaceIsDeclared(FName InterfaceName) const
+{
+	return Builder.IsInterfaceDeclared(InterfaceName);
+}
+
+bool UMetaSoundBuilderBase::NodesAreConnected(const FMetaSoundBuilderNodeOutputHandle& OutputHandle, const FMetaSoundBuilderNodeInputHandle& InputHandle) const
+{
+	const FMetasoundFrontendEdge Edge = { OutputHandle.NodeID, OutputHandle.VertexID, InputHandle.NodeID, InputHandle.VertexID };
+	return Builder.ContainsEdge(Edge);
+}
+
+bool UMetaSoundBuilderBase::NodeInputIsConnected(const FMetaSoundBuilderNodeInputHandle& InputHandle) const
+{
+	return Builder.IsNodeInputConnected(InputHandle.NodeID, InputHandle.VertexID);
+}
+
+bool UMetaSoundBuilderBase::NodeOutputIsConnected(const FMetaSoundBuilderNodeOutputHandle& OutputHandle) const
+{
+	return Builder.IsNodeOutputConnected(OutputHandle.NodeID, OutputHandle.VertexID);
+}
+
+void UMetaSoundBuilderBase::PostInitProperties()
+{
+	if (!HasAnyFlags(RF_ClassDefaultObject))
+	{
+		UClass& MetaSoundUClass = GetBuilderUClass();
+		Builder.InitDocument(MetaSoundUClass.GetFName());
+	}
+
+	Super::PostInitProperties();
+}
+
+void UMetaSoundBuilderBase::RemoveInterface(FName InterfaceName, EMetaSoundBuilderResult& OutResult)
+{
+	const bool bInterfaceRemoved = Builder.RemoveInterface(InterfaceName);
+	OutResult = bInterfaceRemoved ? EMetaSoundBuilderResult::Succeeded : EMetaSoundBuilderResult::Failed;
+}
+
+void UMetaSoundBuilderBase::RemoveNode(const FMetaSoundNodeHandle& NodeHandle, EMetaSoundBuilderResult& OutResult)
+{
+	const bool bNodeRemoved = Builder.RemoveNode(NodeHandle.NodeID);
+	OutResult = bNodeRemoved ? EMetaSoundBuilderResult::Succeeded : EMetaSoundBuilderResult::Failed;
+}
+
+#if WITH_EDITOR
+void UMetaSoundBuilderBase::SetAuthor(const FString& InAuthor)
+{
+	Builder.SetAuthor(InAuthor);
+}
+#endif // WITH_EDITOR
+
+void UMetaSoundBuilderBase::SetNodeInputDefault(const FMetaSoundBuilderNodeInputHandle& InputHandle, const FMetasoundFrontendLiteral& Literal, EMetaSoundBuilderResult& OutResult)
+{
+	using namespace Metasound;
+	using namespace Metasound::Frontend;
+
+	const bool bInputDefaultSet = Builder.SetNodeInputDefault(InputHandle.NodeID, InputHandle.VertexID, Literal);
+	OutResult = bInputDefaultSet ? EMetaSoundBuilderResult::Succeeded : EMetaSoundBuilderResult::Failed;
+}
+
+FMetaSoundNodeHandle UMetaSoundBuilderBase::FindNodeInputParent(const FMetaSoundBuilderNodeInputHandle& InputHandle, EMetaSoundBuilderResult& OutResult) const
+{
+	if (Builder.ContainsNode(InputHandle.NodeID))
+	{
+		OutResult = EMetaSoundBuilderResult::Succeeded;
+		return FMetaSoundNodeHandle { InputHandle.NodeID };
+	}
+
+	OutResult = EMetaSoundBuilderResult::Failed;
+	return { };
+}
+
+FMetaSoundNodeHandle UMetaSoundBuilderBase::FindNodeOutputParent(const FMetaSoundBuilderNodeOutputHandle& OutputHandle, EMetaSoundBuilderResult& OutResult) const
+{
+	if (Builder.ContainsNode(OutputHandle.NodeID))
+	{
+		OutResult = EMetaSoundBuilderResult::Succeeded;
+		return FMetaSoundNodeHandle{ OutputHandle.NodeID };
+	}
+
+	OutResult = EMetaSoundBuilderResult::Failed;
+	return { };
+}
+
+FMetasoundFrontendVersion UMetaSoundBuilderBase::FindNodeClassVersion(const FMetaSoundNodeHandle& NodeHandle, EMetaSoundBuilderResult& OutResult) const
+{
+	if (const FMetasoundFrontendNode* Node = Builder.FindNode(NodeHandle.NodeID))
+	{
+		if (const FMetasoundFrontendClass* Class = Builder.FindDependency(Node->ClassID))
+		{
+			OutResult = EMetaSoundBuilderResult::Succeeded;
+			return FMetasoundFrontendVersion { Class->Metadata.GetClassName().GetFullName(), Class->Metadata.GetVersion() };
+		}
+	}
+
+	OutResult = EMetaSoundBuilderResult::Failed;
+	return FMetasoundFrontendVersion::GetInvalid();
+}
+
+void UMetaSoundBuilderBase::GetNodeInputData(const FMetaSoundBuilderNodeInputHandle& InputHandle, FName& Name, FName& DataType, EMetaSoundBuilderResult& OutResult) const
+{
+	if (const FMetasoundFrontendVertex* Vertex = Builder.FindNodeInput(InputHandle.NodeID, InputHandle.VertexID))
+	{
+		Name = Vertex->Name;
+		DataType = Vertex->TypeName;
+		OutResult = EMetaSoundBuilderResult::Succeeded;
+	}
+	else
+	{
+		Name = { };
+		DataType = { };
+		OutResult = EMetaSoundBuilderResult::Failed;
+	}
+}
+
+void UMetaSoundBuilderBase::GetNodeOutputData(const FMetaSoundBuilderNodeOutputHandle& InputHandle, FName& Name, FName& DataType, EMetaSoundBuilderResult& OutResult) const
+{
+	if (const FMetasoundFrontendVertex* Vertex = Builder.FindNodeInput(InputHandle.NodeID, InputHandle.VertexID))
+	{
+		Name = Vertex->Name;
+		DataType = Vertex->TypeName;
+		OutResult = EMetaSoundBuilderResult::Succeeded;
+	}
+	else
+	{
+		Name = { };
+		DataType = { };
+		OutResult = EMetaSoundBuilderResult::Failed;
+	}
+}
+
+TScriptInterface<IMetaSoundDocumentInterface> UMetaSoundPatchBuilder::Build(UObject* Parent, const FMetaSoundBuilderOptions& BuildOptions) const
+{
+	return BuildInternal<UMetaSoundPatch>(Parent, BuildOptions);
+}
+
+UClass& UMetaSoundPatchBuilder::GetBuilderUClass() const
+{
+	return *UMetaSoundPatch::StaticClass();
+}
+
+UClass& UMetaSoundSourceBuilder::GetBuilderUClass() const
+{
+	return *UMetaSoundSource::StaticClass();
+}
+
+void UMetaSoundSourceBuilder::Audition(UObject* Parent, UAudioComponent* AudioComponent, FOnCreateAuditionGeneratorHandleDelegate CreateGenerator)
+{
+	METASOUND_TRACE_CPUPROFILER_EVENT_SCOPE(UMetaSoundSourceBuilder::Audition);
+
+	using namespace Metasound::Frontend;
+
+	if (!AudioComponent)
+	{
+		UE_LOG(LogMetaSound, Error, TEXT("Failed to audition MetaSoundBuilder '%s': No AudioComponent supplied"), *GetFullName());
+		return;
+	}
+
+	if (AuditionSound.IsValid())
+	{
+		const FMetasoundFrontendDocument& Doc = GetConstBuilder().GetDocument();
+		AuditionSound->SetDocument(Doc);
+		AuditionSound->ConformObjectDataToInterfaces();
+
+		if (USoundBase* Sound = AudioComponent->GetSound())
+		{
+			if (Sound != AuditionSound)
+			{
+				UE_LOG(LogMetaSound, Warning, TEXT("MetaSoundBuilder '%s' supplied AudioComponent with unlinked sound '%s'. Stopping sound and replacing with builder's audition sound."),
+					*GetFullName(),
+					*Sound->GetFullName());
+				AudioComponent->SetSound(AuditionSound.Get());
+			}
+		}
+		else
+		{
+			AudioComponent->SetSound(AuditionSound.Get());
+		}
+	}
+	else
+	{
+		FMetaSoundBuilderOptions BuildOptions;
+		BuildOptions.Name = MakeUniqueObjectName(nullptr, UMetaSoundSource::StaticClass(), FName(GetName() + TEXT("_Audition")));
+		BuildOptions.bAddToRegistry = false;
+
+		TScriptInterface<IMetaSoundDocumentInterface> DocumentInterface = Build(Parent, BuildOptions);
+		AuditionSound = Cast<UMetaSoundSource>(DocumentInterface.GetObject());
+		if (AuditionSound.IsValid())
+		{
+			AudioComponent->SetSound(AuditionSound.Get());
+		}
+	}
+
+	if (CreateGenerator.IsBound())
+	{
+		UMetasoundGeneratorHandle* NewHandle = UMetasoundGeneratorHandle::CreateMetaSoundGeneratorHandle(AudioComponent);
+		checkf(NewHandle, TEXT("BindToGeneratorDelegate Failed when attempting to audition MetaSoundSource builder '%s'"), *GetName());
+		CreateGenerator.Execute(NewHandle);
+	}
+
+	AudioComponent->Play();
+}
+
+TScriptInterface<IMetaSoundDocumentInterface> UMetaSoundSourceBuilder::Build(UObject* Parent, const FMetaSoundBuilderOptions& BuildOptions) const
+{
+	return BuildInternal<UMetaSoundSource>(Parent, BuildOptions);
+}
+
+const Metasound::Engine::FOutputAudioFormatInfoPair* UMetaSoundSourceBuilder::FindOutputAudioFormatInfo() const
+{
+	using namespace Metasound::Engine;
+
+	const FOutputAudioFormatInfoMap& FormatInfo = GetOutputAudioFormatInfo();
+
+	auto Predicate = [this](const FOutputAudioFormatInfoPair& Pair)
+	{
+		const FMetasoundFrontendDocument& Document = Builder.GetDocument();
+		return Document.Interfaces.Contains(Pair.Value.InterfaceVersion);
+	};
+
+	return Algo::FindByPredicate(FormatInfo, Predicate);
+}
+
+void UMetaSoundSourceBuilder::SetFormat(EMetaSoundOutputAudioFormat OutputFormat, EMetaSoundBuilderResult& OutResult)
+{
+	using namespace Metasound::Engine;
+	using namespace Metasound::Frontend;
+
+	// Convert to non-preset MetaSoundSource since interface data is being altered
+	Builder.ConvertFromPreset();
+
+	const FOutputAudioFormatInfoMap& FormatMap = GetOutputAudioFormatInfo();
+
+	// Determine which interfaces to add and remove from the document due to the
+	// output format being changed.
+	TArray<FMetasoundFrontendVersion> OutputFormatsToAdd;
+	if (const FOutputAudioFormatInfo* FormatInfo = FormatMap.Find(OutputFormat))
+	{
+		OutputFormatsToAdd.Add(FormatInfo->InterfaceVersion);
+	}
+
+	TArray<FMetasoundFrontendVersion> OutputFormatsToRemove;
+
+	const FMetasoundFrontendDocument& Document = GetConstBuilder().GetDocument();
+	for (const FOutputAudioFormatInfoPair& Pair : FormatMap)
+	{
+		const FMetasoundFrontendVersion& FormatVersion = Pair.Value.InterfaceVersion;
+		if (Document.Interfaces.Contains(FormatVersion))
+		{
+			if (!OutputFormatsToAdd.Contains(FormatVersion))
+			{
+				OutputFormatsToRemove.Add(FormatVersion);
+			}
+		}
+	}
+
+	const bool bSuccess = Builder.ModifyInterfaces(OutputFormatsToRemove, OutputFormatsToAdd);
+	OutResult = bSuccess ? EMetaSoundBuilderResult::Succeeded : EMetaSoundBuilderResult::Failed;
+}
+
+
+UMetaSoundPatchBuilder* UMetaSoundBuilderSubsystem::CreatePatchBuilder(FName BuilderName, EMetaSoundBuilderResult& OutResult)
+{
+	using namespace Metasound::Engine::BuilderSubsystemPrivate;
+
+	if (!DocumentBuilderEnabledCVar)
+	{
+		UE_LOG(LogMetaSound, Error, TEXT("CreatePatchBuilder Failed to create MetaSound Builder '%s': See CVar au.MetaSound.BuilderSubsystem.Enabled (%s')"), *BuilderName.ToString(), CVarEnableDocumentBuilder->GetHelp());
+		return nullptr;
+	}
+
+	const EObjectFlags NewObjectFlags = RF_Public | RF_Transient;
+	const FName ObjectName = MakeUniqueObjectName(GetTransientPackage(), UMetaSoundPatchBuilder::StaticClass());
+	TObjectPtr<UMetaSoundPatchBuilder> NewBuilder = NewObject<UMetaSoundPatchBuilder>(this, ObjectName, NewObjectFlags);
+	check(NewBuilder);
+	return NewBuilder;
+}
+
+UMetaSoundSourceBuilder* UMetaSoundBuilderSubsystem::CreateSourceBuilder(
+	FName BuilderName,
+	FMetaSoundBuilderNodeOutputHandle& OnPlayNodeOutput,
+	FMetaSoundBuilderNodeInputHandle& OnFinishedNodeInput,
+	TArray<FMetaSoundBuilderNodeInputHandle>& AudioOutNodeInputs,
+	EMetaSoundBuilderResult& OutResult,
+	EMetaSoundOutputAudioFormat OutputFormat,
+	bool bIsOneShot)
+{
+	using namespace Metasound::Frontend;
+	using namespace Metasound::Engine::BuilderSubsystemPrivate;
+
+	OnPlayNodeOutput = { };
+	OnFinishedNodeInput = { };
+	AudioOutNodeInputs.Reset();
+
+	OutResult = EMetaSoundBuilderResult::Failed;
+
+	if (!DocumentBuilderEnabledCVar)
+	{
+		UE_LOG(LogMetaSound, Error, TEXT("CreateSourceBuilder Failed to create MetaSound Builder '%s': See CVar au.MetaSound.BuilderSubsystem.Enabled (%s')"), *BuilderName.ToString(), CVarEnableDocumentBuilder->GetHelp());
+		return nullptr;
+	}
+
+	const EObjectFlags NewObjectFlags = RF_Public | RF_Transient;
+	UPackage* TransientPackage = GetTransientPackage();
+	const FName ObjectName = MakeUniqueObjectName(TransientPackage, UMetaSoundSourceBuilder::StaticClass());
+	TObjectPtr<UMetaSoundSourceBuilder> NewBuilder = NewObject<UMetaSoundSourceBuilder>(TransientPackage, ObjectName, NewObjectFlags);
+	check(NewBuilder);
+
+	bool bFormatSet = OutputFormat == EMetaSoundOutputAudioFormat::Mono;
+	if (!bFormatSet)
+	{
+		EMetaSoundBuilderResult Result = EMetaSoundBuilderResult::Failed;
+		NewBuilder->SetFormat(OutputFormat, Result);
+		bFormatSet = Result == EMetaSoundBuilderResult::Succeeded;
+	}
+
+	if (bFormatSet)
+	{
+		const Metasound::Engine::FOutputAudioFormatInfoPair* FormatInfo = NewBuilder->FindOutputAudioFormatInfo();
+		TArray<FMetaSoundNodeHandle> AudioOutputNodes;
+		if (FormatInfo)
+		{
+			AudioOutputNodes = NewBuilder->FindInterfaceOutputNodes(FormatInfo->Value.InterfaceVersion.Name, OutResult);
+		}
+
+		if (!AudioOutputNodes.IsEmpty())
+		{
+			Algo::Transform(AudioOutputNodes, AudioOutNodeInputs, [&NewBuilder, &BuilderName](const FMetaSoundNodeHandle& AudioOutputNode) -> FMetaSoundBuilderNodeInputHandle
+			{
+				EMetaSoundBuilderResult Result;
+				TArray<FMetaSoundBuilderNodeInputHandle> Inputs = NewBuilder->FindNodeInputs(AudioOutputNode, Result);
+				if (!Inputs.IsEmpty())
+				{
+					return Inputs.Last();
+				}
+
+				UE_LOG(LogMetaSound, Error, TEXT("Builder '%s' Creation Error: Failed to find expected audio output node input vertex. Returned vertices set may be incomplete."), *BuilderName.ToString());
+				return { };
+			});
+		}
+		else
+		{
+			UE_LOG(LogMetaSound, Error, TEXT("Builder '%s' Creation Error: Failed to find expected audio output format and/or associated output nodes."), *BuilderName.ToString());
+			return nullptr;
+		}
+	}
+	else
+	{
+		UE_LOG(LogMetaSound, Error, TEXT("Builder '%s' Creation Error: Failed to set output audio output format when initializing."), *BuilderName.ToString());
+		return nullptr;
+	}
+
+	{
+		FMetaSoundNodeHandle OnPlayNode = NewBuilder->FindGraphInputNode(SourceInterface::Inputs::OnPlay, OutResult);
+		if (OutResult == EMetaSoundBuilderResult::Failed)
+		{
+			UE_LOG(LogMetaSound, Error, TEXT("Builder '%s' Creation Error: Failed to add required interface '%s' when attempting to create MetaSound Source Builder"), *BuilderName.ToString(), * SourceInterface::GetVersion().ToString());
+			return nullptr;
+		}
+
+		TArray<FMetaSoundBuilderNodeOutputHandle> Outputs = NewBuilder->FindNodeOutputs(OnPlayNode, OutResult);
+		if (OutResult == EMetaSoundBuilderResult::Failed)
+		{
+			UE_LOG(LogMetaSound, Error, TEXT("Builder '%s' Creation Error: Failed to find output vertex for 'OnPlay' input node when attempting to create MetaSound Source Builder"), *BuilderName.ToString());
+			return nullptr;
+		}
+
+		check(!Outputs.IsEmpty());
+		OnPlayNodeOutput = Outputs.Last();
+	}
+
+	{
+		if (bIsOneShot)
+		{
+			FMetaSoundNodeHandle OnFinishedNode = NewBuilder->FindGraphOutputNode(SourceOneShotInterface::Outputs::OnFinished, OutResult);
+			if (OutResult == EMetaSoundBuilderResult::Failed)
+			{
+				UE_LOG(LogMetaSound, Error, TEXT("Builder '%s' Creation Error: Failed to add '%s' interface; interface definition may not be registered."), *BuilderName.ToString(), *SourceOneShotInterface::GetVersion().ToString());
+			}
+
+			TArray<FMetaSoundBuilderNodeInputHandle> Inputs = NewBuilder->FindNodeInputs(OnFinishedNode, OutResult);
+			if (OutResult == EMetaSoundBuilderResult::Failed)
+			{
+				UE_LOG(LogMetaSound, Error, TEXT("Builder '%s' Creation Error: Failed to find input vertex for 'OnFinished' output node when attempting to create MetaSound Source Builder"), *BuilderName.ToString());
+				return nullptr;
+			}
+
+			check(!Inputs.IsEmpty());
+			OnFinishedNodeInput = Inputs.Last();
+		}
+		else
+		{
+			NewBuilder->RemoveInterface(SourceOneShotInterface::GetVersion().Name, OutResult);
+		}
+	}
+
+	OutResult = EMetaSoundBuilderResult::Succeeded;
+	return NewBuilder;
+}
+
+UMetaSoundSourceBuilder* UMetaSoundBuilderSubsystem::AttachSourceBuilderToAsset(UMetaSoundSource* InSource) const
+{
+	if (InSource)
+	{
+		TObjectPtr<UMetaSoundSourceBuilder> NewBuilder = NewObject<UMetaSoundSourceBuilder>(InSource);
+		TScriptInterface<IMetaSoundDocumentInterface> MetaSoundSourceDocInterface = InSource;
+		NewBuilder->Builder = FMetaSoundFrontendDocumentBuilder(MetaSoundSourceDocInterface);
+		return NewBuilder;
+	}
+
+	return nullptr;
+}
+
+FMetasoundFrontendLiteral UMetaSoundBuilderSubsystem::CreateBoolMetaSoundLiteral(bool Value, FName& OutDataType)
+{
+	return Metasound::Engine::BuilderSubsystemPrivate::CreatePODMetaSoundLiteral(Value, OutDataType);
+}
+
+FMetasoundFrontendLiteral UMetaSoundBuilderSubsystem::CreateBoolArrayMetaSoundLiteral(const TArray<bool>& Value, FName& OutDataType)
+{
+	return Metasound::Engine::BuilderSubsystemPrivate::CreatePODMetaSoundLiteral(Value, OutDataType);
+}
+
+FMetasoundFrontendLiteral UMetaSoundBuilderSubsystem::CreateFloatMetaSoundLiteral(float Value, FName& OutDataType)
+{
+	return Metasound::Engine::BuilderSubsystemPrivate::CreatePODMetaSoundLiteral(Value, OutDataType);
+}
+
+FMetasoundFrontendLiteral UMetaSoundBuilderSubsystem::CreateFloatArrayMetaSoundLiteral(const TArray<float>& Value, FName& OutDataType)
+{
+	return Metasound::Engine::BuilderSubsystemPrivate::CreatePODMetaSoundLiteral(Value, OutDataType);
+}
+
+FMetasoundFrontendLiteral UMetaSoundBuilderSubsystem::CreateIntMetaSoundLiteral(int32 Value, FName& OutDataType)
+{
+	return Metasound::Engine::BuilderSubsystemPrivate::CreatePODMetaSoundLiteral(Value, OutDataType);
+}
+
+FMetasoundFrontendLiteral UMetaSoundBuilderSubsystem::CreateIntArrayMetaSoundLiteral(const TArray<int32>& Value, FName& OutDataType)
+{
+	return Metasound::Engine::BuilderSubsystemPrivate::CreatePODMetaSoundLiteral(Value, OutDataType);
+}
+
+FMetasoundFrontendLiteral UMetaSoundBuilderSubsystem::CreateStringMetaSoundLiteral(const FString& Value, FName& OutDataType)
+{
+	return Metasound::Engine::BuilderSubsystemPrivate::CreatePODMetaSoundLiteral(Value, OutDataType);
+}
+
+FMetasoundFrontendLiteral UMetaSoundBuilderSubsystem::CreateStringArrayMetaSoundLiteral(const TArray<FString>& Value, FName& OutDataType)
+{
+	return Metasound::Engine::BuilderSubsystemPrivate::CreatePODMetaSoundLiteral(Value, OutDataType);
+}
+
+FMetasoundFrontendLiteral UMetaSoundBuilderSubsystem::CreateObjectMetaSoundLiteral(UObject* Value)
+{
+	FMetasoundFrontendLiteral Literal;
+	Literal.Set(Value);
+	return Literal;
+}
+
+FMetasoundFrontendLiteral UMetaSoundBuilderSubsystem::CreateObjectArrayMetaSoundLiteral(const TArray<UObject*>& Value)
+{
+	FMetasoundFrontendLiteral Literal;
+	Literal.Set(Value);
+	return Literal;
+}
+
+FMetasoundFrontendLiteral UMetaSoundBuilderSubsystem::CreateMetaSoundLiteralFromParam(const FAudioParameter& Param)
+{
+	return FMetasoundFrontendLiteral { Param };
+}
+
+UMetaSoundPatchBuilder* UMetaSoundBuilderSubsystem::FindPatchBuilder(FName BuilderName)
+{
+	return PatchBuilders.FindRef(BuilderName);
+}
+
+UMetaSoundSourceBuilder* UMetaSoundBuilderSubsystem::FindSourceBuilder(FName BuilderName)
+{
+	return SourceBuilders.FindRef(BuilderName);
+}
+
+bool UMetaSoundBuilderSubsystem::IsInterfaceRegistered(FName InInterfaceName) const
+{
+	using namespace Metasound::Frontend;
+
+	FMetasoundFrontendInterface Interface;
+	return ISearchEngine::Get().FindInterfaceWithHighestVersion(InInterfaceName, Interface);
+}
+
+void UMetaSoundBuilderSubsystem::RegisterPatchBuilder(FName BuilderName, UMetaSoundPatchBuilder* Builder)
+{
+	PatchBuilders.FindOrAdd(BuilderName) = Builder;
+}
+
+void UMetaSoundBuilderSubsystem::RegisterSourceBuilder(FName BuilderName, UMetaSoundSourceBuilder* Builder)
+{
+	SourceBuilders.FindOrAdd(BuilderName) = Builder;
+}
+
+bool UMetaSoundBuilderSubsystem::UnregisterPatchBuilder(FName BuilderName)
+{
+	return PatchBuilders.Remove(BuilderName) > 0;
+}
+
+bool UMetaSoundBuilderSubsystem::UnregisterSourceBuilder(FName BuilderName)
+{
+	return SourceBuilders.Remove(BuilderName) > 0;
+}
diff --git a/Engine/Plugins/Runtime/Metasound/Source/MetasoundEngine/Private/MetasoundEngineArchetypes.cpp b/Engine/Plugins/Runtime/Metasound/Source/MetasoundEngine/Private/MetasoundEngineArchetypes.cpp
--- a/Engine/Plugins/Runtime/Metasound/Source/MetasoundEngine/Private/MetasoundEngineArchetypes.cpp
+++ b/Engine/Plugins/Runtime/Metasound/Source/MetasoundEngine/Private/MetasoundEngineArchetypes.cpp
@@ -3,7 +3,7 @@
 #include "MetasoundEngineArchetypes.h"
 
 #include "IAudioParameterInterfaceRegistry.h"
-#include "Interfaces/MetasoundFrontendOutputFormatInterfaces.h"
+#include "Interfaces/MetasoundOutputFormatInterfaces.h"
 #include "MetasoundAudioFormats.h"
 #include "MetasoundDataReference.h"
 #include "MetasoundFrontendArchetypeRegistry.h"
diff --git a/Engine/Plugins/Runtime/Metasound/Source/MetasoundEngine/Private/MetasoundEngineModule.cpp b/Engine/Plugins/Runtime/Metasound/Source/MetasoundEngine/Private/MetasoundEngineModule.cpp
--- a/Engine/Plugins/Runtime/Metasound/Source/MetasoundEngine/Private/MetasoundEngineModule.cpp
+++ b/Engine/Plugins/Runtime/Metasound/Source/MetasoundEngine/Private/MetasoundEngineModule.cpp
@@ -44,7 +44,8 @@ class FMetasoundEngineModule : public IMetasoundEngineModule
 		check(AudioSettings);
 		AudioSettings->RegisterParameterInterfaces();
 
-		// Register interfaces
+		// Register supported MetaSound UClasses & interfaces
+		Metasound::Engine::RegisterUClasses();
 		Metasound::Engine::RegisterInterfaces();
 
 		// Flush node registration queue
diff --git a/Engine/Plugins/Runtime/Metasound/Source/MetasoundEngine/Private/MetasoundInterface.cpp b/Engine/Plugins/Runtime/Metasound/Source/MetasoundEngine/Private/MetasoundInterface.cpp
--- a/Engine/Plugins/Runtime/Metasound/Source/MetasoundEngine/Private/MetasoundInterface.cpp
+++ b/Engine/Plugins/Runtime/Metasound/Source/MetasoundEngine/Private/MetasoundInterface.cpp
@@ -5,13 +5,16 @@
 #include "AudioParameterControllerInterface.h"
 #include "IAudioParameterInterfaceRegistry.h"
 #include "IAudioParameterTransmitter.h"
+#include "Interfaces/MetasoundInputFormatInterfaces.h"
+#include "Interfaces/MetasoundOutputFormatInterfaces.h"
 #include "Interfaces/MetasoundFrontendSourceInterface.h"
+#include "Metasound.h"
 #include "MetasoundEngineArchetypes.h"
+#include "MetasoundFrontendArchetypeRegistry.h"
 #include "MetasoundFrontendDataTypeRegistry.h"
 #include "MetasoundFrontendDocument.h"
 #include "MetasoundFrontendTransform.h"
 #include "MetasoundLog.h"
-#include "Interfaces/MetasoundFrontendOutputFormatInterfaces.h"
 #include "MetasoundParameterTransmitter.h"
 #include "MetasoundSource.h"
 #include "MetasoundUObjectRegistry.h"
@@ -25,17 +28,37 @@ namespace Metasound
 {
 	namespace Engine
 	{
-		struct FInterfaceRegistryOptions
+		struct FInterfaceRegistryUClassOptions
 		{
+			FName ClassName;
 			bool bIsDefault = false;
-			bool bEditorCanAddOrRemove = true;
+			bool bEditorCanAddOrRemove = false;
+		};
+
+		struct FInterfaceRegistryOptions
+		{
 			FName InputSystemName;
-			FName UClassName;
+			TArray<FInterfaceRegistryUClassOptions> UClassOptions;
 		};
 
 		// Entry for registered interface.
 		class FInterfaceRegistryEntry : public Frontend::IInterfaceRegistryEntry
 		{
+			const FInterfaceRegistryUClassOptions* FindOptionsByClassName(FName ClassName) const
+			{
+				auto NameIsEqual = [&ClassName](const FInterfaceRegistryUClassOptions& InOptions)
+				{
+					return InOptions.ClassName == ClassName;
+				};
+
+				return Options.UClassOptions.FindByPredicate(NameIsEqual);
+			}
+
+			static bool SortBindingsByConnectionPriorityPredicate(const FMetasoundFrontendInterfaceBinding& A, const FMetasoundFrontendInterfaceBinding& B)
+			{
+				return A.ConnectionPriority < B.ConnectionPriority;
+			}
+
 		public:
 			FInterfaceRegistryEntry(
 				const FMetasoundFrontendInterface& InInterface,
@@ -48,25 +71,41 @@ namespace Metasound
 			{
 			}
 
-			virtual bool EditorCanAddOrRemove() const override
+			virtual void AddOutputBindings(TArray<FMetasoundFrontendInterfaceBinding>&& InBindings) override
+			{
+				OutputBindings.Append(MoveTemp(InBindings));
+				OutputBindings.Sort(&SortBindingsByConnectionPriorityPredicate);
+			}
+
+			virtual bool EditorCanAddOrRemove(FName InUClassName) const override
 			{
-				return Options.bEditorCanAddOrRemove;
+				if (const FInterfaceRegistryUClassOptions* UClassOptions = FindOptionsByClassName(InUClassName))
+				{
+					return UClassOptions->bEditorCanAddOrRemove;
+				}
+
+				return false;
 			}
 
 			virtual bool UClassIsSupported(FName InUClassName) const override
 			{
-				if (Options.UClassName.IsNone())
+				if (const FInterfaceRegistryUClassOptions* UClassOptions = FindOptionsByClassName(InUClassName))
 				{
 					return true;
 				}
 
 				// TODO: Support child asset class types.
-				return Options.UClassName == InUClassName;
+				return false;
 			}
 
-			virtual bool IsDefault() const override
+			virtual bool IsDefault(FName InUClassName) const override
 			{
-				return Options.bIsDefault;
+				if (const FInterfaceRegistryUClassOptions* UClassOptions = FindOptionsByClassName(InUClassName))
+				{
+					return UClassOptions->bIsDefault;
+				}
+
+				return false;
 			}
 
 			virtual FName GetRouterName() const override
@@ -79,6 +118,24 @@ namespace Metasound
 				return Interface;
 			}
 
+			virtual const TArray<FMetasoundFrontendInterfaceBinding>& GetOutputBindings() const override
+			{
+				return OutputBindings;
+			}
+
+			virtual bool RemoveOutputBinding(const FMetasoundFrontendVersion& InInterfaceVersion) override
+			{
+				return OutputBindings.RemoveAll([&InInterfaceVersion](const FMetasoundFrontendInterfaceBinding& Binding)
+				{
+					return Binding.Version == InInterfaceVersion;
+				}) > 0;
+			}
+
+			virtual void ResetOutputBindings()
+			{
+				OutputBindings.Reset();
+			}
+
 			virtual bool UpdateRootGraphInterface(Frontend::FDocumentHandle InDocument) const override
 			{
 				if (UpdateTransform.IsValid())
@@ -90,28 +147,12 @@ namespace Metasound
 
 		private:
 			FMetasoundFrontendInterface Interface;
+			TArray<FMetasoundFrontendInterfaceBinding> OutputBindings;
 			TUniquePtr<Frontend::IDocumentTransform> UpdateTransform;
 			FInterfaceRegistryOptions Options;
 		};
 
-		template <typename UClassType>
-		void RegisterInterface(const FMetasoundFrontendInterface& InInterface, TUniquePtr<Frontend::IDocumentTransform>&& InUpdateTransform, bool bInIsDefault, bool bEditorCanAddOrRemove, FName InRouterName)
-		{
-			using namespace Frontend;
-
-			FInterfaceRegistryOptions Options
-			{
-				bInIsDefault,
-				bEditorCanAddOrRemove,
-				InRouterName,
-				UClassType::StaticClass()->GetFName()
-			};
-
-			IMetasoundUObjectRegistry::Get().RegisterUClassInterface(MakeUnique<TMetasoundUObjectRegistryEntry<UClassType>>(InInterface.Version));
-			IInterfaceRegistry::Get().RegisterInterface(MakeUnique<FInterfaceRegistryEntry>(InInterface, MoveTemp(InUpdateTransform), MoveTemp(Options)));
-		}
-
-		void RegisterInterface(Audio::FParameterInterfacePtr Interface, TUniquePtr<Frontend::IDocumentTransform>&& InUpdateTransform, bool bInIsDefault, bool bEditorCanAddOrRemove, FName InRouterName)
+		FMetasoundFrontendInterface ConvertParameterToFrontendInterface(const Audio::FParameterInterface& InInterface)
 		{
 			using namespace Frontend;
 
@@ -130,11 +171,11 @@ namespace Metasound
 			};
 
 			FMetasoundFrontendInterface FrontendInterface;
-			FrontendInterface.Version = { Interface->GetName(), FMetasoundFrontendVersionNumber { Interface->GetVersion().Major, Interface->GetVersion().Minor } };
+			FrontendInterface.Version = { InInterface.GetName(), FMetasoundFrontendVersionNumber { InInterface.GetVersion().Major, InInterface.GetVersion().Minor } };
 
 			// Transfer all input data from AudioExtension interface struct to FrontendInterface
 			{
-				Algo::Transform(Interface->GetInputs(), FrontendInterface.Inputs, [&](const Audio::FParameterInterface::FInput& Input)
+				Algo::Transform(InInterface.GetInputs(), FrontendInterface.Inputs, [&](const Audio::FParameterInterface::FInput& Input)
 				{
 					FMetasoundFrontendClassInput ClassInput;
 					ClassInput.Name = Input.InitValue.ParamName;
@@ -167,7 +208,7 @@ namespace Metasound
 
 			// Transfer all output data from AudioExtension interface struct to FrontendInterface
 			{
-				Algo::Transform(Interface->GetOutputs(), FrontendInterface.Outputs, [&](const Audio::FParameterInterface::FOutput& Output)
+				Algo::Transform(InInterface.GetOutputs(), FrontendInterface.Outputs, [&](const Audio::FParameterInterface::FOutput& Output)
 				{
 					FMetasoundFrontendClassOutput ClassOutput;
 					ClassOutput.Name = Output.ParamName;
@@ -197,7 +238,7 @@ namespace Metasound
 				});
 			}
 
-			Algo::Transform(Interface->GetEnvironment(), FrontendInterface.Environment, [&](const Audio::FParameterInterface::FEnvironmentVariable& Environment)
+			Algo::Transform(InInterface.GetEnvironment(), FrontendInterface.Environment, [&](const Audio::FParameterInterface::FEnvironmentVariable& Environment)
 			{
 				FMetasoundFrontendClassEnvironmentVariable EnvironmentVariable;
 				EnvironmentVariable.Name = Environment.ParamName;
@@ -209,72 +250,233 @@ namespace Metasound
 				return EnvironmentVariable;
 			});
 
-			const UClass* SourceClass = UMetaSoundSource::StaticClass();
-			check(SourceClass);
-			const bool bInterfaceSupportsSource = SourceClass->IsChildOf(&Interface->GetType());
-			if (ensureAlwaysMsgf(bInterfaceSupportsSource, TEXT("Interfaces defined using Audio::FParameterInterface currently only supported by UMetaSoundSource asset type.")))
+			return FrontendInterface;
+		}
+
+		void RegisterUClasses()
+		{
+			IMetasoundUObjectRegistry::Get().RegisterUClass(MakeUnique<TMetasoundUObjectRegistryEntry<UMetaSoundPatch>>());
+			IMetasoundUObjectRegistry::Get().RegisterUClass(MakeUnique<TMetasoundUObjectRegistryEntry<UMetaSoundSource>>());
+		}
+
+		void RegisterInterface(Audio::FParameterInterfacePtr Interface, TUniquePtr<Frontend::IDocumentTransform>&& InUpdateTransform, FInterfaceRegistryOptions&& InOptions)
+		{
+			using namespace Frontend;
+
+			const FMetasoundFrontendInterface FrontendInterface = ConvertParameterToFrontendInterface(*Interface);
+			IInterfaceRegistry::Get().RegisterInterface(MakeUnique<FInterfaceRegistryEntry>(FrontendInterface, MoveTemp(InUpdateTransform), MoveTemp(InOptions)));
+		}
+
+		void RegisterInterface(const FMetasoundFrontendInterface& InInterface, TUniquePtr<Frontend::IDocumentTransform>&& InUpdateTransform, FInterfaceRegistryOptions&& InOptions)
+		{
+			using namespace Frontend;
+			IInterfaceRegistry::Get().RegisterInterface(MakeUnique<FInterfaceRegistryEntry>(InInterface, MoveTemp(InUpdateTransform), MoveTemp(InOptions)));
+		}
+
+		template <typename UClassType>
+		void RegisterInterfaceForSingleClass(Audio::FParameterInterfacePtr Interface, TUniquePtr<Frontend::IDocumentTransform>&& InUpdateTransform, bool bInIsDefault = false, bool bInEditorCanAddOrRemove = false, FName InRouterName = IDataReference::RouterName)
+		{
+			using namespace Frontend;
+
+			const FMetasoundFrontendInterface FrontendInterface = ConvertParameterToFrontendInterface(*Interface);
+			RegisterInterface(FrontendInterface, MoveTemp(InUpdateTransform), FInterfaceRegistryOptions
 			{
-				RegisterInterface<UMetaSoundSource>(FrontendInterface, MoveTemp(InUpdateTransform), bInIsDefault, bEditorCanAddOrRemove, InRouterName);
-				UE_LOG(LogMetaSound, Verbose, TEXT("Interface '%s' registered for asset type '%s'."), *Interface->GetName().ToString(), *SourceClass->GetName());
-			}
+				InRouterName,
+				{
+					{
+						UClassType::StaticClass()->GetFName(),
+						bInIsDefault,
+						bInEditorCanAddOrRemove
+					}
+				}
+			});
 		}
 
-		void RegisterInterfaces()
+		template <typename UClassType>
+		void RegisterInterfaceForSingleClass(const FMetasoundFrontendInterface& InInterface, TUniquePtr<Frontend::IDocumentTransform>&& InUpdateTransform, bool bInIsDefault = false, bool bInEditorCanAddOrRemove = false, FName InRouterName = IDataReference::RouterName)
 		{
 			using namespace Frontend;
 
-			// Register Default Internal Interfaces (Not managed directly by end-user & added by default when creating new MetaSound assets).
+			RegisterInterface(InInterface, MoveTemp(InUpdateTransform), FInterfaceRegistryOptions
 			{
-				constexpr bool bIsDefault = true;
-				constexpr bool bEditorCanAddOrRemove = false;
-				RegisterInterface<UMetaSoundPatch>(MetasoundV1_0::GetInterface(), nullptr, bIsDefault, bEditorCanAddOrRemove, IDataReference::RouterName);
+				InRouterName,
+				{
+					{
+						UClassType::StaticClass()->GetFName(),
+						bInIsDefault,
+						bInEditorCanAddOrRemove
+					}
+				}
+			});
+		}
+
+		void RegisterAudioFormatInterfaces()
+		{
+			using namespace Frontend;
 
-				RegisterInterface(SourceInterface::CreateInterface(*UMetaSoundSource::StaticClass()), MakeUnique<SourceInterface::FUpdateInterface>(), bIsDefault, bEditorCanAddOrRemove, IDataReference::RouterName);
-				RegisterInterface(OutputFormatMonoInterface::CreateInterface(*UMetaSoundSource::StaticClass()), nullptr, bIsDefault, bEditorCanAddOrRemove, IDataReference::RouterName);
+			const FName SourceClassName = UMetaSoundSource::StaticClass()->GetFName();
+			const FName PatchClassName = UMetaSoundPatch::StaticClass()->GetFName();
 
-				RegisterInterface(SourceOneShotInterface::CreateInterface(*UMetaSoundSource::StaticClass()), nullptr, bIsDefault, true, IDataReference::RouterName);
+			// Input Formats
+			{
+				RegisterInterface(InputFormatMonoInterface::CreateInterface(), nullptr, FInterfaceRegistryOptions
+				{
+					IDataReference::RouterName,
+					{
+						{ PatchClassName, false /* bIsDefault */, true /* bEditorCanAddOrRemove */ },
+						{ SourceClassName },
+					}
+				});
+				IInterfaceRegistry::Get().AddInterfaceOutputBindings(InputFormatMonoInterface::GetVersion(), InputFormatMonoInterface::CreateOutputBindings());
+
+
+				RegisterInterface(InputFormatStereoInterface::CreateInterface(), nullptr, FInterfaceRegistryOptions
+				{
+					IDataReference::RouterName,
+					{
+						{ PatchClassName, false /* bIsDefault */, true /* bEditorCanAddOrRemove */ },
+						{ SourceClassName },
+					}
+				});
+				IInterfaceRegistry::Get().AddInterfaceOutputBindings(InputFormatStereoInterface::GetVersion(), InputFormatStereoInterface::CreateOutputBindings());
 			}
 
-			// Register Non-Default Internal Interfaces (Not managed directly by end-user & not added by default when creating new MetaSound assets).
+			// Output Formats
 			{
-				constexpr bool bIsDefault = false;
-				constexpr bool bEditorCanAddOrRemove = false;
-				
-				// Added for upgrading old metasounds
-				RegisterInterface(SourceInterfaceV1_0::CreateInterface(*UMetaSoundSource::StaticClass()), nullptr, bIsDefault, bEditorCanAddOrRemove, IDataReference::RouterName);
-
-				// Set default interface with unset version to use base UMetaSoundPatch class implementation (legacy requirement for 5.0 alpha).
-				RegisterInterface<UMetaSoundPatch>(FMetasoundFrontendInterface(), nullptr, bIsDefault, bEditorCanAddOrRemove, IDataReference::RouterName);
-
-				RegisterInterface<UMetaSoundSource>(MetasoundOutputFormatStereoV1_0::GetInterface(), nullptr, bIsDefault, bEditorCanAddOrRemove, IDataReference::RouterName);
-				RegisterInterface<UMetaSoundSource>(MetasoundOutputFormatStereoV1_1::GetInterface(), MakeUnique<MetasoundOutputFormatStereoV1_1::FUpdateInterface>(), bIsDefault, bEditorCanAddOrRemove, IDataReference::RouterName);
-				RegisterInterface<UMetaSoundSource>(MetasoundOutputFormatStereoV1_2::GetInterface(), MakeUnique<MetasoundOutputFormatStereoV1_2::FUpdateInterface>(), bIsDefault, bEditorCanAddOrRemove, IDataReference::RouterName);
-
-				RegisterInterface<UMetaSoundSource>(MetasoundOutputFormatMonoV1_0::GetInterface(), nullptr, bIsDefault, bEditorCanAddOrRemove, IDataReference::RouterName);
-				RegisterInterface<UMetaSoundSource>(MetasoundOutputFormatMonoV1_1::GetInterface(), MakeUnique<MetasoundOutputFormatMonoV1_1::FUpdateInterface>(), bIsDefault, bEditorCanAddOrRemove, IDataReference::RouterName);
-				RegisterInterface<UMetaSoundSource>(MetasoundOutputFormatMonoV1_2::GetInterface(), MakeUnique<MetasoundOutputFormatMonoV1_2::FUpdateInterface>(), bIsDefault, bEditorCanAddOrRemove, IDataReference::RouterName);
-
-				RegisterInterface(OutputFormatStereoInterface::CreateInterface(*UMetaSoundSource::StaticClass()), nullptr, bIsDefault, bEditorCanAddOrRemove, IDataReference::RouterName);			
-				RegisterInterface(OutputFormatQuadInterface::CreateInterface(*UMetaSoundSource::StaticClass()), nullptr, bIsDefault, bEditorCanAddOrRemove, IDataReference::RouterName);			
-				RegisterInterface(OutputFormatFiveDotOneInterface::CreateInterface(*UMetaSoundSource::StaticClass()), nullptr, bIsDefault, bEditorCanAddOrRemove, IDataReference::RouterName);			
-				RegisterInterface(OutputFormatSevenDotOneInterface::CreateInterface(*UMetaSoundSource::StaticClass()), nullptr, bIsDefault, bEditorCanAddOrRemove, IDataReference::RouterName);			
+				RegisterInterface(OutputFormatMonoInterface::CreateInterface(), nullptr, FInterfaceRegistryOptions
+				{
+					IDataReference::RouterName,
+					{
+						{ PatchClassName, false /* bIsDefault */, true /* bEditorCanAddOrRemove */ },
+						{ SourceClassName, true /* bIsDefault */, false /* bEditorCanAddOrRemove */ },
+					}
+				});
+
+				RegisterInterface(OutputFormatStereoInterface::CreateInterface(), nullptr, FInterfaceRegistryOptions
+				{
+					IDataReference::RouterName,
+					{
+						{ PatchClassName, false /* bIsDefault */, true /* bEditorCanAddOrRemove */ },
+						{ SourceClassName },
+					}
+				});
+
+				RegisterInterface(OutputFormatQuadInterface::CreateInterface(), nullptr, FInterfaceRegistryOptions
+				{
+					IDataReference::RouterName,
+					{
+						{ PatchClassName, false /* bIsDefault */, true /* bEditorCanAddOrRemove */ },
+						{ SourceClassName },
+					}
+				});
+
+				RegisterInterface(OutputFormatFiveDotOneInterface::CreateInterface(), nullptr, FInterfaceRegistryOptions
+				{
+					IDataReference::RouterName,
+					{
+						{ PatchClassName, false /* bIsDefault */, true /* bEditorCanAddOrRemove */ },
+						{ SourceClassName },
+					}
+				});
+
+				RegisterInterface(OutputFormatSevenDotOneInterface::CreateInterface(), nullptr, FInterfaceRegistryOptions
+				{
+					IDataReference::RouterName,
+					{
+						{ PatchClassName, false /* bIsDefault */, true /* bEditorCanAddOrRemove */ },
+						{ SourceClassName },
+					}
+				});
 			}
+		}
+
+		void RegisterDeprecatedInterfaces()
+		{
+			using namespace Frontend;
+
+			RegisterInterfaceForSingleClass<UMetaSoundSource>(SourceInterfaceV1_0::CreateInterface(), nullptr);
+
+			// Set default interface with unset version to use base UMetaSoundPatch class implementation (legacy requirement for 5.0 alpha).
+			RegisterInterfaceForSingleClass<UMetaSoundPatch>(FMetasoundFrontendInterface(), nullptr);
+
+			RegisterInterfaceForSingleClass<UMetaSoundSource>(MetasoundOutputFormatStereoV1_0::GetInterface(), nullptr);
+			RegisterInterfaceForSingleClass<UMetaSoundSource>(MetasoundOutputFormatStereoV1_1::GetInterface(), MakeUnique<MetasoundOutputFormatStereoV1_1::FUpdateInterface>());
+			RegisterInterfaceForSingleClass<UMetaSoundSource>(MetasoundOutputFormatStereoV1_2::GetInterface(), MakeUnique<MetasoundOutputFormatStereoV1_2::FUpdateInterface>());
 
+			RegisterInterfaceForSingleClass<UMetaSoundSource>(MetasoundOutputFormatMonoV1_0::GetInterface(), nullptr);
+			RegisterInterfaceForSingleClass<UMetaSoundSource>(MetasoundOutputFormatMonoV1_1::GetInterface(), MakeUnique<MetasoundOutputFormatMonoV1_1::FUpdateInterface>());
+			RegisterInterfaceForSingleClass<UMetaSoundSource>(MetasoundOutputFormatMonoV1_2::GetInterface(), MakeUnique<MetasoundOutputFormatMonoV1_2::FUpdateInterface>());
+		}
+
+		void RegisterExternalInterfaces()
+		{
 			// Register External Interfaces (Interfaces defined externally & can be managed directly by end-user).
 			auto RegisterExternalInterface = [](Audio::FParameterInterfacePtr Interface)
 			{
+				FInterfaceRegistryOptions Options { Audio::IParameterTransmitter::RouterName };
+
 				// Currently, no externally defined interfaces can be added as default for protection against undesired
 				// interfaces automatically being added when creating a new MetaSound asset through the editor. Also,
 				// all parameter interfaces are enabled in the editor for addition/removal.
 				constexpr bool bIsDefault = false;
 				constexpr bool bEditorCanAddOrRemove = true;
 
-				TUniquePtr<Frontend::IDocumentTransform> NullTransform;
-				RegisterInterface(Interface, MoveTemp(NullTransform), bIsDefault, bEditorCanAddOrRemove, Audio::IParameterTransmitter::RouterName);
+				bool bInterfaceSupportsPatch = true;
+				bool bInterfaceSupportsSource = true;
+				const TArray<const UClass*> SupportedUClasses = Interface->FindSupportedUClasses();
+				for (const UClass* SupportedUClass : SupportedUClasses )
+				{
+					const UClass* PatchClass = UMetaSoundPatch::StaticClass();
+					check(PatchClass);
+					bInterfaceSupportsPatch |= PatchClass->IsChildOf(SupportedUClass);
+
+					const UClass* SourceClass = UMetaSoundSource::StaticClass();
+					check(SourceClass);
+					bInterfaceSupportsSource |= SourceClass->IsChildOf(SupportedUClass);
+				}
+
+				const bool bSupported = bInterfaceSupportsSource || bInterfaceSupportsPatch;
+				if (ensureAlwaysMsgf(bSupported, TEXT("Parameter interface '%s' type not supported by MetaSounds"), *Interface->GetName().ToString()))
+				{
+					const FMetasoundFrontendInterface FrontendInterface = ConvertParameterToFrontendInterface(*Interface);
+					if (bInterfaceSupportsPatch)
+					{
+						const FName ClassName = UMetaSoundPatch::StaticClass()->GetFName();
+						Options.UClassOptions.Add(FInterfaceRegistryUClassOptions { ClassName, bIsDefault, bEditorCanAddOrRemove });
+					}
+
+					if (bInterfaceSupportsSource)
+					{
+						const FName ClassName = UMetaSoundSource::StaticClass()->GetFName();
+						Options.UClassOptions.Add(FInterfaceRegistryUClassOptions { ClassName, bIsDefault, bEditorCanAddOrRemove });
+					}
+
+					RegisterInterface(FrontendInterface, nullptr, MoveTemp(Options));
+				}
 			};
 
 			Audio::IAudioParameterInterfaceRegistry::Get().IterateInterfaces(RegisterExternalInterface);
 			Audio::IAudioParameterInterfaceRegistry::Get().OnRegistration(MoveTemp(RegisterExternalInterface));
 		}
+
+		void RegisterInterfaces()
+		{
+			using namespace Frontend;
+
+			// Default Patch Interfaces (legacy)
+			{
+				RegisterInterfaceForSingleClass<UMetaSoundPatch>(MetasoundV1_0::GetInterface(), nullptr, true, false);
+			}
+
+			// Default Source Interfaces
+			{
+				RegisterInterfaceForSingleClass<UMetaSoundSource>(SourceInterface::CreateInterface(), MakeUnique<SourceInterface::FUpdateInterface>(), true, false);
+				RegisterInterfaceForSingleClass<UMetaSoundSource>(SourceOneShotInterface::CreateInterface(), nullptr, true, true);
+			}
+
+			RegisterAudioFormatInterfaces();
+			RegisterDeprecatedInterfaces();
+			RegisterExternalInterfaces();
+		}
 	} // namespace Engine
 } // namespace Metasound
diff --git a/Engine/Plugins/Runtime/Metasound/Source/MetasoundEngine/Private/MetasoundInterface.h b/Engine/Plugins/Runtime/Metasound/Source/MetasoundEngine/Private/MetasoundInterface.h
--- a/Engine/Plugins/Runtime/Metasound/Source/MetasoundEngine/Private/MetasoundInterface.h
+++ b/Engine/Plugins/Runtime/Metasound/Source/MetasoundEngine/Private/MetasoundInterface.h
@@ -6,6 +6,7 @@ namespace Metasound
 {
 	namespace Engine
 	{
+		void RegisterUClasses();
 		void RegisterInterfaces();
 	} // namespace Engine
 } // namespace Metasound
diff --git a/Engine/Plugins/Runtime/Metasound/Source/MetasoundEngine/Private/MetasoundSource.cpp b/Engine/Plugins/Runtime/Metasound/Source/MetasoundEngine/Private/MetasoundSource.cpp
--- a/Engine/Plugins/Runtime/Metasound/Source/MetasoundEngine/Private/MetasoundSource.cpp
+++ b/Engine/Plugins/Runtime/Metasound/Source/MetasoundEngine/Private/MetasoundSource.cpp
@@ -1,22 +1,25 @@
 // Copyright Epic Games, Inc. All Rights Reserved.
 #include "MetasoundSource.h"
 
+#include "Algo/Find.h"
 #include "Algo/Transform.h"
 #include "AssetRegistry/AssetRegistryModule.h"
 #include "AudioDeviceManager.h"
 #include "IAudioParameterInterfaceRegistry.h"
-#include "Interfaces/MetasoundFrontendOutputFormatInterfaces.h"
+#include "Interfaces/MetasoundOutputFormatInterfaces.h"
 #include "Interfaces/MetasoundFrontendSourceInterface.h"
 #include "Internationalization/Text.h"
 #include "MetasoundAssetBase.h"
 #include "MetasoundAssetManager.h"
 #include "MetasoundAudioFormats.h"
+#include "MetasoundBuilderSubsystem.h"
 #include "MetasoundEngineArchetypes.h"
 #include "MetasoundEngineAsset.h"
 #include "MetasoundEngineEnvironment.h"
 #include "MetasoundEnvironment.h"
 #include "MetasoundFrontendController.h"
 #include "MetasoundFrontendDataTypeRegistry.h"
+#include "MetasoundFrontendDocumentBuilder.h"
 #include "MetasoundFrontendQuery.h"
 #include "MetasoundFrontendQuerySteps.h"
 #include "MetasoundFrontendTransform.h"
@@ -32,6 +35,7 @@
 #include "MetasoundTrigger.h"
 #include "MetasoundUObjectRegistry.h"
 #include "UObject/ObjectSaveContext.h"
+#include "UObject/ScriptInterface.h"
 
 #include UE_INLINE_GENERATED_CPP_BY_NAME(MetasoundSource)
 
@@ -44,95 +48,7 @@
 namespace Metasound
 {
 	namespace SourcePrivate
-	{
-		// Contains information on output audio formats
-		struct FOutputFormatInfo
-		{
-			FMetasoundFrontendVersion InterfaceVersion;
-			TArray<Metasound::FVertexName> OutputVertexChannelOrder;
-		};
-
-		using FFormatInfoMap = TMap<EMetasoundSourceAudioFormat, FOutputFormatInfo>;
-		using FFormatInfoPair = TPair<EMetasoundSourceAudioFormat, FOutputFormatInfo>;
-
-		// Return a map containing all the supported audio formats for a MetaSound
-		// Source. 
-		const FFormatInfoMap& GetFormatInfoMap()
-		{
-			auto CreateFormatInfoMap = []()
-			{
-				using namespace Metasound::Frontend;
-
-				return FFormatInfoMap
-				{
-					{
-						EMetasoundSourceAudioFormat::Mono,
-						{
-							OutputFormatMonoInterface::GetVersion(),
-							{
-								OutputFormatMonoInterface::Outputs::MonoOut
-							}
-						}
-					},
-					{
-						EMetasoundSourceAudioFormat::Stereo,
-						{
-							OutputFormatStereoInterface::GetVersion(),
-							{
-								OutputFormatStereoInterface::Outputs::LeftOut,
-								OutputFormatStereoInterface::Outputs::RightOut
-							}
-						}
-					},
-					{
-						EMetasoundSourceAudioFormat::Quad,
-						{
-							OutputFormatQuadInterface::GetVersion(),
-							{
-								OutputFormatQuadInterface::Outputs::FrontLeftOut,
-								OutputFormatQuadInterface::Outputs::FrontRightOut,
-								OutputFormatQuadInterface::Outputs::SideLeftOut,
-								OutputFormatQuadInterface::Outputs::SideRightOut
-							}
-						}
-					},
-					{
-						EMetasoundSourceAudioFormat::FiveDotOne,
-						{
-							OutputFormatFiveDotOneInterface::GetVersion(),
-							{
-								OutputFormatFiveDotOneInterface::Outputs::FrontLeftOut,
-								OutputFormatFiveDotOneInterface::Outputs::FrontRightOut,
-								OutputFormatFiveDotOneInterface::Outputs::FrontCenterOut,
-								OutputFormatFiveDotOneInterface::Outputs::LowFrequencyOut,
-								OutputFormatFiveDotOneInterface::Outputs::SideLeftOut,
-								OutputFormatFiveDotOneInterface::Outputs::SideRightOut
-							}
-						}
-					},
-					{
-						EMetasoundSourceAudioFormat::SevenDotOne,
-						{
-							OutputFormatSevenDotOneInterface::GetVersion(),
-							{
-								OutputFormatSevenDotOneInterface::Outputs::FrontLeftOut,
-								OutputFormatSevenDotOneInterface::Outputs::FrontRightOut,
-								OutputFormatSevenDotOneInterface::Outputs::FrontCenterOut,
-								OutputFormatSevenDotOneInterface::Outputs::LowFrequencyOut,
-								OutputFormatSevenDotOneInterface::Outputs::SideLeftOut,
-								OutputFormatSevenDotOneInterface::Outputs::SideRightOut,
-								OutputFormatSevenDotOneInterface::Outputs::BackLeftOut,
-								OutputFormatSevenDotOneInterface::Outputs::BackRightOut
-							}
-						}
-					}
-				};
-			};
-
-			static const FFormatInfoMap Map = CreateFormatInfoMap();
-			return Map;
-		}
-		
+	{		
 		Frontend::FMetaSoundAssetRegistrationOptions GetInitRegistrationOptions()
 		{
 			Frontend::FMetaSoundAssetRegistrationOptions RegOptions;
@@ -144,25 +60,6 @@ namespace Metasound
 
 			return RegOptions;
 		}
-
-		// Return an array of all the audio format versions.
-		const TArray<FMetasoundFrontendVersion>& GetFormatInterfaceVersions()
-		{
-			auto CreateFormatInterfaceVersions = []() -> TArray<FMetasoundFrontendVersion>
-			{
-				TArray<FMetasoundFrontendVersion> FormatVersions;
-				const FFormatInfoMap& FormatMap = GetFormatInfoMap();
-				for (const auto& Pair : FormatMap)
-				{
-					FormatVersions.Add(Pair.Value.InterfaceVersion);
-				}
-				return FormatVersions;
-			};
-
-			static const TArray<FMetasoundFrontendVersion> Versions = CreateFormatInterfaceVersions();
-			return Versions;
-		}
-
 	} // namespace SourcePrivate
 } // namespace Metasound
 
@@ -211,37 +108,18 @@ void UMetaSoundSource::PostEditChangeProperty(FPropertyChangedEvent& InEvent)
 
 void UMetaSoundSource::PostEditChangeOutputFormat()
 {
-	using namespace Metasound::SourcePrivate;
 	using namespace Metasound::Frontend;
 
-	// If this is a preset, convert to normal metasound source since it is being
-	// altered. 
-	ConvertFromPreset();
-
-	// Determine which interfaces to add and remove from the document due to the
-	// output format being changed.
-	TArray<FMetasoundFrontendVersion> OutputFormatsToAdd;
-	if (const FOutputFormatInfo* FormatInfo = GetFormatInfoMap().Find(OutputFormat))
-	{
-		OutputFormatsToAdd.Add(FormatInfo->InterfaceVersion);
-	}
-
-	TArray<FMetasoundFrontendVersion> OutputFormatsToRemove;
-	for (const FMetasoundFrontendVersion& FormatVersion : GetFormatInterfaceVersions())
-	{
-		if (RootMetasoundDocument.Interfaces.Contains(FormatVersion))
-		{
-			if (!OutputFormatsToAdd.Contains(FormatVersion))
-			{
-				OutputFormatsToRemove.Add(FormatVersion);
-			}
-		}
-	}
-
-	// Add and/or remove interfaces from the root document.
-	const bool bDidModifyDocument = FModifyRootGraphInterfaces(OutputFormatsToRemove, OutputFormatsToAdd).Transform(GetDocumentHandle());
+	// Add and/or remove output format interfaces from the root document.
+	TScriptInterface<IMetaSoundDocumentInterface> DocInterface = this;
+	FMetaSoundFrontendDocumentBuilder Builder(DocInterface);
 
-	if (bDidModifyDocument)
+	const UMetaSoundBuilderSubsystem* BuilderSubsystem = GEngine->GetEngineSubsystem<UMetaSoundBuilderSubsystem>();
+	check(BuilderSubsystem);
+	UMetaSoundSourceBuilder* SourceBuilder = BuilderSubsystem->AttachSourceBuilderToAsset(this);
+	EMetaSoundBuilderResult Result;
+	SourceBuilder->SetFormat(OutputFormat, Result);
+	if (Result == EMetaSoundBuilderResult::Succeeded)
 	{
 		// Update the data in this UMetaSoundSource to reflect what is in the metasound document.
 		ConformObjectDataToInterfaces();
@@ -259,14 +137,15 @@ void UMetaSoundSource::PostEditChangeOutputFormat()
 
 bool UMetaSoundSource::ConformObjectDataToInterfaces()
 {
-	using namespace Metasound::SourcePrivate;
+	using namespace Metasound::Engine;
+	using namespace Metasound::Frontend;
 
 	bool bDidAlterObjectData = false;
 
 	// Update the OutputFormat and NumChannels to match the audio format interface
 	// on the root document.
-	const FFormatInfoMap& FormatInfo = GetFormatInfoMap();
-	for (const FFormatInfoPair& Pair : FormatInfo)
+	const FOutputAudioFormatInfoMap& FormatInfo = GetOutputAudioFormatInfo();
+	for (const FOutputAudioFormatInfoPair& Pair : FormatInfo)
 	{
 		if (RootMetasoundDocument.Interfaces.Contains(Pair.Value.InterfaceVersion))
 		{
@@ -996,9 +875,10 @@ Metasound::FMetasoundEnvironment UMetaSoundSource::CreateEnvironment(const Audio
 
 const TArray<Metasound::FVertexName>& UMetaSoundSource::GetOutputAudioChannelOrder() const
 {
-	using namespace Metasound::SourcePrivate;
+	using namespace Metasound::Engine;
+	using namespace Metasound::Frontend;
 
-	if (const FOutputFormatInfo* FormatInfo = GetFormatInfoMap().Find(OutputFormat))
+	if (const FOutputAudioFormatInfo* FormatInfo = GetOutputAudioFormatInfo().Find(OutputFormat))
 	{
 		return FormatInfo->OutputVertexChannelOrder;
 	}
@@ -1045,6 +925,4 @@ TWeakPtr<Metasound::FMetasoundGenerator> UMetaSoundSource::GetGeneratorForAudioC
 	}
 	return *Result;
 }
-
 #undef LOCTEXT_NAMESPACE // MetaSound
-
diff --git a/Engine/Plugins/Runtime/Metasound/Source/MetasoundEngine/Private/MetasoundUObjectRegistry.cpp b/Engine/Plugins/Runtime/Metasound/Source/MetasoundEngine/Private/MetasoundUObjectRegistry.cpp
--- a/Engine/Plugins/Runtime/Metasound/Source/MetasoundEngine/Private/MetasoundUObjectRegistry.cpp
+++ b/Engine/Plugins/Runtime/Metasound/Source/MetasoundEngine/Private/MetasoundUObjectRegistry.cpp
@@ -25,74 +25,11 @@ namespace Metasound
 		public:
 			FMetasoundUObjectRegistry() = default;
 
-			void RegisterUClassInterface(TUniquePtr<IMetasoundUObjectRegistryEntry>&& InEntry) override
+			void RegisterUClass(TUniquePtr<IMetasoundUObjectRegistryEntry>&& InEntry) override
 			{
 				METASOUND_LLM_SCOPE;
-
-				if (InEntry.IsValid())
-				{
-					Frontend::FInterfaceRegistryKey Key = Frontend::GetInterfaceRegistryKey(InEntry->GetInterfaceVersion());
-
-					EntriesByInterface.Add(Key, InEntry.Get());
-					EntriesByName.Add(InEntry->GetInterfaceVersion().Name, InEntry.Get());
-					Entries.Add(InEntry.Get());
-					Storage.Add(MoveTemp(InEntry));
-				}
-			}
-
-			TArray<const IMetasoundUObjectRegistryEntry*> FindInterfaceEntriesByName(FName InName) const override
-			{
-				TArray<const IMetasoundUObjectRegistryEntry*> EntriesWithName;
-				EntriesByName.MultiFind(InName, EntriesWithName);
-
-				return EntriesWithName;
-			}
-
-			TArray<UClass*> FindSupportedInterfaceClasses(const FMetasoundFrontendVersion& InInterfaceVersion) const override
-			{
-				TArray<UClass*> Classes;
-
-				TArray<const IMetasoundUObjectRegistryEntry*> EntriesForInterface;
-				EntriesByInterface.MultiFind(Frontend::GetInterfaceRegistryKey(InInterfaceVersion), EntriesForInterface);
-
-				for (const IMetasoundUObjectRegistryEntry* Entry : EntriesForInterface)
-				{
-					if (nullptr != Entry)
-					{
-						if (UClass* Class = Entry->GetUClass())
-						{
-							Classes.Add(Class);
-						}
-					}
-				}
-
-				return Classes;
-			}
-
-			UObject* NewObject(UClass* InClass, const FMetasoundFrontendDocument& InDocument, const FString& InPath) const override
-			{
-				METASOUND_LLM_SCOPE;
-
-				TArray<const IMetasoundUObjectRegistryEntry*> AllInterfaceEntries;
-
-				for (const FMetasoundFrontendVersion& InterfaceVersion : InDocument.Interfaces)
-				{
-					TArray<const IMetasoundUObjectRegistryEntry*> EntriesForInterface;
-					EntriesByInterface.MultiFind(Frontend::GetInterfaceRegistryKey(InterfaceVersion), EntriesForInterface);
-					AllInterfaceEntries.Append(MoveTemp(EntriesForInterface));
-				}
-
-				auto IsChildClassOfRegisteredClass = [&](const IMetasoundUObjectRegistryEntry* Entry)
-				{
-					return Entry->IsChildClass(InClass);
-				};
-
-				if (const IMetasoundUObjectRegistryEntry* const* EntryForClass = AllInterfaceEntries.FindByPredicate(IsChildClassOfRegisteredClass))
-				{
-					return NewObject(**EntryForClass, InDocument, InPath);
-				}
-
-				return nullptr;
+				Entries.Add(InEntry.Get());
+				Storage.Add(MoveTemp(InEntry));
 			}
 
 			bool IsRegisteredClass(UObject* InObject) const override
@@ -133,48 +70,6 @@ namespace Metasound
 			}
 
 		private:
-			UObject* NewObject(const IMetasoundUObjectRegistryEntry& InEntry, const FMetasoundFrontendDocument& InDocument, const FString& InPath) const
-			{
-				UPackage* PackageToSaveTo = nullptr;
-
-				if (GIsEditor)
-				{
-					FText InvalidPathReason;
-					bool const bValidPackageName = FPackageName::IsValidLongPackageName(InPath, false, &InvalidPathReason);
-
-					if (!ensureAlwaysMsgf(bValidPackageName, TEXT("Tried to generate a Metasound UObject with an invalid package path/name Falling back to transient package, which means we won't be able to save this asset.")))
-					{
-						PackageToSaveTo = GetTransientPackage();
-					}
-					else
-					{
-						PackageToSaveTo = CreatePackage(*InPath);
-					}
-				}
-				else
-				{
-					PackageToSaveTo = GetTransientPackage();
-				}
-
-				UObject* NewMetasoundObject = InEntry.NewObject(PackageToSaveTo, *InDocument.RootGraph.Metadata.GetClassName().GetFullName().ToString());
-				FMetasoundAssetBase* NewAssetBase = InEntry.Cast(NewMetasoundObject);
-				if (ensure(nullptr != NewAssetBase))
-				{
-					NewAssetBase->SetDocument(InDocument);
-				}
-
-#if WITH_EDITOR
-				AsyncTask(ENamedThreads::GameThread, [NewMetasoundObject]()
-				{
-					FAssetRegistryModule::AssetCreated(NewMetasoundObject);
-					NewMetasoundObject->MarkPackageDirty();
-					// todo: how do you get the package for a uobject and save it? I forget
-				});
-#endif
-
-				return NewMetasoundObject;
-			}
-
 			const IMetasoundUObjectRegistryEntry* FindEntryByPredicate(TFunction<bool (const IMetasoundUObjectRegistryEntry*)> InPredicate) const
 			{
 				const IMetasoundUObjectRegistryEntry* const* Entry = Entries.FindByPredicate(InPredicate);
@@ -212,8 +107,6 @@ namespace Metasound
 			}
 
 			TArray<TUniquePtr<IMetasoundUObjectRegistryEntry>> Storage;
-			TMultiMap<Frontend::FInterfaceRegistryKey, const IMetasoundUObjectRegistryEntry*> EntriesByInterface;
-			TMultiMap<FName, const IMetasoundUObjectRegistryEntry*> EntriesByName;
 			TArray<const IMetasoundUObjectRegistryEntry*> Entries;
 	};
 
diff --git a/Engine/Plugins/Runtime/Metasound/Source/MetasoundEngine/Public/Metasound.h b/Engine/Plugins/Runtime/Metasound/Source/MetasoundEngine/Public/Metasound.h
--- a/Engine/Plugins/Runtime/Metasound/Source/MetasoundEngine/Public/Metasound.h
+++ b/Engine/Plugins/Runtime/Metasound/Source/MetasoundEngine/Public/Metasound.h
@@ -5,6 +5,7 @@
 
 #include "EdGraph/EdGraph.h"
 #include "MetasoundAssetBase.h"
+#include "MetasoundDocumentInterface.h"
 #include "MetasoundFrontend.h"
 #include "MetasoundFrontendDocument.h"
 #include "MetasoundFrontendTransform.h"
@@ -47,12 +48,13 @@ namespace Metasound
 	struct FMetaSoundEngineAssetHelper;
 }
 
+
 /**
  * This asset type is used for Metasound assets that can only be used as nodes in other Metasound graphs.
  * Because of this, they contain no required inputs or outputs.
  */
 UCLASS(hidecategories = object, BlueprintType)
-class METASOUNDENGINE_API UMetaSoundPatch : public UObject, public FMetasoundAssetBase
+class METASOUNDENGINE_API UMetaSoundPatch : public UObject, public FMetasoundAssetBase, public IMetaSoundDocumentInterface
 {
 	GENERATED_BODY()
 
@@ -128,6 +130,10 @@ public:
 
 #endif // #if WITH_EDITORONLY_DATA
 
+	virtual const FMetasoundFrontendDocument& GetDocument() const override
+	{
+		return RootMetaSoundDocument;
+	}
 
 #if WITH_EDITOR
 	virtual void PostDuplicate(EDuplicateMode::Type DuplicateMode) override;
@@ -167,7 +173,7 @@ protected:
 	virtual void SetReferencedAssetClasses(TSet<Metasound::Frontend::IMetaSoundAssetManager::FAssetInfo>&& InAssetClasses) override;
 #endif // #if WITH_EDITOR
 
-	Metasound::Frontend::FDocumentAccessPtr GetDocument() override
+	Metasound::Frontend::FDocumentAccessPtr GetDocumentAccessPtr() override
 	{
 		using namespace Metasound::Frontend;
 		// Return document using FAccessPoint to inform the TAccessPtr when the 
@@ -175,11 +181,17 @@ protected:
 		return MakeAccessPtr<FDocumentAccessPtr>(RootMetaSoundDocument.AccessPoint, RootMetaSoundDocument);
 	}
 
-	Metasound::Frontend::FConstDocumentAccessPtr GetDocument() const override
+	Metasound::Frontend::FConstDocumentAccessPtr GetDocumentConstAccessPtr() const override
 	{
 		using namespace Metasound::Frontend;
 		// Return document using FAccessPoint to inform the TAccessPtr when the 
 		// object is no longer valid.
 		return MakeAccessPtr<FConstDocumentAccessPtr>(RootMetaSoundDocument.AccessPoint, RootMetaSoundDocument);
 	}
+
+private:
+	virtual FMetasoundFrontendDocument& GetDocument() override
+	{
+		return RootMetaSoundDocument;
+	}
 };
