- Add Runtime MetaSound Builder support

diff --git a/Engine/Plugins/Runtime/Metasound/Source/MetasoundEngine/Public/MetasoundBuilderSubsystem.h b/Engine/Plugins/Runtime/Metasound/Source/MetasoundEngine/Public/MetasoundBuilderSubsystem.h
new file mode 100644
--- /dev/null
+++ b/Engine/Plugins/Runtime/Metasound/Source/MetasoundEngine/Public/MetasoundBuilderSubsystem.h
@@ -0,0 +1,442 @@
+// Copyright Epic Games, Inc. All Rights Reserved.
+#pragma once
+
+#include "Containers/Map.h"
+#include "Engine/Engine.h"
+#include "Interfaces/MetasoundOutputFormatInterfaces.h"
+#include "MetasoundFrontendController.h"
+#include "MetasoundFrontendDocument.h"
+#include "MetasoundFrontendDocumentBuilder.h"
+#include "Subsystems/EngineSubsystem.h"
+#include "UObject/NoExportTypes.h"
+#include "UObject/Object.h"
+#include "UObject/ObjectMacros.h"
+
+#include "MetasoundBuilderSubsystem.generated.h"
+
+
+// Forward Declarations
+class UMetaSound;
+class UMetaSoundSource;
+struct FMetasoundFrontendVersion;
+class UAudioComponent;
+enum class EMetaSoundOutputAudioFormat : uint8;
+
+namespace Metasound::Engine
+{
+	struct FOutputAudioFormatInfo;
+} // namespace Metasound::Engine
+
+
+DECLARE_DYNAMIC_DELEGATE_OneParam(FOnCreateAuditionGeneratorHandleDelegate, UMetasoundGeneratorHandle*, GeneratorHandle);
+
+USTRUCT(BlueprintType, meta = (DisplayName = "MetaSound Node Input Handle"))
+struct METASOUNDENGINE_API FMetaSoundBuilderNodeInputHandle : public FMetasoundFrontendVertexHandle
+{
+	GENERATED_BODY()
+};
+
+USTRUCT(BlueprintType, meta = (DisplayName = "MetaSound Node Output Handle"))
+struct METASOUNDENGINE_API FMetaSoundBuilderNodeOutputHandle : public FMetasoundFrontendVertexHandle
+{
+	GENERATED_BODY()
+};
+
+USTRUCT(BlueprintType)
+struct METASOUNDENGINE_API FMetaSoundNodeHandle
+{
+	GENERATED_BODY()
+
+	UPROPERTY()
+	FGuid NodeID;
+
+public:
+	FMetaSoundNodeHandle() = default;
+	FMetaSoundNodeHandle(const FGuid& InNodeID)
+		: NodeID(InNodeID)
+	{
+	}
+
+	// Returns whether or not the vertex handle is set (may or may not be
+	// valid depending on what builder context it is referenced against)
+	bool IsSet() const
+	{
+		return NodeID.IsValid();
+	}
+};
+
+USTRUCT(BlueprintType)
+struct METASOUNDENGINE_API FMetaSoundBuilderOptions
+{
+	GENERATED_BODY()
+
+	UPROPERTY(EditAnywhere, BlueprintReadWrite, Category = "MetaSound|Builder")
+	FName Name;
+
+	// If true, adds MetaSound to node registry, making it available
+	// for reference by other dynamically created MetaSounds.
+	UPROPERTY(EditAnywhere, BlueprintReadWrite, Category = "MetaSound|Builder")
+	bool bAddToRegistry = false;
+
+	TScriptInterface<IMetaSoundDocumentInterface> ExistingAsset;
+};
+
+UENUM(BlueprintType)
+enum class EMetaSoundBuilderResult : uint8
+{
+	Succeeded,
+	Failed
+};
+
+/** Base implementation of MetaSound builder */
+UCLASS(Abstract)
+class METASOUNDENGINE_API UMetaSoundBuilderBase : public UObject
+{
+	GENERATED_BODY()
+
+public:
+	// Adds a graph input node with the given name, DataType, and sets the graph input to default value.
+	// Returns the new input node's output handle if it was successfully created, or an invalid handle if it failed.
+	UFUNCTION(BlueprintCallable, Category = "Audio|MetaSound|Builder", meta = (ExpandEnumAsExecs = "OutResult", AdvancedDisplay = "3"))
+	UPARAM(DisplayName = "Output Handle") FMetaSoundBuilderNodeOutputHandle AddGraphInputNode(FName Name, FName DataType, FMetasoundFrontendLiteral DefaultValue, EMetaSoundBuilderResult& OutResult, bool bIsConstructorInput = false);
+
+	// Adds a graph output node with the given name, DataType, and sets output node's input to default value.
+	// Returns the new output node's input handle if it was successfully created, or an invalid handle if it failed.
+	UFUNCTION(BlueprintCallable, Category = "Audio|MetaSound|Builder", meta = (ExpandEnumAsExecs = "OutResult"))
+	UPARAM(DisplayName = "Input Handle") FMetaSoundBuilderNodeInputHandle AddGraphOutputNode(FName Name, FName DataType, FMetasoundFrontendLiteral DefaultValue, EMetaSoundBuilderResult& OutResult, bool bIsConstructorInput = false);
+
+	// Adds an interface registered with the given name to the graph, adding associated input and output nodes.
+	UFUNCTION(BlueprintCallable, Category = "Audio|MetaSound|Builder", meta = (ExpandEnumAsExecs = "OutResult"))
+	void AddInterface(FName InterfaceName, EMetaSoundBuilderResult& OutResult);
+
+	// Adds a node to the graph using the provided MetaSound asset as its defining NodeClass.
+	// Returns a node handle to the created node if successful, or an invliad handle if it failed.
+	UFUNCTION(BlueprintCallable, Category = "Audio|MetaSound|Builder", DisplayName = "Add MetaSound Node From Asset Class", meta = (ExpandEnumAsExecs = "OutResult"))
+	UPARAM(DisplayName = "Node Handle") FMetaSoundNodeHandle AddNode(TScriptInterface<IMetaSoundDocumentInterface> NodeClass, EMetaSoundBuilderResult& OutResult);
+
+	// Connects node output to a node input. Does *NOT* provide loop detection for performance reasons.  Loop detection is checked on class registration when built or played.
+	// Returns succeeded if connection made, failed if connection already exists with input, the data types do not match, or the connection is not supported due to access type
+	// incompatibility (ex. constructor input to non-constructor input).
+	UFUNCTION(BlueprintCallable, Category = "Audio|MetaSound|Builder", meta = (ExpandEnumAsExecs = "OutResult"))
+	void ConnectNodes(const FMetaSoundBuilderNodeOutputHandle& NodeOutputHandle, const FMetaSoundBuilderNodeInputHandle& NodeInputHandle, EMetaSoundBuilderResult& OutResult);
+
+	// Connects two nodes using defined MetaSound Interface Bindings registered with the MetaSound Interface registry.
+	UFUNCTION(BlueprintCallable, Category = "Audio|MetaSound|Builder", meta = (ExpandEnumAsExecs = "OutResult"))
+	void ConnectNodesByInterfaceBindings(const FMetaSoundNodeHandle& FromNodeHandle, const FMetaSoundNodeHandle& ToNodeHandle, EMetaSoundBuilderResult& OutResult);
+
+	// Connects a given node's outputs to all graph outputs for shared interfaces implemented on both the node's referenced class and the builder's MetaSound graph.
+	UFUNCTION(BlueprintCallable, Category = "Audio|MetaSound|Builder", meta = (ExpandEnumAsExecs = "OutResult"))
+	UPARAM(DisplayName = "Connected Graph Output Node Inputs") TArray<FMetaSoundBuilderNodeInputHandle> ConnectNodeOutputsToMatchingGraphInterfaceOutputs(const FMetaSoundNodeHandle& NodeHandle, EMetaSoundBuilderResult& OutResult);
+
+	// Connects a given node's inputs to all graph inputs for shared interfaces implemented on both the node's referenced class and the builder's MetaSound graph. Returns output input node's outputs that were connected.
+	UFUNCTION(BlueprintCallable, Category = "Audio|MetaSound|Builder", meta = (ExpandEnumAsExecs = "OutResult"))
+	UPARAM(DisplayName = "Connected Graph Input Node Outputs") TArray<FMetaSoundBuilderNodeOutputHandle> ConnectNodeInputsToMatchingGraphInterfaceInputs(const FMetaSoundNodeHandle& NodeHandle, EMetaSoundBuilderResult& OutResult);
+
+	// Connects a given node output to the graph output with the given name.
+	UFUNCTION(BlueprintCallable, Category = "Audio|MetaSound|Builder", meta = (ExpandEnumAsExecs = "OutResult"))
+	void ConnectNodeOutputToGraphOutput(FName GraphOutputName, const FMetaSoundBuilderNodeOutputHandle& NodeOutputHandle, EMetaSoundBuilderResult& OutResult);
+
+	// Connects a given node input to the graph input with the given name.
+	UFUNCTION(BlueprintCallable, Category = "Audio|MetaSound|Builder", meta = (ExpandEnumAsExecs = "OutResult"))
+	void ConnectNodeInputToGraphInput(FName GraphInputName, const FMetaSoundBuilderNodeInputHandle& NodeInputHandle, EMetaSoundBuilderResult& OutResult);
+
+	// Returns whether node exists.
+	UFUNCTION(BlueprintCallable, Category = "Audio|MetaSound|Builder")
+	UPARAM(DisplayName = "IsValid") bool ContainsNode(const FMetaSoundNodeHandle& Node) const;
+
+	// Returns whether node input exists.
+	UFUNCTION(BlueprintCallable, Category = "Audio|MetaSound|Builder")
+	UPARAM(DisplayName = "IsValid") bool ContainsNodeInput(const FMetaSoundBuilderNodeInputHandle& Input) const;
+
+	// Returns whether node output exists.
+	UFUNCTION(BlueprintCallable, Category = "Audio|MetaSound|Builder")
+	UPARAM(DisplayName = "IsValid") bool ContainsNodeOutput(const FMetaSoundBuilderNodeOutputHandle& Output) const;
+
+	// Disconnects node output to a node input. Returns success if connection was removed, failed if not.
+	UFUNCTION(BlueprintCallable, Category = "Audio|MetaSound|Builder", meta = (ExpandEnumAsExecs = "OutResult"))
+	void DisconnectNodes(const FMetaSoundBuilderNodeOutputHandle& NodeOutputHandle, const FMetaSoundBuilderNodeInputHandle& NodeInputHandle, EMetaSoundBuilderResult& OutResult);
+
+	// Removes connection to a given node input. Returns success if connection was removed, failed if not.
+	UFUNCTION(BlueprintCallable, Category = "Audio|MetaSound|Builder", meta = (ExpandEnumAsExecs = "OutResult"))
+	void DisconnectNodeInput(const FMetaSoundBuilderNodeInputHandle& NodeInputHandle, EMetaSoundBuilderResult& OutResult);
+
+	// Removes all connections from a given node output. Returns success if all connections were removed, failed if not.
+	UFUNCTION(BlueprintCallable, Category = "Audio|MetaSound|Builder", meta = (ExpandEnumAsExecs = "OutResult"))
+	void DisconnectNodeOutput(const FMetaSoundBuilderNodeOutputHandle& NodeOutputHandle, EMetaSoundBuilderResult& OutResult);
+
+	// Returns graph input node by the given name if it exists, or an invalid handle if not found.
+	UFUNCTION(BlueprintCallable, Category = "Audio|MetaSound|Builder", meta = (ExpandEnumAsExecs = "OutResult"))
+	UPARAM(DisplayName = "Node Handle") FMetaSoundNodeHandle FindGraphInputNode(FName InputName, EMetaSoundBuilderResult& OutResult) const;
+
+	// Returns graph output node by the given name if it exists, or an invalid handle if not found.
+	UFUNCTION(BlueprintCallable, Category = "Audio|MetaSound|Builder", meta = (ExpandEnumAsExecs = "OutResult"))
+	UPARAM(DisplayName = "Node Handle") FMetaSoundNodeHandle FindGraphOutputNode(FName OutputName, EMetaSoundBuilderResult& OutResult) const;
+
+	// Returns node input by the given name if it exists, or an invalid handle if not found.
+	UFUNCTION(BlueprintCallable, Category = "Audio|MetaSound|Builder", meta = (ExpandEnumAsExecs = "OutResult"))
+	UPARAM(DisplayName = "Input Handle") FMetaSoundBuilderNodeInputHandle FindNodeInputByName(const FMetaSoundNodeHandle& NodeHandle, FName InputName, EMetaSoundBuilderResult& OutResult) const;
+
+	// Returns node output by the given name.
+	UFUNCTION(BlueprintCallable, Category = "Audio|MetaSound|Builder", meta = (ExpandEnumAsExecs = "OutResult"))
+	UPARAM(DisplayName = "Input  Handles") TArray<FMetaSoundBuilderNodeInputHandle> FindNodeInputs(const FMetaSoundNodeHandle& NodeHandle, EMetaSoundBuilderResult& OutResult) const;
+
+	// Returns node inputs by the given DataType (ex. "Audio", "Trigger", "String", "Bool", "Float", "Int32", etc.).
+	UFUNCTION(BlueprintCallable, Category = "Audio|MetaSound|Builder", meta = (ExpandEnumAsExecs = "OutResult"))
+	UPARAM(DisplayName = "Input Handles") TArray<FMetaSoundBuilderNodeInputHandle> FindNodeInputsByDataType(const FMetaSoundNodeHandle& NodeHandle, EMetaSoundBuilderResult& OutResult, FName DataType) const;
+
+	// Returns node output by the given name.
+	UFUNCTION(BlueprintCallable, Category = "Audio|MetaSound|Builder", meta = (ExpandEnumAsExecs = "OutResult"))
+	UPARAM(DisplayName = "Output Handle") FMetaSoundBuilderNodeOutputHandle FindNodeOutputByName(const FMetaSoundNodeHandle& NodeHandle, FName OutputName, EMetaSoundBuilderResult& OutResult) const;
+
+	// Returns all node outputs.
+	UFUNCTION(BlueprintCallable, Category = "Audio|MetaSound|Builder", meta = (ExpandEnumAsExecs = "OutResult"))
+	UPARAM(DisplayName = "Output Handles") TArray<FMetaSoundBuilderNodeOutputHandle> FindNodeOutputs(const FMetaSoundNodeHandle& NodeHandle, EMetaSoundBuilderResult& OutResult) const;
+
+	// Returns node outputs by the given DataType (ex. "Audio", "Trigger", "String", "Bool", "Float", "Int32", etc.).
+	UFUNCTION(BlueprintCallable, Category = "Audio|MetaSound|Builder", meta = (ExpandEnumAsExecs = "OutResult"))
+	UPARAM(DisplayName = "Output Handles") TArray<FMetaSoundBuilderNodeOutputHandle> FindNodeOutputsByDataType(const FMetaSoundNodeHandle& NodeHandle, EMetaSoundBuilderResult& OutResult, FName DataType) const;
+
+	// Returns input nodes associated with a given interface.
+	UFUNCTION(BlueprintCallable, Category = "Audio|MetaSound|Builder", meta = (ExpandEnumAsExecs = "OutResult"))
+	UPARAM(DisplayName = "Input Node Handles") TArray<FMetaSoundNodeHandle> FindInterfaceInputNodes(FName InterfaceName, EMetaSoundBuilderResult& OutResult) const;
+
+	// Returns output nodes associated with a given interface.
+	UFUNCTION(BlueprintCallable, Category = "Audio|MetaSound|Builder", meta = (ExpandEnumAsExecs = "OutResult"))
+	UPARAM(DisplayName = "Output Node Handles") TArray<FMetaSoundNodeHandle> FindInterfaceOutputNodes(FName InterfaceName, EMetaSoundBuilderResult& OutResult) const;
+
+	// Returns input's parent node if the input is valid, otherwise returns invalid node handle.
+	UFUNCTION(BlueprintCallable, Category = "Audio|MetaSound|Builder", meta = (ExpandEnumAsExecs = "OutResult"))
+	UPARAM(DisplayName = "Node Handle") FMetaSoundNodeHandle FindNodeInputParent(const FMetaSoundBuilderNodeInputHandle& InputHandle, EMetaSoundBuilderResult& OutResult) const;
+
+	// Returns output's parent node if the input is valid, otherwise returns invalid node handle.
+	UFUNCTION(BlueprintCallable, Category = "Audio|MetaSound|Builder", meta = (ExpandEnumAsExecs = "OutResult"))
+	UPARAM(DisplayName = "Node Handle") FMetaSoundNodeHandle FindNodeOutputParent(const FMetaSoundBuilderNodeOutputHandle& OutputHandle, EMetaSoundBuilderResult& OutResult) const;
+
+	// Returns output's parent node if the input is valid, otherwise returns invalid node handle.
+	UFUNCTION(BlueprintCallable, Category = "Audio|MetaSound|Builder", meta = (ExpandEnumAsExecs = "OutResult"))
+	UPARAM(DisplayName = "Node ClassVersion") FMetasoundFrontendVersion FindNodeClassVersion(const FMetaSoundNodeHandle& NodeHandle, EMetaSoundBuilderResult& OutResult) const;
+
+	// Returns node input's data if valid (including things like name and datatype).
+	UFUNCTION(BlueprintCallable, Category = "Audio|MetaSound|Builder", meta = (ExpandEnumAsExecs = "OutResult"))
+	void GetNodeInputData(const FMetaSoundBuilderNodeInputHandle& InputHandle, FName& Name, FName& DataType, EMetaSoundBuilderResult& OutResult) const;
+
+	// Returns node output's data if valid (including things like name and datatype).
+	UFUNCTION(BlueprintCallable, Category = "Audio|MetaSound|Builder", meta = (ExpandEnumAsExecs = "OutResult"))
+	void GetNodeOutputData(const FMetaSoundBuilderNodeOutputHandle& OutputHandle, FName& Name, FName& DataType, EMetaSoundBuilderResult& OutResult) const;
+
+	// Returns if a given interface is declared.
+	UFUNCTION(BlueprintCallable, Category = "Audio|MetaSound|Builder")
+	UPARAM(DisplayName = "IsDeclared") bool InterfaceIsDeclared(FName InterfaceName) const;
+
+	// Returns if a given node output and node input are connected.
+	UFUNCTION(BlueprintCallable, Category = "Audio|MetaSound|Builder")
+	UPARAM(DisplayName = "Connected") bool NodesAreConnected(const FMetaSoundBuilderNodeOutputHandle& OutputHandle, const FMetaSoundBuilderNodeInputHandle& InputHandle) const;
+
+	// Returns if a given node input has connections.
+	UFUNCTION(BlueprintCallable, Category = "Audio|MetaSound|Builder")
+	UPARAM(DisplayName = "Connected") bool NodeInputIsConnected(const FMetaSoundBuilderNodeInputHandle& InputHandle) const;
+
+	// Returns if a given node output is connected.
+	UFUNCTION(BlueprintCallable, Category = "Audio|MetaSound|Builder")
+	UPARAM(DisplayName = "Connected") bool NodeOutputIsConnected(const FMetaSoundBuilderNodeOutputHandle& OutputHandle) const;
+
+	// Removes the interface with the given name from the builder's MetaSound. Removes any graph inputs
+	// and outputs associated with the given interface and their respective connections (if any).
+	UFUNCTION(BlueprintCallable, Category = "Audio|MetaSound|Builder", meta = (ExpandEnumAsExecs = "OutResult"))
+	void RemoveInterface(FName InterfaceName, EMetaSoundBuilderResult& OutResult);
+
+	// Removes node and any associated connections from the builder's MetaSound.
+	UFUNCTION(BlueprintCallable, Category = "Audio|MetaSound|Builder", meta = (ExpandEnumAsExecs = "OutResult"))
+	void RemoveNode(const FMetaSoundNodeHandle& NodeHandle, EMetaSoundBuilderResult& OutResult);
+
+#if WITH_EDITOR
+	// Sets the author of the MetaSound.
+	void SetAuthor(const FString& InAuthor);
+#endif // WITH_EDITOR
+
+	// Sets the node's input default value (used if no connection to the given node input is present)
+	UFUNCTION(BlueprintCallable, Category = "Audio|MetaSound|Builder", meta = (ExpandEnumAsExecs = "OutResult"))
+	void SetNodeInputDefault(const FMetaSoundBuilderNodeInputHandle& NodeInputHandle, const FMetasoundFrontendLiteral& Literal, EMetaSoundBuilderResult& OutResult);
+
+	virtual TScriptInterface<IMetaSoundDocumentInterface> Build(UObject* Parent, const FMetaSoundBuilderOptions& Options) const PURE_VIRTUAL(UMetaSoundBuilderBase::Build, return { }; );
+
+	// Returns the builder's supported UClass
+	virtual UClass& GetBuilderUClass() const PURE_VIRTUAL(UMetaSoundBuilderBase::GetBuilderUClass, return *UObject::StaticClass(); );
+
+	virtual void PostInitProperties() override;
+
+	// Initializes and ensures all nodes have a position (required prior to exporting to an asset if expected to be viewed in the editor).
+	void InitNodeLocations();
+
+protected:
+	const FMetaSoundFrontendDocumentBuilder& GetConstBuilder() const
+	{
+		return Builder;
+	}
+
+	template <typename UClassType>
+	TScriptInterface<IMetaSoundDocumentInterface> BuildInternal(UObject* Parent, const FMetaSoundBuilderOptions& BuilderOptions) const
+	{
+		FName ObjectName = BuilderOptions.Name;
+		if (!ObjectName.IsNone())
+		{
+			ObjectName = MakeUniqueObjectName(Parent, UClassType::StaticClass(), BuilderOptions.Name);
+		}
+
+		UClassType* NewMetaSound = BuilderOptions.ExistingAsset
+			? CastChecked<UClassType>(BuilderOptions.ExistingAsset.GetObject())
+			: NewObject<UClassType>(Parent, ObjectName, RF_Public | RF_Transient);
+		if (!NewMetaSound)
+		{
+			return nullptr;
+		}
+
+		FMetasoundFrontendDocument NewDocument = GetConstBuilder().GetDocument();
+		{
+			// This is required to ensure the newly build document has a unique class
+			// identifier to avoid collisions if added to the Frontend class registry
+			// (either below or at a later point in time).
+			constexpr bool bResetVersion = false;
+			FMetaSoundFrontendDocumentBuilder::InitGraphClassMetadata(NewDocument.RootGraph.Metadata, bResetVersion);
+		}
+		NewMetaSound->SetDocument(MoveTemp(NewDocument));
+		NewMetaSound->ConformObjectDataToInterfaces();
+
+		if (BuilderOptions.bAddToRegistry)
+		{
+			NewMetaSound->RegisterGraphWithFrontend();
+		}
+
+		UE_LOG(LogMetaSound, VeryVerbose, TEXT("New MetaSound '%s' built from '%s'"), *BuilderOptions.Name.ToString(), *GetFullName());
+		return NewMetaSound;
+	}
+
+	UPROPERTY()
+	FMetaSoundFrontendDocumentBuilder Builder;
+
+	// Friending allows for swapping the builder in certain circumstances where desired (eg. attaching a builder to an existing asset)
+	friend class UMetaSoundBuilderSubsystem;
+};
+
+/** Builder in charge of building a MetaSound */
+UCLASS(Transient, BlueprintType)
+class METASOUNDENGINE_API UMetaSoundPatchBuilder : public UMetaSoundBuilderBase
+{
+	GENERATED_BODY()
+
+public:
+	UFUNCTION(BlueprintCallable, Category = "Audio|MetaSound|Builder", meta = (WorldContext = "Parent"))
+	virtual UPARAM(DisplayName = "MetaSound") TScriptInterface<IMetaSoundDocumentInterface> Build(UObject* Parent, const FMetaSoundBuilderOptions& Options) const override;
+
+	virtual UClass& GetBuilderUClass() const override;
+};
+
+/** Builder in charge of building a MetaSound Source */
+UCLASS(Transient, BlueprintType)
+class METASOUNDENGINE_API UMetaSoundSourceBuilder : public UMetaSoundBuilderBase
+{
+	GENERATED_BODY()
+
+public:
+	UFUNCTION(BlueprintCallable, Category = "Audio|MetaSound|Builder", meta = (WorldContext = "Parent", AdvancedDisplay = "2"))
+	void Audition(UObject* Parent, UAudioComponent* AudioComponent, FOnCreateAuditionGeneratorHandleDelegate OnCreateGenerator);
+
+	UFUNCTION(BlueprintCallable, Category = "Audio|MetaSound|Builder", meta = (WorldContext = "Parent"))
+	virtual UPARAM(DisplayName = "MetaSound") TScriptInterface<IMetaSoundDocumentInterface> Build(UObject* Parent, const FMetaSoundBuilderOptions& Options) const override;
+
+	// Sets the output audio format of the source
+	UFUNCTION(BlueprintCallable, Category = "Audio|MetaSound|Builder", meta = (ExpandEnumAsExecs = "OutResult"))
+	void SetFormat(EMetaSoundOutputAudioFormat OutputFormat, EMetaSoundBuilderResult& OutResult);
+
+	virtual UClass& GetBuilderUClass() const override;
+
+	const Metasound::Engine::FOutputAudioFormatInfoPair* FindOutputAudioFormatInfo() const;
+
+private:
+	TWeakObjectPtr<UMetaSoundSource> AuditionSound;
+};
+
+/** The subsystem in charge of tracking MetaSound builders */
+UCLASS()
+class METASOUNDENGINE_API UMetaSoundBuilderSubsystem : public UEngineSubsystem
+{
+	GENERATED_BODY()
+
+	UPROPERTY()
+	TMap<FName, TObjectPtr<UMetaSoundPatchBuilder>> PatchBuilders;
+
+	UPROPERTY()
+	TMap<FName, TObjectPtr<UMetaSoundSourceBuilder>> SourceBuilders;
+
+public:
+	UMetaSoundSourceBuilder* AttachSourceBuilderToAsset(UMetaSoundSource* InSource) const;
+
+	UFUNCTION(BlueprintCallable, Category = "Audio|MetaSound|Builder",  meta = (ExpandEnumAsExecs = "OutResult"))
+	UPARAM(DisplayName = "Patch Builder") UMetaSoundPatchBuilder* CreatePatchBuilder(FName BuilderName, EMetaSoundBuilderResult& OutResult);
+
+	UFUNCTION(BlueprintCallable, Category = "Audio|MetaSound|Builder", meta = (ExpandEnumAsExecs = "OutResult"))
+	UPARAM(DisplayName = "Source Builder") UMetaSoundSourceBuilder* CreateSourceBuilder(
+		FName BuilderName,
+		FMetaSoundBuilderNodeOutputHandle& OnPlayNodeOutput,
+		FMetaSoundBuilderNodeInputHandle& OnFinishedNodeInput,
+		TArray<FMetaSoundBuilderNodeInputHandle>& AudioOutNodeInputs,
+		EMetaSoundBuilderResult& OutResult,
+		EMetaSoundOutputAudioFormat OutputFormat = EMetaSoundOutputAudioFormat::Mono,
+		bool bIsOneShot = true);
+
+	UFUNCTION(BlueprintCallable, Category = "Audio|MetaSound|Builder", meta = (DisplayName = "Create MetaSound Bool Literal"))
+	UPARAM(DisplayName = "Bool Literal") FMetasoundFrontendLiteral CreateBoolMetaSoundLiteral(bool Value, FName& DataType);
+
+	UFUNCTION(BlueprintCallable, Category = "Audio|MetaSound|Builder", meta = (DisplayName = "Create MetaSound Bool Array Literal"))
+	UPARAM(DisplayName = "Bool Array Literal") FMetasoundFrontendLiteral CreateBoolArrayMetaSoundLiteral(const TArray<bool>& Value, FName& DataType);
+
+	UFUNCTION(BlueprintCallable, Category = "Audio|MetaSound|Builder", meta = (DisplayName = "Create MetaSound Float Literal"))
+	UPARAM(DisplayName = "Float Literal") FMetasoundFrontendLiteral CreateFloatMetaSoundLiteral(float Value, FName& DataType);
+
+	UFUNCTION(BlueprintCallable, Category = "Audio|MetaSound|Builder", meta = (DisplayName = "Create MetaSound Float Array Literal"))
+	UPARAM(DisplayName = "Float Array Literal") FMetasoundFrontendLiteral CreateFloatArrayMetaSoundLiteral(const TArray<float>& Value, FName& DataType);
+
+	UFUNCTION(BlueprintCallable, Category = "Audio|MetaSound|Builder", meta = (DisplayName = "Create MetaSound Int Literal"))
+	UPARAM(DisplayName = "Int32 Literal") FMetasoundFrontendLiteral CreateIntMetaSoundLiteral(int32 Value, FName& DataType);
+
+	UFUNCTION(BlueprintCallable, Category = "Audio|MetaSound|Builder", meta = (DisplayName = "Create MetaSound Int Array Literal"))
+	UPARAM(DisplayName = "Int32 Array Literal") FMetasoundFrontendLiteral CreateIntArrayMetaSoundLiteral(const TArray<int32>& Value, FName& DataType);
+
+	UFUNCTION(BlueprintCallable, Category = "Audio|MetaSound|Builder", meta = (DisplayName = "Create MetaSound Object Literal"))
+	UPARAM(DisplayName = "Object Literal") FMetasoundFrontendLiteral CreateObjectMetaSoundLiteral(UObject* Value);
+
+	UFUNCTION(BlueprintCallable, Category = "Audio|MetaSound|Builder", meta = (DisplayName = "Create MetaSound Object Array Literal"))
+	UPARAM(DisplayName = "Object Array Literal") FMetasoundFrontendLiteral CreateObjectArrayMetaSoundLiteral(const TArray<UObject*>& Value);
+
+	UFUNCTION(BlueprintCallable, Category = "Audio|MetaSound|Builder", meta = (DisplayName = "Create MetaSound String Literal"))
+	UPARAM(DisplayName = "String Literal") FMetasoundFrontendLiteral CreateStringMetaSoundLiteral(const FString& Value, FName& DataType);
+
+	UFUNCTION(BlueprintCallable, Category = "Audio|MetaSound|Builder", meta = (DisplayName = "Create MetaSound String Array Literal"))
+	UPARAM(DisplayName = "String Array Literal") FMetasoundFrontendLiteral CreateStringArrayMetaSoundLiteral(const TArray<FString>& Value, FName& DataType);
+
+	UFUNCTION(BlueprintCallable, Category = "Audio|MetaSound|Builder", meta = (DisplayName = "Create MetaSound Literal From AudioParameter"))
+	UPARAM(DisplayName = "Param Literal") FMetasoundFrontendLiteral CreateMetaSoundLiteralFromParam(const FAudioParameter& Param);
+
+	UFUNCTION(BlueprintCallable, Category = "Audio|MetaSound|Builder")
+	UPARAM(DisplayName = "Patch Builder") UMetaSoundPatchBuilder* FindPatchBuilder(FName BuilderName);
+
+	UFUNCTION(BlueprintCallable, Category = "Audio|MetaSound|Builder")
+	UPARAM(DisplayName = "Source Builder") UMetaSoundSourceBuilder* FindSourceBuilder(FName BuilderName);
+
+	UFUNCTION(BlueprintCallable, Category = "Audio|MetaSound|Builder")
+	UPARAM(DisplayName = "Is Registered") bool IsInterfaceRegistered(FName InInterfaceName) const;
+
+	// Adds builder to subsystem's registry to make it persistent and easily accessible by multiple systems or Blueprints
+	UFUNCTION(BlueprintCallable, Category = "Audio|MetaSound|Builder")
+	void RegisterPatchBuilder(FName BuilderName, UMetaSoundPatchBuilder* Builder);
+
+	// Adds builder to subsystem's registry to make it persistent and easily accessible by multiple systems or Blueprints
+	UFUNCTION(BlueprintCallable, Category = "Audio|MetaSound|Builder")
+	void RegisterSourceBuilder(FName BuilderName, UMetaSoundSourceBuilder* Builder);
+
+	UFUNCTION(BlueprintCallable, Category = "Audio|MetaSound|Builder")
+	UPARAM(DisplayName = "Unregistered") bool UnregisterPatchBuilder(FName BuilderName);
+
+	UFUNCTION(BlueprintCallable, Category = "Audio|MetaSound|Builder")
+	UPARAM(DisplayName = "Unregistered") bool UnregisterSourceBuilder(FName BuilderName);
+};
diff --git a/Engine/Plugins/Runtime/Metasound/Source/MetasoundEngine/Public/MetasoundSource.h b/Engine/Plugins/Runtime/Metasound/Source/MetasoundEngine/Public/MetasoundSource.h
--- a/Engine/Plugins/Runtime/Metasound/Source/MetasoundEngine/Public/MetasoundSource.h
+++ b/Engine/Plugins/Runtime/Metasound/Source/MetasoundEngine/Public/MetasoundSource.h
@@ -1,10 +1,10 @@
 // Copyright Epic Games, Inc. All Rights Reserved.
 #pragma once
 
-#include "CoreMinimal.h"
 #include "EdGraph/EdGraph.h"
 #include "HAL/CriticalSection.h"
 #include "IAudioParameterTransmitter.h"
+#include "Interfaces/MetasoundOutputFormatInterfaces.h"
 #include "Metasound.h"
 #include "MetasoundAssetBase.h"
 #include "MetasoundFrontend.h"
@@ -18,25 +18,14 @@
 
 #include "MetasoundSource.generated.h"
 
+
+// Forward Declarations
 namespace Metasound
 {
-	// Forward declare
 	struct FMetaSoundEngineAssetHelper;
 	class FMetasoundGenerator;
-}
 
-/** Declares the output audio format of the UMetaSoundSource */
-UENUM()
-enum class EMetasoundSourceAudioFormat : uint8
-{
-	Mono,
-	Stereo,
-	Quad,
-	FiveDotOne UMETA(DisplayName="5.1"),
-	SevenDotOne UMETA(DisplayName="7.1"),
-
-	COUNT UMETA(Hidden)
-};
+} // namespace Metasound
 
 DECLARE_MULTICAST_DELEGATE_TwoParams(FOnGeneratorInstanceCreated, uint64, TSharedPtr<Metasound::FMetasoundGenerator>);
 DECLARE_MULTICAST_DELEGATE_TwoParams(FOnGeneratorInstanceDestroyed, uint64, TSharedPtr<Metasound::FMetasoundGenerator>);
@@ -45,7 +34,7 @@ DECLARE_MULTICAST_DELEGATE_TwoParams(FOnGeneratorInstanceDestroyed, uint64, TSha
  * This Metasound type can be played as an audio source.
  */
 UCLASS(hidecategories = object, BlueprintType)
-class METASOUNDENGINE_API UMetaSoundSource : public USoundWaveProcedural, public FMetasoundAssetBase
+class METASOUNDENGINE_API UMetaSoundSource : public USoundWaveProcedural, public FMetasoundAssetBase, public IMetaSoundDocumentInterface
 {
 	GENERATED_BODY()
 
@@ -73,7 +62,7 @@ public:
 
 	// The output audio format of the metasound source.
 	UPROPERTY(EditAnywhere, BlueprintReadOnly, Category = Metasound)
-	EMetasoundSourceAudioFormat OutputFormat;
+	EMetaSoundOutputAudioFormat OutputFormat;
 
 	UPROPERTY(AssetRegistrySearchable)
 	FGuid AssetClassID;
@@ -201,8 +190,7 @@ public:
 	FOnGeneratorInstanceDestroyed OnGeneratorInstanceDestroyed;
 
 protected:
-
-	Metasound::Frontend::FDocumentAccessPtr GetDocument() override
+	Metasound::Frontend::FDocumentAccessPtr GetDocumentAccessPtr() override
 	{
 		using namespace Metasound::Frontend;
 		// Return document using FAccessPoint to inform the TAccessPtr when the 
@@ -210,7 +198,7 @@ protected:
 		return MakeAccessPtr<FDocumentAccessPtr>(RootMetasoundDocument.AccessPoint, RootMetasoundDocument);
 	}
 
-	Metasound::Frontend::FConstDocumentAccessPtr GetDocument() const override
+	Metasound::Frontend::FConstDocumentAccessPtr GetDocumentConstAccessPtr() const override
 	{
 		using namespace Metasound::Frontend;
 		// Return document using FAccessPoint to inform the TAccessPtr when the 
@@ -218,6 +206,11 @@ protected:
 		return MakeAccessPtr<FConstDocumentAccessPtr>(RootMetasoundDocument.AccessPoint, RootMetasoundDocument);
 	}
 
+	virtual const FMetasoundFrontendDocument& GetDocument() const override
+	{
+		return RootMetasoundDocument;
+	}
+
 	/** Gets all the default parameters for this Asset.  */
 	virtual bool GetAllDefaultParameters(TArray<FAudioParameter>& OutParameters) const override;
 
@@ -226,6 +219,10 @@ protected:
 #endif // #if WITH_EDITOR
 
 private:
+	virtual FMetasoundFrontendDocument& GetDocument() override
+	{
+		return RootMetasoundDocument;
+	}
 
 	bool IsParameterValid(const FAudioParameter& InParameter, const TMap<FName, FMetasoundFrontendVertex>& InInputNameVertexMap) const;
 
diff --git a/Engine/Plugins/Runtime/Metasound/Source/MetasoundEngine/Public/MetasoundUObjectRegistry.h b/Engine/Plugins/Runtime/Metasound/Source/MetasoundEngine/Public/MetasoundUObjectRegistry.h
--- a/Engine/Plugins/Runtime/Metasound/Source/MetasoundEngine/Public/MetasoundUObjectRegistry.h
+++ b/Engine/Plugins/Runtime/Metasound/Source/MetasoundEngine/Public/MetasoundUObjectRegistry.h
@@ -26,8 +26,12 @@ namespace Metasound
 	public:
 		virtual ~IMetasoundUObjectRegistryEntry() = default;
 
-		/** Interface version associated with this entry. */
-		virtual const FMetasoundFrontendVersion& GetInterfaceVersion() const = 0;
+		UE_DEPRECATED(5.3, "Interfaces queries are no longer restricted to a single UObjectRegistryEntry. Options set upon registering interface determines supported UClasses.")
+		virtual const FMetasoundFrontendVersion& GetInterfaceVersion() const
+		{
+			static const FMetasoundFrontendVersion InterfaceVersion;
+			return InterfaceVersion;
+		}
 
 		/** UClass associated with this entry. */
 		virtual UClass* GetUClass() const = 0;
@@ -45,7 +49,12 @@ namespace Metasound
 		virtual const FMetasoundAssetBase* Cast(const UObject* InObject) const = 0;
 
 		/** Creates a new object of the UClass type. */
-		virtual UObject* NewObject(UPackage* InPackage, const FName& InName) const = 0;
+		UE_DEPRECATED(5.3, "NewObject is now done via the MetaSoundBuilder system. Attempting to create a MetaSound via this call will now always fail.")
+		virtual UObject* NewObject(UPackage* InPackage, const FName& InName) const
+		{
+			checkNoEntry();
+			return nullptr;
+		}
 
 	private:
 		IMetasoundUObjectRegistryEntry() = default;
@@ -67,18 +76,9 @@ namespace Metasound
 		static_assert(std::is_base_of<UObject, UClassType>::value, "UClass must be derived from UObject");
 
 	public:
-		TMetasoundUObjectRegistryEntry(const FMetasoundFrontendVersion& InInterfaceVersion)
-		:	InterfaceVersion(InInterfaceVersion)
-		{
-		}
-
+		TMetasoundUObjectRegistryEntry() = default;
 		virtual ~TMetasoundUObjectRegistryEntry() = default;
 
-		virtual const FMetasoundFrontendVersion& GetInterfaceVersion() const override
-		{
-			return InterfaceVersion;
-		}
-
 		UClass* GetUClass() const override
 		{
 			return UClassType::StaticClass();
@@ -119,15 +119,6 @@ namespace Metasound
 			}
 			return static_cast<const FMetasoundAssetBase*>(CastChecked<const UClassType>(InObject));
 		}
-
-		UObject* NewObject(UPackage* InPackage, const FName& InName) const override
-		{
-			return ::NewObject<UClassType>(InPackage, InName);
-		}
-
-	private:
-
-		FMetasoundFrontendVersion InterfaceVersion;
 	};
 
 
@@ -145,13 +136,16 @@ namespace Metasound
 			static IMetasoundUObjectRegistry& Get();
 
 			/** Adds an entry to the registry. */
-			virtual void RegisterUClassInterface(TUniquePtr<IMetasoundUObjectRegistryEntry>&& InEntry) = 0;
+			virtual void RegisterUClass(TUniquePtr<IMetasoundUObjectRegistryEntry>&& InEntry) = 0;
+
+			UE_DEPRECATED(5.3, "Interfaces are no longer registered with the UObject registry as interfaces now support multiple UClasses that are registered with the interface registry.")
+			virtual void RegisterUClassInterface(TUniquePtr<IMetasoundUObjectRegistryEntry>&& InEntry) { }
 
-			/** Returns all RegistryEntries with the given name */
-			virtual TArray<const IMetasoundUObjectRegistryEntry*> FindInterfaceEntriesByName(FName InName) const = 0;
+			UE_DEPRECATED(5.3, "Interfaces are no longer registered with the UObject registry as interfaces now support multiple UClasses that are registered with the interface registry.")
+			virtual TArray<const IMetasoundUObjectRegistryEntry*> FindInterfaceEntriesByName(FName InName) const { return { }; }
 
-			/** Returns all UClasses registered to the interface version. */
-			virtual TArray<UClass*> FindSupportedInterfaceClasses(const FMetasoundFrontendVersion& InInterfaceVersion) const = 0;
+			UE_DEPRECATED(5.3, "Interfaces are no longer registered with the UObject registry as interfaces now support multiple UClasses that are registered with the interface registry.")
+			virtual TArray<UClass*> FindSupportedInterfaceClasses(const FMetasoundFrontendVersion& InInterfaceVersion) const { return { }; }
 
 			/** Creates a new object from a MetaSound document.
 			 *
@@ -162,7 +156,8 @@ namespace Metasound
 			 *
 			 * @return A new object. A nullptr on error.
 			 */
-			virtual UObject* NewObject(UClass* InClass, const FMetasoundFrontendDocument& InDocument, const FString& InPath) const = 0;
+			 UE_DEPRECATED(5.3, "UObject registry form of NewObject is no longer used. Use UMetaSoundBuilderSubsystem to author MetaSounds instead.")
+			virtual UObject* NewObject(UClass* InClass, const FMetasoundFrontendDocument& InDocument, const FString& InPath) const { return nullptr; }
 
 			/** Iterate all registered UClasses that serve as MetaSound assets.*/
 			virtual void IterateRegisteredUClasses(TFunctionRef<void(UClass&)> InFunc) const = 0;
diff --git a/Engine/Plugins/Runtime/Metasound/Source/MetasoundEngineTest/Private/EngineTestMetasoundSource.cpp b/Engine/Plugins/Runtime/Metasound/Source/MetasoundEngineTest/Private/EngineTestMetasoundSource.cpp
--- a/Engine/Plugins/Runtime/Metasound/Source/MetasoundEngineTest/Private/EngineTestMetasoundSource.cpp
+++ b/Engine/Plugins/Runtime/Metasound/Source/MetasoundEngineTest/Private/EngineTestMetasoundSource.cpp
@@ -4,7 +4,7 @@
 #include "Components/AudioComponent.h"
 #include "IAudioParameterInterfaceRegistry.h"
 #include "Interfaces/IPluginManager.h"
-#include "Interfaces/MetasoundFrontendOutputFormatInterfaces.h"
+#include "Interfaces/MetasoundOutputFormatInterfaces.h"
 #include "Interfaces/MetasoundFrontendSourceInterface.h"
 #include "MetasoundFrontendController.h"
 #include "MetasoundFrontendSearchEngine.h"
@@ -69,6 +69,7 @@ namespace EngineTestMetasoundSourcePrivate
 	{
 		using namespace Audio;
 		using namespace Metasound;
+		using namespace Metasound::Engine;
 		using namespace Metasound::Frontend;
 
 		FMetasoundFrontendDocument Document;
@@ -83,10 +84,10 @@ namespace EngineTestMetasoundSourcePrivate
 		// Add default source & mono interface members (OnPlay, OnFinished & Mono Out)
 		FModifyRootGraphInterfaces InterfaceTransform(
 		{ },
-		{	
-			SourceInterface::GetVersion(), 
+		{
+			SourceInterface::GetVersion(),
 			SourceOneShotInterface::GetVersion(),
-			OutputFormatMonoInterface::GetVersion() 
+			OutputFormatMonoInterface::GetVersion()
 		}); 
 		InterfaceTransform.Transform(DocumentHandle);
 
diff --git a/Engine/Plugins/Runtime/Metasound/Source/MetasoundFrontend/Private/Interfaces/MetasoundFrontendSourceInterface.cpp b/Engine/Plugins/Runtime/Metasound/Source/MetasoundFrontend/Private/Interfaces/MetasoundFrontendSourceInterface.cpp
--- a/Engine/Plugins/Runtime/Metasound/Source/MetasoundFrontend/Private/Interfaces/MetasoundFrontendSourceInterface.cpp
+++ b/Engine/Plugins/Runtime/Metasound/Source/MetasoundFrontend/Private/Interfaces/MetasoundFrontendSourceInterface.cpp
@@ -14,7 +14,6 @@ namespace Metasound
 {
 	namespace Frontend
 	{
-	
 #define AUDIO_PARAMETER_INTERFACE_NAMESPACE "UE.Source.OneShot"
 
 		namespace SourceOneShotInterface
@@ -30,12 +29,17 @@ namespace Metasound
 				const FName OnFinished = AUDIO_PARAMETER_INTERFACE_MEMBER_DEFINE("OnFinished");
 			}
 
-			Audio::FParameterInterfacePtr CreateInterface(const UClass& InUClass)
+			Audio::FParameterInterfacePtr CreateInterface(const UClass& InClass)
+			{
+				return CreateInterface();
+			}
+
+			Audio::FParameterInterfacePtr CreateInterface()
 			{
 				struct FInterface : public Audio::FParameterInterface
 				{
-					FInterface(const UClass& InAssetClass)
-						: FParameterInterface(SourceOneShotInterface::GetVersion().Name, SourceOneShotInterface::GetVersion().Number.ToInterfaceVersion(), InAssetClass)
+					FInterface()
+						: FParameterInterface(SourceOneShotInterface::GetVersion().Name, SourceOneShotInterface::GetVersion().Number.ToInterfaceVersion())
 					{
 						Outputs =
 						{
@@ -43,20 +47,19 @@ namespace Metasound
 								LOCTEXT("OnFinished", "On Finished"),
 								LOCTEXT("OnFinishedDescription", "Trigger executed to initiate stopping the source."),
 								GetMetasoundDataTypeName<FTrigger>(),
-								Outputs::OnFinished,								
+								Outputs::OnFinished,
 								LOCTEXT("OnFinishedWarning", "\"On Finished\" should be connected for OneShot MetaSound sources. For sources with undefined duration (e.g. looping), remove the OneShot interface and use an audio component to avoid leaking the source."),
 							}
 						};
 					}
 				};
 
-				return MakeShared<FInterface>(InUClass);
+				return MakeShared<FInterface>();
 			}
 		} // namespace SourceOneShotInterface
 
 #undef AUDIO_PARAMETER_INTERFACE_NAMESPACE
 
-
 #define AUDIO_PARAMETER_INTERFACE_NAMESPACE "UE.Source"
 
 		namespace SourceInterfaceV1_0
@@ -86,12 +89,17 @@ namespace Metasound
 				const FName TransmitterID = AUDIO_PARAMETER_INTERFACE_MEMBER_DEFINE("TransmitterID");
 			}
 
-			Audio::FParameterInterfacePtr CreateInterface(const UClass& InUClass)
+			Audio::FParameterInterfacePtr CreateInterface(const UClass& InClass)
+			{
+				return CreateInterface();
+			}
+
+			Audio::FParameterInterfacePtr CreateInterface()
 			{
 				struct FInterface : public Audio::FParameterInterface
 				{
-					FInterface(const UClass& InAssetClass)
-						: FParameterInterface(SourceInterfaceV1_0::GetVersion().Name, SourceInterfaceV1_0::GetVersion().Number.ToInterfaceVersion(), InAssetClass)
+					FInterface()
+						: FParameterInterface(SourceInterfaceV1_0::GetVersion().Name, SourceInterfaceV1_0::GetVersion().Number.ToInterfaceVersion())
 					{
 						Inputs =
 						{
@@ -149,11 +157,10 @@ namespace Metasound
 					}
 				};
 
-				return MakeShared<FInterface>(InUClass);
+				return MakeShared<FInterface>();
 			}
 		} // namespace SourceInterfaceV1_0
 
-
 		namespace SourceInterface
 		{
 			const FMetasoundFrontendVersion& GetVersion()
@@ -177,12 +184,17 @@ namespace Metasound
 				const FName AudioMixerNumOutputFrames = AUDIO_PARAMETER_INTERFACE_MEMBER_DEFINE("AudioMixerNumOutputFrames");
 			}
 
-			Audio::FParameterInterfacePtr CreateInterface(const UClass& InUClass)
+			Audio::FParameterInterfacePtr CreateInterface(const UClass& InClass)
+			{
+				return CreateInterface();
+			}
+
+			Audio::FParameterInterfacePtr CreateInterface()
 			{
 				struct FInterface : public Audio::FParameterInterface
 				{
-					FInterface(const UClass& InAssetClass)
-						: FParameterInterface(SourceInterface::GetVersion().Name, SourceInterface::GetVersion().Number.ToInterfaceVersion(), InAssetClass)
+					FInterface()
+						: FParameterInterface(SourceInterface::GetVersion().Name, SourceInterface::GetVersion().Number.ToInterfaceVersion())
 					{
 						Inputs =
 						{
@@ -236,7 +248,7 @@ namespace Metasound
 					}
 				};
 
-				return MakeShared<FInterface>(InUClass);
+				return MakeShared<FInterface>();
 			}
 
 			bool FUpdateInterface::Transform(Frontend::FDocumentHandle InDocument) const
@@ -245,9 +257,9 @@ namespace Metasound
 
 				// When upgrading, we only want to add the one-shot interface if the MetaSound actually has the OnFinished trigger connected.
 				bool bIsOnFinishedConnected = false;
-				InDocument->GetRootGraph()->IterateConstNodes([&](FConstNodeHandle NodeHandle)
+				InDocument->GetRootGraph()->IterateConstNodes([&](Frontend::FConstNodeHandle NodeHandle)
 				{
-					NodeHandle->IterateConstInputs([&](FConstInputHandle InputHandle)
+					NodeHandle->IterateConstInputs([&](Frontend::FConstInputHandle InputHandle)
 					{
 						if (InputHandle->GetName() == SourceInterfaceV1_0::Outputs::OnFinished)
 						{
@@ -271,14 +283,13 @@ namespace Metasound
 					InterfacesToAdd.Add(SourceOneShotInterface::GetVersion());
 				}
 
-				FModifyRootGraphInterfaces InterfaceTransform(InterfacesToRemove, InterfacesToAdd);
+				Frontend::FModifyRootGraphInterfaces InterfaceTransform(InterfacesToRemove, InterfacesToAdd);
 				return InterfaceTransform.Transform(InDocument);
 			}
 
 } // namespace SourceInterface
 
 #undef AUDIO_PARAMETER_INTERFACE_NAMESPACE
-
 	} // namespace Frontend
 } // namespace Metasound
 
diff --git a/Engine/Plugins/Runtime/Metasound/Source/MetasoundFrontend/Private/MetasoundAssetBase.cpp b/Engine/Plugins/Runtime/Metasound/Source/MetasoundFrontend/Private/MetasoundAssetBase.cpp
--- a/Engine/Plugins/Runtime/Metasound/Source/MetasoundFrontend/Private/MetasoundAssetBase.cpp
+++ b/Engine/Plugins/Runtime/Metasound/Source/MetasoundFrontend/Private/MetasoundAssetBase.cpp
@@ -18,6 +18,7 @@
 #include "MetasoundFrontendArchetypeRegistry.h"
 #include "MetasoundFrontendController.h"
 #include "MetasoundFrontendDocument.h"
+#include "MetasoundFrontendDocumentBuilder.h"
 #include "MetasoundFrontendDocumentVersioning.h"
 #include "MetasoundFrontendGraph.h"
 #include "MetasoundFrontendNodeTemplateRegistry.h"
@@ -26,6 +27,7 @@
 #include "MetasoundFrontendTransform.h"
 #include "MetasoundJsonBackend.h"
 #include "MetasoundLog.h"
+#include "MetasoundParameterPack.h"
 #include "MetasoundParameterTransmitter.h"
 #include "MetasoundTrace.h"
 #include "MetasoundVertex.h"
@@ -33,7 +35,6 @@
 #include "StructSerializer.h"
 #include "Templates/SharedPointer.h"
 #include "UObject/MetaData.h"
-#include "MetasoundParameterPack.h"
 
 #define LOCTEXT_NAMESPACE "MetaSound"
 
@@ -295,30 +296,7 @@ void FMetasoundAssetBase::SetMetadata(FMetasoundFrontendClassMetadata& InMetadat
 
 bool FMetasoundAssetBase::GetDeclaredInterfaces(TArray<const Metasound::Frontend::IInterfaceRegistryEntry*>& OutInterfaces) const
 {
-	using namespace Metasound;
-	using namespace Metasound::Frontend;
-
-	if (const FMetasoundFrontendDocument* Document = GetDocument().Get())
-	{
-		bool bInterfacesFound = true;
-
-		Algo::Transform(Document->Interfaces, OutInterfaces, [&](const FMetasoundFrontendVersion& Version)
-		{
-			const FInterfaceRegistryKey InterfaceKey = GetInterfaceRegistryKey(Version);
-			const IInterfaceRegistryEntry* RegistryEntry = IInterfaceRegistry::Get().FindInterfaceRegistryEntry(InterfaceKey);
-			if (!RegistryEntry)
-			{
-				bInterfacesFound = false;
-				UE_LOG(LogMetaSound, Warning, TEXT("No registered interface matching interface version on document [InterfaceVersion:%s]"), *Version.ToString());
-			}
-
-			return RegistryEntry;
-		});
-	
-		return bInterfacesFound;
-	}
-
-	return false;
+	return FMetaSoundFrontendDocumentBuilder::FindDeclaredInterfaces(GetDocumentChecked(), OutInterfaces);
 }
 
 bool FMetasoundAssetBase::IsInterfaceDeclared(const FMetasoundFrontendVersion& InVersion) const
@@ -349,6 +327,13 @@ void FMetasoundAssetBase::AddDefaultInterfaces()
 	FModifyRootGraphInterfaces({ }, InitInterfaces).Transform(DocumentHandle);
 }
 
+void FMetasoundAssetBase::SetDocument(FMetasoundFrontendDocument&& InDocument)
+{
+	FMetasoundFrontendDocument& Document = GetDocumentChecked();
+	Document = MoveTemp(InDocument);
+	MarkMetasoundDocumentDirty();
+}
+
 bool FMetasoundAssetBase::VersionAsset()
 {
 	using namespace Metasound;
@@ -363,7 +348,7 @@ bool FMetasoundAssetBase::VersionAsset()
 		AssetPath = OwningAsset->GetPathName();
 	}
 
-	FMetasoundFrontendDocument* Doc = GetDocument().Get();
+	FMetasoundFrontendDocument* Doc = GetDocumentAccessPtr().Get();
 	if (!ensure(Doc))
 	{
 		return false;
@@ -410,7 +395,7 @@ void FMetasoundAssetBase::CacheRegistryMetadata()
 {
 	using namespace Metasound::Frontend;
 
-	FMetasoundFrontendDocument* Document = GetDocument().Get();
+	FMetasoundFrontendDocument* Document = GetDocumentAccessPtr().Get();
 	if (!ensure(Document))
 	{
 		return;
@@ -710,12 +695,12 @@ bool FMetasoundAssetBase::MarkMetasoundDocumentDirty() const
 
 Metasound::Frontend::FDocumentHandle FMetasoundAssetBase::GetDocumentHandle()
 {
-	return Metasound::Frontend::IDocumentController::CreateDocumentHandle(GetDocument());
+	return Metasound::Frontend::IDocumentController::CreateDocumentHandle(GetDocumentAccessPtr());
 }
 
 Metasound::Frontend::FConstDocumentHandle FMetasoundAssetBase::GetDocumentHandle() const
 {
-	return Metasound::Frontend::IDocumentController::CreateDocumentHandle(GetDocument());
+	return Metasound::Frontend::IDocumentController::CreateDocumentHandle(GetDocumentConstAccessPtr());
 }
 
 Metasound::Frontend::FGraphHandle FMetasoundAssetBase::GetRootGraphHandle()
@@ -732,7 +717,7 @@ bool FMetasoundAssetBase::ImportFromJSON(const FString& InJSON)
 {
 	METASOUND_TRACE_CPUPROFILER_EVENT_SCOPE(MetaSoundAssetBase::ImportFromJSON);
 
-	FMetasoundFrontendDocument* Document = GetDocument().Get();
+	FMetasoundFrontendDocument* Document = GetDocumentAccessPtr().Get();
 	if (ensure(nullptr != Document))
 	{
 		bool bSuccess = Metasound::Frontend::ImportJSONToMetasound(InJSON, *Document);
@@ -751,7 +736,7 @@ bool FMetasoundAssetBase::ImportFromJSONAsset(const FString& InAbsolutePath)
 {
 	METASOUND_TRACE_CPUPROFILER_EVENT_SCOPE(MetaSoundAssetBase::ImportFromJSONAsset);
 
-	Metasound::Frontend::FDocumentAccessPtr DocumentPtr = GetDocument();
+	Metasound::Frontend::FDocumentAccessPtr DocumentPtr = GetDocumentAccessPtr();
 	if (FMetasoundFrontendDocument* Document = DocumentPtr.Get())
 	{
 		bool bSuccess = Metasound::Frontend::ImportJSONAssetToMetasound(InAbsolutePath, *Document);
@@ -768,14 +753,14 @@ bool FMetasoundAssetBase::ImportFromJSONAsset(const FString& InAbsolutePath)
 
 FMetasoundFrontendDocument& FMetasoundAssetBase::GetDocumentChecked()
 {
-	FMetasoundFrontendDocument* Document = GetDocument().Get();
+	FMetasoundFrontendDocument* Document = GetDocumentAccessPtr().Get();
 	check(nullptr != Document);
 	return *Document;
 }
 
 const FMetasoundFrontendDocument& FMetasoundAssetBase::GetDocumentChecked() const
 {
-	const FMetasoundFrontendDocument* Document = GetDocument().Get();
+	const FMetasoundFrontendDocument* Document = GetDocumentConstAccessPtr().Get();
 
 	check(nullptr != Document);
 	return *Document;
@@ -802,7 +787,7 @@ TArray<FMetasoundFrontendClassInput> FMetasoundAssetBase::GetPublicClassInputs()
 	// Inputs which are controlled by an interface are private. 
 	TArray<const IInterfaceRegistryEntry*> Interfaces;
 	TSet<FVertexName> PrivateInputs;
-	GetDeclaredInterfaces(Interfaces);
+	FMetaSoundFrontendDocumentBuilder::FindDeclaredInterfaces(GetDocumentChecked(), Interfaces);
 	for (const IInterfaceRegistryEntry* InterfaceEntry : Interfaces)
 	{
 		if (InterfaceEntry)
diff --git a/Engine/Plugins/Runtime/Metasound/Source/MetasoundFrontend/Private/MetasoundFrontendArchetypeRegistry.cpp b/Engine/Plugins/Runtime/Metasound/Source/MetasoundFrontend/Private/MetasoundFrontendArchetypeRegistry.cpp
--- a/Engine/Plugins/Runtime/Metasound/Source/MetasoundFrontend/Private/MetasoundFrontendArchetypeRegistry.cpp
+++ b/Engine/Plugins/Runtime/Metasound/Source/MetasoundFrontend/Private/MetasoundFrontendArchetypeRegistry.cpp
@@ -62,6 +62,50 @@ namespace Metasound
 			return false;
 		}
 
+		void FInterfaceRegistry::AddInterfaceOutputBindings(const FMetasoundFrontendVersion& InInputInterfaceVersion, TArray<FMetasoundFrontendInterfaceBinding>&& InOutputBindings)
+		{
+			using namespace Frontend;
+
+			FInterfaceRegistryKey Key = GetInterfaceRegistryKey(InInputInterfaceVersion);
+			if (TUniquePtr<IInterfaceRegistryEntry>* EntryPtr = Entries.Find(Key))
+			{
+				IInterfaceRegistryEntry* Entry = (*EntryPtr).Get();
+				check(Entry);
+				Entry->AddOutputBindings(MoveTemp(InOutputBindings));
+			}
+		}
+
+		bool FInterfaceRegistry::RemoveInterfaceOutputBinding(const FMetasoundFrontendVersion& InInputInterfaceVersion, const FMetasoundFrontendVersion& InOutputInterfaceVersion)
+		{
+			using namespace Frontend;
+
+			FInterfaceRegistryKey Key = GetInterfaceRegistryKey(InInputInterfaceVersion);
+			if (TUniquePtr<IInterfaceRegistryEntry>* EntryPtr = Entries.Find(Key))
+			{
+				IInterfaceRegistryEntry* Entry = (*EntryPtr).Get();
+				check(Entry);
+				return Entry->RemoveOutputBinding(InOutputInterfaceVersion);
+			}
+
+			return false;
+		}
+
+		bool FInterfaceRegistry::RemoveAllInterfaceOutputBindings(const FMetasoundFrontendVersion& InInputInterfaceVersion)
+		{
+			using namespace Frontend;
+
+			FInterfaceRegistryKey Key = GetInterfaceRegistryKey(InInputInterfaceVersion);
+			if (TUniquePtr<IInterfaceRegistryEntry>* EntryPtr = Entries.Find(Key))
+			{
+				IInterfaceRegistryEntry* Entry = (*EntryPtr).Get();
+				check(Entry);
+				Entry->ResetOutputBindings();
+				return true;
+			}
+
+			return false;
+		}
+
 		TUniquePtr<FInterfaceTransactionStream> FInterfaceRegistry::CreateTransactionStream()
 		{
 			return MakeUnique<FInterfaceTransactionStream>(TransactionBuffer);
