- Add Runtime MetaSound Builder support

diff --git a/Engine/Plugins/Runtime/Metasound/Source/MetasoundFrontend/Private/MetasoundFrontendDocumentBuilder.cpp b/Engine/Plugins/Runtime/Metasound/Source/MetasoundFrontend/Private/MetasoundFrontendDocumentBuilder.cpp
new file mode 100644
--- /dev/null
+++ b/Engine/Plugins/Runtime/Metasound/Source/MetasoundFrontend/Private/MetasoundFrontendDocumentBuilder.cpp
@@ -0,0 +1,1385 @@
+// Copyright Epic Games, Inc. All Rights Reserved.
+#include "MetasoundFrontendDocumentBuilder.h"
+
+#include "Algo/AnyOf.h"
+#include "Algo/Find.h"
+#include "Algo/NoneOf.h"
+#include "Algo/Transform.h"
+#include "MetasoundAssetBase.h"
+#include "MetasoundAssetManager.h"
+#include "MetasoundDocumentInterface.h"
+#include "MetasoundFrontendArchetypeRegistry.h"
+#include "MetasoundFrontendDocumentCache.h"
+#include "MetasoundFrontendDocument.h"
+#include "MetasoundFrontendRegistries.h"
+#include "MetasoundFrontendSearchEngine.h"
+#include "MetasoundFrontendTransform.h"
+#include "MetasoundTrace.h"
+
+#include UE_INLINE_GENERATED_CPP_BY_NAME(MetasoundFrontendDocumentBuilder)
+
+
+namespace Metasound::Frontend
+{
+	namespace DocumentBuilderPrivate
+	{
+		void FinalizeGraphVertexNode(FMetasoundFrontendNode& InOutNode, const FMetasoundFrontendClassVertex& InVertex)
+		{
+			InOutNode.Name = InVertex.Name;
+			// Set name on related vertices of input node
+			auto IsVertexWithTypeName = [&InVertex](const FMetasoundFrontendVertex& Vertex) { return Vertex.TypeName == InVertex.TypeName; };
+			if (FMetasoundFrontendVertex* InputVertex = InOutNode.Interface.Inputs.FindByPredicate(IsVertexWithTypeName))
+			{
+				InputVertex->Name = InVertex.Name;
+			}
+			else
+			{
+				UE_LOG(LogMetaSound, Error, TEXT("Node associated with graph vertex of type '%s' does not contain input vertex with type '%s'"), *InVertex.TypeName.ToString());
+			}
+
+			if (FMetasoundFrontendVertex* OutputVertex = InOutNode.Interface.Outputs.FindByPredicate(IsVertexWithTypeName))
+			{
+				OutputVertex->Name = InVertex.Name;
+			}
+			else
+			{
+				UE_LOG(LogMetaSound, Error, TEXT("Node associated with graph vertex of type '%s' does not contain output vertex of matching type."), *InVertex.TypeName.ToString());
+			}
+		}
+	} // namespace DocumentBuilderPrivate
+} // namespace Metasound::Frontend
+
+FMetaSoundFrontendDocumentBuilder::FMetaSoundFrontendDocumentBuilder()
+	: DocumentInterface(NewObject<UMetaSoundBuilderDocument>())
+{
+	ReloadCache();
+}
+
+FMetaSoundFrontendDocumentBuilder::FMetaSoundFrontendDocumentBuilder(TScriptInterface<IMetaSoundDocumentInterface> InDocumentInterface)
+	: DocumentInterface(InDocumentInterface)
+{
+	ReloadCache();
+}
+
+FMetaSoundFrontendDocumentBuilder::FMetaSoundFrontendDocumentBuilder(const FMetaSoundFrontendDocumentBuilder& InBuilder)
+{
+	// A new object doc has to be created when copied to avoid multiple builders/caches operating on a single document.
+	UMetaSoundBuilderDocument* NewDoc = NewObject<UMetaSoundBuilderDocument>();
+	NewDoc->Document = InBuilder.GetDocument();
+	DocumentInterface = NewDoc;
+	ReloadCache();
+}
+
+FMetaSoundFrontendDocumentBuilder::FMetaSoundFrontendDocumentBuilder(FMetaSoundFrontendDocumentBuilder&& InBuilder)
+	: DocumentInterface(MoveTemp(InBuilder.DocumentInterface))
+	, DocumentCache(MoveTemp(InBuilder.DocumentCache))
+{
+}
+
+FMetaSoundFrontendDocumentBuilder& FMetaSoundFrontendDocumentBuilder::operator=(FMetaSoundFrontendDocumentBuilder&& InRHS)
+{
+	using namespace Metasound::Frontend;
+	DocumentInterface = MoveTemp(InRHS.DocumentInterface);
+	DocumentCache = MoveTemp(InRHS.DocumentCache);
+	return *this;
+}
+
+FMetaSoundFrontendDocumentBuilder& FMetaSoundFrontendDocumentBuilder::operator=(const FMetaSoundFrontendDocumentBuilder& InRHS)
+{
+	using namespace Metasound::Frontend;
+
+	// A new object doc has to be created when copied to avoid multiple builders/caches operating on a single document.
+	UMetaSoundBuilderDocument* NewDoc = NewObject<UMetaSoundBuilderDocument>();
+	NewDoc->Document = InRHS.GetDocument();
+	DocumentInterface = NewDoc;
+	ReloadCache();
+	return *this;
+}
+
+const FMetasoundFrontendClass* FMetaSoundFrontendDocumentBuilder::AddGraphDependency(const FMetasoundFrontendGraphClass& InClass)
+{
+	using namespace Metasound::Frontend;
+
+	checkf(InClass.Metadata.GetType() == EMetasoundFrontendClassType::Graph, TEXT("Call should only be used for 'Graph' type dependencies. Use 'AddNativeDependency' for other class types."));
+
+	FMetasoundFrontendDocument& Document = GetDocument();
+	const FMetasoundFrontendClass* Dependency = nullptr;
+
+	FMetasoundFrontendClass NewDependency = InClass;
+
+	// All dependencies are listed as 'External' from the perspective of the owning document.
+	// This makes them implementation agnostic to accommodate nativization of assets.
+	NewDependency.Metadata.SetType(EMetasoundFrontendClassType::External);
+	NewDependency.ID = FGuid::NewGuid();
+	Dependency = &Document.Dependencies.Emplace_GetRef(MoveTemp(NewDependency));
+
+	const int32 NewIndex = Document.Dependencies.Num() - 1;
+	DocumentCache->OnDependencyAdded(NewIndex);
+
+	return Dependency;
+}
+
+const FMetasoundFrontendClass* FMetaSoundFrontendDocumentBuilder::AddNativeDependency(const FMetasoundFrontendClassMetadata& InClassMetadata)
+{
+	using namespace Metasound::Frontend;
+
+	checkf(InClassMetadata.GetType() != EMetasoundFrontendClassType::Graph, TEXT("Graph dependencies must be added via 'AddGraphDependency' to avoid overhead of registration where not necessary"));
+
+	FMetasoundFrontendDocument& Document = GetDocument();
+	const FMetasoundFrontendClass* Dependency = nullptr;
+
+	const FNodeRegistryKey RegistryKey = NodeRegistryKey::CreateKey(InClassMetadata);
+	FMetasoundFrontendClass NewDependency = GenerateClass(RegistryKey);
+	NewDependency.ID = FGuid::NewGuid();
+	Dependency = &Document.Dependencies.Emplace_GetRef(MoveTemp(NewDependency));
+
+	const int32 NewIndex = Document.Dependencies.Num() - 1;
+	DocumentCache->OnDependencyAdded(NewIndex);
+
+	return Dependency;
+}
+
+const FMetasoundFrontendEdge* FMetaSoundFrontendDocumentBuilder::AddEdge(FMetasoundFrontendEdge&& InNewEdge)
+{
+	using namespace Metasound::Frontend;
+
+	FMetasoundFrontendDocument& Document = GetDocument();
+	FMetasoundFrontendGraph& Graph = Document.RootGraph.Graph;
+	if (CanAddEdge(InNewEdge))
+	{
+		IDocumentGraphEdgeCache& EdgeCache = DocumentCache->GetEdgeCache();
+
+		const FMetasoundFrontendEdge& NewEdge = Graph.Edges.Add_GetRef(MoveTemp(InNewEdge));
+		const int32 NewIndex = Graph.Edges.Num() - 1;
+		EdgeCache.OnEdgeAdded(NewIndex);
+		return &NewEdge;
+	}
+
+	return nullptr;
+}
+
+bool FMetaSoundFrontendDocumentBuilder::AddEdgesByNamedConnections(const TSet<Metasound::Frontend::FNamedEdge>& ConnectionsToMake, TArray<const FMetasoundFrontendEdge*>* OutNewEdges)
+{
+	using namespace Metasound::Frontend;
+
+	FMetasoundFrontendDocument& Document = GetDocument();
+	const IDocumentGraphNodeCache& NodeCache = DocumentCache->GetNodeCache();
+
+	if (OutNewEdges)
+	{
+		OutNewEdges->Reset();
+	}
+
+	bool bSuccess = true;
+
+	TArray<FMetasoundFrontendEdge> EdgesToAdd;
+	for (const FNamedEdge& Connection : ConnectionsToMake)
+	{
+		const FMetasoundFrontendGraph& Graph = Document.RootGraph.Graph;
+		const FMetasoundFrontendVertex* OutputVertex = NodeCache.FindOutputVertex(Connection.OutputNodeID, Connection.OutputName);
+		const FMetasoundFrontendVertex* InputVertex = NodeCache.FindInputVertex(Connection.InputNodeID, Connection.InputName);
+
+		if (OutputVertex && InputVertex)
+		{
+			FMetasoundFrontendEdge NewEdge = { Connection.OutputNodeID, OutputVertex->VertexID, Connection.InputNodeID, InputVertex->VertexID };
+			if (CanAddEdge(NewEdge))
+			{
+				EdgesToAdd.Add(MoveTemp(NewEdge));
+			}
+			else
+			{
+				bSuccess = false;
+				UE_LOG(LogMetaSound, Warning, TEXT("Failed to add connections between MetaSound output '%s' and input '%s': Vertex types either incompatable (eg. construction pin to non-construction pin) or input already connected."), *Connection.OutputName.ToString(), *Connection.InputName.ToString());
+			}
+		}
+	}
+
+	for (FMetasoundFrontendEdge& EdgeToAdd : EdgesToAdd)
+	{
+		const FMetasoundFrontendEdge* NewEdge = AddEdge(MoveTemp(EdgeToAdd));
+		if (ensureAlwaysMsgf(NewEdge, TEXT("Failed to add MetaSound graph edge via DocumentBuilder when prior step validated edge add was valid")))
+		{
+			if (OutNewEdges)
+			{
+				OutNewEdges->Add(NewEdge);
+			}
+		}
+		else
+		{
+			bSuccess = false;
+		}
+	}
+
+	return bSuccess;
+}
+
+bool FMetaSoundFrontendDocumentBuilder::AddEdgesByNodeClassInterfaceBindings(const FGuid& InFromNodeID, const FGuid& InToNodeID)
+{
+	using namespace Metasound;
+	using namespace Metasound::Frontend;
+
+	auto FindNodeClassDocAndInterfaces = [this](const FGuid & InNodeID) -> const FMetasoundFrontendDocument*
+	{
+		const TScriptInterface<IMetaSoundDocumentInterface> NodeDocumentUInterface = FindOrLoadNodeClassDocument(InNodeID);
+		if (NodeDocumentUInterface)
+		{
+			return &NodeDocumentUInterface->GetDocument();
+		}
+
+		return nullptr;
+	};
+
+	const FMetasoundFrontendDocument* OutputNodeDocument = FindNodeClassDocAndInterfaces(InFromNodeID);
+	const FMetasoundFrontendDocument* InputNodeDocument = FindNodeClassDocAndInterfaces(InToNodeID);
+
+	if (!InputNodeDocument || !OutputNodeDocument)
+	{
+		return false;
+	}
+
+	const TSet<FMetasoundFrontendVersion>& OutputInterfaceVersions = OutputNodeDocument->Interfaces;
+
+	TArray<const IInterfaceRegistryEntry*> InputNodeInterfaces;
+	FindDeclaredInterfaces(*InputNodeDocument, InputNodeInterfaces);
+
+	TSet<FNamedEdge> ConnectionsToMake;
+	TSet<FName> InputsToConnect;
+	for (const IInterfaceRegistryEntry* InputInterfaceEntry : InputNodeInterfaces)
+	{
+		const FMetasoundFrontendInterface& InputInterface = InputInterfaceEntry->GetInterface();
+		const TArray<FMetasoundFrontendInterfaceBinding>& OutputBindings = InputInterfaceEntry->GetOutputBindings();
+
+		// Bindings are sorted in registry with earlier entries being higher priority to apply connections,
+		// so earlier listed connections are selected over potential collisions with later entries.
+		for (const FMetasoundFrontendInterfaceBinding& OutputBinding : OutputBindings)
+		{
+			if (OutputInterfaceVersions.Contains(OutputBinding.Version))
+			{
+				for (const FMetasoundFrontendInterfaceBindingConnection& Connection : OutputBinding.Connections)
+				{
+					if (!InputsToConnect.Contains(Connection.InputName))
+					{
+						InputsToConnect.Add(Connection.InputName);
+						ConnectionsToMake.Add(FNamedEdge { InFromNodeID, Connection.OutputName, InToNodeID, Connection.InputName });
+					}
+				}
+			}
+		}
+	}
+
+	return AddEdgesByNamedConnections(ConnectionsToMake);
+}
+
+bool FMetaSoundFrontendDocumentBuilder::AddEdgesFromMatchingInterfaceNodeOutputsToGraphOutputs(const FGuid& InNodeID, TArray<const FMetasoundFrontendEdge*>& OutEdgesCreated)
+{
+	METASOUND_TRACE_CPUPROFILER_EVENT_SCOPE(FMetaSoundFrontendDocumentBuilder::AddEdgesFromMatchingInterfaceNodeOutputsToGraphOutputs);
+
+	using namespace Metasound::Frontend;
+
+	OutEdgesCreated.Reset();
+
+	const TScriptInterface<IMetaSoundDocumentInterface> NodeDocumentInterface = FindOrLoadNodeClassDocument(InNodeID);
+	if (!NodeDocumentInterface)
+	{
+		return false;
+	}
+
+	const IDocumentGraphNodeCache& NodeCache = DocumentCache->GetNodeCache();
+
+	const TSet<FMetasoundFrontendVersion>& NodeInterfaces = NodeDocumentInterface->GetDocument().Interfaces;
+	const TSet<FMetasoundFrontendVersion> CommonInterfaces = NodeInterfaces.Intersect(GetDocument().Interfaces);
+
+	TSet<FNamedEdge> ConnectionsToMake;
+	for (const FMetasoundFrontendVersion& Version : CommonInterfaces)
+	{
+		const FInterfaceRegistryKey InterfaceKey = GetInterfaceRegistryKey(Version);
+		if (const IInterfaceRegistryEntry* RegistryEntry = IInterfaceRegistry::Get().FindInterfaceRegistryEntry(InterfaceKey))
+		{
+			Algo::Transform(RegistryEntry->GetInterface().Outputs, ConnectionsToMake, [this, &NodeCache, InNodeID](const FMetasoundFrontendClassOutput& Output)
+			{
+				const FMetasoundFrontendGraph& Graph = GetDocument().RootGraph.Graph;
+				const FMetasoundFrontendVertex* NodeVertex = NodeCache.FindOutputVertex(InNodeID, Output.Name);
+				check(NodeVertex);
+				const FMetasoundFrontendNode* OutputNode = NodeCache.FindOutputNode(Output.Name);
+				check(OutputNode);
+				const TArray<FMetasoundFrontendVertex>& Inputs = OutputNode->Interface.Inputs;
+				check(!Inputs.IsEmpty());
+				return FNamedEdge { InNodeID, NodeVertex->Name, OutputNode->GetID(), Inputs.Last().Name };
+			});
+		}
+	}
+
+	return AddEdgesByNamedConnections(ConnectionsToMake, &OutEdgesCreated);
+}
+
+bool FMetaSoundFrontendDocumentBuilder::AddEdgesFromMatchingInterfaceNodeInputsToGraphInputs(const FGuid& InNodeID, TArray<const FMetasoundFrontendEdge*>& OutEdgesCreated)
+{
+	METASOUND_TRACE_CPUPROFILER_EVENT_SCOPE(FMetaSoundFrontendDocumentBuilder::AddEdgesFromMatchingInterfaceNodeInputsToGraphInputs);
+
+	using namespace Metasound::Frontend;
+
+	OutEdgesCreated.Reset();
+
+	const TScriptInterface<IMetaSoundDocumentInterface> NodeDocumentInterface = FindOrLoadNodeClassDocument(InNodeID);
+	if (!NodeDocumentInterface)
+	{
+		return false;
+	}
+
+	const IDocumentGraphNodeCache& NodeCache = DocumentCache->GetNodeCache();
+
+	const TSet<FMetasoundFrontendVersion>& NodeInterfaces = NodeDocumentInterface->GetDocument().Interfaces;
+	const TSet<FMetasoundFrontendVersion> CommonInterfaces = NodeInterfaces.Intersect(GetDocument().Interfaces);
+
+	TSet<FNamedEdge> ConnectionsToMake;
+	for (const FMetasoundFrontendVersion& Version : CommonInterfaces)
+	{
+		const FInterfaceRegistryKey InterfaceKey = GetInterfaceRegistryKey(Version);
+		if (const IInterfaceRegistryEntry* RegistryEntry = IInterfaceRegistry::Get().FindInterfaceRegistryEntry(InterfaceKey))
+		{
+			Algo::Transform(RegistryEntry->GetInterface().Inputs, ConnectionsToMake, [this, &NodeCache, InNodeID](const FMetasoundFrontendClassInput& Input)
+			{
+				const FMetasoundFrontendGraph& Graph = GetDocument().RootGraph.Graph;
+				const FMetasoundFrontendVertex* NodeVertex = NodeCache.FindOutputVertex(InNodeID, Input.Name);
+				check(NodeVertex);
+				const FMetasoundFrontendNode* InputNode = NodeCache.FindInputNode(Input.Name);
+				check(InputNode);
+				const TArray<FMetasoundFrontendVertex>& Outputs = InputNode->Interface.Outputs;
+				check(!Outputs.IsEmpty());
+				return FNamedEdge { InNodeID, NodeVertex->Name, InputNode->GetID(), Outputs.Last().Name };
+			});
+		}
+	}
+
+	return AddEdgesByNamedConnections(ConnectionsToMake, &OutEdgesCreated);
+}
+
+const FMetasoundFrontendNode* FMetaSoundFrontendDocumentBuilder::AddGraphInput(const FMetasoundFrontendClassInput& InClassInput)
+{
+	using namespace Metasound::Frontend;
+
+	FMetasoundFrontendGraphClass& RootGraph = GetDocument().RootGraph;
+
+	checkf(InClassInput.NodeID.IsValid(), TEXT("Unassigned NodeID when adding graph input"));
+	checkf(InClassInput.VertexID.IsValid(), TEXT("Unassigned VertexID when adding graph input"));
+
+	if (InClassInput.TypeName.IsNone())
+	{
+		UE_LOG(LogMetaSound, Error, TEXT("TypeName unset when attempting to add class input '%s'"), *InClassInput.Name.ToString());
+		return nullptr;
+	}
+	else if (const FMetasoundFrontendNode* Node = DocumentCache->GetNodeCache().FindInputNode(InClassInput.Name))
+	{
+		UE_LOG(LogMetaSound, Error, TEXT("Attempting to add MetaSound graph input '%s' when input with name already exists"), *InClassInput.Name.ToString());
+		return Node;
+	}
+
+	auto FindRegistryClass = [&InClassInput](FMetasoundFrontendClass& OutClass) -> bool
+	{
+		switch (InClassInput.AccessType)
+		{
+		case EMetasoundFrontendVertexAccessType::Value:
+		{
+			return IDataTypeRegistry::Get().GetFrontendConstructorInputClass(InClassInput.TypeName, OutClass);
+		}
+		break;
+
+		case EMetasoundFrontendVertexAccessType::Reference:
+		{
+			return IDataTypeRegistry::Get().GetFrontendInputClass(InClassInput.TypeName, OutClass);
+		}
+		break;
+
+		case EMetasoundFrontendVertexAccessType::Unset:
+		default:
+		{
+			checkNoEntry();
+		}
+		break;
+		}
+
+		return false;
+	};
+
+	FMetasoundFrontendClass Class;
+	if (FindRegistryClass(Class))
+	{
+		if(!FindDependency(Class.Metadata))
+		{
+			AddNativeDependency(Class.Metadata);
+		}
+
+		auto FinalizeNode = [&InClassInput](FMetasoundFrontendNode& InOutNode, const Metasound::Frontend::FNodeRegistryKey&)
+		{
+			DocumentBuilderPrivate::FinalizeGraphVertexNode(InOutNode, InClassInput);
+		};
+		if (FMetasoundFrontendNode* NewNode = AddNodeInternal(Class.Metadata, FinalizeNode, InClassInput.NodeID))
+		{
+			
+			FMetasoundFrontendClassInput& NewInput = RootGraph.Interface.Inputs.Add_GetRef(InClassInput);
+			if (!NewInput.VertexID.IsValid())
+			{
+				NewInput.VertexID = FGuid::NewGuid();
+			}
+			return NewNode;
+		}
+	}
+
+	return nullptr;
+}
+
+const FMetasoundFrontendNode* FMetaSoundFrontendDocumentBuilder::AddGraphOutput(const FMetasoundFrontendClassOutput& InClassOutput)
+{
+	using namespace Metasound::Frontend;
+
+	checkf(InClassOutput.NodeID.IsValid(), TEXT("Unassigned NodeID when adding graph output"));
+	checkf(InClassOutput.VertexID.IsValid(), TEXT("Unassigned VertexID when adding graph output"));
+
+	FMetasoundFrontendGraphClass& RootGraph = GetDocument().RootGraph;
+
+	if (InClassOutput.TypeName.IsNone())
+	{
+		UE_LOG(LogMetaSound, Error, TEXT("TypeName unset when attempting to add class output '%s'"), *InClassOutput.Name.ToString());
+		return nullptr;
+	}
+	else if (const FMetasoundFrontendNode* Node = DocumentCache->GetNodeCache().FindInputNode(InClassOutput.Name))
+	{
+		UE_LOG(LogMetaSound, Error, TEXT("Attempting to add MetaSound graph output '%s' when output with name already exists"), *InClassOutput.Name.ToString());
+		return Node;
+	}
+
+	auto FindRegistryClass = [&InClassOutput](FMetasoundFrontendClass& OutClass) -> bool
+	{
+		switch (InClassOutput.AccessType)
+		{
+		case EMetasoundFrontendVertexAccessType::Value:
+		{
+			return IDataTypeRegistry::Get().GetFrontendConstructorOutputClass(InClassOutput.TypeName, OutClass);
+		}
+		break;
+
+		case EMetasoundFrontendVertexAccessType::Reference:
+		{
+			return IDataTypeRegistry::Get().GetFrontendOutputClass(InClassOutput.TypeName, OutClass);
+		}
+		break;
+
+		case EMetasoundFrontendVertexAccessType::Unset:
+		default:
+		{
+			checkNoEntry();
+		}
+		break;
+		}
+
+		return false;
+	};
+
+	FMetasoundFrontendClass Class;
+	if (FindRegistryClass(Class))
+	{
+		if (!FindDependency(Class.Metadata))
+		{
+			AddNativeDependency(Class.Metadata);
+		}
+
+		auto FinalizeNode = [&InClassOutput](FMetasoundFrontendNode& InOutNode, const Metasound::Frontend::FNodeRegistryKey&)
+		{
+			DocumentBuilderPrivate::FinalizeGraphVertexNode(InOutNode, InClassOutput);
+		};
+		if (FMetasoundFrontendNode* NewNode = AddNodeInternal(Class.Metadata, FinalizeNode, InClassOutput.NodeID))
+		{
+			FMetasoundFrontendClassOutput& NewOutput = RootGraph.Interface.Outputs.Add_GetRef(InClassOutput);
+			if (!NewOutput.VertexID.IsValid())
+			{
+				NewOutput.VertexID = FGuid::NewGuid();
+			}
+			return NewNode;
+		}
+	}
+
+	return nullptr;
+}
+
+bool FMetaSoundFrontendDocumentBuilder::AddInterface(FName InterfaceName)
+{
+	using namespace Metasound::Frontend;
+
+	FMetasoundFrontendInterface Interface;
+	if (ISearchEngine::Get().FindInterfaceWithHighestVersion(InterfaceName, Interface))
+	{
+		if (GetDocument().Interfaces.Contains(Interface.Version))
+		{
+			return true;
+		}
+
+		const FModifyRootGraphInterfaces ModifyInterfaces({ }, { Interface.Version });
+		if (ModifyInterfaces.Transform(GetDocument()))
+		{
+			// TODO: Add ability to update local builder caches dynamically
+			// in ModifyRootGraphInterfaces transform once it moves away
+			// from using controllers internally. Then cache updates can
+			// be more isolated and not use this full rebuild call.
+			ReloadCache();
+			return true;
+		}
+	}
+
+	return false;
+}
+
+const FMetasoundFrontendNode* FMetaSoundFrontendDocumentBuilder::AddGraphNode(const FMetasoundFrontendGraphClass& InGraphClass, FGuid InNodeID)
+{
+	auto FinalizeNode = [](const FMetasoundFrontendNode& Node, const Metasound::Frontend::FNodeRegistryKey& ClassKey)
+	{
+#if WITH_EDITOR
+		using namespace Metasound::Frontend;
+
+		// Cache the asset name on the node if it node is reference to asset-defined graph.
+		if (IMetaSoundAssetManager* AssetManager = IMetaSoundAssetManager::Get())
+		{
+			if (const FSoftObjectPath* Path = AssetManager->FindObjectPathFromKey(ClassKey))
+			{
+				return FName(*Path->GetAssetName());
+			}
+		}
+#endif // WITH_EDITOR
+
+		return Node.Name;
+	};
+
+
+	// Dependency is considered "External" when looked up or added on another graph
+	FMetasoundFrontendClassMetadata NewClassMetadata = InGraphClass.Metadata;
+	NewClassMetadata.SetType(EMetasoundFrontendClassType::External);
+
+	if (!FindDependency(NewClassMetadata))
+	{
+		AddGraphDependency(InGraphClass);
+	}
+
+	return AddNodeInternal(NewClassMetadata, FinalizeNode, InNodeID);
+}
+
+FMetasoundFrontendNode* FMetaSoundFrontendDocumentBuilder::AddNodeInternal(const FMetasoundFrontendClassMetadata& InClassMetadata, FFinalizeNodeFunctionRef FinalizeNode, FGuid InNodeID)
+{
+	METASOUND_TRACE_CPUPROFILER_EVENT_SCOPE(FMetaSoundFrontendDocumentBuilder::AddNodeInternal);
+
+	using namespace Metasound::Frontend;
+
+	FMetasoundFrontendDocument& Document = GetDocument();
+	const FNodeRegistryKey ClassKey = NodeRegistryKey::CreateKey(InClassMetadata);
+	if (const FMetasoundFrontendClass* Dependency = DocumentCache->FindDependency(ClassKey))
+	{
+		FMetasoundFrontendGraph& Graph = Document.RootGraph.Graph;
+		TArray<FMetasoundFrontendNode>& Nodes = Graph.Nodes;
+		FMetasoundFrontendNode& Node = Nodes.Emplace_GetRef(*Dependency);
+		Node.UpdateID(InNodeID);
+		FinalizeNode(Node, ClassKey);
+
+		const int32 NewIndex = Nodes.Num() - 1;
+		IDocumentGraphNodeCache& NodeCache = DocumentCache->GetNodeCache();
+		NodeCache.OnNodeAdded(NewIndex);
+		return &Node;
+	}
+
+	return nullptr;
+}
+
+bool FMetaSoundFrontendDocumentBuilder::CanAddEdge(const FMetasoundFrontendEdge& InEdge) const
+{
+	using namespace Metasound::Frontend;
+
+	const FMetasoundFrontendDocument& Document = GetDocument();
+	const IDocumentGraphEdgeCache& EdgeCache = DocumentCache->GetEdgeCache();
+
+	if (!EdgeCache.IsNodeInputConnected(InEdge.ToNodeID, InEdge.ToVertexID))
+	{
+		const IDocumentGraphNodeCache& NodeCache = DocumentCache->GetNodeCache();
+		const FMetasoundFrontendGraph& Graph = Document.RootGraph.Graph;
+		const FMetasoundFrontendVertex* OutputVertex = NodeCache.FindOutputVertex(InEdge.FromNodeID, InEdge.FromVertexID);
+		const FMetasoundFrontendVertex* InputVertex = NodeCache.FindInputVertex(InEdge.ToNodeID, InEdge.ToVertexID);
+		if (OutputVertex && InputVertex)
+		{
+			if (OutputVertex->TypeName == InputVertex->TypeName)
+			{
+				const FMetasoundFrontendClassOutput* ClassOutput = FindNodeOutputClassOutput(InEdge.FromNodeID, InEdge.FromVertexID);
+				const FMetasoundFrontendClassInput* ClassInput = FindNodeInputClassInput(InEdge.ToNodeID, InEdge.ToVertexID);
+				if (ClassInput && ClassOutput)
+				{
+					return FMetasoundFrontendClassVertex::CanConnectVertexAccessTypes(ClassOutput->AccessType, ClassInput->AccessType);
+				}
+			}
+		}
+	}
+
+	return false;
+}
+
+bool FMetaSoundFrontendDocumentBuilder::ContainsEdge(const FMetasoundFrontendEdge& InEdge) const
+{
+	using namespace Metasound::Frontend;
+	IDocumentGraphEdgeCache& EdgeCache = DocumentCache->GetEdgeCache();
+	return EdgeCache.ContainsEdge(InEdge);
+}
+
+bool FMetaSoundFrontendDocumentBuilder::ContainsNode(const FGuid& InNodeID) const
+{
+	using namespace Metasound::Frontend;
+	const IDocumentGraphNodeCache& NodeCache = DocumentCache->GetNodeCache();
+	return NodeCache.ContainsNode(InNodeID);
+}
+
+bool FMetaSoundFrontendDocumentBuilder::ConvertFromPreset()
+{
+	using namespace Metasound::Frontend;
+
+	FMetasoundFrontendDocument& Document = GetDocument();
+	FMetasoundFrontendGraphClass& RootGraphClass = Document.RootGraph;
+	FMetasoundFrontendGraphClassPresetOptions& PresetOptions = RootGraphClass.PresetOptions;
+	PresetOptions.bIsPreset = false;
+
+#if WITH_EDITOR
+	FMetasoundFrontendGraphStyle Style = Document.RootGraph.Graph.Style;
+	Style.bIsGraphEditable = true;
+#endif // WITH_EDITOR
+
+	return true;
+}
+
+bool FMetaSoundFrontendDocumentBuilder::FindDeclaredInterfaces(TArray<const Metasound::Frontend::IInterfaceRegistryEntry*>& OutInterfaces) const
+{
+	return FindDeclaredInterfaces(GetDocument(), OutInterfaces);
+}
+
+bool FMetaSoundFrontendDocumentBuilder::FindDeclaredInterfaces(const FMetasoundFrontendDocument& InDocument, TArray<const Metasound::Frontend::IInterfaceRegistryEntry*>& OutInterfaces)
+{
+	using namespace Metasound;
+	using namespace Metasound::Frontend;
+
+	bool bInterfacesFound = true;
+
+	Algo::Transform(InDocument.Interfaces, OutInterfaces, [&bInterfacesFound](const FMetasoundFrontendVersion& Version)
+	{
+		const FInterfaceRegistryKey InterfaceKey = GetInterfaceRegistryKey(Version);
+		const IInterfaceRegistryEntry* RegistryEntry = IInterfaceRegistry::Get().FindInterfaceRegistryEntry(InterfaceKey);
+		if (!RegistryEntry)
+		{
+			bInterfacesFound = false;
+			UE_LOG(LogMetaSound, Warning, TEXT("No registered interface matching interface version on document [InterfaceVersion:%s]"), *Version.ToString());
+		}
+
+		return RegistryEntry;
+	});
+
+	return bInterfacesFound;
+}
+
+const FMetasoundFrontendClass* FMetaSoundFrontendDocumentBuilder::FindDependency(const FGuid& InClassID) const
+{
+	return DocumentCache->FindDependency(InClassID);
+}
+
+const FMetasoundFrontendClass* FMetaSoundFrontendDocumentBuilder::FindDependency(const FMetasoundFrontendClassMetadata& InMetadata) const
+{
+	using namespace Metasound::Frontend;
+
+	checkf(InMetadata.GetType() != EMetasoundFrontendClassType::Graph,
+		TEXT("Dependencies are never listed as 'Graph' types. Graphs are considered 'External' from the perspective of the parent document to allow for nativization."));
+	const FNodeRegistryKey RegistryKey = NodeRegistryKey::CreateKey(InMetadata);
+	return DocumentCache->FindDependency(RegistryKey);
+}
+
+bool FMetaSoundFrontendDocumentBuilder::FindInterfaceInputNodes(FName InterfaceName, TArray<const FMetasoundFrontendNode*>& OutInputs) const
+{
+	using namespace Metasound::Frontend;
+
+	OutInputs.Reset();
+
+	const IDocumentGraphNodeCache& NodeCache = DocumentCache->GetNodeCache();
+
+	FMetasoundFrontendInterface Interface;
+	if (ISearchEngine::Get().FindInterfaceWithHighestVersion(InterfaceName, Interface))
+	{
+		if (GetDocument().Interfaces.Contains(Interface.Version))
+		{
+			TArray<const FMetasoundFrontendNode*> InterfaceInputs;
+			for (const FMetasoundFrontendClassInput& Input : Interface.Inputs)
+			{
+				if (const int32* NodeIndex = NodeCache.FindInputNodeIndex(Input.Name))
+				{
+					const FMetasoundFrontendNode& Node = GetDocument().RootGraph.Graph.Nodes[*NodeIndex];
+					InterfaceInputs.Add(&Node);
+				}
+				else
+				{
+					return false;
+				}
+			}
+
+			OutInputs = MoveTemp(InterfaceInputs);
+			return true;
+		}
+	}
+
+	return false;
+}
+
+bool FMetaSoundFrontendDocumentBuilder::FindInterfaceOutputNodes(FName InterfaceName, TArray<const FMetasoundFrontendNode*>& OutOutputs) const
+{
+	using namespace Metasound::Frontend;
+
+	OutOutputs.Reset();
+
+	FMetasoundFrontendInterface Interface;
+	if (ISearchEngine::Get().FindInterfaceWithHighestVersion(InterfaceName, Interface))
+	{
+		if (GetDocument().Interfaces.Contains(Interface.Version))
+		{
+			TArray<const FMetasoundFrontendNode*> InterfaceOutputs;
+			for (const FMetasoundFrontendClassOutput& Output : Interface.Outputs)
+			{
+				if (const int32* NodeIndex = DocumentCache->GetNodeCache().FindOutputNodeIndex(Output.Name))
+				{
+					const FMetasoundFrontendNode& Node = GetDocument().RootGraph.Graph.Nodes[*NodeIndex];
+					InterfaceOutputs.Add(&Node);
+				}
+				else
+				{
+					return false;
+				}
+			}
+
+			OutOutputs = MoveTemp(InterfaceOutputs);
+			return true;
+		}
+	}
+
+	return false;
+}
+
+const FMetasoundFrontendNode* FMetaSoundFrontendDocumentBuilder::FindGraphInputNode(FName InputName) const
+{
+	using namespace Metasound::Frontend;
+	const IDocumentGraphNodeCache& NodeCache = DocumentCache->GetNodeCache();
+	return NodeCache.FindInputNode(InputName);
+}
+
+const FMetasoundFrontendNode* FMetaSoundFrontendDocumentBuilder::FindGraphOutputNode(FName OutputName) const
+{
+	using namespace Metasound::Frontend;
+	const IDocumentGraphNodeCache& NodeCache = DocumentCache->GetNodeCache();
+	return NodeCache.FindOutputNode(OutputName);
+}
+
+const FMetasoundFrontendVertex* FMetaSoundFrontendDocumentBuilder::FindNodeInput(const FGuid& InNodeID, const FGuid& InVertexID) const
+{
+	using namespace Metasound::Frontend;
+	const IDocumentGraphNodeCache& NodeCache = DocumentCache->GetNodeCache();
+	return NodeCache.FindInputVertex(InNodeID, InVertexID);
+}
+
+const FMetasoundFrontendVertex* FMetaSoundFrontendDocumentBuilder::FindNodeInput(const FGuid& InNodeID, FName InVertexName) const
+{
+	using namespace Metasound::Frontend;
+	const IDocumentGraphNodeCache& NodeCache = DocumentCache->GetNodeCache();
+	return NodeCache.FindInputVertex(InNodeID, InVertexName);
+}
+
+const FMetasoundFrontendVertex* FMetaSoundFrontendDocumentBuilder::FindNodeOutput(const FGuid& InNodeID, const FGuid& InVertexID) const
+{
+	using namespace Metasound::Frontend;
+	const IDocumentGraphNodeCache& NodeCache = DocumentCache->GetNodeCache();
+	return NodeCache.FindOutputVertex(InNodeID, InVertexID);
+}
+
+const FMetasoundFrontendVertex* FMetaSoundFrontendDocumentBuilder::FindNodeOutput(const FGuid& InNodeID, FName InVertexName) const
+{
+	using namespace Metasound::Frontend;
+	const IDocumentGraphNodeCache& NodeCache = DocumentCache->GetNodeCache();
+	return NodeCache.FindOutputVertex(InNodeID, InVertexName);
+}
+
+const FMetasoundFrontendClassInput* FMetaSoundFrontendDocumentBuilder::FindNodeInputClassInput(const FGuid& InNodeID, const FGuid& InVertexID) const
+{
+	using namespace Metasound::Frontend;
+
+	const FMetasoundFrontendDocument& Document = GetDocument();
+	const IDocumentGraphNodeCache& NodeCache = DocumentCache->GetNodeCache();
+
+	if (const FMetasoundFrontendNode* Node = NodeCache.FindNode(InNodeID))
+	{
+		if (const FMetasoundFrontendClass* Class = DocumentCache->FindDependency(Node->ClassID))
+		{
+			auto InputMatchesVertex = [&InVertexID](const FMetasoundFrontendClassVertex& ClassVertex) { return ClassVertex.VertexID == InVertexID; };
+			return Class->Interface.Inputs.FindByPredicate(InputMatchesVertex);
+		}
+	}
+
+	return nullptr;
+}
+
+const FMetasoundFrontendClassOutput* FMetaSoundFrontendDocumentBuilder::FindNodeOutputClassOutput(const FGuid& InNodeID, const FGuid& InVertexID) const
+{
+	using namespace Metasound::Frontend;
+
+	const FMetasoundFrontendDocument& Document = GetDocument();
+	const IDocumentGraphNodeCache& NodeCache = DocumentCache->GetNodeCache();
+
+	if (const FMetasoundFrontendNode* Node = NodeCache.FindNode(InNodeID))
+	{
+		if (const FMetasoundFrontendClass* Class = DocumentCache->FindDependency(Node->ClassID))
+		{
+			auto OutputMatchesVertex = [&InVertexID](const FMetasoundFrontendClassVertex& ClassVertex) { return ClassVertex.VertexID == InVertexID; };
+			return Class->Interface.Outputs.FindByPredicate(OutputMatchesVertex);
+		}
+	}
+
+	return nullptr;
+}
+
+const FMetasoundFrontendNode* FMetaSoundFrontendDocumentBuilder::FindNode(const FGuid& InNodeID) const
+{
+	using namespace Metasound::Frontend;
+	const IDocumentGraphNodeCache& NodeCache = DocumentCache->GetNodeCache();
+	return NodeCache.FindNode(InNodeID);
+}
+
+TArray<const FMetasoundFrontendVertex*> FMetaSoundFrontendDocumentBuilder::FindNodeInputs(const FGuid& InNodeID, FName TypeName) const
+{
+	return DocumentCache->GetNodeCache().FindNodeInputs(InNodeID, TypeName);
+}
+
+TArray<const FMetasoundFrontendVertex*> FMetaSoundFrontendDocumentBuilder::FindNodeOutputs(const FGuid& InNodeID, FName TypeName) const
+{
+	return DocumentCache->GetNodeCache().FindNodeOutputs(InNodeID, TypeName);
+}
+
+const TScriptInterface<IMetaSoundDocumentInterface> FMetaSoundFrontendDocumentBuilder::FindOrLoadNodeClassDocument(const FGuid& InNodeID) const
+{
+	using namespace Metasound;
+	using namespace Metasound::Frontend;
+
+	const FMetasoundFrontendDocument& Document = GetDocument();
+	const IDocumentGraphNodeCache& NodeCache = DocumentCache->GetNodeCache();
+	if (const FMetasoundFrontendNode* Node = NodeCache.FindNode(InNodeID))
+	{
+		if (const FMetasoundFrontendClass* NodeClass = DocumentCache->FindDependency(Node->ClassID))
+		{
+			const FNodeRegistryKey NodeClassRegistryKey = NodeRegistryKey::CreateKey(NodeClass->Metadata);
+			if (FMetasoundAssetBase* Asset = IMetaSoundAssetManager::GetChecked().TryLoadAssetFromKey(NodeClassRegistryKey))
+			{
+				UObject* Object = Asset->GetOwningAsset();
+				check(Object);
+				return Object;
+			}
+		}
+	}
+
+	return nullptr;
+}
+
+FMetasoundFrontendDocument& FMetaSoundFrontendDocumentBuilder::GetDocument()
+{
+	IMetaSoundDocumentInterface* Interface = DocumentInterface.GetInterface();
+	return Interface->GetDocument();
+}
+
+const FMetasoundFrontendDocument& FMetaSoundFrontendDocumentBuilder::GetDocument() const
+{
+	return DocumentInterface->GetDocument();
+}
+
+void FMetaSoundFrontendDocumentBuilder::InitGraphClassMetadata(FMetasoundFrontendClassMetadata& InOutMetadata, bool bResetVersion)
+{
+	InOutMetadata.SetClassName(FMetasoundFrontendClassName(FName(), *FGuid::NewGuid().ToString(), FName()));
+
+	if (bResetVersion)
+	{
+		InOutMetadata.SetVersion({ 1, 0 });
+	}
+
+	InOutMetadata.SetType(EMetasoundFrontendClassType::Graph);
+}
+
+void FMetaSoundFrontendDocumentBuilder::InitDocument(FName UClassName)
+{
+	METASOUND_TRACE_CPUPROFILER_EVENT_SCOPE(FMetaSoundFrontendDocumentBuilder::InitDocument);
+
+	using namespace Metasound::Frontend;
+
+	// 1. Set default class Metadata
+	{
+		constexpr bool bResetVersion = true;
+		FMetasoundFrontendClassMetadata& ClassMetadata = GetDocument().RootGraph.Metadata;
+		InitGraphClassMetadata(ClassMetadata, bResetVersion);
+	}
+
+	// 2. Set default doc version Metadata
+	{
+		FMetasoundFrontendDocumentMetadata& DocMetadata = GetDocument().Metadata;
+		DocMetadata.Version.Number = FMetasoundFrontendDocument::GetMaxVersion();
+	}
+
+	// 3. Add default interfaces for given UClass
+	{
+		TArray<FMetasoundFrontendInterface> InitInterfaces = ISearchEngine::Get().FindUClassDefaultInterfaces(UClassName);
+		FModifyRootGraphInterfaces ModifyRootGraphInterfaces({ }, InitInterfaces);
+		ModifyRootGraphInterfaces.Transform(GetDocument());
+	}
+
+	// 5. Reload the whole cache as transforms above mutate cached collections. Can be safely removed once transforms above
+	// are re-implemented in builder and support cache to enable marginally better performance on document initialization.
+	ReloadCache();
+}
+
+void FMetaSoundFrontendDocumentBuilder::InitNodeLocations()
+{
+#if WITH_EDITORONLY_DATA
+	using namespace Metasound;
+	using namespace Metasound::Frontend;
+
+	FMetasoundFrontendGraph& Graph = GetDocument().RootGraph.Graph;
+
+	FVector2D InputNodeLocation = FVector2D::ZeroVector;
+	FVector2D ExternalNodeLocation = InputNodeLocation + DisplayStyle::NodeLayout::DefaultOffsetX;
+	FVector2D OutputNodeLocation = ExternalNodeLocation + DisplayStyle::NodeLayout::DefaultOffsetX;
+
+	TArray<FMetasoundFrontendNode>& Nodes = Graph.Nodes;
+	for (FMetasoundFrontendNode& Node : Nodes)
+	{
+		if (const int32* ClassIndex = DocumentCache->FindDependencyIndex(Node.ClassID))
+		{
+			FMetasoundFrontendClass& Class = GetDocument().Dependencies[*ClassIndex];
+
+			const EMetasoundFrontendClassType NodeType = Class.Metadata.GetType();
+			FVector2D NewLocation;
+			if (NodeType == EMetasoundFrontendClassType::Input)
+			{
+				NewLocation = InputNodeLocation;
+				InputNodeLocation += DisplayStyle::NodeLayout::DefaultOffsetY;
+			}
+			else if (NodeType == EMetasoundFrontendClassType::Output)
+			{
+				NewLocation = OutputNodeLocation;
+				OutputNodeLocation += DisplayStyle::NodeLayout::DefaultOffsetY;
+			}
+			else
+			{
+				NewLocation = ExternalNodeLocation;
+				ExternalNodeLocation += DisplayStyle::NodeLayout::DefaultOffsetY;
+			}
+
+			// TODO: Find consistent location for controlling node locations.
+			// Currently it is split between MetasoundEditor and MetasoundFrontend modules.
+			FMetasoundFrontendNodeStyle& Style = Node.Style;
+			Style.Display.Locations = { { FGuid::NewGuid(), NewLocation } };
+		}
+	}
+#endif // WITH_EDITORONLY_DATA
+}
+
+bool FMetaSoundFrontendDocumentBuilder::IsNodeInputConnected(const FGuid& InNodeID, const FGuid& InVertexID) const
+{
+	return DocumentCache->GetEdgeCache().IsNodeInputConnected(InNodeID, InVertexID);
+}
+
+bool FMetaSoundFrontendDocumentBuilder::IsNodeOutputConnected(const FGuid& InNodeID, const FGuid& InVertexID) const
+{
+	return DocumentCache->GetEdgeCache().IsNodeOutputConnected(InNodeID, InVertexID);
+}
+
+bool FMetaSoundFrontendDocumentBuilder::IsInterfaceDeclared(FName InInterfaceName) const
+{
+	using namespace Metasound::Frontend;
+
+	FMetasoundFrontendInterface Interface;
+	if (ISearchEngine::Get().FindInterfaceWithHighestVersion(InInterfaceName, Interface))
+	{
+		return IsInterfaceDeclared(Interface.Version);
+	}
+
+	return false;
+}
+
+bool FMetaSoundFrontendDocumentBuilder::IsInterfaceDeclared(const FMetasoundFrontendVersion& InInterfaceVersion) const
+{
+	return GetDocument().Interfaces.Contains(InInterfaceVersion);
+}
+
+void FMetaSoundFrontendDocumentBuilder::ReloadCache()
+{
+	using namespace Metasound::Frontend;
+	DocumentCache = MakeUnique<FDocumentCache>(GetDocument());
+}
+
+bool FMetaSoundFrontendDocumentBuilder::RemoveDependency(EMetasoundFrontendClassType ClassType, const FMetasoundFrontendClassName& InClassName, const FMetasoundFrontendVersionNumber& InClassVersionNumber)
+{
+	using namespace Metasound::Frontend;
+
+	FMetasoundFrontendDocument& Document = GetDocument();
+	const IDocumentGraphNodeCache& NodeCache = DocumentCache->GetNodeCache();
+
+	TArray<FMetasoundFrontendClass>& Dependencies = Document.Dependencies;
+	const FNodeRegistryKey ClassKey = NodeRegistryKey::CreateKey(ClassType, InClassName.GetFullName().ToString(), InClassVersionNumber.Major, InClassVersionNumber.Minor);
+	if (const int32* IndexPtr = DocumentCache->FindDependencyIndex(ClassKey))
+	{
+		const int32 Index = *IndexPtr;
+
+		TArray<const FMetasoundFrontendNode*> Nodes = NodeCache.FindNodesOfClassID(Dependencies[Index].ID);
+		for (const FMetasoundFrontendNode* Node : Nodes)
+		{
+			if (!RemoveNode(Node->GetID()))
+			{
+				return false;
+			}
+		}
+
+		const int32 LastIndex = Dependencies.Num() - 1;
+		DocumentCache->OnRemovingDependency(Index);
+		if (Index != LastIndex)
+		{
+			DocumentCache->OnRemovingDependency(LastIndex);
+		}
+		Dependencies.RemoveAtSwap(Index, 1, false);
+		if (Index != LastIndex)
+		{
+			DocumentCache->OnDependencyAdded(Index);
+		}
+	}
+
+	return true;
+}
+
+bool FMetaSoundFrontendDocumentBuilder::RemoveEdge(const FMetasoundFrontendEdge& EdgeToRemove)
+{
+	using namespace Metasound::Frontend;
+
+	FMetasoundFrontendDocument& Document = GetDocument();
+	FMetasoundFrontendGraph& Graph = Document.RootGraph.Graph;
+	TArray<FMetasoundFrontendEdge>& Edges = Graph.Edges;
+	IDocumentGraphEdgeCache& EdgeCache = DocumentCache->GetEdgeCache();
+	if (const int32* IndexPtr = EdgeCache.FindEdgeIndexToNodeInput(EdgeToRemove.ToNodeID, EdgeToRemove.ToVertexID))
+	{
+		const int32 Index = *IndexPtr;
+		const int32 LastIndex = Edges.Num() - 1;
+		EdgeCache.OnRemovingEdge(Index);
+		if (Index != LastIndex)
+		{
+			EdgeCache.OnRemovingEdge(LastIndex);
+		}
+		Edges.RemoveAtSwap(Index, 1, false);
+		if (Index != LastIndex)
+		{
+			EdgeCache.OnEdgeAdded(Index);
+		}
+	}
+
+	return false;
+}
+
+bool FMetaSoundFrontendDocumentBuilder::RemoveEdgesByNamedConnections(const TSet<Metasound::Frontend::FNamedEdge>& ConnectionsToRemove, TArray<FMetasoundFrontendEdge>* OutRemovedEdges)
+{
+	using namespace Metasound::Frontend;
+
+	FMetasoundFrontendDocument& Document = GetDocument();
+	const IDocumentGraphNodeCache& NodeCache = DocumentCache->GetNodeCache();
+
+	if (OutRemovedEdges)
+	{
+		OutRemovedEdges->Reset();
+	}
+
+	bool bSuccess = true;
+
+	TArray<FMetasoundFrontendEdge> EdgesToRemove;
+	for (const FNamedEdge& Connection : ConnectionsToRemove)
+	{
+		const FMetasoundFrontendGraph& Graph = Document.RootGraph.Graph;
+		const FMetasoundFrontendVertex* OutputVertex = NodeCache.FindOutputVertex(Connection.OutputNodeID, Connection.OutputName);
+		const FMetasoundFrontendVertex* InputVertex = NodeCache.FindInputVertex(Connection.InputNodeID, Connection.InputName);
+
+		if (OutputVertex && InputVertex)
+		{
+			FMetasoundFrontendEdge NewEdge = { Connection.OutputNodeID, OutputVertex->VertexID, Connection.InputNodeID, InputVertex->VertexID };
+			if (ContainsEdge(NewEdge))
+			{
+				EdgesToRemove.Add(MoveTemp(NewEdge));
+			}
+			else
+			{
+				bSuccess = false;
+				UE_LOG(LogMetaSound, Warning, TEXT("Failed to remove connections between MetaSound node output '%s' and input '%s': No connection found."), *Connection.OutputName.ToString(), *Connection.InputName.ToString());
+			}
+		}
+	}
+
+	for (FMetasoundFrontendEdge& EdgeToRemove : EdgesToRemove)
+	{
+		const bool bRemovedEdge = RemoveEdgeToNodeInput(EdgeToRemove.ToNodeID, EdgeToRemove.ToVertexID);
+		if (ensureAlwaysMsgf(bRemovedEdge, TEXT("Failed to remove MetaSound graph edge via DocumentBuilder when prior step validated edge remove was valid")))
+		{
+			if (OutRemovedEdges)
+			{
+				OutRemovedEdges->Add(EdgeToRemove);
+			}
+		}
+		else
+		{
+			bSuccess = false;
+		}
+	}
+
+	return bSuccess;
+}
+
+bool FMetaSoundFrontendDocumentBuilder::RemoveEdgesByNodeClassInterfaceBindings(const FGuid& InFromNodeID, const FGuid& InToNodeID)
+{
+	using namespace Metasound;
+	using namespace Metasound::Frontend;
+
+	auto FindNodeClassDocAndInterfaces = [this](const FGuid& InNodeID) -> const FMetasoundFrontendDocument*
+	{
+		const TScriptInterface<IMetaSoundDocumentInterface> NodeDocumentUInterface = FindOrLoadNodeClassDocument(InNodeID);
+		if (NodeDocumentUInterface)
+		{
+			return &NodeDocumentUInterface->GetDocument();
+		}
+
+		return nullptr;
+	};
+
+	const FMetasoundFrontendDocument* OutputNodeDocument = FindNodeClassDocAndInterfaces(InFromNodeID);
+	const FMetasoundFrontendDocument* InputNodeDocument = FindNodeClassDocAndInterfaces(InToNodeID);
+
+	if (!InputNodeDocument || !OutputNodeDocument)
+	{
+		return false;
+	}
+
+	const TSet<FMetasoundFrontendVersion>& OutputInterfaceVersions = OutputNodeDocument->Interfaces;
+
+	TArray<const IInterfaceRegistryEntry*> InputNodeInterfaces;
+	FindDeclaredInterfaces(*InputNodeDocument, InputNodeInterfaces);
+
+	bool bSuccess = true;
+
+	TSet<FNamedEdge> ConnectionsToRemove;
+	TSet<FName> InputsToDisconnect;
+	for (const IInterfaceRegistryEntry* InputInterfaceEntry : InputNodeInterfaces)
+	{
+		const FMetasoundFrontendInterface& InputInterface = InputInterfaceEntry->GetInterface();
+		const TArray<FMetasoundFrontendInterfaceBinding>& OutputBindings = InputInterfaceEntry->GetOutputBindings();
+
+		// Bindings are sorted in registry with earlier entries being higher priority to apply connections,
+		// so earlier listed connections are selected over potential collisions with later entries.
+		for (const FMetasoundFrontendInterfaceBinding& OutputBinding : OutputBindings)
+		{
+			if (OutputInterfaceVersions.Contains(OutputBinding.Version))
+			{
+				for (const FMetasoundFrontendInterfaceBindingConnection& Connection : OutputBinding.Connections)
+				{
+					if (!InputsToDisconnect.Contains(Connection.InputName))
+					{
+						InputsToDisconnect.Add(Connection.InputName);
+						ConnectionsToRemove.Add(FNamedEdge{ InFromNodeID, Connection.OutputName, InToNodeID, Connection.InputName });
+					}
+				}
+			}
+		}
+	}
+
+	return RemoveEdgesByNamedConnections(ConnectionsToRemove);
+}
+
+bool FMetaSoundFrontendDocumentBuilder::RemoveEdgesFromNodeOutput(const FGuid& InNodeID, const FGuid& InVertexID)
+{
+	using namespace Metasound::Frontend;
+
+	FMetasoundFrontendDocument& Document = GetDocument();
+	FMetasoundFrontendGraph& Graph = Document.RootGraph.Graph;
+	IDocumentGraphEdgeCache& EdgeCache = DocumentCache->GetEdgeCache();
+
+	if (const TArray<int32>* Indices = EdgeCache.FindEdgeIndicesFromNodeOutput(InNodeID, InVertexID))
+	{
+		TArray<int32> IndicesCopy = *Indices; // Copy off indices as the array may be modified when notifying the cache in the loop below
+		for (int32 Index : IndicesCopy)
+		{
+			const int32 LastIndex = Graph.Edges.Num() - 1;
+			EdgeCache.OnRemovingEdge(Index);
+			if (Index != LastIndex)
+			{
+				EdgeCache.OnRemovingEdge(LastIndex);
+			}
+			Graph.Edges.RemoveAtSwap(Index, 1, false);
+			if (Index != LastIndex)
+			{
+				EdgeCache.OnEdgeAdded(Index);
+			}
+		}
+
+		return true;
+	}
+
+	return false;
+}
+
+bool FMetaSoundFrontendDocumentBuilder::RemoveEdgeToNodeInput(const FGuid& InNodeID, const FGuid& InVertexID)
+{
+	using namespace Metasound::Frontend;
+
+	FMetasoundFrontendDocument& Document = GetDocument();
+	IDocumentGraphEdgeCache& EdgeCache = DocumentCache->GetEdgeCache();
+	FMetasoundFrontendGraph& Graph = Document.RootGraph.Graph;
+	if (const int32* IndexPtr = EdgeCache.FindEdgeIndexToNodeInput(InNodeID, InVertexID))
+	{
+		const int32 Index = *IndexPtr; // Copy off indices as the pointer may be modified when notifying the cache below
+		const int32 LastIndex = Graph.Edges.Num() - 1;
+		EdgeCache.OnRemovingEdge(Index);
+		if (Index != LastIndex)
+		{
+			EdgeCache.OnRemovingEdge(LastIndex);
+		}
+		Graph.Edges.RemoveAtSwap(Index, 1, false);
+		if (Index != LastIndex)
+		{
+			EdgeCache.OnEdgeAdded(Index);
+		}
+
+		return true;
+	}
+
+	return false;
+}
+
+bool FMetaSoundFrontendDocumentBuilder::RemoveInterface(FName InterfaceName)
+{
+	using namespace Metasound::Frontend;
+
+	FMetasoundFrontendInterface Interface;
+	if (ISearchEngine::Get().FindInterfaceWithHighestVersion(InterfaceName, Interface))
+	{
+		if (!GetDocument().Interfaces.Contains(Interface.Version))
+		{
+			return true;
+		}
+
+		const FModifyRootGraphInterfaces ModifyInterfaces({ { Interface.Version }, { } });
+		if (ModifyInterfaces.Transform(GetDocument()))
+		{
+			// TODO: Add ability to update local builder caches dynamically
+			// in ModifyRootGraphInterfaces transform once it moves away
+			// from using controllers internally. Then cache updates can
+			// be more isolated and not use this full rebuild call.
+			ReloadCache();
+			return true;
+		}
+	}
+
+	return false;
+}
+
+bool FMetaSoundFrontendDocumentBuilder::RemoveNode(const FGuid& InNodeID)
+{
+	METASOUND_TRACE_CPUPROFILER_EVENT_SCOPE(FMetaSoundFrontendDocumentBuilder::RemoveNode);
+
+	using namespace Metasound::Frontend;
+
+	FMetasoundFrontendDocument& Document = GetDocument();
+	IDocumentGraphNodeCache& NodeCache = DocumentCache->GetNodeCache();
+
+	FMetasoundFrontendGraph& Graph = Document.RootGraph.Graph;
+	TArray<FMetasoundFrontendNode>& Nodes = Graph.Nodes;
+	if (const int32* IndexPtr = NodeCache.FindNodeIndex(InNodeID))
+	{
+		const int32 Index = *IndexPtr; // Copy off indices as the pointer may be modified when notifying the cache below
+		const FMetasoundFrontendNode& Node = Nodes[Index];
+
+		for (const FMetasoundFrontendVertex& Vertex : Node.Interface.Inputs)
+		{
+			if (!RemoveEdgeToNodeInput(Node.GetID(), Vertex.VertexID))
+			{
+				return false;
+			}
+		}
+
+		for (const FMetasoundFrontendVertex& Vertex : Node.Interface.Outputs)
+		{
+			RemoveEdgesFromNodeOutput(Node.GetID(), Vertex.VertexID);
+		}
+
+		NodeCache.OnRemovingNode(Index);
+
+		const int32 LastIndex = Nodes.Num() - 1;
+		if (Index != LastIndex)
+		{
+			NodeCache.OnRemovingNode(LastIndex);
+		}
+		Nodes.RemoveAtSwap(Index, 1, false);
+		if (Index != LastIndex)
+		{
+			NodeCache.OnNodeAdded(Index);
+		}
+
+		return true;
+	}
+
+	return false;
+}
+
+#if WITH_EDITOR
+void FMetaSoundFrontendDocumentBuilder::SetAuthor(const FString& InAuthor)
+{
+	FMetasoundFrontendClassMetadata& ClassMetadata = GetDocument().RootGraph.Metadata;
+	ClassMetadata.SetAuthor(InAuthor);
+}
+#endif // WITH_EDITOR
+
+bool FMetaSoundFrontendDocumentBuilder::SetNodeInputDefault(const FGuid& InNodeID, const FGuid& InVertexID, const FMetasoundFrontendLiteral& InLiteral)
+{
+	using namespace Metasound::Frontend;
+
+	if (const int32* Index = DocumentCache->GetNodeCache().FindNodeIndex(InNodeID))
+	{
+		FMetasoundFrontendGraph& Graph = GetDocument().RootGraph.Graph;
+		FMetasoundFrontendNode& Node = Graph.Nodes[*Index];
+		TArray<FMetasoundFrontendVertex>& Inputs = Node.Interface.Inputs;
+
+		auto IsVertex = [&InVertexID](const FMetasoundFrontendVertex& Vertex) { return Vertex.VertexID == InVertexID; };
+		if (Inputs.ContainsByPredicate(IsVertex))
+		{
+			FMetasoundFrontendVertexLiteral NewVertexLiteral;
+			NewVertexLiteral.VertexID = InVertexID;
+			NewVertexLiteral.Value = InLiteral;
+
+			auto IsLiteral = [&InVertexID](const FMetasoundFrontendVertexLiteral& Literal) { return Literal.VertexID == InVertexID; };
+			if (FMetasoundFrontendVertexLiteral* InputLiteral = Node.InputLiterals.FindByPredicate(IsLiteral))
+			{
+				*InputLiteral = MoveTemp(NewVertexLiteral);
+			}
+			else
+			{
+				Node.InputLiterals.Add(MoveTemp(NewVertexLiteral));
+			}
+
+			return true;
+		}
+	}
+
+	return false;
+}
+
+bool FMetaSoundFrontendDocumentBuilder::ModifyInterfaces(TArray<FMetasoundFrontendVersion> OutputFormatsToRemove, TArray<FMetasoundFrontendVersion> OutputFormatsToAdd)
+{
+	using namespace Metasound::Frontend;
+
+	// Integrating this transform into the builder will allow for more selective updates and not having to rebuild the cache from scratch after performing the transform.
+	FModifyRootGraphInterfaces ModifyRootGraph(OutputFormatsToRemove, OutputFormatsToAdd);
+	if (ModifyRootGraph.Transform(GetDocument()))
+	{
+		ReloadCache();
+		return true;
+	}
+
+	return false;
+}
diff --git a/Engine/Plugins/Runtime/Metasound/Source/MetasoundFrontend/Private/MetasoundFrontendDocumentCache.cpp b/Engine/Plugins/Runtime/Metasound/Source/MetasoundFrontend/Private/MetasoundFrontendDocumentCache.cpp
new file mode 100644
--- /dev/null
+++ b/Engine/Plugins/Runtime/Metasound/Source/MetasoundFrontend/Private/MetasoundFrontendDocumentCache.cpp
@@ -0,0 +1,439 @@
+// Copyright Epic Games, Inc. All Rights Reserved.
+
+#include "MetasoundFrontendDocumentCache.h"
+
+#include "MetasoundFrontendDocument.h"
+#include "MetasoundFrontendDocumentBuilder.h"
+#include "MetasoundFrontendRegistries.h"
+#include "MetasoundTrace.h"
+
+
+namespace Metasound::Frontend
+{
+	FDocumentCache::FDocumentCache(const FMetasoundFrontendDocument& InDocument)
+		: Document(&InDocument)
+	{
+		METASOUND_TRACE_CPUPROFILER_EVENT_SCOPE(FDocumentCache::FDocumentCache);
+
+		for (int32 Index = 0; Index < Document->Dependencies.Num(); ++Index)
+		{
+			const FMetasoundFrontendClass& Class = Document->Dependencies[Index];
+			const FMetasoundFrontendClassMetadata& Metadata = Class.Metadata;
+			FNodeRegistryKey Key = NodeRegistryKey::CreateKey(Metadata);
+			KeyToIndex.Add(MoveTemp(Key), Index);
+			IDToIndex.Add(Class.ID, Index);
+		}
+
+		// Dependencies must be cached first as the NodeCache is dependent on them
+		NodeCache = MakeUnique<FDocumentGraphNodeCache>(InDocument, *this);
+		EdgeCache = MakeUnique<FDocumentGraphEdgeCache>(InDocument);
+	}
+
+	const FMetasoundFrontendDocument& FDocumentCache::GetDocument() const
+	{
+		check(Document);
+		return *Document;
+	}
+
+	bool FDocumentCache::ContainsDependency(const FNodeRegistryKey& InClassKey) const
+	{
+		return KeyToIndex.Contains(InClassKey);
+	}
+
+	const FMetasoundFrontendClass* FDocumentCache::FindDependency(const FNodeRegistryKey& InClassKey) const
+	{
+		if (const int32* Index = KeyToIndex.Find(InClassKey))
+		{
+			return &GetDocument().Dependencies[*Index];
+		}
+
+		return nullptr;
+	}
+
+	const FMetasoundFrontendClass* FDocumentCache::FindDependency(const FGuid& InClassID) const
+	{
+		if (const int32* Index = IDToIndex.Find(InClassID))
+		{
+			return &GetDocument().Dependencies[*Index];
+		}
+
+		return nullptr;
+	}
+
+	const int32* FDocumentCache::FindDependencyIndex(const Metasound::Frontend::FNodeRegistryKey& InClassKey) const
+	{
+		return KeyToIndex.Find(InClassKey);
+	}
+
+	const int32* FDocumentCache::FindDependencyIndex(const FGuid& InClassID) const
+	{
+		return IDToIndex.Find(InClassID);
+	}
+
+	void FDocumentCache::OnDependencyAdded(int32 NewIndex)
+	{
+		using namespace Metasound::Frontend;
+
+		const FMetasoundFrontendClass& Dependency = GetDocument().Dependencies[NewIndex];
+		const FNodeRegistryKey Key = NodeRegistryKey::CreateKey(Dependency.Metadata);
+		KeyToIndex.Add(Key, NewIndex);
+		IDToIndex.Add(Dependency.ID, NewIndex);
+	}
+
+	void FDocumentCache::OnRemovingDependency(int32 IndexBeingRemoved)
+	{
+		using namespace Metasound::Frontend;
+
+		const FMetasoundFrontendClass& DependencyBeingRemoved = GetDocument().Dependencies[IndexBeingRemoved];
+		const FNodeRegistryKey Key = NodeRegistryKey::CreateKey(DependencyBeingRemoved.Metadata);
+		KeyToIndex.Remove(Key);
+		IDToIndex.Remove(DependencyBeingRemoved.ID);
+	}
+
+	FDocumentGraphNodeCache::FDocumentGraphNodeCache(const FMetasoundFrontendDocument& InDocument, const FDocumentCache& InParent)
+		: Document(&InDocument)
+		, Parent(&InParent)
+	{
+		const FMetasoundFrontendGraph& Graph = GetDocument().RootGraph.Graph;
+		const TArray<FMetasoundFrontendNode>& Nodes = Graph.Nodes;
+		for (int32 Index = 0; Index < Nodes.Num(); ++Index)
+		{
+			const FMetasoundFrontendNode& Node = Nodes[Index];
+			const FMetasoundFrontendClass* Class = InParent.FindDependency(Node.ClassID);
+			check(Class);
+
+			IDToIndex.Add(Node.GetID(), Index);
+
+			switch (Class->Metadata.GetType())
+			{
+				case EMetasoundFrontendClassType::Output:
+				{
+					OutputNameToIndex.Add(Node.Name, Index);
+				}
+				break;
+
+				case EMetasoundFrontendClassType::Input:
+				{
+					InputNameToIndex.Add(Node.Name, Index);
+				}
+				break;
+
+				default:
+				break;
+			}
+		}
+	}
+
+	const FMetasoundFrontendDocument& FDocumentGraphNodeCache::GetDocument() const
+	{
+		check(Document);
+		return *Document;
+	}
+
+	bool FDocumentGraphNodeCache::ContainsNode(const FGuid& InNodeID) const
+	{
+		return IDToIndex.Contains(InNodeID);
+	}
+
+	const int32* FDocumentGraphNodeCache::FindInputNodeIndex(FName InputName) const
+	{
+		return InputNameToIndex.Find(InputName);
+	}
+
+	const int32* FDocumentGraphNodeCache::FindOutputNodeIndex(FName OutputName) const
+	{
+		return OutputNameToIndex.Find(OutputName);
+	}
+
+	const int32* FDocumentGraphNodeCache::FindNodeIndex(const FGuid& InNodeID) const
+	{
+		return IDToIndex.Find(InNodeID);
+	}
+
+	TArray<const FMetasoundFrontendNode*> FDocumentGraphNodeCache::FindNodesOfClassID(const FGuid& InClassID) const
+	{
+		TArray<const FMetasoundFrontendNode*> Nodes;
+		if (const TArray<int32>* NodeIndices = ClassIDToNodeIndices.Find(InClassID))
+		{
+			const FMetasoundFrontendGraph& Graph = GetDocument().RootGraph.Graph;
+			Algo::Transform(*NodeIndices, Nodes, [&Graph](const int32& Index) { return &Graph.Nodes[Index]; });
+		}
+		return Nodes;
+	}
+
+	const FMetasoundFrontendNode* FDocumentGraphNodeCache::FindNode(const FGuid& InNodeID) const
+	{
+		if (const int32* NodeIndex = IDToIndex.Find(InNodeID))
+		{
+			return &GetDocument().RootGraph.Graph.Nodes[*NodeIndex];
+		}
+
+		return nullptr;
+	}
+
+	bool FDocumentGraphNodeCache::ContainsNodesOfClassID(const FGuid& InClassID) const
+	{
+		if (const TArray<int32>* NodeIndices = ClassIDToNodeIndices.Find(InClassID))
+		{
+			return !NodeIndices->IsEmpty();
+		}
+
+		return false;
+	}
+
+	TArray<const FMetasoundFrontendVertex*> FDocumentGraphNodeCache::FindNodeInputs(const FGuid& InNodeID, FName TypeName) const
+	{
+		TArray<const FMetasoundFrontendVertex*> FoundVertices;
+		if (const FMetasoundFrontendNode* Node = FindNode(InNodeID))
+		{
+			auto IsVertexOfType = [&TypeName](const FMetasoundFrontendVertex& Vertex) { return TypeName.IsNone() || Vertex.TypeName == TypeName; };
+			auto GetVertexPtr = [](const FMetasoundFrontendVertex& Vertex) { return &Vertex; };
+			Algo::TransformIf(Node->Interface.Inputs, FoundVertices, IsVertexOfType, GetVertexPtr);
+		}
+
+		return FoundVertices;
+	}
+
+	TArray<const FMetasoundFrontendVertex*> FDocumentGraphNodeCache::FindNodeOutputs(const FGuid& InNodeID, FName TypeName) const
+	{
+		TArray<const FMetasoundFrontendVertex*> FoundVertices;
+		if (const FMetasoundFrontendNode* Node = FindNode(InNodeID))
+		{
+			auto IsVertexOfType = [&TypeName](const FMetasoundFrontendVertex& Vertex) { return TypeName.IsNone() || Vertex.TypeName == TypeName; };
+			auto GetVertexPtr = [](const FMetasoundFrontendVertex& Vertex) { return &Vertex; };
+			Algo::TransformIf(Node->Interface.Outputs, FoundVertices, IsVertexOfType, GetVertexPtr);
+		}
+
+		return FoundVertices;
+	}
+
+	const FMetasoundFrontendNode* FDocumentGraphNodeCache::FindOutputNode(FName OutputName) const
+	{
+		if (const int32* Index = OutputNameToIndex.Find(OutputName))
+		{
+			return &GetDocument().RootGraph.Graph.Nodes[*Index];
+		}
+
+		return nullptr;
+	}
+
+	const FMetasoundFrontendNode* FDocumentGraphNodeCache::FindInputNode(FName InputName) const
+	{
+		if (const int32* Index = InputNameToIndex.Find(InputName))
+		{
+			return &GetDocument().RootGraph.Graph.Nodes[*Index];
+		}
+
+		return nullptr;
+	}
+
+	const FMetasoundFrontendVertex* FDocumentGraphNodeCache::FindInputVertex(const FGuid& InNodeID, const FGuid& InVertexID) const
+	{
+		if (const FMetasoundFrontendNode* Node = FindNode(InNodeID))
+		{
+			auto VertexMatchesPredicate = [&InVertexID](const FMetasoundFrontendVertex& Vertex)
+			{
+				return Vertex.VertexID == InVertexID;
+			};
+			return Node->Interface.Inputs.FindByPredicate(VertexMatchesPredicate);
+		}
+
+		return nullptr;
+	}
+
+	const FMetasoundFrontendVertex* FDocumentGraphNodeCache::FindInputVertex(const FGuid& InNodeID, FName InVertexName) const
+	{
+		if (const FMetasoundFrontendNode* Node = FindNode(InNodeID))
+		{
+			auto VertexMatchesPredicate = [&InVertexName](const FMetasoundFrontendVertex& Vertex)
+			{
+				return Vertex.Name == InVertexName;
+			};
+			return Node->Interface.Inputs.FindByPredicate(VertexMatchesPredicate);
+		}
+
+		return nullptr;
+	}
+
+	const FMetasoundFrontendVertex* FDocumentGraphNodeCache::FindOutputVertex(const FGuid& InNodeID, const FGuid& InVertexID) const
+	{
+		if (const FMetasoundFrontendNode* Node = FindNode(InNodeID))
+		{
+			auto VertexMatchesPredicate = [&InVertexID](const FMetasoundFrontendVertex& Vertex)
+			{
+				return Vertex.VertexID == InVertexID;
+			};
+			return Node->Interface.Outputs.FindByPredicate(VertexMatchesPredicate);
+		}
+
+		return nullptr;
+	}
+
+	const FMetasoundFrontendVertex* FDocumentGraphNodeCache::FindOutputVertex(const FGuid& InNodeID, FName InVertexName) const
+	{
+		if (const FMetasoundFrontendNode* Node = FindNode(InNodeID))
+		{
+			auto VertexMatchesPredicate = [&InVertexName](const FMetasoundFrontendVertex& Vertex)
+			{
+				return Vertex.Name == InVertexName;
+			};
+			return Node->Interface.Outputs.FindByPredicate(VertexMatchesPredicate);
+		}
+
+		return nullptr;
+	}
+
+	void FDocumentGraphNodeCache::OnNodeAdded(int32 InNewIndex)
+	{
+		const FMetasoundFrontendNode& Node = GetDocument().RootGraph.Graph.Nodes[InNewIndex];
+		IDToIndex.Add(Node.GetID(), InNewIndex);
+		ClassIDToNodeIndices.FindOrAdd(Node.ClassID).Add(InNewIndex);
+
+		check(Parent);
+		const FMetasoundFrontendClass* Dependency = Parent->FindDependency(Node.ClassID);
+		check(Dependency);
+
+		switch (Dependency->Metadata.GetType())
+		{
+			case EMetasoundFrontendClassType::Input:
+			{
+				InputNameToIndex.Add(Node.Name, InNewIndex);
+			}
+			break;
+
+			case EMetasoundFrontendClassType::Output:
+			{
+				OutputNameToIndex.Add(Node.Name, InNewIndex);
+			}
+			break;
+
+			default:
+			break;
+		}
+	}
+
+	void FDocumentGraphNodeCache::OnRemovingNode(int32 IndexBeingRemoved)
+	{
+		using namespace Metasound::Frontend;
+
+		// If not found, transaction was missed
+		const FMetasoundFrontendNode& NodeBeingRemoved = GetDocument().RootGraph.Graph.Nodes[IndexBeingRemoved];
+
+		IDToIndex.Remove(NodeBeingRemoved.GetID());
+		TArray<int32>& NodeIndices = ClassIDToNodeIndices.FindChecked(NodeBeingRemoved.ClassID);
+		NodeIndices.Remove(IndexBeingRemoved);
+		if (NodeIndices.IsEmpty())
+		{
+			ClassIDToNodeIndices.Remove(NodeBeingRemoved.ClassID);
+		}
+
+		check(Parent);
+		const FMetasoundFrontendClass* Dependency = Parent->FindDependency(NodeBeingRemoved.ClassID);
+		check(Dependency);
+		switch (Dependency->Metadata.GetType())
+		{
+			case EMetasoundFrontendClassType::Input:
+			{
+				InputNameToIndex.Add(NodeBeingRemoved.Name);
+			}
+			break;
+
+			case EMetasoundFrontendClassType::Output:
+			{
+				OutputNameToIndex.Remove(NodeBeingRemoved.Name);
+			}
+			break;
+
+			default:
+			break;
+		}
+	}
+
+	FDocumentGraphEdgeCache::FDocumentGraphEdgeCache(const FMetasoundFrontendDocument& InDocument)
+		: Document(&InDocument)
+	{
+		const TArray<FMetasoundFrontendEdge>& Edges = GetDocument().RootGraph.Graph.Edges;
+		for (int32 Index = 0; Index < Edges.Num(); ++Index)
+		{
+			const FMetasoundFrontendEdge& Edge = Edges[Index];
+			OutputToEdgeIndices.FindOrAdd(Edge.GetFromVertexHandle()).Add(Index);
+			InputToEdgeIndex.Add(Edge.GetToVertexHandle()) = Index;
+		}
+	}
+
+	bool FDocumentGraphEdgeCache::ContainsEdge(const FMetasoundFrontendEdge& InEdge) const
+	{
+		const FMetasoundFrontendVertexHandle InputPair = { InEdge.FromNodeID, InEdge.FromVertexID };
+		if (const int32* Index = InputToEdgeIndex.Find(InputPair))
+		{
+			const FMetasoundFrontendEdge& Edge = GetDocument().RootGraph.Graph.Edges[*Index];
+			return Edge.ToNodeID == InEdge.ToNodeID && Edge.ToVertexID == InEdge.ToVertexID;
+		}
+
+		return false;
+	}
+
+	TArray<const FMetasoundFrontendEdge*> FDocumentGraphEdgeCache::FindEdges(const FGuid& InNodeID, const FGuid& InVertexID) const
+	{
+		TArray<const FMetasoundFrontendEdge*> Edges;
+
+		const FMetasoundFrontendGraph& Graph = GetDocument().RootGraph.Graph;
+		const FMetasoundFrontendVertexHandle Handle { InNodeID, InVertexID };
+		if (const int32* Index = InputToEdgeIndex.Find(Handle))
+		{
+			Edges.Add(&Graph.Edges[*Index]);
+		}
+
+		if (const TArray<int32>* Indices = OutputToEdgeIndices.Find(Handle))
+		{
+			Algo::Transform(*Indices, Edges, [&Graph](const int32& Index) { return &Graph.Edges[Index]; });
+		}
+
+		return Edges;
+	}
+
+	const int32* FDocumentGraphEdgeCache::FindEdgeIndexToNodeInput(const FGuid& InNodeID, const FGuid& InVertexID) const
+	{
+		return InputToEdgeIndex.Find({ InNodeID, InVertexID });
+	}
+
+	const TArray<int32>* FDocumentGraphEdgeCache::FindEdgeIndicesFromNodeOutput(const FGuid& InNodeID, const FGuid& InVertexID) const
+	{
+		return OutputToEdgeIndices.Find({ InNodeID, InVertexID });
+	}
+
+	const FMetasoundFrontendDocument& FDocumentGraphEdgeCache::GetDocument() const
+	{
+		check(Document);
+		return *Document;
+	}
+
+	bool FDocumentGraphEdgeCache::IsNodeInputConnected(const FGuid& InNodeID, const FGuid& InVertexID) const
+	{
+		return InputToEdgeIndex.Contains(FMetasoundFrontendVertexHandle { InNodeID, InVertexID });
+	}
+
+	bool FDocumentGraphEdgeCache::IsNodeOutputConnected(const FGuid& InNodeID, const FGuid& InVertexID) const
+	{
+		return OutputToEdgeIndices.Contains(FMetasoundFrontendVertexHandle { InNodeID, InVertexID });
+	}
+
+	void FDocumentGraphEdgeCache::OnEdgeAdded(int32 InNewIndex)
+	{
+		const FMetasoundFrontendEdge& NewEdge = GetDocument().RootGraph.Graph.Edges[InNewIndex];
+		InputToEdgeIndex.Add(FMetasoundFrontendVertexHandle { NewEdge.ToNodeID, NewEdge.ToVertexID }, InNewIndex);
+		OutputToEdgeIndices.FindOrAdd({ NewEdge.FromNodeID, NewEdge.FromVertexID }).Add(InNewIndex);
+	}
+
+	void FDocumentGraphEdgeCache::OnRemovingEdge(int32 IndexToRemove)
+	{
+		const FMetasoundFrontendEdge& EdgeToRemove = GetDocument().RootGraph.Graph.Edges[IndexToRemove];
+		InputToEdgeIndex.Remove(FMetasoundFrontendVertexHandle { EdgeToRemove.ToNodeID, EdgeToRemove.ToVertexID });
+		TArray<int32>& Indices = OutputToEdgeIndices.FindChecked(EdgeToRemove.GetFromVertexHandle());
+		Indices.RemoveAllSwap([&IndexToRemove](const int32& Index) { return Index == IndexToRemove; }, false /* bAllowShrinking */);
+		if (Indices.IsEmpty())
+		{
+			OutputToEdgeIndices.Remove(FMetasoundFrontendVertexHandle { EdgeToRemove.FromNodeID, EdgeToRemove.FromVertexID });
+		}
+	}
+} // namespace Metasound::Frontend
diff --git a/Engine/Plugins/Runtime/Metasound/Source/MetasoundFrontend/Private/MetasoundFrontendDocumentCache.h b/Engine/Plugins/Runtime/Metasound/Source/MetasoundFrontend/Private/MetasoundFrontendDocumentCache.h
new file mode 100644
--- /dev/null
+++ b/Engine/Plugins/Runtime/Metasound/Source/MetasoundFrontend/Private/MetasoundFrontendDocumentCache.h
@@ -0,0 +1,148 @@
+// Copyright Epic Games, Inc. All Rights Reserved.
+#pragma once
+
+#include "HAL/Platform.h"
+#include "MetasoundDocumentInterface.h"
+#include "MetasoundFrontendDocument.h"
+#include "MetasoundFrontendDocumentCacheInterface.h"
+#include "Containers/SortedMap.h"
+
+// Forward Declarations
+struct FMetaSoundFrontendDocumentBuilder;
+
+
+namespace Metasound::Frontend
+{
+	// Forward Declarations
+	class FDocumentCache;
+
+	class FDocumentGraphEdgeCache : public IDocumentGraphEdgeCache
+	{
+	public:
+		FDocumentGraphEdgeCache() = default;
+		FDocumentGraphEdgeCache(const FMetasoundFrontendDocument& InDocument);
+		virtual ~FDocumentGraphEdgeCache() = default;
+
+		// IDocumentGraphEdgeCache implementation
+		virtual bool ContainsEdge(const FMetasoundFrontendEdge& InEdge) const override;
+
+		virtual bool IsNodeInputConnected(const FGuid& InNodeID, const FGuid& InVertexID) const override;
+		virtual bool IsNodeOutputConnected(const FGuid& InNodeID, const FGuid& InVertexID) const override;
+
+		virtual TArray<const FMetasoundFrontendEdge*> FindEdges(const FGuid& InNodeID, const FGuid& InVertexID) const override;
+		virtual const int32* FindEdgeIndexToNodeInput(const FGuid& InNodeID, const FGuid& InVertexID) const override;
+		virtual const TArray<int32>* FindEdgeIndicesFromNodeOutput(const FGuid& InNodeID, const FGuid& InVertexID) const override;
+
+	protected:
+		virtual void OnEdgeAdded(int32 InNewIndex) override;
+		virtual void OnRemovingEdge(int32 IndexBeingRemoved) override;
+		virtual const FMetasoundFrontendDocument& GetDocument() const override;
+
+		// Friended to ensure only builder is the one managing cache manipulation via add/remove calls above
+		friend struct ::FMetaSoundFrontendDocumentBuilder;
+
+	private:
+
+		// Cache of outputs NodeId/VertexId pairs to associated edge indices
+		TMap<FMetasoundFrontendVertexHandle, TArray<int32>> OutputToEdgeIndices;
+
+		// Cache of Input NodeId/VertexId pairs to associated edge indices
+		TMap<FMetasoundFrontendVertexHandle, int32> InputToEdgeIndex;
+
+		const FMetasoundFrontendDocument* Document = nullptr;
+	};
+
+	class FDocumentGraphNodeCache : public IDocumentGraphNodeCache
+	{
+	public:
+		FDocumentGraphNodeCache() = default;
+		FDocumentGraphNodeCache(const FMetasoundFrontendDocument& InDocument, const FDocumentCache& InParent);
+		virtual ~FDocumentGraphNodeCache() = default;
+
+		// IDocumentGraphNodeCache implementation
+		virtual bool ContainsNode(const FGuid& InNodeID) const override;
+		virtual bool ContainsNodesOfClassID(const FGuid& InClassID) const override;
+
+		virtual const int32* FindInputNodeIndex(FName InputName) const override;
+		virtual const int32* FindOutputNodeIndex(FName OutputName) const override;
+		virtual const int32* FindNodeIndex(const FGuid& InNodeID) const override;
+		virtual TArray<const FMetasoundFrontendNode*> FindNodesOfClassID(const FGuid& InClassID) const override;
+		virtual const FMetasoundFrontendNode* FindNode(const FGuid& InNodeID) const override;
+		virtual const FMetasoundFrontendNode* FindInputNode(FName InputName) const override;
+		virtual const FMetasoundFrontendNode* FindOutputNode(FName OutputName) const override;
+
+		virtual TArray<const FMetasoundFrontendVertex*> FindNodeInputs(const FGuid& InNodeID, FName TypeName = FName()) const override;
+		virtual TArray<const FMetasoundFrontendVertex*> FindNodeOutputs(const FGuid& InNodeID, FName TypeName = FName()) const override;
+
+		virtual const FMetasoundFrontendVertex* FindInputVertex(const FGuid& InNodeID, const FGuid& InVertexID) const override;
+		virtual const FMetasoundFrontendVertex* FindInputVertex(const FGuid& InNodeID, FName InVertexName) const override;
+		virtual const FMetasoundFrontendVertex* FindOutputVertex(const FGuid& InNodeID, const FGuid& InVertexID) const override;
+		virtual const FMetasoundFrontendVertex* FindOutputVertex(const FGuid& InNodeID, FName InVertexName) const override;
+
+	protected:
+		virtual void OnNodeAdded(int32 NewIndex) override;
+		virtual void OnRemovingNode(int32 IndexBeingRemoved) override;
+		virtual const FMetasoundFrontendDocument& GetDocument() const override;
+
+		// Friended to ensure only builder is the one managing cache manipulation via add/remove calls above
+		friend struct ::FMetaSoundFrontendDocumentBuilder;
+
+	private:
+
+		// Cache of Input name to array index of node
+		TMap<FName, int32> InputNameToIndex;
+
+		// Cache of Output name to array index of node
+		TMap<FName, int32> OutputNameToIndex;
+
+		// Cache of NodeId to array index of node
+		TSortedMap<FGuid, int32> IDToIndex;
+
+		// Cache of ClassID to referencing node indices
+		TSortedMap<FGuid, TArray<int32>> ClassIDToNodeIndices;
+
+		const FMetasoundFrontendDocument* Document = nullptr;
+		const FDocumentCache* Parent = nullptr;
+	};
+
+	class FDocumentCache : public IDocumentCache
+	{
+	public:
+		FDocumentCache() = default;
+		FDocumentCache(const FMetasoundFrontendDocument& InDocument);
+		virtual ~FDocumentCache() = default;
+
+		virtual bool ContainsDependency(const FNodeRegistryKey& InClassKey) const override;
+
+		virtual const FMetasoundFrontendClass* FindDependency(const FNodeRegistryKey& InClassKey) const override;
+		virtual const FMetasoundFrontendClass* FindDependency(const FGuid& InClassID) const override;
+		virtual const int32* FindDependencyIndex(const Metasound::Frontend::FNodeRegistryKey& InClassKey) const override;
+		virtual const int32* FindDependencyIndex(const FGuid& InClassID) const override;
+
+
+		virtual IDocumentGraphEdgeCache& GetEdgeCache() override { return *EdgeCache; }
+		virtual const IDocumentGraphEdgeCache& GetEdgeCache() const override { return *EdgeCache; }
+		virtual IDocumentGraphNodeCache& GetNodeCache() override { return *NodeCache; }
+		virtual const IDocumentGraphNodeCache& GetNodeCache() const override { return *NodeCache; }
+
+	protected:
+		virtual void OnDependencyAdded(int32 InNewIndex) override;
+		virtual void OnRemovingDependency(int32 IndexBeingRemoved) override;
+		virtual const FMetasoundFrontendDocument& GetDocument() const override;
+
+	private:
+
+		// Cache of dependency (Class) ID to corresponding class dependency index
+		TSortedMap<FGuid, int32> IDToIndex;
+
+		// Cache of version data to corresponding class dependency index
+		TSortedMap<FNodeRegistryKey, int32> KeyToIndex;
+
+		TUniquePtr<FDocumentGraphEdgeCache> EdgeCache;
+		TUniquePtr<FDocumentGraphNodeCache> NodeCache;
+
+		const FMetasoundFrontendDocument* Document = nullptr;
+
+		friend struct ::FMetaSoundFrontendDocumentBuilder;
+	};
+} // namespace Metasound::Frontend
diff --git a/Engine/Plugins/Runtime/Metasound/Source/MetasoundFrontend/Private/MetasoundFrontendInterfaceRegistryPrivate.h b/Engine/Plugins/Runtime/Metasound/Source/MetasoundFrontend/Private/MetasoundFrontendInterfaceRegistryPrivate.h
--- a/Engine/Plugins/Runtime/Metasound/Source/MetasoundFrontend/Private/MetasoundFrontendInterfaceRegistryPrivate.h
+++ b/Engine/Plugins/Runtime/Metasound/Source/MetasoundFrontend/Private/MetasoundFrontendInterfaceRegistryPrivate.h
@@ -17,7 +17,6 @@ namespace Metasound
 		class FInterfaceRegistry : public IInterfaceRegistry
 		{
 		public:
-
 			static FInterfaceRegistry& Get();
 
 			FInterfaceRegistry();
@@ -30,9 +29,16 @@ namespace Metasound
 
 			virtual bool FindInterface(const FInterfaceRegistryKey& InKey, FMetasoundFrontendInterface& OutInterface) const override;
 
+			virtual void AddInterfaceOutputBindings(const FMetasoundFrontendVersion& InInputInterfaceVersion, TArray<FMetasoundFrontendInterfaceBinding>&& InBinding) override;
+
+			virtual bool RemoveInterfaceOutputBinding(const FMetasoundFrontendVersion& InInputInterfaceVersion, const FMetasoundFrontendVersion& InOutputInterfaceVersion) override;
+
+			virtual bool RemoveAllInterfaceOutputBindings(const FMetasoundFrontendVersion& InInputInterfaceVersion) override;
+
 			TUniquePtr<FInterfaceTransactionStream> CreateTransactionStream();
 
 			virtual void ForEachRegistryTransactionSince(FRegistryTransactionID InSince, FRegistryTransactionID* OutCurrentRegistryTransactionID, TFunctionRef<void(const FInterfaceRegistryTransaction&)> InFunc) const override;
+
 		private:
 			using FInterfaceTransactionBuffer = TTransactionBuffer<FInterfaceRegistryTransaction>;
 
diff --git a/Engine/Plugins/Runtime/Metasound/Source/MetasoundFrontend/Private/MetasoundFrontendSearchEngineCore.cpp b/Engine/Plugins/Runtime/Metasound/Source/MetasoundFrontend/Private/MetasoundFrontendSearchEngineCore.cpp
--- a/Engine/Plugins/Runtime/Metasound/Source/MetasoundFrontend/Private/MetasoundFrontendSearchEngineCore.cpp
+++ b/Engine/Plugins/Runtime/Metasound/Source/MetasoundFrontend/Private/MetasoundFrontendSearchEngineCore.cpp
@@ -177,21 +177,6 @@ namespace Metasound
 				return false;
 			}
 
-			bool FRemoveInterfacesWhichAreNotDefault::Filter(const FFrontendQueryEntry& InEntry) const
-			{
-				if (InEntry.Value.IsType<FMetasoundFrontendInterface>())
-				{
-					FInterfaceRegistryKey RegistryKey = GetInterfaceRegistryKey(InEntry.Value.Get<FMetasoundFrontendInterface>());
-					const IInterfaceRegistryEntry* RegEntry = IInterfaceRegistry::Get().FindInterfaceRegistryEntry(RegistryKey);
-
-					if (nullptr != RegEntry)
-					{
-						return RegEntry->IsDefault();
-					}
-				}
-				return false;
-			}
-
 			class FMapNodeRegistrationEventsToClassAndMajorVersion : public IFrontendQueryMapStep
 			{
 			public:
@@ -281,9 +266,7 @@ namespace Metasound
 				}
 				return Result;
 			}
-
-
-		}
+		} // namespace SearchEngineQuerySteps
 
 
 		FFrontendQuery FFindClassWithHighestMinorVersionQueryPolicy::CreateQuery()
@@ -323,46 +306,44 @@ namespace Metasound
 			return Key;
 		}
 
-		FFrontendQuery FFindAllDefaultInterfacesQueryPolicy::CreateQuery()
+		FFrontendQuery FFindAllRegisteredInterfacesWithNameQueryPolicy::CreateQuery()
 		{
 			using namespace SearchEngineQuerySteps;
 			FFrontendQuery Query;
 			Query.AddStep<FInterfaceRegistryTransactionSource>()
 				.AddStep<FMapInterfaceRegistryTransactionsToInterfaceRegistryKeys>()
 				.AddStep<FReduceInterfaceRegistryTransactionsToCurrentStatus>()
-				.AddStep<FTransformInterfaceRegistryTransactionToInterface>()
-				.AddStep<FMapInterfaceToInterfaceNameAndMajorVersion>()
-				.AddStep<FReduceInterfacesToHighestVersion>()
-				.AddStep<FRemoveInterfacesWhichAreNotDefault>()
-				.AddStep<FMapToNull>();
+				.AddStep<FMapInterfaceRegistryTransactionToInterfaceName>()
+				.AddStep<FTransformInterfaceRegistryTransactionToInterfaceVersion>();
 			return Query;
 		}
 
-		TArray<FMetasoundFrontendInterface> FFindAllDefaultInterfacesQueryPolicy::BuildResult(const FFrontendQueryPartition& InPartition)
+		TArray<FMetasoundFrontendVersion> FFindAllRegisteredInterfacesWithNameQueryPolicy::BuildResult(const FFrontendQueryPartition& InPartition)
 		{
 			using namespace SearchEngineQuerySteps;
-			return BuildArrayOfInterfacesFromPartition(InPartition);
+			return BuildArrayOfVersionsFromPartition(InPartition);
 		}
 
-		FFrontendQuery FFindAllRegisteredInterfacesWithNameQueryPolicy::CreateQuery()
+		FFrontendQuery FFindInterfaceWithHighestVersionQueryPolicy::CreateQuery()
 		{
 			using namespace SearchEngineQuerySteps;
 			FFrontendQuery Query;
 			Query.AddStep<FInterfaceRegistryTransactionSource>()
 				.AddStep<FMapInterfaceRegistryTransactionsToInterfaceRegistryKeys>()
 				.AddStep<FReduceInterfaceRegistryTransactionsToCurrentStatus>()
-				.AddStep<FMapInterfaceRegistryTransactionToInterfaceName>()
-				.AddStep<FTransformInterfaceRegistryTransactionToInterfaceVersion>();
+				.AddStep<FTransformInterfaceRegistryTransactionToInterface>()
+				.AddStep<FMapInterfaceToInterfaceName>()
+				.AddStep<FReduceInterfacesToHighestVersion>();
 			return Query;
 		}
 
-		TArray<FMetasoundFrontendVersion> FFindAllRegisteredInterfacesWithNameQueryPolicy::BuildResult(const FFrontendQueryPartition& InPartition)
+		FMetasoundFrontendInterface FFindInterfaceWithHighestVersionQueryPolicy::BuildResult(const FFrontendQueryPartition& InPartition)
 		{
 			using namespace SearchEngineQuerySteps;
-			return BuildArrayOfVersionsFromPartition(InPartition);
+			return BuildSingleInterfaceFromPartition(InPartition);
 		}
 
-		FFrontendQuery FFindInterfaceWithHighestVersionQueryPolicy::CreateQuery()
+		FFrontendQuery FFindAllInterfacesQueryPolicy::CreateQuery()
 		{
 			using namespace SearchEngineQuerySteps;
 			FFrontendQuery Query;
@@ -370,26 +351,46 @@ namespace Metasound
 				.AddStep<FMapInterfaceRegistryTransactionsToInterfaceRegistryKeys>()
 				.AddStep<FReduceInterfaceRegistryTransactionsToCurrentStatus>()
 				.AddStep<FTransformInterfaceRegistryTransactionToInterface>()
-				.AddStep<FMapInterfaceToInterfaceName>()
-				.AddStep<FReduceInterfacesToHighestVersion>();
+				.AddStep<FMapInterfaceToInterfaceNameAndMajorVersion>()
+				.AddStep<FReduceInterfacesToHighestVersion>()
+				.AddStep<FMapToNull>();
 			return Query;
 		}
 
-		FMetasoundFrontendInterface FFindInterfaceWithHighestVersionQueryPolicy::BuildResult(const FFrontendQueryPartition& InPartition)
+		TArray<FMetasoundFrontendInterface> FFindAllInterfacesQueryPolicy::BuildResult(const FFrontendQueryPartition& InPartition)
 		{
 			using namespace SearchEngineQuerySteps;
-			return BuildSingleInterfaceFromPartition(InPartition);
+			return BuildArrayOfInterfacesFromPartition(InPartition);
+		}
+
+		FFrontendQuery FFindAllInterfacesIncludingAllVersionsQueryPolicy::CreateQuery()
+		{
+			using namespace SearchEngineQuerySteps;
+			FFrontendQuery Query;
+			Query.AddStep<FInterfaceRegistryTransactionSource>()
+				.AddStep<FMapInterfaceRegistryTransactionsToInterfaceRegistryKeys>()
+				.AddStep<FReduceInterfaceRegistryTransactionsToCurrentStatus>()
+				.AddStep<FTransformInterfaceRegistryTransactionToInterface>()
+				.AddStep<FMapToNull>();
+			return Query;
 		}
 
+		TArray<FMetasoundFrontendInterface> FFindAllInterfacesIncludingAllVersionsQueryPolicy::BuildResult(const FFrontendQueryPartition& InPartition)
+			{
+				using namespace SearchEngineQuerySteps;
+				return BuildArrayOfInterfacesFromPartition(InPartition);
+			}
+
 		void FSearchEngineCore::Prime()
 		{
 			METASOUND_LLM_SCOPE;
 			METASOUND_TRACE_CPUPROFILER_EVENT_SCOPE(metasound::FSearchEngineCore::Prime);
 
 			FindClassWithHighestMinorVersionQuery.Prime();
-			FindAllDefaultInterfacesQuery.Prime();
 			FindAllRegisteredInterfacesWithNameQuery.Prime();
 			FindInterfaceWithHighestVersionQuery.Prime();
+			FindAllInterfacesIncludingAllVersionsQuery.Prime();
+			FindAllInterfacesQuery.Prime();
 		}
 
 		bool FSearchEngineCore::FindClassWithHighestMinorVersion(const FMetasoundFrontendClassName& InName, int32 InMajorVersion, FMetasoundFrontendClass& OutClass)
@@ -413,26 +414,27 @@ namespace Metasound
 		{
 			METASOUND_TRACE_CPUPROFILER_EVENT_SCOPE(metasound::FSearchEngineCore::FindUClassDefaultInterfaces);
 
+			constexpr bool bIncludeAllVersions = false;
 			const FFrontendQueryKey NullKey;
-			TArray<FMetasoundFrontendInterface> DefaultInterfaces = FindAllDefaultInterfacesQuery.UpdateAndFindResult(NullKey);
+			TArray<FMetasoundFrontendInterface> DefaultInterfaces = FindAllInterfaces(bIncludeAllVersions);
 
 			// All default interfaces are filtered here because a UClass name cannot effectively 
 			// be used as a "Key" for this query.
 			//
 			// We cannot know all the UClasses associated with each default interface given only a UClassName
 			// because we need to inspect UClass inheritance to handle derived UClasses
-			// which should be compatible with their parent's default interfaces. 
+			// which should be compatible with their parent's default interfaces.
 			//
 			// This module does not have access to the UClass inheritance structure.
 			// Instead, this information is accessed in the IInterfaceRegistryEntry
-			// concrete implementation. 
+			// concrete implementation.
 			auto DoesNotSupportUClass = [&InUClassName](const FMetasoundFrontendInterface& InInterface)
 			{
 				FInterfaceRegistryKey RegistryKey = GetInterfaceRegistryKey(InInterface);
 				const IInterfaceRegistryEntry* RegEntry = IInterfaceRegistry::Get().FindInterfaceRegistryEntry(RegistryKey);
 				if (ensure(RegEntry))
 				{
-					return !RegEntry->UClassIsSupported(InUClassName);
+					return !RegEntry->IsDefault(InUClassName);
 				}
 
 				return true;
@@ -458,5 +460,20 @@ namespace Metasound
 			const FFrontendQueryKey Key{InInterfaceName};
 			return FindInterfaceWithHighestVersionQuery.UpdateAndFindResult(Key, OutInterface);
 		}
-	}
-}
+
+		TArray<FMetasoundFrontendInterface> FSearchEngineCore::FindAllInterfaces(bool bInIncludeAllVersions)
+		{
+			METASOUND_TRACE_CPUPROFILER_EVENT_SCOPE(metasound::FSearchEngine::FindAllInterfaces);
+
+			const FFrontendQueryKey NullKey;
+			if (bInIncludeAllVersions)
+			{
+				return FindAllInterfacesIncludingAllVersionsQuery.UpdateAndFindResult(NullKey);
+			}
+			else
+			{
+				return FindAllInterfacesQuery.UpdateAndFindResult(NullKey);
+			}
+		}
+	} // namespace Frontend
+} // namespace Metasound
diff --git a/Engine/Plugins/Runtime/Metasound/Source/MetasoundFrontend/Private/MetasoundFrontendSearchEngineCore.h b/Engine/Plugins/Runtime/Metasound/Source/MetasoundFrontend/Private/MetasoundFrontendSearchEngineCore.h
--- a/Engine/Plugins/Runtime/Metasound/Source/MetasoundFrontend/Private/MetasoundFrontendSearchEngineCore.h
+++ b/Engine/Plugins/Runtime/Metasound/Source/MetasoundFrontend/Private/MetasoundFrontendSearchEngineCore.h
@@ -83,11 +83,6 @@ namespace Metasound
 				virtual bool Filter(const FFrontendQueryEntry& InEntry) const override;
 			};
 
-			class FRemoveInterfacesWhichAreNotDefault : public IFrontendQueryFilterStep
-			{
-				virtual bool Filter(const FFrontendQueryEntry& InEntry) const override;
-			};
-
 			class FMapInterfaceToInterfaceName : public IFrontendQueryMapStep
 			{
 			public:
@@ -248,15 +243,6 @@ namespace Metasound
 			static ResultType BuildResult(const FFrontendQueryPartition& InPartition);
 		};
 
-		// Policy for finding all registered default interfaces. 
-		struct FFindAllDefaultInterfacesQueryPolicy
-		{
-			using ResultType = TArray<FMetasoundFrontendInterface>;
-
-			static FFrontendQuery CreateQuery();
-			static ResultType BuildResult(const FFrontendQueryPartition& InPartition);
-		};
-
 		// Policy for finding all registered interface versions (name & version number)
 		// by name.
 		struct FFindAllRegisteredInterfacesWithNameQueryPolicy
@@ -276,6 +262,24 @@ namespace Metasound
 			static FMetasoundFrontendInterface BuildResult(const FFrontendQueryPartition& InPartition);
 		};
 
+		// Policy for finding all registered interfaces. 
+		struct FFindAllInterfacesQueryPolicy
+		{
+			using ResultType = TArray<FMetasoundFrontendInterface>;
+
+			static FFrontendQuery CreateQuery();
+			static ResultType BuildResult(const FFrontendQueryPartition& InPartition);
+		};
+
+		// Policy for finding all registered interfaces (including deprecated). 
+		struct FFindAllInterfacesIncludingAllVersionsQueryPolicy
+		{
+			using ResultType = TArray<FMetasoundFrontendInterface>;
+
+			static FFrontendQuery CreateQuery();
+			static TArray<FMetasoundFrontendInterface> BuildResult(const FFrontendQueryPartition& InPartition);
+		};
+
 		/** Supports essential search engine functionality needed for runtime. */
 		class FSearchEngineCore : public ISearchEngine
 		{
@@ -294,13 +298,15 @@ namespace Metasound
 			virtual TArray<FMetasoundFrontendInterface> FindUClassDefaultInterfaces(FName InUClassName) override;
 			virtual TArray<FMetasoundFrontendVersion> FindAllRegisteredInterfacesWithName(FName InInterfaceName) override;
 			virtual bool FindInterfaceWithHighestVersion(FName InInterfaceName, FMetasoundFrontendInterface& OutInterface) override;
+			virtual TArray<FMetasoundFrontendInterface> FindAllInterfaces(bool bInIncludeAllVersions) override;
 
 		private:
 
 			TSearchEngineQuery<FFindClassWithHighestMinorVersionQueryPolicy> FindClassWithHighestMinorVersionQuery;
-			TSearchEngineQuery<FFindAllDefaultInterfacesQueryPolicy> FindAllDefaultInterfacesQuery;
 			TSearchEngineQuery<FFindAllRegisteredInterfacesWithNameQueryPolicy> FindAllRegisteredInterfacesWithNameQuery;
 			TSearchEngineQuery<FFindInterfaceWithHighestVersionQueryPolicy> FindInterfaceWithHighestVersionQuery;
+			TSearchEngineQuery<FFindAllInterfacesQueryPolicy> FindAllInterfacesQuery;
+			TSearchEngineQuery<FFindAllInterfacesIncludingAllVersionsQueryPolicy> FindAllInterfacesIncludingAllVersionsQuery;
 		};
 	}
 }
diff --git a/Engine/Plugins/Runtime/Metasound/Source/MetasoundFrontend/Private/MetasoundFrontendSearchEngineEditorOnly.cpp b/Engine/Plugins/Runtime/Metasound/Source/MetasoundFrontend/Private/MetasoundFrontendSearchEngineEditorOnly.cpp
--- a/Engine/Plugins/Runtime/Metasound/Source/MetasoundFrontend/Private/MetasoundFrontendSearchEngineEditorOnly.cpp
+++ b/Engine/Plugins/Runtime/Metasound/Source/MetasoundFrontend/Private/MetasoundFrontendSearchEngineEditorOnly.cpp
@@ -112,44 +112,6 @@ namespace Metasound
 			return BuildSingleClassFromPartition(InPartition);
 		}
 
-		FFrontendQuery FFindAllInterfacesQueryPolicy::CreateQuery()
-		{
-			using namespace SearchEngineQuerySteps;
-			FFrontendQuery Query;
-			Query.AddStep<FInterfaceRegistryTransactionSource>()
-				.AddStep<FMapInterfaceRegistryTransactionsToInterfaceRegistryKeys>()
-				.AddStep<FReduceInterfaceRegistryTransactionsToCurrentStatus>()
-				.AddStep<FTransformInterfaceRegistryTransactionToInterface>()
-				.AddStep<FMapInterfaceToInterfaceNameAndMajorVersion>()
-				.AddStep<FReduceInterfacesToHighestVersion>()
-				.AddStep<FMapToNull>();
-			return Query;
-		}
-
-		TArray<FMetasoundFrontendInterface> FFindAllInterfacesQueryPolicy::BuildResult(const FFrontendQueryPartition& InPartition)
-		{
-			using namespace SearchEngineQuerySteps;
-			return BuildArrayOfInterfacesFromPartition(InPartition);
-		}
-
-		FFrontendQuery FFindAllInterfacesIncludingAllVersionsQueryPolicy::CreateQuery()
-		{
-			using namespace SearchEngineQuerySteps;
-			FFrontendQuery Query;
-			Query.AddStep<FInterfaceRegistryTransactionSource>()
-				.AddStep<FMapInterfaceRegistryTransactionsToInterfaceRegistryKeys>()
-				.AddStep<FReduceInterfaceRegistryTransactionsToCurrentStatus>()
-				.AddStep<FTransformInterfaceRegistryTransactionToInterface>()
-				.AddStep<FMapToNull>();
-			return Query;
-		}
-
-		TArray<FMetasoundFrontendInterface> FFindAllInterfacesIncludingAllVersionsQueryPolicy::BuildResult(const FFrontendQueryPartition& InPartition)
-		{
-			using namespace SearchEngineQuerySteps;
-			return BuildArrayOfInterfacesFromPartition(InPartition);
-		}
-
 		void FSearchEngineEditorOnly::Prime()
 		{
 			Super::Prime();
@@ -162,8 +124,6 @@ namespace Metasound
 				FindClassesWithNameUnsortedQuery.Prime();
 				FindClassesWithNameSortedQuery.Prime();
 				FindClassWithHighestVersionQuery.Prime();
-				FindAllInterfacesIncludingAllVersionsQuery.Prime();
-				FindAllInterfacesQuery.Prime();
 			}
 		}
 
@@ -204,22 +164,6 @@ namespace Metasound
 			const FFrontendQueryKey Key{InName.GetFullName()};
 			return FindClassWithHighestVersionQuery.UpdateAndFindResult(Key, OutClass);
 		}
-
-		TArray<FMetasoundFrontendInterface> FSearchEngineEditorOnly::FindAllInterfaces(bool bInIncludeAllVersions)
-		{
-			METASOUND_TRACE_CPUPROFILER_EVENT_SCOPE(metasound::FSearchEngineEditorOnly::FindAllInterfaces);
-
-			const FFrontendQueryKey NullKey;
-			if (bInIncludeAllVersions)
-			{
-				return FindAllInterfacesIncludingAllVersionsQuery.UpdateAndFindResult(NullKey);
-			}
-			else
-			{
-				return FindAllInterfacesQuery.UpdateAndFindResult(NullKey);
-			}
-		}
-
 	}
 }
 
diff --git a/Engine/Plugins/Runtime/Metasound/Source/MetasoundFrontend/Private/MetasoundFrontendSearchEngineEditorOnly.h b/Engine/Plugins/Runtime/Metasound/Source/MetasoundFrontend/Private/MetasoundFrontendSearchEngineEditorOnly.h
--- a/Engine/Plugins/Runtime/Metasound/Source/MetasoundFrontend/Private/MetasoundFrontendSearchEngineEditorOnly.h
+++ b/Engine/Plugins/Runtime/Metasound/Source/MetasoundFrontend/Private/MetasoundFrontendSearchEngineEditorOnly.h
@@ -60,24 +60,6 @@ namespace Metasound
 			static ResultType BuildResult(const FFrontendQueryPartition& InPartition);
 		};
 
-		// Policy for finding all registered interfaces. 
-		struct FFindAllInterfacesQueryPolicy
-		{
-			using ResultType = TArray<FMetasoundFrontendInterface>;
-
-			static FFrontendQuery CreateQuery();
-			static ResultType BuildResult(const FFrontendQueryPartition& InPartition);
-		};
-
-		// Policy for finding all registered interfaces (including deprecated). 
-		struct FFindAllInterfacesIncludingAllVersionsQueryPolicy
-		{
-			using ResultType = TArray<FMetasoundFrontendInterface>;
-
-			static FFrontendQuery CreateQuery();
-			static TArray<FMetasoundFrontendInterface> BuildResult(const FFrontendQueryPartition& InPartition);
-		};
-
 		// To minimize runtime memory costs during gampeplay, some search engine
 		// queries are only exposed when the editor only data is available. This
 		// class supports the editor only functionality of the search engine.
@@ -98,7 +80,6 @@ namespace Metasound
 			virtual TArray<FMetasoundFrontendClass> FindAllClasses(bool bInIncludeAllVersions) override;
 			virtual TArray<FMetasoundFrontendClass> FindClassesWithName(const FMetasoundFrontendClassName& InName, bool bInSortByVersion) override;
 			virtual bool FindClassWithHighestVersion(const FMetasoundFrontendClassName& InName, FMetasoundFrontendClass& OutClass) override;
-			virtual TArray<FMetasoundFrontendInterface> FindAllInterfaces(bool bInIncludeAllVersions) override;
 
 		private:
 
@@ -107,8 +88,6 @@ namespace Metasound
 			TSearchEngineQuery<FFindClassesWithNameUnsortedQueryPolicy> FindClassesWithNameUnsortedQuery;
 			TSearchEngineQuery<FFindClassesWithNameSortedQueryPolicy> FindClassesWithNameSortedQuery;
 			TSearchEngineQuery<FFindClassWithHighestVersionQueryPolicy> FindClassWithHighestVersionQuery;
-			TSearchEngineQuery<FFindAllInterfacesQueryPolicy> FindAllInterfacesQuery;
-			TSearchEngineQuery<FFindAllInterfacesIncludingAllVersionsQueryPolicy> FindAllInterfacesIncludingAllVersionsQuery;
 		};
 	}
 }
diff --git a/Engine/Plugins/Runtime/Metasound/Source/MetasoundFrontend/Private/MetasoundFrontendTransform.cpp b/Engine/Plugins/Runtime/Metasound/Source/MetasoundFrontend/Private/MetasoundFrontendTransform.cpp
--- a/Engine/Plugins/Runtime/Metasound/Source/MetasoundFrontend/Private/MetasoundFrontendTransform.cpp
+++ b/Engine/Plugins/Runtime/Metasound/Source/MetasoundFrontend/Private/MetasoundFrontendTransform.cpp
@@ -453,6 +453,13 @@ namespace Metasound
 			return bDidEdit;
 		}
 
+		bool FModifyRootGraphInterfaces::Transform(FMetasoundFrontendDocument& InOutDocument) const
+		{
+			// TODO: Swap implementation to not use access pointers/controllers
+			FDocumentAccessPtr DocAccessPtr = MakeAccessPtr<FDocumentAccessPtr>(InOutDocument.AccessPoint, InOutDocument);
+			return Transform(FDocumentController::CreateDocumentHandle(DocAccessPtr));
+		}
+
 		bool FUpdateRootGraphInterface::Transform(FDocumentHandle InDocument) const
 		{
 			bool bDidEdit = false;
diff --git a/Engine/Plugins/Runtime/Metasound/Source/MetasoundFrontend/Private/NodeTemplates/MetasoundFrontendDocumentTemplatePreprocessor.h b/Engine/Plugins/Runtime/Metasound/Source/MetasoundFrontend/Private/NodeTemplates/MetasoundFrontendDocumentTemplatePreprocessor.h
--- a/Engine/Plugins/Runtime/Metasound/Source/MetasoundFrontend/Private/NodeTemplates/MetasoundFrontendDocumentTemplatePreprocessor.h
+++ b/Engine/Plugins/Runtime/Metasound/Source/MetasoundFrontend/Private/NodeTemplates/MetasoundFrontendDocumentTemplatePreprocessor.h
@@ -35,9 +35,6 @@ namespace Metasound
 			}
 
 			virtual ~FNodeTemplatePreprocessTransformBase() = default;
-
-			virtual FMetasoundFrontendDocument& GetOwningDocument() const override { return Document; };
-			virtual FMetasoundFrontendGraph& GetOwningGraph() const override { return Graph; };
 		};
 
 	} // namespace Frontend
diff --git a/Engine/Plugins/Runtime/Metasound/Source/MetasoundFrontend/Private/NodeTemplates/MetasoundFrontendNodeTemplateReroute.cpp b/Engine/Plugins/Runtime/Metasound/Source/MetasoundFrontend/Private/NodeTemplates/MetasoundFrontendNodeTemplateReroute.cpp
--- a/Engine/Plugins/Runtime/Metasound/Source/MetasoundFrontend/Private/NodeTemplates/MetasoundFrontendNodeTemplateReroute.cpp
+++ b/Engine/Plugins/Runtime/Metasound/Source/MetasoundFrontend/Private/NodeTemplates/MetasoundFrontendNodeTemplateReroute.cpp
@@ -24,8 +24,8 @@ namespace Metasound
 					for (int32 i = 0; i < Edges.Num(); ++i)
 					{
 						FMetasoundFrontendEdge& Edge = Edges[i];
-						InputEdgeMap.Add({ Edge.ToNodeID, Edge.ToVertexID }, &Edge);
-						OutputEdgeMap.FindOrAdd({ Edge.FromNodeID, Edge.FromVertexID }).Add(&Edge);
+						InputEdgeMap.Add(Edge.GetToVertexHandle(), &Edge);
+						OutputEdgeMap.FindOrAdd(Edge.GetFromVertexHandle()).Add(&Edge);
 					}
 				}
 
@@ -34,10 +34,8 @@ namespace Metasound
 				virtual bool Transform(FMetasoundFrontendNode& InOutNode) const override;
 
 			private:
-				using FNodeVertexGuidPair = TPair<FGuid, FGuid>;
-
-				mutable TMap<FNodeVertexGuidPair, FMetasoundFrontendEdge*> InputEdgeMap;
-				mutable TMap<FNodeVertexGuidPair, TArray<FMetasoundFrontendEdge*>> OutputEdgeMap;
+				mutable TMap<FMetasoundFrontendVertexHandle, FMetasoundFrontendEdge*> InputEdgeMap;
+				mutable TMap<FMetasoundFrontendVertexHandle, TArray<FMetasoundFrontendEdge*>> OutputEdgeMap;
 			};
 
 			bool FRerouteNodeTemplatePreprocessTransform::Transform(FMetasoundFrontendNode& InOutNode) const
@@ -53,8 +51,8 @@ namespace Metasound
 					}
 
 					const FMetasoundFrontendVertex& InputVertex = InOutNode.Interface.Inputs.Last();
-					FNodeVertexGuidPair InputNodeVertexPair{ InOutNode.GetID(), InputVertex.VertexID };
-					InputEdge = InputEdgeMap.FindRef(InputNodeVertexPair);
+					FMetasoundFrontendVertexHandle InputNodeVertexHandle { InOutNode.GetID(), InputVertex.VertexID };
+					InputEdge = InputEdgeMap.FindRef(InputNodeVertexHandle);
 
 					// This can happen if the reroute node isn't provided an input, so its perfectly
 					// acceptable to just ignore this node as it ultimately provides no sourced input.
@@ -66,14 +64,14 @@ namespace Metasound
 
 				TArray<FMetasoundFrontendEdge*>* OutputEdges = nullptr;
 				const FMetasoundFrontendVertex& OutputVertex = InOutNode.Interface.Outputs.Last();
-				FNodeVertexGuidPair OutputNodeVertexPair{ InOutNode.GetID(), OutputVertex.VertexID };
+				const FMetasoundFrontendVertexHandle OutputVertexHandle { InOutNode.GetID(), OutputVertex.VertexID };
 				{
 					if (!ensure(InOutNode.Interface.Outputs.Num() == 1))
 					{
 						return false;
 					}
 
-					OutputEdges = OutputEdgeMap.Find(OutputNodeVertexPair);
+					OutputEdges = OutputEdgeMap.Find(OutputVertexHandle);
 
 					// This can happen if the reroute node isn't provided any outputs to connect to, so its
 					// perfectly acceptable to just ignore this node as it ultimately provides no sourced input.
@@ -84,16 +82,17 @@ namespace Metasound
 				}
 
 				// Update the output edges with the input edge 
+				FMetasoundFrontendVertexHandle NewOutputEdgeNodeVertexHandle { InputEdge->FromNodeID, InputEdge->FromVertexID };
 				FNodeVertexGuidPair NewOutputEdgeNodeVertexPair{ InputEdge->FromNodeID, InputEdge->FromVertexID };
 
 				for (FMetasoundFrontendEdge* OutputEdge : *OutputEdges)
 				{
-					OutputEdge->FromNodeID = NewOutputEdgeNodeVertexPair.Key;
-					OutputEdge->FromVertexID = NewOutputEdgeNodeVertexPair.Value;
+					OutputEdge->FromNodeID = NewOutputEdgeNodeVertexHandle.NodeID;
+					OutputEdge->FromVertexID = NewOutputEdgeNodeVertexHandle.VertexID;
 				}
 
-				OutputEdgeMap.FindOrAdd(NewOutputEdgeNodeVertexPair).Append(*OutputEdges);
-				OutputEdgeMap.Remove(OutputNodeVertexPair);
+				OutputEdgeMap.FindOrAdd(NewOutputEdgeNodeVertexHandle).Append(*OutputEdges);
+				OutputEdgeMap.Remove(OutputVertexHandle);
 
 				return true;
 			}
diff --git a/Engine/Plugins/Runtime/Metasound/Source/MetasoundFrontend/Public/Interfaces/MetasoundFrontendSourceInterface.h b/Engine/Plugins/Runtime/Metasound/Source/MetasoundFrontend/Public/Interfaces/MetasoundFrontendSourceInterface.h
--- a/Engine/Plugins/Runtime/Metasound/Source/MetasoundFrontend/Public/Interfaces/MetasoundFrontendSourceInterface.h
+++ b/Engine/Plugins/Runtime/Metasound/Source/MetasoundFrontend/Public/Interfaces/MetasoundFrontendSourceInterface.h
@@ -22,7 +22,10 @@ namespace Metasound
 			}
 
 			METASOUNDFRONTEND_API const FMetasoundFrontendVersion& GetVersion();
+
+			UE_DEPRECATED(5.3, "ParameterInterface is now only registered for given UClasses at registration time, not in interface definition.")
 			METASOUNDFRONTEND_API Audio::FParameterInterfacePtr CreateInterface(const UClass& InClass);
+			METASOUNDFRONTEND_API Audio::FParameterInterfacePtr CreateInterface();
 		}
 
 		namespace SourceInterfaceV1_0
@@ -47,7 +50,10 @@ namespace Metasound
 			}
 
 			METASOUNDFRONTEND_API const FMetasoundFrontendVersion& GetVersion();
+
+			UE_DEPRECATED(5.3, "ParameterInterface is now only registered for given UClasses at registration time, not in interface definition.")
 			METASOUNDFRONTEND_API Audio::FParameterInterfacePtr CreateInterface(const UClass& InClass);
+			METASOUNDFRONTEND_API Audio::FParameterInterfacePtr CreateInterface();
 		}
 
 		namespace SourceInterface
@@ -68,7 +74,10 @@ namespace Metasound
 			}
 
 			METASOUNDFRONTEND_API const FMetasoundFrontendVersion& GetVersion();
+
+			UE_DEPRECATED(5.3, "ParameterInterface is now only registered for given UClasses at registration time, not in interface definition.")
 			METASOUNDFRONTEND_API Audio::FParameterInterfacePtr CreateInterface(const UClass& InClass);
+			METASOUNDFRONTEND_API Audio::FParameterInterfacePtr CreateInterface();
 
 			class METASOUNDFRONTEND_API FUpdateInterface : public Frontend::IDocumentTransform
 			{
