- Add Runtime MetaSound Builder support

diff --git a/Engine/Plugins/Runtime/Metasound/Source/MetasoundFrontend/Public/MetasoundAssetBase.h b/Engine/Plugins/Runtime/Metasound/Source/MetasoundFrontend/Public/MetasoundAssetBase.h
--- a/Engine/Plugins/Runtime/Metasound/Source/MetasoundFrontend/Public/MetasoundAssetBase.h
+++ b/Engine/Plugins/Runtime/Metasound/Source/MetasoundFrontend/Public/MetasoundAssetBase.h
@@ -76,6 +76,7 @@ public:
 	void UnregisterGraphWithFrontend();
 
 	// Sets/overwrites the root class metadata
+	UE_DEPRECATED(5.3, "Directly setting graph class Metadata is no longer be supported. Use the FMetaSoundFrontendDocumentBuilder to modify class data.")
 	virtual void SetMetadata(FMetasoundFrontendClassMetadata& InMetadata);
 
 #if WITH_EDITOR
@@ -84,6 +85,7 @@ public:
 #endif // WITH_EDITOR
 
 	// Returns the interface entries declared by the given asset's document from the InterfaceRegistry.
+	UE_DEPRECATED(5.3, "Use static FMetaSoundFrontendDocumentBuilder 'FindDeclaredInterfaces instead.")
 	bool GetDeclaredInterfaces(TArray<const Metasound::Frontend::IInterfaceRegistryEntry*>& OutInterfaces) const;
 
 	// Returns whether an interface with the given version is declared by the given asset's document.
@@ -108,7 +110,10 @@ public:
 
 	bool AddingReferenceCausesLoop(const FSoftObjectPath& InReferencePath) const;
 	bool IsReferencedAsset(const FMetasoundAssetBase& InAssetToCheck) const;
+
+	UE_DEPRECATED(5.3, "ConvertFromPreset moved to FMetaSoundFrontendDocumentBuilder.")
 	void ConvertFromPreset();
+
 	bool IsRegistered() const;
 
 	// Imports data from a JSON string directly
@@ -128,10 +133,12 @@ public:
 	// Overwrites the existing document. If the document's interface is not supported,
 	// the FMetasoundAssetBase be while queried for a new one using `GetPreferredInterface`.
 	void SetDocument(const FMetasoundFrontendDocument& InDocument);
+	void SetDocument(FMetasoundFrontendDocument&& InDocument);
 
 	FMetasoundFrontendDocument& GetDocumentChecked();
 	const FMetasoundFrontendDocument& GetDocumentChecked() const;
 
+	UE_DEPRECATED(5.3, "AddDefaultInterfaces is included in now applied via FMetaSoundFrontendDocumentBuilder::InitDocument and no longer directly supported via this function.")
 	void AddDefaultInterfaces();
 
 	bool VersionAsset();
@@ -179,10 +186,10 @@ protected:
 #endif // WITH_EDITORONLY_DATA
 
 	// Returns an access pointer to the document.
-	virtual Metasound::Frontend::FDocumentAccessPtr GetDocument() = 0;
+	virtual Metasound::Frontend::FDocumentAccessPtr GetDocumentAccessPtr() = 0;
 
 	// Returns an access pointer to the document.
-	virtual Metasound::Frontend::FConstDocumentAccessPtr GetDocument() const = 0;
+	virtual Metasound::Frontend::FConstDocumentAccessPtr GetDocumentConstAccessPtr() const = 0;
 
 protected:
 	// Container for runtime data of MetaSound graph.
diff --git a/Engine/Plugins/Runtime/Metasound/Source/MetasoundFrontend/Public/MetasoundDocumentInterface.h b/Engine/Plugins/Runtime/Metasound/Source/MetasoundFrontend/Public/MetasoundDocumentInterface.h
new file mode 100644
--- /dev/null
+++ b/Engine/Plugins/Runtime/Metasound/Source/MetasoundFrontend/Public/MetasoundDocumentInterface.h
@@ -0,0 +1,33 @@
+// Copyright Epic Games, Inc. All Rights Reserved.
+#pragma once
+
+#include "UObject/Interface.h"
+#include "MetasoundFrontendDocument.h"
+
+#include "MetasoundDocumentInterface.generated.h"
+
+
+// Forward Declarations
+struct FMetaSoundFrontendDocumentBuilder;
+
+
+// UInterface for all MetaSound UClasses that implement a MetaSound document
+// as a means for accessing data via scripting, execution, or node class generation.
+UINTERFACE(BlueprintType, meta = (CannotImplementInterfaceInBlueprint))
+class METASOUNDFRONTEND_API UMetaSoundDocumentInterface : public UInterface
+{
+	GENERATED_BODY()
+};
+
+class METASOUNDFRONTEND_API IMetaSoundDocumentInterface : public IInterface
+{
+	GENERATED_BODY()
+
+public:
+	virtual const FMetasoundFrontendDocument& GetDocument() const = 0;
+
+private:
+	virtual FMetasoundFrontendDocument& GetDocument() = 0;
+
+	friend struct FMetaSoundFrontendDocumentBuilder;
+};
\ No newline at end of file
diff --git a/Engine/Plugins/Runtime/Metasound/Source/MetasoundFrontend/Public/MetasoundFrontendArchetypeRegistry.h b/Engine/Plugins/Runtime/Metasound/Source/MetasoundFrontend/Public/MetasoundFrontendArchetypeRegistry.h
--- a/Engine/Plugins/Runtime/Metasound/Source/MetasoundFrontend/Public/MetasoundFrontendArchetypeRegistry.h
+++ b/Engine/Plugins/Runtime/Metasound/Source/MetasoundFrontend/Public/MetasoundFrontendArchetypeRegistry.h
@@ -25,12 +25,18 @@ namespace Metasound
 			// MetaSound Interface definition
 			virtual const FMetasoundFrontendInterface& GetInterface() const = 0;
 
-			// Whether or not to immediately apply the given interface when a MetaSound
-			// asset of the supported UClass is created.
-			virtual bool IsDefault() const = 0;
+			// Whether or not to immediately apply the given interface
+			// when a MetaSound asset of the supported UClass is created.
+			virtual bool IsDefault(FName InUClassName) const = 0;
+
+			UE_DEPRECATED(5.3, "UClass name must now be provided to determine if interface is to be apply by default.")
+			virtual bool IsDefault() const { return false; }
 
 			// Whether or not the interface can be added or removed by the editor.
-			virtual bool EditorCanAddOrRemove() const = 0;
+			virtual bool EditorCanAddOrRemove(FName InUClassName) const = 0;
+
+			UE_DEPRECATED(5.3, "UClass name must now be provided to determine if editor can directly add or remove interface.")
+			virtual bool EditorCanAddOrRemove() const { return false; }
 
 			// Name of routing system used to update interface inputs (ex. ParameterInterface or DataReference).
 			virtual FName GetRouterName() const = 0;
@@ -40,6 +46,18 @@ namespace Metasound
 
 			// How to update a given document if versioning is required to this interface from a deprecated version.
 			virtual bool UpdateRootGraphInterface(FDocumentHandle InDocument) const = 0;
+
+			// Adds bindings that informs how this interface can be used to connect its inputs to node that implements another interface's outputs. Faster than adding iteratively.
+			virtual void AddOutputBindings(TArray<FMetasoundFrontendInterfaceBinding>&& InBindings) = 0;
+
+			// Returns all bindings that inform how this interface can be used to connect its inputs to a node that implements another interface's outputs.
+			virtual const TArray<FMetasoundFrontendInterfaceBinding>& GetOutputBindings() const = 0;
+
+			// Remove binding that informs how this interface can be used to connect its inputs to node that implements another interface's outputs.
+			virtual bool RemoveOutputBinding(const FMetasoundFrontendVersion& InVersion) = 0;
+
+			// Resets output bindings collection.
+			virtual void ResetOutputBindings() = 0;
 		};
 
 		class METASOUNDFRONTEND_API FInterfaceRegistryTransaction
@@ -86,6 +104,15 @@ namespace Metasound
 			// Find an interface with the given key. Returns true if interface is found, false if not.
 			virtual bool FindInterface(const FInterfaceRegistryKey& InKey, FMetasoundFrontendInterface& OutInterface) const = 0;
 
+			// Add interface output bindings to the interface with the given key. Orders the bindings using the given sort order, where lower values are evaluated earlier if multiple bindings are provided.
+			virtual void AddInterfaceOutputBindings(const FMetasoundFrontendVersion& InInputInterfaceVersion, TArray<FMetasoundFrontendInterfaceBinding>&& InOutputBindings) = 0;
+
+			// Removes an interface binding to the input interface with the given key from the interface with the given key. Returns true if interface binding was removed, false if not.
+			virtual bool RemoveInterfaceOutputBinding(const FMetasoundFrontendVersion& InInputInterfaceVersion, const FMetasoundFrontendVersion& InOutputInterfaceVersion) = 0;
+
+			// Removes all output interface bindings registered with the given interface version. Returns true if found and removed, false if not found.
+			virtual bool RemoveAllInterfaceOutputBindings(const FMetasoundFrontendVersion& InInputInterfaceVersion) = 0;
+
 			// Executes a function against all transactions since a provided ID and provides the current transaction ID.
 			UE_DEPRECATED(5.1, "ForEachRegistryTransactionSince is no longer supported. Use Metasound::Frontend::ISearchEngine to query for all interfaces.")
 			virtual void ForEachRegistryTransactionSince(FRegistryTransactionID InSince, FRegistryTransactionID* OutCurrentRegistryTransactionID, TFunctionRef<void(const FInterfaceRegistryTransaction&)> InFunc) const = 0;
diff --git a/Engine/Plugins/Runtime/Metasound/Source/MetasoundFrontend/Public/MetasoundFrontendDocument.h b/Engine/Plugins/Runtime/Metasound/Source/MetasoundFrontend/Public/MetasoundFrontendDocument.h
--- a/Engine/Plugins/Runtime/Metasound/Source/MetasoundFrontend/Public/MetasoundFrontendDocument.h
+++ b/Engine/Plugins/Runtime/Metasound/Source/MetasoundFrontend/Public/MetasoundFrontendDocument.h
@@ -23,8 +23,11 @@
 // Forward Declarations
 struct FMetasoundFrontendClass;
 struct FMetasoundFrontendClassInterface;
+
+
 namespace Metasound
 {
+	// Forward Declarations
 	struct FLiteral;
 
 	extern const FGuid METASOUNDFRONTEND_API FrontendInvalidID;
@@ -153,17 +156,17 @@ enum class EMetasoundFrontendClassType : uint8
 };
 
 // General purpose version number for Metasound Frontend objects.
-USTRUCT()
+USTRUCT(BlueprintType)
 struct METASOUNDFRONTEND_API FMetasoundFrontendVersionNumber
 {
 	GENERATED_BODY()
 
 	// Major version number.
-	UPROPERTY(VisibleAnywhere, Category = General)
+	UPROPERTY(VisibleAnywhere, BlueprintReadOnly, Category = General)
 	int32 Major = 1;
 
 	// Minor version number.
-	UPROPERTY(VisibleAnywhere, Category = General)
+	UPROPERTY(VisibleAnywhere, BlueprintReadOnly, Category = General)
 	int32 Minor = 0;
 
 	static const FMetasoundFrontendVersionNumber& GetInvalid()
@@ -244,17 +247,17 @@ struct METASOUNDFRONTEND_API FMetasoundFrontendVersionNumber
 };
 
 // General purpose version info for Metasound Frontend objects.
-USTRUCT()
+USTRUCT(BlueprintType)
 struct METASOUNDFRONTEND_API FMetasoundFrontendVersion
 {
 	GENERATED_BODY()
 
 	// Name of version.
-	UPROPERTY(VisibleAnywhere, Category = CustomView)
+	UPROPERTY(VisibleAnywhere, BlueprintReadOnly, Category = CustomView)
 	FName Name;
 
 	// Version number.
-	UPROPERTY(VisibleAnywhere, Category = CustomView)
+	UPROPERTY(VisibleAnywhere, BlueprintReadOnly, Category = CustomView)
 	FMetasoundFrontendVersionNumber Number;
 
 	FString ToString() const;
@@ -341,6 +344,42 @@ struct METASOUNDFRONTEND_API FMetasoundFrontendVertex
 	static bool IsFunctionalEquivalent(const FMetasoundFrontendVertex& InLHS, const FMetasoundFrontendVertex& InRHS);
 };
 
+// Pair of guids used to address location of vertex within a FrontendDocument graph
+USTRUCT()
+struct METASOUNDFRONTEND_API FMetasoundFrontendVertexHandle
+{
+	GENERATED_BODY()
+
+public:
+	UPROPERTY()
+	FGuid NodeID;
+
+	UPROPERTY()
+	FGuid VertexID;
+
+	// Returns whether or not the vertex handle is set (may or may not be
+	// valid depending on what builder context it is referenced against)
+	bool IsSet() const
+	{
+		return NodeID.IsValid() && VertexID.IsValid();
+	}
+
+	friend bool operator==(const FMetasoundFrontendVertexHandle& InLHS, const FMetasoundFrontendVertexHandle& InRHS)
+	{
+		return InLHS.NodeID == InRHS.NodeID && InLHS.VertexID == InRHS.VertexID;
+	}
+
+	friend bool operator!=(const FMetasoundFrontendVertexHandle& InLHS, const FMetasoundFrontendVertexHandle& InRHS)
+	{
+		return InLHS.NodeID != InRHS.NodeID || InLHS.VertexID != InRHS.VertexID;
+	}
+
+	friend FORCEINLINE uint32 GetTypeHash(const FMetasoundFrontendVertexHandle& InHandle)
+	{
+		return HashCombineFast(InHandle.NodeID.A, InHandle.VertexID.D);
+	}
+};
+
 // Contains a default value for a single vertex ID
 USTRUCT() 
 struct FMetasoundFrontendVertexLiteral
@@ -548,6 +587,16 @@ struct FMetasoundFrontendEdge
 	// ID of destination point on destination node.
 	UPROPERTY()
 	FGuid ToVertexID = Metasound::FrontendInvalidID;
+
+	FMetasoundFrontendVertexHandle GetFromVertexHandle() const
+	{
+		return FMetasoundFrontendVertexHandle { FromNodeID, FromVertexID };
+	}
+
+	FMetasoundFrontendVertexHandle GetToVertexHandle() const
+	{
+		return FMetasoundFrontendVertexHandle { ToNodeID, ToVertexID };
+	}
 };
 
 USTRUCT()
@@ -742,7 +791,7 @@ struct METASOUNDFRONTEND_API FMetasoundFrontendClassVertex : public FMetasoundFr
 	FGuid NodeID = Metasound::FrontendInvalidID;
 
 #if WITH_EDITORONLY_DATA
-	// Metadata associated with input.
+	// Metadata associated with vertex.
 	UPROPERTY(EditAnywhere, Category = CustomView)
 	FMetasoundFrontendVertexMetadata Metadata;
 #endif // WITH_EDITORONLY_DATA
@@ -1039,6 +1088,56 @@ public:
 	friend struct FMetasoundFrontendClass;
 };
 
+USTRUCT()
+struct METASOUNDFRONTEND_API FMetasoundFrontendInterfaceBindingConnection
+{
+	GENERATED_BODY()
+
+	UPROPERTY()
+	FName OutputName;
+
+	UPROPERTY()
+	FName InputName;
+
+	friend bool operator==(const FMetasoundFrontendInterfaceBindingConnection& InLHS, const FMetasoundFrontendInterfaceBindingConnection& InRHS)
+	{
+		return InLHS.OutputName == InRHS.OutputName && InLHS.InputName == InRHS.InputName;
+	}
+
+	friend bool operator!=(const FMetasoundFrontendInterfaceBindingConnection& InLHS, const FMetasoundFrontendInterfaceBindingConnection& InRHS)
+	{
+		return InLHS.OutputName != InRHS.OutputName || InLHS.InputName != InRHS.InputName;
+	}
+
+	FString ToString() const
+	{
+		return FString::Format(TEXT("{0}->{1}"), { OutputName.ToString(), InputName.ToString() });
+	}
+
+	friend FORCEINLINE uint32 GetTypeHash(const FMetasoundFrontendInterfaceBindingConnection& InBinding)
+	{
+		return HashCombineFast(GetTypeHash(InBinding.OutputName), GetTypeHash(InBinding.InputName));
+	}
+};
+
+
+USTRUCT()
+struct METASOUNDFRONTEND_API FMetasoundFrontendInterfaceBinding
+{
+	GENERATED_BODY()
+
+	UPROPERTY()
+	FMetasoundFrontendVersion Version;
+
+	// Priority to determine if binding connections are higher or lower than another interface
+	// that may be shared between nodes attempting to be connected via shared interfaces.
+	UPROPERTY()
+	int32 ConnectionPriority = 0;
+
+	UPROPERTY()
+	TArray<FMetasoundFrontendInterfaceBindingConnection> Connections;
+};
+
 
 USTRUCT()
 struct METASOUNDFRONTEND_API FMetasoundFrontendInterface : public FMetasoundFrontendClassInterface
@@ -1428,6 +1527,7 @@ struct METASOUNDFRONTEND_API FMetasoundFrontendDocument
 {
 	GENERATED_BODY()
 
+public:
 	static FMetasoundFrontendVersionNumber GetMaxVersion();
 
 	Metasound::Frontend::FAccessPoint AccessPoint;
@@ -1437,7 +1537,6 @@ struct METASOUNDFRONTEND_API FMetasoundFrontendDocument
 	UPROPERTY(EditAnywhere, Category = Metadata)
 	FMetasoundFrontendDocumentMetadata Metadata;
 
-public:
 	UPROPERTY(VisibleAnywhere, Category = CustomView)
 	TSet<FMetasoundFrontendVersion> Interfaces;
 
@@ -1480,4 +1579,5 @@ public:
 	}
 };
 
+
 METASOUNDFRONTEND_API const TCHAR* LexToString(EMetasoundFrontendVertexAccessType InVertexAccess);
diff --git a/Engine/Plugins/Runtime/Metasound/Source/MetasoundFrontend/Public/MetasoundFrontendDocumentBuilder.h b/Engine/Plugins/Runtime/Metasound/Source/MetasoundFrontend/Public/MetasoundFrontendDocumentBuilder.h
new file mode 100644
--- /dev/null
+++ b/Engine/Plugins/Runtime/Metasound/Source/MetasoundFrontend/Public/MetasoundFrontendDocumentBuilder.h
@@ -0,0 +1,187 @@
+// Copyright Epic Games, Inc. All Rights Reserved.
+#pragma once
+
+#include "Containers/Map.h"
+#include "MetasoundDocumentInterface.h"
+#include "MetasoundFrontendArchetypeRegistry.h"
+#include "MetasoundFrontendController.h"
+#include "MetasoundFrontendDocument.h"
+#include "MetasoundFrontendDocumentCacheInterface.h"
+#include "MetasoundFrontendRegistries.h"
+#include "MetasoundFrontendTransform.h"
+#include "MetasoundVertex.h"
+#include "Templates/Function.h"
+#include "UObject/ScriptInterface.h"
+
+#include "MetasoundFrontendDocumentBuilder.generated.h"
+
+
+// Forward Declarations
+class FMetasoundAssetBase;
+
+
+namespace Metasound::Frontend
+{
+	struct FNamedEdge
+	{
+		const FGuid OutputNodeID;
+		const FName OutputName;
+		const FGuid InputNodeID;
+		const FName InputName;
+
+		friend bool operator==(const FNamedEdge& InLHS, const FNamedEdge& InRHS)
+		{
+			return InLHS.OutputNodeID == InRHS.OutputNodeID
+				&& InLHS.OutputName == InRHS.OutputName
+				&& InLHS.InputNodeID == InRHS.InputNodeID
+				&& InLHS.InputName == InRHS.InputName;
+		}
+
+		friend bool operator!=(const FNamedEdge& InLHS, const FNamedEdge& InRHS)
+		{
+			return !(InLHS == InRHS);
+		}
+
+		friend FORCEINLINE uint32 GetTypeHash(const FNamedEdge& InBinding)
+		{
+			const int32 NameHash = HashCombineFast(GetTypeHash(InBinding.OutputName), GetTypeHash(InBinding.InputName));
+			const int32 GuidHash = HashCombineFast(GetTypeHash(InBinding.OutputNodeID), GetTypeHash(InBinding.InputNodeID));
+			return HashCombineFast(NameHash, GuidHash);
+		}
+	};
+} // namespace Metasound::Frontend
+
+
+UCLASS()
+class  METASOUNDFRONTEND_API UMetaSoundBuilderDocument : public UObject, public IMetaSoundDocumentInterface
+{
+	GENERATED_BODY()
+
+public:
+	virtual const FMetasoundFrontendDocument& GetDocument() const override
+	{
+		return Document;
+	}
+
+private:
+	virtual FMetasoundFrontendDocument& GetDocument() override
+	{
+		return Document;
+	}
+
+	UPROPERTY(Transient)
+	FMetasoundFrontendDocument Document;
+
+	friend struct FMetaSoundFrontendDocumentBuilder;
+};
+
+// Builder used to support dynamically generating MetaSound documents at runtime. Builder contains caches that speed up
+// common search and modification operations on a given document, which may result in slower performance on construction,
+// but faster manipulation of its managed document.  The builder's managed copy of a document is expected to not be modified
+// by any external system to avoid cache to not become stale.
+USTRUCT()
+struct METASOUNDFRONTEND_API FMetaSoundFrontendDocumentBuilder
+{
+	GENERATED_BODY()
+
+public:
+	FMetaSoundFrontendDocumentBuilder();
+	FMetaSoundFrontendDocumentBuilder(FMetaSoundFrontendDocumentBuilder&& InBuilder);
+	FMetaSoundFrontendDocumentBuilder(TScriptInterface<IMetaSoundDocumentInterface> InDocumentInterface);
+
+	// Copying actively used builders is highly discouraged (as only one builder should be acting on a document at any given time
+	// to avoid internal cache state to become out-of-date) but must contractually be implemented in this case as a USTRUCT (builders
+	// are USTRUCTS to keep document dependencies from being GC'ed).  To protect against this kind of "back door" building from another
+	// builder, copying builders also creates a new document & builder owner.
+	FMetaSoundFrontendDocumentBuilder(const FMetaSoundFrontendDocumentBuilder& InBuilder);
+	FMetaSoundFrontendDocumentBuilder& operator=(const FMetaSoundFrontendDocumentBuilder& InRHS);
+
+	FMetaSoundFrontendDocumentBuilder& operator=(FMetaSoundFrontendDocumentBuilder&& InRHS);
+
+	const FMetasoundFrontendClass* AddGraphDependency(const FMetasoundFrontendGraphClass& InClass);
+	const FMetasoundFrontendClass* AddNativeDependency(const FMetasoundFrontendClassMetadata& InClassMetadata);
+	const FMetasoundFrontendEdge* AddEdge(FMetasoundFrontendEdge&& InNewEdge);
+	bool AddEdgesByNamedConnections(const TSet<Metasound::Frontend::FNamedEdge>& ConnectionsToMake, TArray<const FMetasoundFrontendEdge*>* OutNewEdges = nullptr);
+	bool AddEdgesByNodeClassInterfaceBindings(const FGuid& InFromNodeID, const FGuid& InToNodeID);
+	bool AddEdgesFromMatchingInterfaceNodeOutputsToGraphOutputs(const FGuid& InNodeID, TArray<const FMetasoundFrontendEdge*>& OutEdgesCreated);
+	bool AddEdgesFromMatchingInterfaceNodeInputsToGraphInputs(const FGuid& InNodeID, TArray<const FMetasoundFrontendEdge*>& OutEdgesCreated);
+	const FMetasoundFrontendNode* AddGraphInput(const FMetasoundFrontendClassInput& InClassInput);
+	const FMetasoundFrontendNode* AddGraphOutput(const FMetasoundFrontendClassOutput& InClassOutput);
+	bool AddInterface(FName InterfaceName);
+
+	const FMetasoundFrontendNode* AddGraphNode(const FMetasoundFrontendGraphClass& InGraphClass, FGuid InNodeID = FGuid::NewGuid());
+
+	bool CanAddEdge(const FMetasoundFrontendEdge& InEdge) const;
+
+	bool ContainsEdge(const FMetasoundFrontendEdge& InEdge) const;
+	bool ContainsNode(const FGuid& InNodeID) const;
+	bool ConvertFromPreset();
+
+	static bool FindDeclaredInterfaces(const FMetasoundFrontendDocument& InDocument, TArray<const Metasound::Frontend::IInterfaceRegistryEntry*>& OutInterfaces);
+	bool FindDeclaredInterfaces(TArray<const Metasound::Frontend::IInterfaceRegistryEntry*>& OutInterfaces) const;
+
+	const FMetasoundFrontendClass* FindDependency(const FGuid& InClassID) const;
+	const FMetasoundFrontendClass* FindDependency(const FMetasoundFrontendClassMetadata& InMetadata) const;
+
+	bool FindInterfaceInputNodes(FName InterfaceName, TArray<const FMetasoundFrontendNode*>& OutInputs) const;
+	bool FindInterfaceOutputNodes(FName InterfaceName, TArray<const FMetasoundFrontendNode*>& OutOutputs) const;
+
+	const FMetasoundFrontendNode* FindGraphInputNode(FName InputName) const;
+	const FMetasoundFrontendNode* FindGraphOutputNode(FName OutputName) const;
+
+	const FMetasoundFrontendVertex* FindNodeInput(const FGuid& InNodeID, const FGuid& InVertexID) const;
+	const FMetasoundFrontendVertex* FindNodeInput(const FGuid& InNodeID, FName InVertexName) const;
+	const FMetasoundFrontendVertex* FindNodeOutput(const FGuid& InNodeID, const FGuid& InVertexID) const;
+	const FMetasoundFrontendVertex* FindNodeOutput(const FGuid& InNodeID, FName InVertexName) const;
+
+	const FMetasoundFrontendNode* FindNode(const FGuid& InNodeID) const;
+	TArray<const FMetasoundFrontendVertex*> FindNodeInputs(const FGuid& InNodeID, FName TypeName = FName()) const;
+	TArray<const FMetasoundFrontendVertex*> FindNodeOutputs(const FGuid& InNodeID, FName TypeName = FName()) const;
+
+	const FMetasoundFrontendDocument& GetDocument() const;
+
+	static void InitGraphClassMetadata(FMetasoundFrontendClassMetadata& InOutMetadata, bool bResetVersion = false);
+	void InitDocument(FName UClassName);
+	void InitNodeLocations();
+
+	bool IsNodeInputConnected(const FGuid& InNodeID, const FGuid& InVertexID) const;
+	bool IsNodeOutputConnected(const FGuid& InNodeID, const FGuid& InVertexID) const;
+
+	bool IsInterfaceDeclared(FName InInterfaceName) const;
+	bool IsInterfaceDeclared(const FMetasoundFrontendVersion& InInterfaceVersion) const;
+
+	bool ModifyInterfaces(TArray<FMetasoundFrontendVersion> OutputFormatsToRemove, TArray<FMetasoundFrontendVersion> OutputFormatsToAdd);
+
+	bool RemoveDependency(EMetasoundFrontendClassType ClassType, const FMetasoundFrontendClassName& InClassName, const FMetasoundFrontendVersionNumber& InClassVersionNumber);
+	bool RemoveEdge(const FMetasoundFrontendEdge& EdgeToRemove);
+	bool RemoveEdgesByNamedConnections(const TSet<Metasound::Frontend::FNamedEdge>& ConnectionsToRemove, TArray<FMetasoundFrontendEdge>* OutRemovedEdges = nullptr);
+	bool RemoveEdgesByNodeClassInterfaceBindings(const FGuid& InOutputNodeID, const FGuid& InInputNodeID);
+	bool RemoveEdgesFromNodeOutput(const FGuid& InNodeID, const FGuid& InVertexID);
+	bool RemoveEdgeToNodeInput(const FGuid& InNodeID, const FGuid& InVertexID);
+	bool RemoveInterface(FName InName);
+	bool RemoveNode(const FGuid& InNodeID);
+
+#if WITH_EDITOR
+	void SetAuthor(const FString& InAuthor);
+#endif // WITH_EDITOR
+
+	bool SetNodeInputDefault(const FGuid& InNodeID, const FGuid& InVertexID, const FMetasoundFrontendLiteral& InLiteral);
+
+private:
+	using FFinalizeNodeFunctionRef = TFunctionRef<void(FMetasoundFrontendNode&, const Metasound::Frontend::FNodeRegistryKey&)>;
+
+	UPROPERTY(Transient)
+	TScriptInterface<IMetaSoundDocumentInterface> DocumentInterface;
+
+	TUniquePtr<Metasound::Frontend::IDocumentCache> DocumentCache;
+
+	FMetasoundFrontendDocument& GetDocument();
+
+	FMetasoundFrontendNode* AddNodeInternal(const FMetasoundFrontendClassMetadata& InClassMetadata, FFinalizeNodeFunctionRef FinalizeNode, FGuid InNodeID = FGuid::NewGuid());
+
+	const TScriptInterface<IMetaSoundDocumentInterface> FindOrLoadNodeClassDocument(const FGuid& InNodeID) const;
+	const FMetasoundFrontendClassInput* FindNodeInputClassInput(const FGuid& InNodeID, const FGuid& InVertexID) const;
+	const FMetasoundFrontendClassOutput* FindNodeOutputClassOutput(const FGuid& InNodeID, const FGuid& InVertexID) const;
+
+	void ReloadCache();
+};
diff --git a/Engine/Plugins/Runtime/Metasound/Source/MetasoundFrontend/Public/MetasoundFrontendLiteral.h b/Engine/Plugins/Runtime/Metasound/Source/MetasoundFrontend/Public/MetasoundFrontendLiteral.h
--- a/Engine/Plugins/Runtime/Metasound/Source/MetasoundFrontend/Public/MetasoundFrontendLiteral.h
+++ b/Engine/Plugins/Runtime/Metasound/Source/MetasoundFrontend/Public/MetasoundFrontendLiteral.h
@@ -55,7 +55,7 @@ static_assert(std::is_same<std::underlying_type<Metasound::ELiteralType>::type,
 
 
 // Represents the serialized version of variant literal types. 
-USTRUCT()
+USTRUCT(BlueprintType, meta = (DisplayName = "MetaSound Literal"))
 struct METASOUNDFRONTEND_API FMetasoundFrontendLiteral
 {
 	GENERATED_BODY()
diff --git a/Engine/Plugins/Runtime/Metasound/Source/MetasoundFrontend/Public/MetasoundFrontendTransform.h b/Engine/Plugins/Runtime/Metasound/Source/MetasoundFrontend/Public/MetasoundFrontendTransform.h
--- a/Engine/Plugins/Runtime/Metasound/Source/MetasoundFrontend/Public/MetasoundFrontendTransform.h
+++ b/Engine/Plugins/Runtime/Metasound/Source/MetasoundFrontend/Public/MetasoundFrontendTransform.h
@@ -49,8 +49,8 @@ namespace Metasound
 		public:
 			virtual ~IGraphTransform() = default;
 
-			// Returns reference to the node's owning document.
-			virtual FMetasoundFrontendDocument& GetOwningDocument() const = 0;
+			UE_DEPRECATED(5.3, "Deprecated: unused function which encouraged breaking const behavior of transform state.")
+			virtual FMetasoundFrontendDocument& GetOwningDocument() const { static FMetasoundFrontendDocument BaseDoc; return BaseDoc; };
 
 			/** Return true if the graph was modified, false otherwise. */
 			virtual bool Transform(FMetasoundFrontendGraph& InOutGraph) const = 0;
@@ -62,11 +62,11 @@ namespace Metasound
 		public:
 			virtual ~INodeTransform() = default;
 
-			// Returns reference to the node's owning document.
-			virtual FMetasoundFrontendDocument& GetOwningDocument() const = 0;
+			UE_DEPRECATED(5.3, "Deprecated: unused function which encouraged breaking const behavior of transform state.")
+			virtual FMetasoundFrontendDocument& GetOwningDocument() const { static FMetasoundFrontendDocument BaseDoc; return BaseDoc; };
 
-			// Returns reference to the node's owning graph.
-			virtual FMetasoundFrontendGraph& GetOwningGraph() const = 0;
+			UE_DEPRECATED(5.3, "Deprecated: unused function which encouraged breaking const behavior of transform state.")
+			virtual FMetasoundFrontendGraph& GetOwningGraph() const { static FMetasoundFrontendGraph BaseGraph; return BaseGraph; }
 
 			/** Return true if the node was modified, false otherwise. */
 			virtual bool Transform(FMetasoundFrontendNode& InOutNode) const = 0;
@@ -90,7 +90,8 @@ namespace Metasound
 			// that may be related but not be named the same.
 			void SetNamePairingFunction(const TFunction<bool(FName, FName)>& InNamePairingFunction);
 
-			bool Transform(FDocumentHandle InDocument) const override;
+			virtual bool Transform(FDocumentHandle InDocument) const override;
+			virtual bool Transform(FMetasoundFrontendDocument& InOutDocument) const override;
 
 		private:
 			void Init(const TFunction<bool(FName, FName)>* InNamePairingFunction = nullptr);
diff --git a/Engine/Source/Runtime/AudioExtensions/Private/IAudioParameterInterfaceRegistry.cpp b/Engine/Source/Runtime/AudioExtensions/Private/IAudioParameterInterfaceRegistry.cpp
--- a/Engine/Source/Runtime/AudioExtensions/Private/IAudioParameterInterfaceRegistry.cpp
+++ b/Engine/Source/Runtime/AudioExtensions/Private/IAudioParameterInterfaceRegistry.cpp
@@ -1,6 +1,8 @@
 // Copyright Epic Games, Inc. All Rights Reserved.
 #include "IAudioParameterInterfaceRegistry.h"
 
+#include "Algo/Transform.h"
+
 
 namespace Audio
 {
@@ -33,6 +35,78 @@ namespace Audio
 		};
 	} // namespace ParameterInterfaceRegistryPrivate
 
+	FParameterInterface::FParameterInterface(
+		FName InName,
+		const FVersion& InVersion,
+		const UClass& InType
+	)
+		: NamePrivate(InName)
+		, VersionPrivate(InVersion)
+		, SupportedUClassNames({ InType.GetFullName() })
+	{
+	}
+
+	FParameterInterface::FParameterInterface(FName InName, const FVersion& InVersion)
+		: NamePrivate(InName)
+		, VersionPrivate(InVersion)
+	{
+	}
+
+	FParameterInterface::FParameterInterface(FName InName, const FVersion& InVersion, const TArray<UClass*>& InClasses)
+		: NamePrivate(InName)
+		, VersionPrivate(InVersion)
+	{
+		Algo::Transform(InClasses, SupportedUClassNames, [](const UClass* Class)
+		{
+			check(Class);
+			return Class->GetFullName();
+		});
+	}
+
+	FName FParameterInterface::GetName() const
+	{
+		return NamePrivate;
+	}
+
+	const FParameterInterface::FVersion& FParameterInterface::GetVersion() const
+	{
+		return VersionPrivate;
+	}
+
+	const UClass& FParameterInterface::GetType() const
+	{
+		return *UObject::StaticClass();
+	}
+
+	const TArray<FParameterInterface::FInput>& FParameterInterface::GetInputs() const
+	{
+		return Inputs;
+	}
+
+	const TArray<FParameterInterface::FOutput>& FParameterInterface::GetOutputs() const
+	{
+		return Outputs;
+	}
+
+	const TArray<FParameterInterface::FEnvironmentVariable>& FParameterInterface::GetEnvironment() const
+	{
+		return Environment;
+	}
+
+	TArray<const UClass*> FParameterInterface::FindSupportedUClasses() const
+	{
+		TArray<const UClass*> SupportedUClasses;
+		for (const FString& Name : SupportedUClassNames)
+		{
+			if (const UClass* Class = FindObject<const UClass>(nullptr, *Name))
+			{
+				SupportedUClasses.Add(Class);
+			}
+		}
+
+		return SupportedUClasses;
+	}
+
 	TUniquePtr<IAudioParameterInterfaceRegistry> IAudioParameterInterfaceRegistry::Instance;
 
 	IAudioParameterInterfaceRegistry& IAudioParameterInterfaceRegistry::Get()
diff --git a/Engine/Source/Runtime/AudioExtensions/Public/IAudioParameterInterfaceRegistry.h b/Engine/Source/Runtime/AudioExtensions/Public/IAudioParameterInterfaceRegistry.h
--- a/Engine/Source/Runtime/AudioExtensions/Public/IAudioParameterInterfaceRegistry.h
+++ b/Engine/Source/Runtime/AudioExtensions/Public/IAudioParameterInterfaceRegistry.h
@@ -23,13 +23,13 @@ namespace Audio
 
 	struct AUDIOEXTENSIONS_API FParameterInterface
 	{
-		struct FVersion
+		struct AUDIOEXTENSIONS_API FVersion
 		{
 			const int32 Major = 1;
 			const int32 Minor = 0;
 		};
 
-		struct FInput
+		struct AUDIOEXTENSIONS_API FInput
 		{
 			const FText DisplayName;
 			const FText Description;
@@ -41,7 +41,7 @@ namespace Audio
 			const int32 SortOrderIndex = 0;
 		};
 
-		struct FOutput
+		struct AUDIOEXTENSIONS_API FOutput
 		{
 			const FText DisplayName;
 			const FText Description;
@@ -54,7 +54,7 @@ namespace Audio
 			const int32 SortOrderIndex = 0;
 		};
 
-		struct FEnvironmentVariable
+		struct AUDIOEXTENSIONS_API FEnvironmentVariable
 		{
 			const FText DisplayName;
 			const FText Description;
@@ -65,52 +65,31 @@ namespace Audio
 		};
 
 		FParameterInterface() = default;
-		FParameterInterface(
-			FName InName,
-			const FVersion& InVersion,
-			const UClass& InType
-		)
-			: NamePrivate(InName)
-			, VersionPrivate(InVersion)
-			, Type(&InType)
-		{
-		}
 
-		FName GetName() const
-		{
-			return NamePrivate;
-		}
+		// Constructor used for parameter interface not limited to any particular UClass types
+		FParameterInterface(FName InName, const FVersion& InVersion);
 
-		const FVersion& GetVersion() const
-		{
-			return VersionPrivate;
-		}
+		// Constructor used for parameter interface with support for explicit UClass types
+		FParameterInterface(FName InName, const FVersion& InVersion, const TArray<UClass*>& InClasses);
 
-		const UClass& GetType() const
-		{
-			check(Type);
-			return *Type;
-		}
+		UE_DEPRECATED(5.3, "Use FParameterInterface ctor that supports array of classes.")
+		FParameterInterface(FName InName, const FVersion& InVersion, const UClass& InClass);
 
-		const TArray<FInput>& GetInputs() const
-		{
-			return Inputs;
-		}
+		FName GetName() const;
+		const FVersion& GetVersion() const;
 
-		const TArray<FOutput>& GetOutputs() const
-		{
-			return Outputs;
-		}
+		UE_DEPRECATED(5.3, "Use FParameterInterface::FindSupportedUClasses instead")
+		const UClass& GetType() const;
 
-		const TArray<FEnvironmentVariable>& GetEnvironment() const
-		{
-			return Environment;
-		}
+		TArray<const UClass*> FindSupportedUClasses() const;
+		const TArray<FInput>& GetInputs() const;
+		const TArray<FOutput>& GetOutputs() const;
+		const TArray<FEnvironmentVariable>& GetEnvironment() const;
 
 	private:
 		FName NamePrivate;
 		FVersion VersionPrivate;
-		const UClass* Type = nullptr;
+		TArray<FString> SupportedUClassNames;
 
 	protected:
 		TArray<FInput> Inputs;
diff --git a/Engine/Source/Runtime/Engine/Private/SoundAttenuation.cpp b/Engine/Source/Runtime/Engine/Private/SoundAttenuation.cpp
--- a/Engine/Source/Runtime/Engine/Private/SoundAttenuation.cpp
+++ b/Engine/Source/Runtime/Engine/Private/SoundAttenuation.cpp
@@ -162,7 +162,7 @@ namespace Audio
 			struct FInterface : public Audio::FParameterInterface
 			{
 				FInterface()
-					: FParameterInterface(AttenuationInterface::Name, { 1, 0 }, *USoundBase::StaticClass())
+					: FParameterInterface(AttenuationInterface::Name, { 1, 0 })
 				{
 					Inputs =
 					{
@@ -203,7 +203,7 @@ namespace Audio
 			struct FInterface : public Audio::FParameterInterface
 			{
 				FInterface()
-					: FParameterInterface(SpatializationInterface::Name, { 1, 0 }, *USoundBase::StaticClass())
+					: FParameterInterface(SpatializationInterface::Name, { 1, 0 })
 				{
 					Inputs =
 					{
@@ -250,7 +250,7 @@ namespace Audio
 			struct FInterface : public Audio::FParameterInterface
 			{
 				FInterface()
-					: FParameterInterface(SourceOrientationInterface::Name, { 1, 0 }, *USoundBase::StaticClass())
+					: FParameterInterface(SourceOrientationInterface::Name, { 1, 0 })
 				{
 					Inputs =
 					{
@@ -297,7 +297,7 @@ namespace Audio
 			struct FInterface : public Audio::FParameterInterface
 			{
 				FInterface()
-					: FParameterInterface(ListenerOrientationInterface::Name, { 1, 0 }, *USoundBase::StaticClass())
+					: FParameterInterface(ListenerOrientationInterface::Name, { 1, 0 })
 				{
 					Inputs =
 					{
