Fix borked merge
#rnx

diff --git a/Engine/Plugins/Runtime/Metasound/Source/MetasoundFrontend/Private/NodeTemplates/MetasoundFrontendNodeTemplateReroute.cpp b/Engine/Plugins/Runtime/Metasound/Source/MetasoundFrontend/Private/NodeTemplates/MetasoundFrontendNodeTemplateReroute.cpp
--- a/Engine/Plugins/Runtime/Metasound/Source/MetasoundFrontend/Private/NodeTemplates/MetasoundFrontendNodeTemplateReroute.cpp
+++ b/Engine/Plugins/Runtime/Metasound/Source/MetasoundFrontend/Private/NodeTemplates/MetasoundFrontendNodeTemplateReroute.cpp
@@ -83,7 +83,6 @@ namespace Metasound
 
 				// Update the output edges with the input edge 
 				FMetasoundFrontendVertexHandle NewOutputEdgeNodeVertexHandle { InputEdge->FromNodeID, InputEdge->FromVertexID };
-				FNodeVertexGuidPair NewOutputEdgeNodeVertexPair{ InputEdge->FromNodeID, InputEdge->FromVertexID };
 
 				for (FMetasoundFrontendEdge* OutputEdge : *OutputEdges)
 				{
