[Metasound] Add MetasoundGenerator API to get input write refs and output read refs.

See unit tests for usage examples.

diff --git a/Engine/Plugins/Runtime/Metasound/Source/MetasoundEngineTest/MetasoundEngineTest.Build.cs b/Engine/Plugins/Runtime/Metasound/Source/MetasoundEngineTest/MetasoundEngineTest.Build.cs
--- a/Engine/Plugins/Runtime/Metasound/Source/MetasoundEngineTest/MetasoundEngineTest.Build.cs
+++ b/Engine/Plugins/Runtime/Metasound/Source/MetasoundEngineTest/MetasoundEngineTest.Build.cs
@@ -18,10 +18,6 @@ namespace UnrealBuildTool.Rules
 					"Engine",
 					"AudioExtensions",
 					"AudioMixer",
-					"MetasoundGraphCore",
-					"MetasoundGenerator",
-					"MetasoundFrontend",
-					"MetasoundEngine",
 					"SignalProcessing",
 					"Projects"
 				}
@@ -31,6 +27,10 @@ namespace UnrealBuildTool.Rules
 			(
 				new string[]
 				{
+					"MetasoundGraphCore",
+					"MetasoundGenerator",
+					"MetasoundFrontend",
+					"MetasoundEngine"
 				}
 			);
 		}
diff --git a/Engine/Plugins/Runtime/Metasound/Source/MetasoundEngineTest/Private/MetasoundGeneratorTests.cpp b/Engine/Plugins/Runtime/Metasound/Source/MetasoundEngineTest/Private/MetasoundGeneratorTests.cpp
new file mode 100644
--- /dev/null
+++ b/Engine/Plugins/Runtime/Metasound/Source/MetasoundEngineTest/Private/MetasoundGeneratorTests.cpp
@@ -0,0 +1,276 @@
+// Copyright Epic Games, Inc. All Rights Reserved.
+
+#include "NodeTestGraphBuilder.h"
+#include "Misc/AutomationTest.h"
+
+#if WITH_DEV_AUTOMATION_TESTS
+
+namespace Metasound::Test::Generator
+{
+	using GraphBuilder = FNodeTestGraphBuilder;
+	using namespace Frontend;
+	
+	template<typename DataType>
+	TUniquePtr<FMetasoundGenerator> BuildPassthroughGraph(
+		FAutomationTestBase& Test,
+		const FName& InputName,
+		const FName& OutputName,
+		const FSampleRate SampleRate,
+		const int32 NumSamplesPerBlock)
+	{
+		GraphBuilder Builder;
+		const FNodeHandle InputNode = Builder.AddInput(InputName, Metasound::GetMetasoundDataTypeName<DataType>());
+		const FNodeHandle OutputNode = Builder.AddOutput(OutputName, Metasound::GetMetasoundDataTypeName<DataType>());
+		const FOutputHandle OutputToConnect = InputNode->GetOutputWithVertexName(InputName);
+		const FInputHandle InputToConnect = OutputNode->GetInputWithVertexName(OutputName);
+		if (!Test.TestTrue("Connected input to output", InputToConnect->Connect(*OutputToConnect)))
+		{
+			return nullptr;
+		}
+		return Builder.BuildGenerator(SampleRate, NumSamplesPerBlock);
+	}
+	
+	template<typename DataType>
+	bool RunSimpleOutputTest(FAutomationTestBase& Test, DataType ExpectedValue)
+	{
+		const FName InputName = "MyInput";
+		const FName OutputName = "MyOutput";
+		constexpr FSampleRate SampleRate = 48000;
+		constexpr int32 NumSamplesPerBlock = 128;
+
+		// Build a passthrough graph
+		const TUniquePtr<FMetasoundGenerator> Generator =
+			BuildPassthroughGraph<DataType>(Test, InputName, OutputName, SampleRate, NumSamplesPerBlock);
+		if (!Test.TestNotNull("Generator built", Generator.Get()))
+		{
+			return false;
+		}
+		
+		// Set the input value
+		const TOptional<TDataWriteReference<DataType>> InputWriteRef = Generator->GetInputWriteReference<DataType>(InputName);
+		if (!Test.TestTrue("Got the write ref", InputWriteRef.IsSet()))
+		{
+			return false;
+		}
+		*InputWriteRef.GetValue() = ExpectedValue;
+
+		// Render a block
+		{
+			TArray<float> Buffer;
+			Buffer.Reserve(NumSamplesPerBlock);
+			if (!Test.TestEqual(
+				"Generated the right number of samples.",
+				Generator->OnGenerateAudio(Buffer.GetData(), NumSamplesPerBlock),
+				NumSamplesPerBlock))
+			{
+				return false;
+			}
+		}
+
+		// Check the output
+		const TOptional<TDataReadReference<DataType>> OutputReadRef = Generator->GetOutputReadReference<DataType>(OutputName);
+		if (!Test.TestTrue("Got the read ref", OutputReadRef.IsSet()))
+		{
+			return false;
+		}
+		if (!Test.TestEqual("Input was passed through", *OutputReadRef.GetValue(), ExpectedValue))
+		{
+			return false;
+		}
+
+		return true;
+	}
+	
+	IMPLEMENT_SIMPLE_AUTOMATION_TEST(
+	FMetasoundGeneratorOutputSimpleFloatTest,
+	"Audio.Metasound.Generator.Outputs.Simple.Float",
+	EAutomationTestFlags::EditorContext | EAutomationTestFlags::EngineFilter)
+	bool FMetasoundGeneratorOutputSimpleFloatTest::RunTest(const FString&)
+	{
+		return RunSimpleOutputTest<float>(*this, 123.456f);
+	}
+	
+	IMPLEMENT_SIMPLE_AUTOMATION_TEST(
+	FMetasoundGeneratorOutputSimpleIntTest,
+	"Audio.Metasound.Generator.Outputs.Simple.Int",
+	EAutomationTestFlags::EditorContext | EAutomationTestFlags::EngineFilter)
+	bool FMetasoundGeneratorOutputSimpleIntTest::RunTest(const FString&)
+	{
+		return RunSimpleOutputTest<int32>(*this, 123456);
+	}
+		
+	IMPLEMENT_SIMPLE_AUTOMATION_TEST(
+	FMetasoundGeneratorOutputSimpleBoolTest,
+	"Audio.Metasound.Generator.Outputs.Simple.Bool",
+	EAutomationTestFlags::EditorContext | EAutomationTestFlags::EngineFilter)
+	bool FMetasoundGeneratorOutputSimpleBoolTest::RunTest(const FString&)
+	{
+		return RunSimpleOutputTest<bool>(*this, true);
+	}
+
+	template<typename DataType>
+	bool RunSetInputTest(FAutomationTestBase& Test, DataType ExpectedValue)
+	{
+		const FName InputName = "MyInput";
+		const FName OutputName = "MyOutput";
+		constexpr FSampleRate SampleRate = 48000;
+		constexpr int32 NumSamplesPerBlock = 128;
+
+		// Build a passthrough graph
+		const TUniquePtr<FMetasoundGenerator> Generator =
+			BuildPassthroughGraph<DataType>(Test, InputName, OutputName, SampleRate, NumSamplesPerBlock);
+		if (!Test.TestNotNull("Generator built", Generator.Get()))
+		{
+			return false;
+		}
+		
+		// Set the input value
+		Generator->SetInputValue<DataType>(InputName, ExpectedValue);
+
+		// Render a block
+		{
+			TArray<float> Buffer;
+			Buffer.Reserve(NumSamplesPerBlock);
+			if (!Test.TestEqual(
+				"Generated the right number of samples.",
+				Generator->OnGenerateAudio(Buffer.GetData(), NumSamplesPerBlock),
+				NumSamplesPerBlock))
+			{
+				return false;
+			}
+		}
+
+		// Check the output
+		const TOptional<TDataReadReference<DataType>> OutputReadRef = Generator->GetOutputReadReference<DataType>(OutputName);
+		if (!Test.TestTrue("Got the read ref", OutputReadRef.IsSet()))
+		{
+			return false;
+		}
+		if (!Test.TestEqual("Input was passed through", *OutputReadRef.GetValue(), ExpectedValue))
+		{
+			return false;
+		}
+
+		return true;
+	}
+	
+	IMPLEMENT_SIMPLE_AUTOMATION_TEST(
+	FMetasoundGeneratorSetInputFloatTest,
+	"Audio.Metasound.Generator.SetInput.Float",
+	EAutomationTestFlags::EditorContext | EAutomationTestFlags::EngineFilter)
+	bool FMetasoundGeneratorSetInputFloatTest::RunTest(const FString&)
+	{
+		return RunSetInputTest<float>(*this, 123.456f);
+	}
+	
+	IMPLEMENT_SIMPLE_AUTOMATION_TEST(
+	FMetasoundGeneratorSetInputIntTest,
+	"Audio.Metasound.Generator.SetInput.Int",
+	EAutomationTestFlags::EditorContext | EAutomationTestFlags::EngineFilter)
+	bool FMetasoundGeneratorSetInputIntTest::RunTest(const FString&)
+	{
+		return RunSetInputTest<int32>(*this, 123456);
+	}
+
+	IMPLEMENT_SIMPLE_AUTOMATION_TEST(
+	FMetasoundGeneratorSetInputBoolTest,
+	"Audio.Metasound.Generator.SetInput.Bool",
+	EAutomationTestFlags::EditorContext | EAutomationTestFlags::EngineFilter)
+	bool FMetasoundGeneratorSetInputBoolTest::RunTest(const FString&)
+	{
+		return RunSetInputTest<bool>(*this, true);
+	}
+
+	template<typename DataType>
+	bool RunApplyToInputValueTest(FAutomationTestBase& Test, TFunctionRef<void(DataType&)> InFunc, DataType ExpectedValue)
+	{
+		const FName InputName = "MyInput";
+		const FName OutputName = "MyOutput";
+		constexpr FSampleRate SampleRate = 48000;
+		constexpr int32 NumSamplesPerBlock = 128;
+
+		// Build a passthrough graph
+		const TUniquePtr<FMetasoundGenerator> Generator =
+			BuildPassthroughGraph<DataType>(Test, InputName, OutputName, SampleRate, NumSamplesPerBlock);
+		if (!Test.TestNotNull("Generator built", Generator.Get()))
+		{
+			return false;
+		}
+		
+		// Set the input value
+		Generator->ApplyToInputValue(InputName, InFunc);
+
+		// Render a block
+		{
+			TArray<float> Buffer;
+			Buffer.Reserve(NumSamplesPerBlock);
+			if (!Test.TestEqual(
+				"Generated the right number of samples.",
+				Generator->OnGenerateAudio(Buffer.GetData(), NumSamplesPerBlock),
+				NumSamplesPerBlock))
+			{
+				return false;
+			}
+		}
+
+		// Check the output
+		const TOptional<TDataReadReference<DataType>> OutputReadRef = Generator->GetOutputReadReference<DataType>(OutputName);
+		if (!Test.TestTrue("Got the read ref", OutputReadRef.IsSet()))
+		{
+			return false;
+		}
+		if (!Test.TestEqual("Input was passed through", *OutputReadRef.GetValue(), ExpectedValue))
+		{
+			return false;
+		}
+
+		return true;
+	}
+
+	IMPLEMENT_SIMPLE_AUTOMATION_TEST(
+	FMetasoundGeneratorApplyToInputFloatTest,
+	"Audio.Metasound.Generator.ApplyToInput.Float",
+	EAutomationTestFlags::EditorContext | EAutomationTestFlags::EngineFilter)
+	bool FMetasoundGeneratorApplyToInputFloatTest::RunTest(const FString&)
+	{
+		using DataType = float;
+		constexpr DataType ExpectedValue = 123.456f;
+		const TUniqueFunction<void(DataType&)> Fn = [ExpectedValue](DataType& Value)
+		{
+			Value = ExpectedValue;
+		};
+		return RunApplyToInputValueTest<DataType>(*this, Fn, ExpectedValue);
+	}
+
+	IMPLEMENT_SIMPLE_AUTOMATION_TEST(
+	FMetasoundGeneratorApplyToInputIntTest,
+	"Audio.Metasound.Generator.ApplyToInput.Int",
+	EAutomationTestFlags::EditorContext | EAutomationTestFlags::EngineFilter)
+	bool FMetasoundGeneratorApplyToInputIntTest::RunTest(const FString&)
+	{
+		using DataType = int32;
+		constexpr DataType ExpectedValue = 123456;
+		const TUniqueFunction<void(DataType&)> Fn = [ExpectedValue](DataType& Value)
+		{
+			Value = ExpectedValue;
+		};
+		return RunApplyToInputValueTest<DataType>(*this, Fn, ExpectedValue);
+	}
+	
+	IMPLEMENT_SIMPLE_AUTOMATION_TEST(
+	FMetasoundGeneratorApplyToInputBoolTest,
+	"Audio.Metasound.Generator.ApplyToInput.Bool",
+	EAutomationTestFlags::EditorContext | EAutomationTestFlags::EngineFilter)
+	bool FMetasoundGeneratorApplyToInputBoolTest::RunTest(const FString&)
+	{
+		using DataType = bool;
+		constexpr DataType ExpectedValue = true;
+		const TUniqueFunction<void(DataType&)> Fn = [ExpectedValue](DataType& Value)
+		{
+			Value = ExpectedValue;
+		};
+		return RunApplyToInputValueTest<DataType>(*this, Fn, ExpectedValue);
+	}
+}
+
+#endif
\ No newline at end of file
diff --git a/Engine/Plugins/Runtime/Metasound/Source/MetasoundEngineTest/Private/NodeTestGraphBuilder.cpp b/Engine/Plugins/Runtime/Metasound/Source/MetasoundEngineTest/Private/NodeTestGraphBuilder.cpp
--- a/Engine/Plugins/Runtime/Metasound/Source/MetasoundEngineTest/Private/NodeTestGraphBuilder.cpp
+++ b/Engine/Plugins/Runtime/Metasound/Source/MetasoundEngineTest/Private/NodeTestGraphBuilder.cpp
@@ -83,6 +83,47 @@ namespace Metasound::Test
 		return nullptr;
 	}
 
+	TUniquePtr<FMetasoundGenerator> FNodeTestGraphBuilder::BuildGenerator(FSampleRate SampleRate, int32 SamplesPerBlock)
+	{
+		// Because a generator is tied to the concept of a source, go ahead and add the "OnPlay" input.
+		// Otherwise we get warnings when we run our tests.
+		AddInput(SourceInterface::Inputs::OnPlay, Metasound::GetMetasoundDataTypeName<FTrigger>());
+		
+		const TSet<FName> TransmittableInputNames;
+		const FString UnknownAsset = TEXT("UnknownAsset");
+		if (const TUniquePtr<FFrontendGraph> Graph = FFrontendGraphBuilder::CreateGraph(Document, TransmittableInputNames, UnknownAsset))
+		{
+			FOperatorBuilderSettings BuilderSettings;
+			BuilderSettings.bFailOnAnyError = true;
+			BuilderSettings.bPopulateInternalDataReferences = true;
+			BuilderSettings.bValidateEdgeDataTypesMatch = true;
+			BuilderSettings.bValidateNoCyclesInGraph = true;
+			BuilderSettings.bValidateNoDuplicateInputs = true;
+			BuilderSettings.bValidateOperatorOutputsAreBound = true;
+			BuilderSettings.bValidateVerticesExist = true;
+
+			const FOperatorSettings OperatorSettings{SampleRate, static_cast<float>(SampleRate) / SamplesPerBlock};
+
+			FMetasoundEnvironment Environment;
+			Environment.SetValue<uint64>(SourceInterface::Environment::TransmitterID, 123);
+			
+			FMetasoundGeneratorInitParams GeneratorInitParams{
+				OperatorSettings,
+				BuilderSettings,
+				MakeShared<const FFrontendGraph, ESPMode::ThreadSafe>(*Graph),
+				Environment,
+				"TestMetasound",
+				TArray<FVertexName>(),
+				TArray<FAudioParameter>(),
+				true // bBuildSynchronous, so we don't have to do latent tasks
+			};
+
+			return MakeUnique<FMetasoundGenerator>(MoveTemp(GeneratorInitParams));
+		}
+
+		return nullptr;
+	}
+
 	TUniquePtr<IOperator> FNodeTestGraphBuilder::MakeSingleNodeGraph(
 		const FNodeClassName& ClassName,
 		int32 MajorVersion,
diff --git a/Engine/Plugins/Runtime/Metasound/Source/MetasoundEngineTest/Public/NodeTestGraphBuilder.h b/Engine/Plugins/Runtime/Metasound/Source/MetasoundEngineTest/Public/NodeTestGraphBuilder.h
--- a/Engine/Plugins/Runtime/Metasound/Source/MetasoundEngineTest/Public/NodeTestGraphBuilder.h
+++ b/Engine/Plugins/Runtime/Metasound/Source/MetasoundEngineTest/Public/NodeTestGraphBuilder.h
@@ -4,6 +4,7 @@
 
 #include "MetasoundFrontendController.h"
 #include "MetasoundFrontendDocument.h"
+#include "MetasoundGenerator.h"
 #include "MetasoundNodeInterface.h"
 #include "MetasoundOperatorInterface.h"
 
@@ -27,6 +28,8 @@ namespace Metasound::Test
 		/** Build the graph and get the graph's operator to work with */
 		TUniquePtr<IOperator> BuildGraph(FSampleRate SampleRate = 48000, int32 SamplesPerBlock = 256);
 
+		TUniquePtr<FMetasoundGenerator> BuildGenerator(FSampleRate SampleRate = 48000, int32 SamplesPerBlock = 256);
+
 		/** Helper that will add a single node, wire up the node's inputs and outputs, and hand back the graph's operator */
 		static TUniquePtr<IOperator> MakeSingleNodeGraph(
 			const FNodeClassName& ClassName,
diff --git a/Engine/Plugins/Runtime/Metasound/Source/MetasoundGenerator/Private/MetasoundGenerator.cpp b/Engine/Plugins/Runtime/Metasound/Source/MetasoundGenerator/Private/MetasoundGenerator.cpp
--- a/Engine/Plugins/Runtime/Metasound/Source/MetasoundGenerator/Private/MetasoundGenerator.cpp
+++ b/Engine/Plugins/Runtime/Metasound/Source/MetasoundGenerator/Private/MetasoundGenerator.cpp
@@ -145,7 +145,55 @@ namespace Metasound
 			// Collect data for generator
 			FMetasoundGeneratorData GeneratorData = BuildGeneratorData(InitParams, MoveTemp(GraphOperator), MoveTemp(GraphAnalyzer));
 
-			// Update generator with new metasound graph
+			// Gather relevant input and output references
+			FVertexInterfaceData VertexData(InitParams.Graph->GetVertexInterface());
+			GraphOperator->Bind(VertexData);
+
+			// Get inputs
+			FTriggerWriteRef PlayTrigger = VertexData.GetInputs().GetOrConstructDataWriteReference<FTrigger>(SourceInterface::Inputs::OnPlay, InitParams.OperatorSettings, false);
+
+			// Get outputs
+			TArray<FAudioBufferReadRef> OutputBuffers = FindOutputAudioBuffers(VertexData);
+			FTriggerReadRef FinishTrigger = TDataReadReferenceFactory<FTrigger>::CreateExplicitArgs(InitParams.OperatorSettings, false);
+
+			if (InitParams.Graph->GetVertexInterface().GetOutputInterface().Contains(SourceOneShotInterface::Outputs::OnFinished))
+			{
+				// Only attempt to retrieve the on finished trigger if it exists.
+				// Attempting to retrieve a data reference from a non-existent vertex 
+				// will log an error. 
+				FinishTrigger = VertexData.GetOutputs().GetOrConstructDataReadReference<FTrigger>(SourceOneShotInterface::Outputs::OnFinished, InitParams.OperatorSettings, false);
+			}
+
+			// Create the parameter setter map so parameter packs can be cracked
+			// open and distributed as appropriate...
+			TMap<FName, FParameterSetter> ParameterSetters;
+			FInputVertexInterfaceData& GraphInputs = VertexData.GetInputs();
+			for (auto InputIterator = GraphInputs.begin(); InputIterator != GraphInputs.end(); ++InputIterator)
+			{
+				const FInputDataVertex& InputVertex = (*InputIterator).GetVertex();
+				const Frontend::IParameterAssignmentFunction& Setter = IDataTypeRegistry::Get().GetRawAssignmentFunction(InputVertex.DataTypeName);
+				if (Setter)
+				{
+					FParameterSetter ParameterSetter(InputVertex.DataTypeName,
+						(*InputIterator).GetDataReference()->GetRaw(),
+						Setter);
+					ParameterSetters.Add(InputVertex.VertexName, ParameterSetter);
+				}
+			}
+
+			// Set data needed for graph
+			FMetasoundGeneratorData GeneratorData
+			{
+				InitParams.OperatorSettings,
+				MoveTemp(GraphOperator),
+				MoveTemp(VertexData),
+				MoveTemp(ParameterSetters),
+				MoveTemp(GraphAnalyzer),
+				MoveTemp(OutputBuffers),
+				MoveTemp(PlayTrigger),
+				MoveTemp(FinishTrigger),
+			};
+
 			Generator->SetPendingGraph(MoveTemp(GeneratorData), bTriggerGenerator);
 		}
 		else 
@@ -326,7 +374,7 @@ namespace Metasound
 
 		BuilderTask = MakeUnique<FBuilderTask>(this, MoveTemp(InParams), true /* bTriggerGenerator */);
 		
-		if (Metasound::ConsoleVariables::bEnableAsyncMetaSoundGeneratorBuilder)
+		if (!InParams.bBuildSynchronous && Metasound::ConsoleVariables::bEnableAsyncMetaSoundGeneratorBuilder)
 		{
 			// Build operator asynchronously
 			BuilderTask->StartBackgroundTask(GBackgroundPriorityThreadPool);
@@ -396,6 +444,9 @@ namespace Metasound
 
 		InterleavedAudioBuffer.Reset();
 
+		// Copy off all vertex interface data
+		VertexInterfaceData = MoveTemp(InData->VertexInterfaceData);
+
 		GraphOutputAudio.Reset();
 		if (InData->OutputBuffers.Num() == NumChannels)
 		{
diff --git a/Engine/Plugins/Runtime/Metasound/Source/MetasoundGenerator/Public/MetasoundGenerator.h b/Engine/Plugins/Runtime/Metasound/Source/MetasoundGenerator/Public/MetasoundGenerator.h
--- a/Engine/Plugins/Runtime/Metasound/Source/MetasoundGenerator/Public/MetasoundGenerator.h
+++ b/Engine/Plugins/Runtime/Metasound/Source/MetasoundGenerator/Public/MetasoundGenerator.h
@@ -33,6 +33,7 @@ namespace Metasound
 		FString MetaSoundName;
 		TArray<FVertexName> AudioOutputNames;
 		TArray<FAudioParameter> DefaultParameters;
+		bool bBuildSynchronous = false;
 
 		void Release();
 	};
@@ -58,6 +59,7 @@ namespace Metasound
 	{
 		FOperatorSettings OperatorSettings;
 		TUniquePtr<IOperator> GraphOperator;
+		FVertexInterfaceData VertexInterfaceData;
 		TMap<FName, FParameterSetter> ParameterSetters;
 		TUniquePtr<Frontend::FGraphAnalyzer> GraphAnalyzer;
 		TArray<TDataReadReference<FAudioBuffer>> OutputBuffers;
@@ -119,20 +121,11 @@ namespace Metasound
 		 * @param InData - The value to assign.
 		 */
 		template<typename DataType>
-		void SetInputValue(const FString& InName, DataType InData)
+		void SetInputValue(const FVertexName& InName, DataType InData)
 		{
-			typedef TDataWriteReference< typename TDecay<DataType>::Type > FDataWriteRef;
-
-			const FDataReferenceCollection& InputCollection = RootExecuter.GetInputs();
-
-			// Check if an input data reference with the given name and type exist in the graph.
-			bool bContainsWriteReference = InputCollection.ContainsDataWriteReference< typename TDecay<DataType>::Type >(InName);
-			if (ensureMsgf(bContainsWriteReference, TEXT("Operator does not contain write reference name \"%s\" of type \"%s\""), *InName, *GetMetasoundDataTypeName<DataType>().ToString()))
+			if (const FAnyDataReference* Ref = VertexInterfaceData.GetInputs().FindDataReference(InName))
 			{
-				FDataWriteRef WriteRef = InputCollection.GetDataWriteReference<DataType>(InName);
-
-				// call assignment operator of DataType.
-				*WriteRef = InData;
+				*(Ref->GetDataWriteReference<typename TDecay<DataType>::Type>()) = InData;
 			}
 		}
 
@@ -142,26 +135,58 @@ namespace Metasound
 		 * @param InFunc - A function which takes the DataType as an input.
 		 */ 
 		template<typename DataType>
-		void ApplyToInputValue(const FString& InName, TUniqueFunction<void(DataType&)> InFunc)
+		void ApplyToInputValue(const FVertexName& InName, TFunctionRef<void(DataType&)> InFunc)
 		{
-			// Get decayed type as InFunc could take a const qualified type.
-			typedef TDataWriteReference< typename TDecay<DataType>::Type > FDataWriteRef;
-
-			const FDataReferenceCollection& InputCollection = RootExecuter.GetInputs();
-
-			// Check if an input data reference with the given name and type exists in the graph.
-			bool bContainsWriteReference = InputCollection.ContainsDataWriteReference< typename TDecay<DataType>::Type >(InName);
-			if (ensureMsgf(bContainsWriteReference, TEXT("Operator does not contain write reference name \"%s\" of type \"%s\""), *InName, *GetMetasoundDataTypeName<DataType>().ToString()))
+			if (const FAnyDataReference* Ref = VertexInterfaceData.GetInputs().FindDataReference(InName))
 			{
-				FDataWriteRef WriteRef = InputCollection.GetDataWriteReference<DataType>(InName);
-
-				// Apply function to DataType
-				InFunc(*WriteRef);
+				InFunc(*(Ref->GetDataWriteReference<typename TDecay<DataType>::Type>()));
 			}
 		}
 
 		void QueueParameterPack(TSharedPtr<FMetasoundParameterPackStorage> ParameterPack);
 
+		/**
+		 * Get a write reference to one of the generator's inputs, if it exists.
+		 * NOTE: This reference is only safe to use immediately on the same thread that this generator's
+		 * OnGenerateAudio() is called.
+		 *
+		 * @tparam DataType - The expected data type of the input
+		 * @param InputName - The user-defined name of the input
+		 */
+		template<typename DataType>
+		TOptional<TDataWriteReference<DataType>> GetInputWriteReference(const FVertexName InputName)
+		{
+			TOptional<TDataWriteReference<DataType>> WriteRef;
+			
+			if (const FAnyDataReference* Ref = VertexInterfaceData.GetInputs().FindDataReference(InputName))
+			{
+				WriteRef = Ref->GetDataWriteReference<typename TDecay<DataType>::Type>();
+			}
+			
+			return WriteRef;
+		}
+		
+		/**
+		 * Get a read reference to one of the generator's outputs, if it exists.
+		 * NOTE: This reference is only safe to use immediately on the same thread that this generator's
+		 * OnGenerateAudio() is called.
+		 *
+		 * @tparam DataType - The expected data type of the output
+		 * @param OutputName - The user-defined name of the output
+		 */
+		template<typename DataType>
+		TOptional<TDataReadReference<DataType>> GetOutputReadReference(const FVertexName OutputName)
+		{
+			TOptional<TDataReadReference<DataType>> ReadRef;
+
+			if (const FAnyDataReference* Ref = VertexInterfaceData.GetOutputs().FindDataReference(OutputName))
+			{
+				ReadRef = Ref->GetDataReadReference<typename TDecay<DataType>::Type>();
+			}
+			
+			return ReadRef;
+		}
+		
 		/** Return the number of audio channels. */
 		int32 GetNumChannels() const;
 
@@ -197,6 +222,7 @@ namespace Metasound
 		void UnpackAndTransmitUpdatedParameters();
 
 		FExecuter RootExecuter;
+		FVertexInterfaceData VertexInterfaceData;
 
 		bool bIsGraphBuilding;
 		bool bIsFinishTriggered;
