Fix crash in metasound output node
#rnx

diff --git a/Engine/Plugins/Runtime/Metasound/Source/MetasoundFrontend/Private/MetasoundOutputNode.cpp b/Engine/Plugins/Runtime/Metasound/Source/MetasoundFrontend/Private/MetasoundOutputNode.cpp
--- a/Engine/Plugins/Runtime/Metasound/Source/MetasoundFrontend/Private/MetasoundOutputNode.cpp
+++ b/Engine/Plugins/Runtime/Metasound/Source/MetasoundFrontend/Private/MetasoundOutputNode.cpp
@@ -75,7 +75,7 @@ namespace Metasound
 					TOptional<FAnyDataReference> DataReference = Entry->CreateDataReference(EDataReferenceAccessType::Value, DefaultLiteral, InParams.OperatorSettings);
 					if (DataReference.IsSet())
 					{
-						return MakeUnique<FOutputOperator>(VertexName, *Ref);
+						return MakeUnique<FOutputOperator>(VertexName, *DataReference);
 					}
 				}
 			}
