Fix ensure when creating variables on a MetaSound graph that has been converted from preset

diff --git a/Engine/Plugins/Runtime/Metasound/Source/MetasoundFrontend/Private/MetasoundFrontendTransform.cpp b/Engine/Plugins/Runtime/Metasound/Source/MetasoundFrontend/Private/MetasoundFrontendTransform.cpp
--- a/Engine/Plugins/Runtime/Metasound/Source/MetasoundFrontend/Private/MetasoundFrontendTransform.cpp
+++ b/Engine/Plugins/Runtime/Metasound/Source/MetasoundFrontend/Private/MetasoundFrontendTransform.cpp
@@ -711,7 +711,12 @@ namespace Metasound
 			PresetGraphHandle->IterateConstNodes([InPresetNodeID = &PresetNodeID](FConstNodeHandle PresetNodeHandle)
 			{
 				*InPresetNodeID = PresetNodeHandle->GetID();
-			}, EMetasoundFrontendClassType::External);
+			}, EMetasoundFrontendClassType::External); 
+
+			if (!PresetNodeID.IsValid())
+			{
+				PresetNodeID = FGuid::NewGuid();
+			}
 
 			// Clear the root graph so it can be rebuilt.
 			PresetGraphHandle->ClearGraph();
