Quick fix for automated test build failures by removing object creation from MS Builders on default ctor/copy ctor/copy assignment operators. Have a longer term fix in the works that's a bit safer, but in general builders are transient and should never be copied.
#rnx

diff --git a/Engine/Plugins/Runtime/Metasound/Source/MetasoundFrontend/Private/MetasoundFrontendDocumentBuilder.cpp b/Engine/Plugins/Runtime/Metasound/Source/MetasoundFrontend/Private/MetasoundFrontendDocumentBuilder.cpp
--- a/Engine/Plugins/Runtime/Metasound/Source/MetasoundFrontend/Private/MetasoundFrontendDocumentBuilder.cpp
+++ b/Engine/Plugins/Runtime/Metasound/Source/MetasoundFrontend/Private/MetasoundFrontendDocumentBuilder.cpp
@@ -50,9 +50,7 @@ namespace Metasound::Frontend
 } // namespace Metasound::Frontend
 
 FMetaSoundFrontendDocumentBuilder::FMetaSoundFrontendDocumentBuilder()
-	: DocumentInterface(NewObject<UMetaSoundBuilderDocument>())
 {
-	ReloadCache();
 }
 
 FMetaSoundFrontendDocumentBuilder::FMetaSoundFrontendDocumentBuilder(TScriptInterface<IMetaSoundDocumentInterface> InDocumentInterface)
@@ -63,11 +61,11 @@ FMetaSoundFrontendDocumentBuilder::FMetaSoundFrontendDocumentBuilder(TScriptInte
 
 FMetaSoundFrontendDocumentBuilder::FMetaSoundFrontendDocumentBuilder(const FMetaSoundFrontendDocumentBuilder& InBuilder)
 {
-	// A new object doc has to be created when copied to avoid multiple builders/caches operating on a single document.
-	UMetaSoundBuilderDocument* NewDoc = NewObject<UMetaSoundBuilderDocument>();
-	NewDoc->Document = InBuilder.GetDocument();
-	DocumentInterface = NewDoc;
-	ReloadCache();
+	DocumentInterface = InBuilder.DocumentInterface;
+	if (DocumentInterface)
+	{
+		ReloadCache();
+	}	
 }
 
 FMetaSoundFrontendDocumentBuilder::FMetaSoundFrontendDocumentBuilder(FMetaSoundFrontendDocumentBuilder&& InBuilder)
@@ -88,11 +86,11 @@ FMetaSoundFrontendDocumentBuilder& FMetaSoundFrontendDocumentBuilder::operator=(
 {
 	using namespace Metasound::Frontend;
 
-	// A new object doc has to be created when copied to avoid multiple builders/caches operating on a single document.
-	UMetaSoundBuilderDocument* NewDoc = NewObject<UMetaSoundBuilderDocument>();
-	NewDoc->Document = InRHS.GetDocument();
-	DocumentInterface = NewDoc;
-	ReloadCache();
+	DocumentInterface = InRHS.DocumentInterface;
+	if (DocumentInterface)
+	{
+		ReloadCache();
+	}	
 	return *this;
 }
 
