[Metasound] TVertexAnalyzerForwardValue - Remember last value to prevent unnecessary data pushes.

diff --git a/Engine/Plugins/Runtime/Metasound/Source/MetasoundFrontend/Public/Analysis/MetasoundFrontendVertexAnalyzerForwardValue.h b/Engine/Plugins/Runtime/Metasound/Source/MetasoundFrontend/Public/Analysis/MetasoundFrontendVertexAnalyzerForwardValue.h
--- a/Engine/Plugins/Runtime/Metasound/Source/MetasoundFrontend/Public/Analysis/MetasoundFrontendVertexAnalyzerForwardValue.h
+++ b/Engine/Plugins/Runtime/Metasound/Source/MetasoundFrontend/Public/Analysis/MetasoundFrontendVertexAnalyzerForwardValue.h
@@ -69,6 +69,7 @@ namespace Metasound
 				const DataType& Value = GetVertexData<DataType>();
 				if (LastValue != Value)
 				{
+					LastValue = Value;
 					MarkOutputDirty();
 				}
 			}
