[Metasound] Add/remove vertex analyzers on generator outputs.

diff --git a/Engine/Plugins/Runtime/Metasound/Source/MetasoundEngineTest/Private/MetasoundGeneratorTests.cpp b/Engine/Plugins/Runtime/Metasound/Source/MetasoundEngineTest/Private/MetasoundGeneratorTests.cpp
--- a/Engine/Plugins/Runtime/Metasound/Source/MetasoundEngineTest/Private/MetasoundGeneratorTests.cpp
+++ b/Engine/Plugins/Runtime/Metasound/Source/MetasoundEngineTest/Private/MetasoundGeneratorTests.cpp
@@ -1,6 +1,7 @@
 // Copyright Epic Games, Inc. All Rights Reserved.
 
 #include "NodeTestGraphBuilder.h"
+#include "Analysis/MetasoundFrontendAnalyzerView.h"
 #include "Misc/AutomationTest.h"
 
 #if WITH_DEV_AUTOMATION_TESTS
@@ -16,17 +17,28 @@ namespace Metasound::Test::Generator
 		const FName& InputName,
 		const FName& OutputName,
 		const FSampleRate SampleRate,
-		const int32 NumSamplesPerBlock)
+		const int32 NumSamplesPerBlock,
+		FGuid* OutputGuid = nullptr)
 	{
 		GraphBuilder Builder;
 		const FNodeHandle InputNode = Builder.AddInput(InputName, Metasound::GetMetasoundDataTypeName<DataType>());
 		const FNodeHandle OutputNode = Builder.AddOutput(OutputName, Metasound::GetMetasoundDataTypeName<DataType>());
 		const FOutputHandle OutputToConnect = InputNode->GetOutputWithVertexName(InputName);
 		const FInputHandle InputToConnect = OutputNode->GetInputWithVertexName(OutputName);
+
 		if (!Test.TestTrue("Connected input to output", InputToConnect->Connect(*OutputToConnect)))
 		{
 			return nullptr;
 		}
+
+		if (nullptr != OutputGuid)
+		{
+			*OutputGuid = OutputNode->GetID();
+		}
+
+		// have to add an audio output for the generator to render
+		Builder.AddOutput("Audio", Metasound::GetMetasoundDataTypeName<FAudioBuffer>());
+		
 		return Builder.BuildGenerator(SampleRate, NumSamplesPerBlock);
 	}
 	
@@ -271,6 +283,98 @@ namespace Metasound::Test::Generator
 		};
 		return RunApplyToInputValueTest<DataType>(*this, Fn, ExpectedValue);
 	}
+
+	template<typename DataType>
+	bool RunOutputAnalyzerForwardValueTest(FAutomationTestBase& Test, DataType ExpectedValue, FName AnalyzerName)
+	{
+		const FName InputName = "MyInput";
+		const FName OutputName = "MyOutput";
+		constexpr FSampleRate SampleRate = 48000;
+		constexpr int32 NumSamplesPerBlock = 128;
+
+		// Build a passthrough graph
+		FGuid OutputNodeId;
+		const TUniquePtr<FMetasoundGenerator> Generator =
+			BuildPassthroughGraph<DataType>(Test, InputName, OutputName, SampleRate, NumSamplesPerBlock, &OutputNodeId);
+		if (!Test.TestNotNull("Generator built", Generator.Get()))
+		{
+			return false;
+		}
+
+		// Add an analyzer to the output
+		FAnalyzerAddress AnalyzerAddress;
+		AnalyzerAddress.DataType = GetMetasoundDataTypeName<DataType>();
+		AnalyzerAddress.InstanceID = 1234;
+		AnalyzerAddress.OutputName = OutputName;
+		AnalyzerAddress.AnalyzerName = AnalyzerName;
+		AnalyzerAddress.AnalyzerInstanceID = FGuid::NewGuid();
+		AnalyzerAddress.NodeID = OutputNodeId;
+		Generator->AddOutputVertexAnalyzer(AnalyzerAddress);
+
+
+		// Add an analyzer view to watch the output
+		FMetasoundAnalyzerView AnalyzerView{ MoveTemp(AnalyzerAddress) };
+		AnalyzerView.BindToAllOutputs(Generator->OperatorSettings);
+		
+		// Set the input value
+		Generator->SetInputValue(InputName, ExpectedValue);
+
+		// Render a block
+		{
+			TArray<float> Buffer;
+			Buffer.Reserve(NumSamplesPerBlock);
+			if (!Test.TestEqual(
+				"Generated the right number of samples.",
+				Generator->OnGenerateAudio(Buffer.GetData(), NumSamplesPerBlock),
+				NumSamplesPerBlock))
+			{
+				return false;
+			}
+		}
+
+		// Check the output
+		DataType Value;
+
+		if (!Test.TestTrue("Got output data.", AnalyzerView.TryGetOutputData("Value", Value)))
+		{
+			return false;
+		}
+
+		return Test.TestEqual("Input was passed through", Value, ExpectedValue);
+	}
+
+	IMPLEMENT_SIMPLE_AUTOMATION_TEST(
+	FMetasoundGeneratorOututAnalyzerForwardValueFloatTest,
+	"Audio.Metasound.Generator.OutputAnalyzer.ForwardValue.Float",
+	EAutomationTestFlags::EditorContext | EAutomationTestFlags::EngineFilter)
+	bool FMetasoundGeneratorOututAnalyzerForwardValueFloatTest::RunTest(const FString&)
+	{
+		using DataType = float;
+		constexpr DataType ExpectedValue = 678.345f;
+		return RunOutputAnalyzerForwardValueTest<DataType>(*this, ExpectedValue, "UE.Forward.Float");
+	}
+
+	IMPLEMENT_SIMPLE_AUTOMATION_TEST(
+	FMetasoundGeneratorOututAnalyzerForwardValueIntTest,
+	"Audio.Metasound.Generator.OutputAnalyzer.ForwardValue.Int",
+	EAutomationTestFlags::EditorContext | EAutomationTestFlags::EngineFilter)
+	bool FMetasoundGeneratorOututAnalyzerForwardValueIntTest::RunTest(const FString&)
+	{
+		using DataType = int32;
+		constexpr DataType ExpectedValue = 678345;
+		return RunOutputAnalyzerForwardValueTest<DataType>(*this, ExpectedValue, "UE.Forward.Int32");
+	}
+
+	IMPLEMENT_SIMPLE_AUTOMATION_TEST(
+	FMetasoundGeneratorOututAnalyzerForwardValueBoolTest,
+	"Audio.Metasound.Generator.OutputAnalyzer.ForwardValue.Bool",
+	EAutomationTestFlags::EditorContext | EAutomationTestFlags::EngineFilter)
+	bool FMetasoundGeneratorOututAnalyzerForwardValueBoolTest::RunTest(const FString&)
+	{
+		using DataType = bool;
+		constexpr DataType ExpectedValue = true;
+		return RunOutputAnalyzerForwardValueTest<DataType>(*this, ExpectedValue, "UE.Forward.Bool");
+	}
 }
 
 #endif
\ No newline at end of file
diff --git a/Engine/Plugins/Runtime/Metasound/Source/MetasoundEngineTest/Private/NodeTestGraphBuilder.cpp b/Engine/Plugins/Runtime/Metasound/Source/MetasoundEngineTest/Private/NodeTestGraphBuilder.cpp
--- a/Engine/Plugins/Runtime/Metasound/Source/MetasoundEngineTest/Private/NodeTestGraphBuilder.cpp
+++ b/Engine/Plugins/Runtime/Metasound/Source/MetasoundEngineTest/Private/NodeTestGraphBuilder.cpp
@@ -5,6 +5,7 @@
 #include "MetasoundFrontendGraph.h"
 #include "MetasoundFrontendSearchEngine.h"
 #include "MetasoundOperatorBuilder.h"
+#include "Interfaces/MetasoundOutputFormatInterfaces.h"
 
 namespace Metasound::Test
 {
@@ -53,7 +54,15 @@ namespace Metasound::Test
 		Output.Name = OutputName;
 		Output.TypeName = TypeName;
 		Output.VertexID = FGuid::NewGuid();
-		return RootGraph->AddOutputVertex(Output);
+		FNodeHandle NodeHandle = RootGraph->AddOutputVertex(Output);
+
+		static const FName AudioBufferTypeName = GetMetasoundDataTypeName<FAudioBuffer>();
+		if (TypeName == AudioBufferTypeName)
+		{
+			AudioOutputNames.Add(OutputName);
+		}
+
+		return NodeHandle;
 	}
 
 	TUniquePtr<IOperator> FNodeTestGraphBuilder::BuildGraph(FSampleRate SampleRate, int32 SamplesPerBlock)
@@ -113,7 +122,7 @@ namespace Metasound::Test
 				MakeShared<const FFrontendGraph, ESPMode::ThreadSafe>(*Graph),
 				Environment,
 				"TestMetasound",
-				TArray<FVertexName>(),
+				AudioOutputNames,
 				TArray<FAudioParameter>(),
 				true // bBuildSynchronous, so we don't have to do latent tasks
 			};
diff --git a/Engine/Plugins/Runtime/Metasound/Source/MetasoundEngineTest/Public/NodeTestGraphBuilder.h b/Engine/Plugins/Runtime/Metasound/Source/MetasoundEngineTest/Public/NodeTestGraphBuilder.h
--- a/Engine/Plugins/Runtime/Metasound/Source/MetasoundEngineTest/Public/NodeTestGraphBuilder.h
+++ b/Engine/Plugins/Runtime/Metasound/Source/MetasoundEngineTest/Public/NodeTestGraphBuilder.h
@@ -41,5 +41,6 @@ namespace Metasound::Test
 		FMetasoundFrontendDocument Document;
 		Frontend::FDocumentHandle DocumentHandle = Frontend::IDocumentController::GetInvalidHandle();
 		Frontend::FGraphHandle RootGraph = Frontend::IGraphController::GetInvalidHandle();
+		TArray<FVertexName> AudioOutputNames;
 	};
 }
diff --git a/Engine/Plugins/Runtime/Metasound/Source/MetasoundGenerator/Private/MetasoundGenerator.cpp b/Engine/Plugins/Runtime/Metasound/Source/MetasoundGenerator/Private/MetasoundGenerator.cpp
--- a/Engine/Plugins/Runtime/Metasound/Source/MetasoundGenerator/Private/MetasoundGenerator.cpp
+++ b/Engine/Plugins/Runtime/Metasound/Source/MetasoundGenerator/Private/MetasoundGenerator.cpp
@@ -1,10 +1,8 @@
 // Copyright Epic Games, Inc. All Rights Reserved.
 #include "MetasoundGenerator.h"
 
-#include "Analysis/MetasoundFrontendVertexAnalyzer.h"
 #include "AudioParameter.h"
-#include "DSP/FloatArrayMath.h"
-#include "Interfaces/MetasoundFrontendSourceInterface.h"
+#include "MetasoundFrontendDataTypeRegistry.h"
 #include "MetasoundOperatorBuilder.h"
 #include "MetasoundOperatorInterface.h"
 #include "MetasoundOutputNode.h"
@@ -13,7 +11,11 @@
 #include "MetasoundTrace.h"
 #include "MetasoundTrigger.h"
 #include "MetasoundVertexData.h"
-#include "MetasoundFrontendDataTypeRegistry.h"
+#include "Analysis/MetasoundFrontendAnalyzerFactory.h"
+#include "Analysis/MetasoundFrontendAnalyzerRegistry.h"
+#include "Analysis/MetasoundFrontendVertexAnalyzer.h"
+#include "DSP/FloatArrayMath.h"
+#include "Interfaces/MetasoundFrontendSourceInterface.h"
 
 #ifndef METASOUNDGENERATOR_ENABLE_INVALID_SAMPLE_VALUE_LOGGING 
 #define METASOUNDGENERATOR_ENABLE_INVALID_SAMPLE_VALUE_LOGGING !UE_BUILD_SHIPPING
@@ -305,6 +307,7 @@ namespace Metasound
 
 	FMetasoundGenerator::FMetasoundGenerator(FMetasoundGeneratorInitParams&& InParams)
 		: MetasoundName(InParams.MetaSoundName)
+		, OperatorSettings(InParams.OperatorSettings)
 		, bIsFinishTriggered(false)
 		, bIsFinished(false)
 		, NumChannels(0)
@@ -464,6 +467,43 @@ namespace Metasound
 		return NumChannels;
 	}
 
+	void FMetasoundGenerator::AddOutputVertexAnalyzer(const Frontend::FAnalyzerAddress& AnalyzerAddress)
+	{
+		OutputAnalyzerModificationQueue.Enqueue([this, AnalyzerAddress]
+		{
+			const FAnyDataReference* OutputReference = VertexInterfaceData.GetOutputs().FindDataReference(AnalyzerAddress.OutputName);
+			
+			if (nullptr == OutputReference)
+			{
+				return;
+			}
+			
+			const Frontend::IVertexAnalyzerFactory* Factory = Frontend::IVertexAnalyzerRegistry::Get().FindAnalyzerFactory(AnalyzerAddress.AnalyzerName);
+			
+			if (nullptr == Factory)
+			{
+				return;
+			}
+			
+			TUniquePtr<Frontend::IVertexAnalyzer> Analyzer = Factory->CreateAnalyzer({ AnalyzerAddress, OperatorSettings, *OutputReference});
+			
+			if (nullptr == Analyzer)
+			{
+				return;
+			}
+			
+			OutputAnalyzers.Emplace(AnalyzerAddress.OutputName, MoveTemp(Analyzer));
+		});
+	}
+	
+	void FMetasoundGenerator::RemoveOutputVertexAnalyzer(const FName& OutputName)
+	{
+		OutputAnalyzerModificationQueue.Enqueue([this, OutputName]
+		{
+			OutputAnalyzers.Remove(OutputName);
+		});
+	}
+
 	int32 FMetasoundGenerator::OnGenerateAudio(float* OutAudio, int32 NumSamplesRemaining)
 	{
 		METASOUND_TRACE_CPUPROFILER_EVENT_SCOPE_TEXT(*FString::Printf(TEXT("MetasoundGenerator::OnGenerateAudio %s"), *MetasoundName));
@@ -498,6 +538,14 @@ namespace Metasound
 			return NumSamplesRemaining;
 		}
 
+		// Modify the output analyzers as needed
+		{
+			while (TOptional<TUniqueFunction<void()>> ModFn = OutputAnalyzerModificationQueue.Dequeue())
+			{
+				(*ModFn)();
+			}
+		}
+
 		// If we have any audio left in the internal overflow buffer from 
 		// previous calls, write that to the output before generating more audio.
 		int32 NumSamplesWritten = FillWithBuffer(OverflowBuffer, OutAudio, NumSamplesRemaining);
@@ -520,6 +568,12 @@ namespace Metasound
 				GraphAnalyzer->Execute();
 			}
 
+			// Execute the output analyzers
+			for (const auto& Pair : OutputAnalyzers)
+			{
+				Pair.Value->Execute();
+			}
+
 			// Check if generated finished during this execute call
 			if (*OnFinishedTriggerRef)
 			{
diff --git a/Engine/Plugins/Runtime/Metasound/Source/MetasoundGenerator/Public/MetasoundGenerator.h b/Engine/Plugins/Runtime/Metasound/Source/MetasoundGenerator/Public/MetasoundGenerator.h
--- a/Engine/Plugins/Runtime/Metasound/Source/MetasoundGenerator/Public/MetasoundGenerator.h
+++ b/Engine/Plugins/Runtime/Metasound/Source/MetasoundGenerator/Public/MetasoundGenerator.h
@@ -1,25 +1,19 @@
 // Copyright Epic Games, Inc. All Rights Reserved.
 #pragma once
 
-#include "Analysis/MetasoundFrontendGraphAnalyzer.h"
-#include "Async/Async.h"
-#include "Async/AsyncWork.h"
-#include "Containers/MpscQueue.h"
-#include "DSP/Dsp.h"
-#include "MetasoundAudioFormats.h"
 #include "MetasoundExecutableOperator.h"
-#include "MetasoundFrontendController.h"
 #include "MetasoundGraphOperator.h"
 #include "MetasoundOperatorBuilder.h"
 #include "MetasoundOperatorInterface.h"
 #include "MetasoundParameterPack.h"
-#include "MetasoundPrimitives.h"
 #include "MetasoundRouter.h"
 #include "MetasoundTrigger.h"
 #include "MetasoundVertex.h"
 #include "MetasoundVertexData.h"
+#include "Analysis/MetasoundFrontendGraphAnalyzer.h"
+#include "Async/AsyncWork.h"
+#include "Containers/MpscQueue.h"
 #include "Sound/SoundGenerator.h"
-#include "Tickable.h"
 
 namespace Metasound
 {
@@ -106,14 +100,15 @@ namespace Metasound
 
 		FString MetasoundName;
 
+		const FOperatorSettings OperatorSettings;
+
 		/** Create the generator with a graph operator and an output audio reference.
 		 *
-		 * @param InGraphOperator - Unique pointer to the IOperator which executes the entire graph.
-		 * @param InGraphOutputAudioRef - Read reference to the audio buffer filled by the InGraphOperator.
+		 * @param InParams - The generator initialization parameters
 		 */
-		FMetasoundGenerator(FMetasoundGeneratorInitParams&& InParams);
+		explicit FMetasoundGenerator(FMetasoundGeneratorInitParams&& InParams);
 
-		virtual ~FMetasoundGenerator();
+		virtual ~FMetasoundGenerator() override;
 
 		/** Set the value of a graph's input data using the assignment operator.
 		 *
@@ -186,6 +181,20 @@ namespace Metasound
 			
 			return ReadRef;
 		}
+
+		/**
+		 * Add a vertex analyzer for a named output with the given address info.
+		 *
+		 * @param AnalyzerAddress - Address information for the analyzer to use when making views
+		 */
+		void AddOutputVertexAnalyzer(const Frontend::FAnalyzerAddress& AnalyzerAddress);
+		
+		/**
+		 * Remove a vertex analyzer for a named output
+		 *
+		 * @param OutputName - The name of the output in the MetaSound graph
+		 */
+		void RemoveOutputVertexAnalyzer(const FName& OutputName);
 		
 		/** Return the number of audio channels. */
 		int32 GetNumChannels() const;
@@ -270,5 +279,8 @@ namespace Metasound
 		// a faster method of sending parameters is via the QueueParameterPack function 
 		// and this queue.
 		TMpscQueue<TSharedPtr<FMetasoundParameterPackStorage>> ParameterPackQueue;
+
+		TMpscQueue<TUniqueFunction<void()>> OutputAnalyzerModificationQueue;
+		TMap<FName, TUniquePtr<Frontend::IVertexAnalyzer>> OutputAnalyzers;
 	};
 }
