MetaSound WavePlayer node: Immediately fire OnFinished if WaveAsset passed in is None

diff --git a/Engine/Plugins/Runtime/Metasound/Source/MetasoundEngine/Private/MetasoundWavePlayerNode.cpp b/Engine/Plugins/Runtime/Metasound/Source/MetasoundEngine/Private/MetasoundWavePlayerNode.cpp
--- a/Engine/Plugins/Runtime/Metasound/Source/MetasoundEngine/Private/MetasoundWavePlayerNode.cpp
+++ b/Engine/Plugins/Runtime/Metasound/Source/MetasoundEngine/Private/MetasoundWavePlayerNode.cpp
@@ -529,7 +529,11 @@ namespace Metasound
 				// execute the next trigger
 				if (CurrAudioFrame == NextPlayFrame)
 				{
-					StartPlaying();
+					if (!StartPlaying())
+					{
+						TriggerOnDone->TriggerFrame(CurrAudioFrame);
+					}
+
 					++PlayTrigIndex;
 				}
 
@@ -635,7 +639,7 @@ namespace Metasound
 
 		// Start playing the current wave by creating a wave proxy reader and
 		// recreating the DSP stack.
-		void StartPlaying()
+		bool StartPlaying()
 		{
 			using namespace WavePlayerNodePrivate;
 			METASOUND_TRACE_CPUPROFILER_EVENT_SCOPE(Metasound::FWavePlayerOperator::StartPlaying);
@@ -725,6 +729,7 @@ namespace Metasound
 
 			// If everything was created successfully, start playing.
 			bIsPlaying = WaveProxyReader.IsValid() && ConvertDeinterleave.IsValid() && Resampler.IsValid();
+			return bIsPlaying;
 		}
 
 		/** Removes all samples from the source buffer and resets SourceState. */
