- Fix crash w/PostInitProperties initialization of MetaSoundDocumentBuilder
- Fix bugs with interfaces not registering with the correct UClass types
- Add missing RemoveByInterfaceBinding MS Builder function
- Builder cache fix for ContainsEdge
#rnx

diff --git a/Engine/Plugins/Runtime/Metasound/Source/MetasoundEngine/Private/MetasoundBuilderSubsystem.cpp b/Engine/Plugins/Runtime/Metasound/Source/MetasoundEngine/Private/MetasoundBuilderSubsystem.cpp
--- a/Engine/Plugins/Runtime/Metasound/Source/MetasoundEngine/Private/MetasoundBuilderSubsystem.cpp
+++ b/Engine/Plugins/Runtime/Metasound/Source/MetasoundEngine/Private/MetasoundBuilderSubsystem.cpp
@@ -1,6 +1,7 @@
 // Copyright Epic Games, Inc. All Rights Reserved.
 #include "MetasoundBuilderSubsystem.h"
 
+#include "Algo/Find.h"
 #include "Algo/Transform.h"
 #include "Components/AudioComponent.h"
 #include "HAL/IConsoleManager.h"
@@ -16,7 +17,6 @@
 #include "MetasoundSource.h"
 #include "MetasoundTrace.h"
 #include "MetasoundUObjectRegistry.h"
-#include "Algo/Find.h"
 
 #include UE_INLINE_GENERATED_CPP_BY_NAME(MetasoundBuilderSubsystem)
 
@@ -42,6 +42,18 @@ namespace Metasound::Engine
 			Literal.Set(Value);
 			return Literal;
 		}
+
+		template <typename BuilderClass>
+		BuilderClass& CreateTransientBuilder()
+		{
+			const EObjectFlags NewObjectFlags = RF_Public | RF_Transient;
+			UPackage* TransientPackage = GetTransientPackage();
+			const FName ObjectName = MakeUniqueObjectName(TransientPackage, BuilderClass::StaticClass());
+			TObjectPtr<BuilderClass> NewBuilder = NewObject<BuilderClass>(TransientPackage, ObjectName, NewObjectFlags);
+			check(NewBuilder);
+			NewBuilder->CreateTransientDocument();
+			return *NewBuilder.Get();
+		}
 	} // namespace BuilderSubsystemPrivate
 } // namespace Metasound::Engine
 
@@ -166,21 +178,21 @@ bool UMetaSoundBuilderBase::ContainsNodeOutput(const FMetaSoundBuilderNodeOutput
 	return Builder.FindNodeOutput(OutputHandle.NodeID, OutputHandle.VertexID) != nullptr;
 }
 
-void UMetaSoundBuilderBase::ConnectNodes(const FMetaSoundBuilderNodeOutputHandle& NodeOutput, const FMetaSoundBuilderNodeInputHandle& NodeInput, EMetaSoundBuilderResult& OutResult)
+void UMetaSoundBuilderBase::ConnectNodes(const FMetaSoundBuilderNodeOutputHandle& NodeOutputHandle, const FMetaSoundBuilderNodeInputHandle& NodeInputHandle, EMetaSoundBuilderResult& OutResult)
 {
 	const FMetasoundFrontendEdge* NewEdge = Builder.AddEdge(FMetasoundFrontendEdge
 	{
-		NodeOutput.NodeID,
-		NodeOutput.VertexID,
-		NodeInput.NodeID,
-		NodeInput.VertexID,
+		NodeOutputHandle.NodeID,
+		NodeOutputHandle.VertexID,
+		NodeInputHandle.NodeID,
+		NodeInputHandle.VertexID,
 	});
 	OutResult = NewEdge ? EMetaSoundBuilderResult::Succeeded : EMetaSoundBuilderResult::Failed;
 }
 
-void UMetaSoundBuilderBase::ConnectNodesByInterfaceBindings(const FMetaSoundNodeHandle& FromNode, const FMetaSoundNodeHandle& ToNode, EMetaSoundBuilderResult& OutResult)
+void UMetaSoundBuilderBase::ConnectNodesByInterfaceBindings(const FMetaSoundNodeHandle& FromNodeHandle, const FMetaSoundNodeHandle& ToNodeHandle, EMetaSoundBuilderResult& OutResult)
 {
-	const bool bEdgesAdded = Builder.AddEdgesByNodeClassInterfaceBindings(FromNode.NodeID, ToNode.NodeID);
+	const bool bEdgesAdded = Builder.AddEdgesByNodeClassInterfaceBindings(FromNodeHandle.NodeID, ToNodeHandle.NodeID);
 	OutResult = bEdgesAdded ? EMetaSoundBuilderResult::Succeeded : EMetaSoundBuilderResult::Failed;
 }
 
@@ -218,7 +230,7 @@ TArray<FMetaSoundBuilderNodeOutputHandle> UMetaSoundBuilderBase::ConnectNodeInpu
 	return ConnectedVertices;
 }
 
-void UMetaSoundBuilderBase::ConnectNodeOutputToGraphOutput(FName GraphOutputName, const FMetaSoundBuilderNodeOutputHandle& NodeOutput, EMetaSoundBuilderResult& OutResult)
+void UMetaSoundBuilderBase::ConnectNodeOutputToGraphOutput(FName GraphOutputName, const FMetaSoundBuilderNodeOutputHandle& NodeOutputHandle, EMetaSoundBuilderResult& OutResult)
 {
 	OutResult = EMetaSoundBuilderResult::Failed;
 
@@ -227,8 +239,8 @@ void UMetaSoundBuilderBase::ConnectNodeOutputToGraphOutput(FName GraphOutputName
 		const FMetasoundFrontendVertex& InputVertex = GraphOutputNode->Interface.Inputs.Last();
 		const FMetasoundFrontendEdge* EdgeAdded = Builder.AddEdge(FMetasoundFrontendEdge
 		{
-			NodeOutput.NodeID,
-			NodeOutput.VertexID,
+			NodeOutputHandle.NodeID,
+			NodeOutputHandle.VertexID,
 			GraphOutputNode->GetID(),
 			InputVertex.VertexID
 		});
@@ -236,7 +248,7 @@ void UMetaSoundBuilderBase::ConnectNodeOutputToGraphOutput(FName GraphOutputName
 	}
 }
 
-void UMetaSoundBuilderBase::ConnectNodeInputToGraphInput(FName GraphInputName, const FMetaSoundBuilderNodeInputHandle& NodeInput, EMetaSoundBuilderResult& OutResult)
+void UMetaSoundBuilderBase::ConnectNodeInputToGraphInput(FName GraphInputName, const FMetaSoundBuilderNodeInputHandle& NodeInputHandle, EMetaSoundBuilderResult& OutResult)
 {
 	const FMetasoundFrontendNode* GraphInputNode = Builder.FindGraphInputNode(GraphInputName);
 	if (!GraphInputNode)
@@ -250,36 +262,42 @@ void UMetaSoundBuilderBase::ConnectNodeInputToGraphInput(FName GraphInputName, c
 	{
 		GraphInputNode->GetID(),
 		OutputVertex.VertexID,
-		NodeInput.NodeID,
-		NodeInput.VertexID
+		NodeInputHandle.NodeID,
+		NodeInputHandle.VertexID
 	});
 	OutResult = EdgeAdded ? EMetaSoundBuilderResult::Succeeded : EMetaSoundBuilderResult::Failed;
 }
 
-void UMetaSoundBuilderBase::DisconnectNodes(const FMetaSoundBuilderNodeOutputHandle& NodeOutput, const FMetaSoundBuilderNodeInputHandle& NodeInput, EMetaSoundBuilderResult& OutResult)
+void UMetaSoundBuilderBase::DisconnectNodes(const FMetaSoundBuilderNodeOutputHandle& NodeOutputHandle, const FMetaSoundBuilderNodeInputHandle& NodeInputHandle, EMetaSoundBuilderResult& OutResult)
 {
 	const bool bEdgeRemoved = Builder.RemoveEdge(FMetasoundFrontendEdge
 	{
-		NodeOutput.NodeID,
-		NodeOutput.VertexID,
-		NodeInput.NodeID,
-		NodeInput.VertexID,
+		NodeOutputHandle.NodeID,
+		NodeOutputHandle.VertexID,
+		NodeInputHandle.NodeID,
+		NodeInputHandle.VertexID,
 	});
 	OutResult = bEdgeRemoved ? EMetaSoundBuilderResult::Succeeded : EMetaSoundBuilderResult::Failed;
 }
 
-void UMetaSoundBuilderBase::DisconnectNodeInput(const FMetaSoundBuilderNodeInputHandle& NodeInput, EMetaSoundBuilderResult& OutResult)
+void UMetaSoundBuilderBase::DisconnectNodeInput(const FMetaSoundBuilderNodeInputHandle& NodeInputHandle, EMetaSoundBuilderResult& OutResult)
 {
-	const bool bEdgeRemoved = Builder.RemoveEdgeToNodeInput(NodeInput.NodeID, NodeInput.VertexID);
+	const bool bEdgeRemoved = Builder.RemoveEdgeToNodeInput(NodeInputHandle.NodeID, NodeInputHandle.VertexID);
 	OutResult = bEdgeRemoved ? EMetaSoundBuilderResult::Succeeded : EMetaSoundBuilderResult::Failed;
 }
 
-void UMetaSoundBuilderBase::DisconnectNodeOutput(const FMetaSoundBuilderNodeOutputHandle& NodeOutput, EMetaSoundBuilderResult& OutResult)
+void UMetaSoundBuilderBase::DisconnectNodeOutput(const FMetaSoundBuilderNodeOutputHandle& NodeOutputHandle, EMetaSoundBuilderResult& OutResult)
 {
-	const bool bEdgeRemoved = Builder.RemoveEdgesFromNodeOutput(NodeOutput.NodeID, NodeOutput.VertexID);
+	const bool bEdgeRemoved = Builder.RemoveEdgesFromNodeOutput(NodeOutputHandle.NodeID, NodeOutputHandle.VertexID);
 	OutResult = bEdgeRemoved ? EMetaSoundBuilderResult::Succeeded : EMetaSoundBuilderResult::Failed;
 }
 
+void UMetaSoundBuilderBase::DisconnectNodesByInterfaceBindings(const FMetaSoundNodeHandle& FromNodeHandle, const FMetaSoundNodeHandle& ToNodeHandle, EMetaSoundBuilderResult& OutResult)
+{
+	const bool bEdgesRemoved = Builder.RemoveEdgesByNodeClassInterfaceBindings(FromNodeHandle.NodeID, ToNodeHandle.NodeID);
+	OutResult = bEdgesRemoved ? EMetaSoundBuilderResult::Succeeded : EMetaSoundBuilderResult::Failed;
+}
+
 FMetaSoundBuilderNodeInputHandle UMetaSoundBuilderBase::FindNodeInputByName(const FMetaSoundNodeHandle& NodeHandle, FName InputName, EMetaSoundBuilderResult& OutResult) const
 {
 	if (const FMetasoundFrontendNode* Node = Builder.FindNode(NodeHandle.NodeID))
@@ -443,6 +461,15 @@ FMetaSoundNodeHandle UMetaSoundBuilderBase::FindGraphOutputNode(FName OutputName
 	return { };
 }
 
+void UMetaSoundBuilderBase::CreateTransientDocument()
+{
+	TObjectPtr<UMetaSoundBuilderDocument> NewBuilder = NewObject<UMetaSoundBuilderDocument>();
+	TScriptInterface<IMetaSoundDocumentInterface> MetaSoundSourceDocInterface = NewBuilder;
+	Builder = FMetaSoundFrontendDocumentBuilder(MetaSoundSourceDocInterface);
+	const UClass& BuilderClass = GetBuilderUClass();
+	Builder.InitDocument(BuilderClass.GetFName());
+}
+
 void UMetaSoundBuilderBase::InitNodeLocations()
 {
 	Builder.InitNodeLocations();
@@ -469,17 +496,6 @@ bool UMetaSoundBuilderBase::NodeOutputIsConnected(const FMetaSoundBuilderNodeOut
 	return Builder.IsNodeOutputConnected(OutputHandle.NodeID, OutputHandle.VertexID);
 }
 
-void UMetaSoundBuilderBase::PostInitProperties()
-{
-	if (!HasAnyFlags(RF_ClassDefaultObject))
-	{
-		UClass& MetaSoundUClass = GetBuilderUClass();
-		Builder.InitDocument(MetaSoundUClass.GetFName());
-	}
-
-	Super::PostInitProperties();
-}
-
 void UMetaSoundBuilderBase::RemoveInterface(FName InterfaceName, EMetaSoundBuilderResult& OutResult)
 {
 	const bool bInterfaceRemoved = Builder.RemoveInterface(InterfaceName);
@@ -708,7 +724,6 @@ void UMetaSoundSourceBuilder::SetFormat(EMetaSoundOutputAudioFormat OutputFormat
 	OutResult = bSuccess ? EMetaSoundBuilderResult::Succeeded : EMetaSoundBuilderResult::Failed;
 }
 
-
 UMetaSoundPatchBuilder* UMetaSoundBuilderSubsystem::CreatePatchBuilder(FName BuilderName, EMetaSoundBuilderResult& OutResult)
 {
 	using namespace Metasound::Engine::BuilderSubsystemPrivate;
@@ -718,12 +733,7 @@ UMetaSoundPatchBuilder* UMetaSoundBuilderSubsystem::CreatePatchBuilder(FName Bui
 		UE_LOG(LogMetaSound, Error, TEXT("CreatePatchBuilder Failed to create MetaSound Builder '%s': See CVar au.MetaSound.BuilderSubsystem.Enabled (%s')"), *BuilderName.ToString(), CVarEnableDocumentBuilder->GetHelp());
 		return nullptr;
 	}
-
-	const EObjectFlags NewObjectFlags = RF_Public | RF_Transient;
-	const FName ObjectName = MakeUniqueObjectName(GetTransientPackage(), UMetaSoundPatchBuilder::StaticClass());
-	TObjectPtr<UMetaSoundPatchBuilder> NewBuilder = NewObject<UMetaSoundPatchBuilder>(this, ObjectName, NewObjectFlags);
-	check(NewBuilder);
-	return NewBuilder;
+	return &CreateTransientBuilder<UMetaSoundPatchBuilder>();
 }
 
 UMetaSoundSourceBuilder* UMetaSoundBuilderSubsystem::CreateSourceBuilder(
@@ -750,27 +760,23 @@ UMetaSoundSourceBuilder* UMetaSoundBuilderSubsystem::CreateSourceBuilder(
 		return nullptr;
 	}
 
-	const EObjectFlags NewObjectFlags = RF_Public | RF_Transient;
-	UPackage* TransientPackage = GetTransientPackage();
-	const FName ObjectName = MakeUniqueObjectName(TransientPackage, UMetaSoundSourceBuilder::StaticClass());
-	TObjectPtr<UMetaSoundSourceBuilder> NewBuilder = NewObject<UMetaSoundSourceBuilder>(TransientPackage, ObjectName, NewObjectFlags);
-	check(NewBuilder);
+	UMetaSoundSourceBuilder& NewBuilder = CreateTransientBuilder<UMetaSoundSourceBuilder>();
 
 	bool bFormatSet = OutputFormat == EMetaSoundOutputAudioFormat::Mono;
 	if (!bFormatSet)
 	{
 		EMetaSoundBuilderResult Result = EMetaSoundBuilderResult::Failed;
-		NewBuilder->SetFormat(OutputFormat, Result);
+		NewBuilder.SetFormat(OutputFormat, Result);
 		bFormatSet = Result == EMetaSoundBuilderResult::Succeeded;
 	}
 
 	if (bFormatSet)
 	{
-		const Metasound::Engine::FOutputAudioFormatInfoPair* FormatInfo = NewBuilder->FindOutputAudioFormatInfo();
+		const Metasound::Engine::FOutputAudioFormatInfoPair* FormatInfo = NewBuilder.FindOutputAudioFormatInfo();
 		TArray<FMetaSoundNodeHandle> AudioOutputNodes;
 		if (FormatInfo)
 		{
-			AudioOutputNodes = NewBuilder->FindInterfaceOutputNodes(FormatInfo->Value.InterfaceVersion.Name, OutResult);
+			AudioOutputNodes = NewBuilder.FindInterfaceOutputNodes(FormatInfo->Value.InterfaceVersion.Name, OutResult);
 		}
 
 		if (!AudioOutputNodes.IsEmpty())
@@ -778,7 +784,7 @@ UMetaSoundSourceBuilder* UMetaSoundBuilderSubsystem::CreateSourceBuilder(
 			Algo::Transform(AudioOutputNodes, AudioOutNodeInputs, [&NewBuilder, &BuilderName](const FMetaSoundNodeHandle& AudioOutputNode) -> FMetaSoundBuilderNodeInputHandle
 			{
 				EMetaSoundBuilderResult Result;
-				TArray<FMetaSoundBuilderNodeInputHandle> Inputs = NewBuilder->FindNodeInputs(AudioOutputNode, Result);
+				TArray<FMetaSoundBuilderNodeInputHandle> Inputs = NewBuilder.FindNodeInputs(AudioOutputNode, Result);
 				if (!Inputs.IsEmpty())
 				{
 					return Inputs.Last();
@@ -801,14 +807,14 @@ UMetaSoundSourceBuilder* UMetaSoundBuilderSubsystem::CreateSourceBuilder(
 	}
 
 	{
-		FMetaSoundNodeHandle OnPlayNode = NewBuilder->FindGraphInputNode(SourceInterface::Inputs::OnPlay, OutResult);
+		FMetaSoundNodeHandle OnPlayNode = NewBuilder.FindGraphInputNode(SourceInterface::Inputs::OnPlay, OutResult);
 		if (OutResult == EMetaSoundBuilderResult::Failed)
 		{
 			UE_LOG(LogMetaSound, Error, TEXT("Builder '%s' Creation Error: Failed to add required interface '%s' when attempting to create MetaSound Source Builder"), *BuilderName.ToString(), * SourceInterface::GetVersion().ToString());
 			return nullptr;
 		}
 
-		TArray<FMetaSoundBuilderNodeOutputHandle> Outputs = NewBuilder->FindNodeOutputs(OnPlayNode, OutResult);
+		TArray<FMetaSoundBuilderNodeOutputHandle> Outputs = NewBuilder.FindNodeOutputs(OnPlayNode, OutResult);
 		if (OutResult == EMetaSoundBuilderResult::Failed)
 		{
 			UE_LOG(LogMetaSound, Error, TEXT("Builder '%s' Creation Error: Failed to find output vertex for 'OnPlay' input node when attempting to create MetaSound Source Builder"), *BuilderName.ToString());
@@ -822,13 +828,13 @@ UMetaSoundSourceBuilder* UMetaSoundBuilderSubsystem::CreateSourceBuilder(
 	{
 		if (bIsOneShot)
 		{
-			FMetaSoundNodeHandle OnFinishedNode = NewBuilder->FindGraphOutputNode(SourceOneShotInterface::Outputs::OnFinished, OutResult);
+			FMetaSoundNodeHandle OnFinishedNode = NewBuilder.FindGraphOutputNode(SourceOneShotInterface::Outputs::OnFinished, OutResult);
 			if (OutResult == EMetaSoundBuilderResult::Failed)
 			{
 				UE_LOG(LogMetaSound, Error, TEXT("Builder '%s' Creation Error: Failed to add '%s' interface; interface definition may not be registered."), *BuilderName.ToString(), *SourceOneShotInterface::GetVersion().ToString());
 			}
 
-			TArray<FMetaSoundBuilderNodeInputHandle> Inputs = NewBuilder->FindNodeInputs(OnFinishedNode, OutResult);
+			TArray<FMetaSoundBuilderNodeInputHandle> Inputs = NewBuilder.FindNodeInputs(OnFinishedNode, OutResult);
 			if (OutResult == EMetaSoundBuilderResult::Failed)
 			{
 				UE_LOG(LogMetaSound, Error, TEXT("Builder '%s' Creation Error: Failed to find input vertex for 'OnFinished' output node when attempting to create MetaSound Source Builder"), *BuilderName.ToString());
@@ -840,12 +846,12 @@ UMetaSoundSourceBuilder* UMetaSoundBuilderSubsystem::CreateSourceBuilder(
 		}
 		else
 		{
-			NewBuilder->RemoveInterface(SourceOneShotInterface::GetVersion().Name, OutResult);
+			NewBuilder.RemoveInterface(SourceOneShotInterface::GetVersion().Name, OutResult);
 		}
 	}
 
 	OutResult = EMetaSoundBuilderResult::Succeeded;
-	return NewBuilder;
+	return &NewBuilder;
 }
 
 UMetaSoundSourceBuilder* UMetaSoundBuilderSubsystem::AttachSourceBuilderToAsset(UMetaSoundSource* InSource) const
diff --git a/Engine/Plugins/Runtime/Metasound/Source/MetasoundEngine/Private/MetasoundInterface.cpp b/Engine/Plugins/Runtime/Metasound/Source/MetasoundEngine/Private/MetasoundInterface.cpp
--- a/Engine/Plugins/Runtime/Metasound/Source/MetasoundEngine/Private/MetasoundInterface.cpp
+++ b/Engine/Plugins/Runtime/Metasound/Source/MetasoundEngine/Private/MetasoundInterface.cpp
@@ -423,16 +423,22 @@ namespace Metasound
 
 				bool bInterfaceSupportsPatch = true;
 				bool bInterfaceSupportsSource = true;
+
 				const TArray<const UClass*> SupportedUClasses = Interface->FindSupportedUClasses();
-				for (const UClass* SupportedUClass : SupportedUClasses )
+				if (!SupportedUClasses.IsEmpty())
 				{
-					const UClass* PatchClass = UMetaSoundPatch::StaticClass();
-					check(PatchClass);
-					bInterfaceSupportsPatch |= PatchClass->IsChildOf(SupportedUClass);
+					bInterfaceSupportsPatch = false;
+					bInterfaceSupportsSource = false;
+					for (const UClass* SupportedUClass : SupportedUClasses)
+					{
+						const UClass* PatchClass = UMetaSoundPatch::StaticClass();
+						check(PatchClass);
+						bInterfaceSupportsPatch |= PatchClass->IsChildOf(SupportedUClass);
 
-					const UClass* SourceClass = UMetaSoundSource::StaticClass();
-					check(SourceClass);
-					bInterfaceSupportsSource |= SourceClass->IsChildOf(SupportedUClass);
+						const UClass* SourceClass = UMetaSoundSource::StaticClass();
+						check(SourceClass);
+						bInterfaceSupportsSource |= SourceClass->IsChildOf(SupportedUClass);
+					}
 				}
 
 				const bool bSupported = bInterfaceSupportsSource || bInterfaceSupportsPatch;
diff --git a/Engine/Plugins/Runtime/Metasound/Source/MetasoundEngine/Public/MetasoundBuilderSubsystem.h b/Engine/Plugins/Runtime/Metasound/Source/MetasoundEngine/Public/MetasoundBuilderSubsystem.h
--- a/Engine/Plugins/Runtime/Metasound/Source/MetasoundEngine/Public/MetasoundBuilderSubsystem.h
+++ b/Engine/Plugins/Runtime/Metasound/Source/MetasoundEngine/Public/MetasoundBuilderSubsystem.h
@@ -164,6 +164,11 @@ public:
 	UFUNCTION(BlueprintCallable, Category = "Audio|MetaSound|Builder", meta = (ExpandEnumAsExecs = "OutResult"))
 	void DisconnectNodeOutput(const FMetaSoundBuilderNodeOutputHandle& NodeOutputHandle, EMetaSoundBuilderResult& OutResult);
 
+	// Disconnects two nodes using defined MetaSound Interface Bindings registered with the MetaSound Interface registry. Returns success if
+	// all connections were found and removed, failed if any connections were not.
+	UFUNCTION(BlueprintCallable, Category = "Audio|MetaSound|Builder", meta = (ExpandEnumAsExecs = "OutResult"))
+	void DisconnectNodesByInterfaceBindings(const FMetaSoundNodeHandle& FromNodeHandle, const FMetaSoundNodeHandle& ToNodeHandle, EMetaSoundBuilderResult& OutResult);
+
 	// Returns graph input node by the given name if it exists, or an invalid handle if not found.
 	UFUNCTION(BlueprintCallable, Category = "Audio|MetaSound|Builder", meta = (ExpandEnumAsExecs = "OutResult"))
 	UPARAM(DisplayName = "Node Handle") FMetaSoundNodeHandle FindGraphInputNode(FName InputName, EMetaSoundBuilderResult& OutResult) const;
@@ -263,7 +268,8 @@ public:
 	// Returns the builder's supported UClass
 	virtual UClass& GetBuilderUClass() const PURE_VIRTUAL(UMetaSoundBuilderBase::GetBuilderUClass, return *UObject::StaticClass(); );
 
-	virtual void PostInitProperties() override;
+	// Constructs a transient UMetaSoundBuilderDocument & attaches builder to it
+	void CreateTransientDocument();
 
 	// Initializes and ensures all nodes have a position (required prior to exporting to an asset if expected to be viewed in the editor).
 	void InitNodeLocations();
diff --git a/Engine/Plugins/Runtime/Metasound/Source/MetasoundFrontend/Private/MetasoundFrontendDocumentBuilder.cpp b/Engine/Plugins/Runtime/Metasound/Source/MetasoundFrontend/Private/MetasoundFrontendDocumentBuilder.cpp
--- a/Engine/Plugins/Runtime/Metasound/Source/MetasoundFrontend/Private/MetasoundFrontendDocumentBuilder.cpp
+++ b/Engine/Plugins/Runtime/Metasound/Source/MetasoundFrontend/Private/MetasoundFrontendDocumentBuilder.cpp
@@ -216,7 +216,7 @@ bool FMetaSoundFrontendDocumentBuilder::AddEdgesByNodeClassInterfaceBindings(con
 	using namespace Metasound;
 	using namespace Metasound::Frontend;
 
-	auto FindNodeClassDocAndInterfaces = [this](const FGuid & InNodeID) -> const FMetasoundFrontendDocument*
+	auto FindNodeClassDoc = [this](const FGuid & InNodeID) -> const FMetasoundFrontendDocument*
 	{
 		const TScriptInterface<IMetaSoundDocumentInterface> NodeDocumentUInterface = FindOrLoadNodeClassDocument(InNodeID);
 		if (NodeDocumentUInterface)
@@ -227,8 +227,8 @@ bool FMetaSoundFrontendDocumentBuilder::AddEdgesByNodeClassInterfaceBindings(con
 		return nullptr;
 	};
 
-	const FMetasoundFrontendDocument* OutputNodeDocument = FindNodeClassDocAndInterfaces(InFromNodeID);
-	const FMetasoundFrontendDocument* InputNodeDocument = FindNodeClassDocAndInterfaces(InToNodeID);
+	const FMetasoundFrontendDocument* OutputNodeDocument = FindNodeClassDoc(InFromNodeID);
+	const FMetasoundFrontendDocument* InputNodeDocument = FindNodeClassDoc(InToNodeID);
 
 	if (!InputNodeDocument || !OutputNodeDocument)
 	{
@@ -1107,7 +1107,7 @@ bool FMetaSoundFrontendDocumentBuilder::RemoveEdgesByNamedConnections(const TSet
 			else
 			{
 				bSuccess = false;
-				UE_LOG(LogMetaSound, Warning, TEXT("Failed to remove connections between MetaSound node output '%s' and input '%s': No connection found."), *Connection.OutputName.ToString(), *Connection.InputName.ToString());
+				UE_LOG(LogMetaSound, Warning, TEXT("Failed to remove connection between MetaSound node output '%s' and input '%s': No connection found."), *Connection.OutputName.ToString(), *Connection.InputName.ToString());
 			}
 		}
 	}
@@ -1136,7 +1136,7 @@ bool FMetaSoundFrontendDocumentBuilder::RemoveEdgesByNodeClassInterfaceBindings(
 	using namespace Metasound;
 	using namespace Metasound::Frontend;
 
-	auto FindNodeClassDocAndInterfaces = [this](const FGuid& InNodeID) -> const FMetasoundFrontendDocument*
+	auto FindNodeClassDoc = [this](const FGuid& InNodeID) -> const FMetasoundFrontendDocument*
 	{
 		const TScriptInterface<IMetaSoundDocumentInterface> NodeDocumentUInterface = FindOrLoadNodeClassDocument(InNodeID);
 		if (NodeDocumentUInterface)
@@ -1147,8 +1147,8 @@ bool FMetaSoundFrontendDocumentBuilder::RemoveEdgesByNodeClassInterfaceBindings(
 		return nullptr;
 	};
 
-	const FMetasoundFrontendDocument* OutputNodeDocument = FindNodeClassDocAndInterfaces(InFromNodeID);
-	const FMetasoundFrontendDocument* InputNodeDocument = FindNodeClassDocAndInterfaces(InToNodeID);
+	const FMetasoundFrontendDocument* OutputNodeDocument = FindNodeClassDoc(InFromNodeID);
+	const FMetasoundFrontendDocument* InputNodeDocument = FindNodeClassDoc(InToNodeID);
 
 	if (!InputNodeDocument || !OutputNodeDocument)
 	{
diff --git a/Engine/Plugins/Runtime/Metasound/Source/MetasoundFrontend/Private/MetasoundFrontendDocumentCache.cpp b/Engine/Plugins/Runtime/Metasound/Source/MetasoundFrontend/Private/MetasoundFrontendDocumentCache.cpp
--- a/Engine/Plugins/Runtime/Metasound/Source/MetasoundFrontend/Private/MetasoundFrontendDocumentCache.cpp
+++ b/Engine/Plugins/Runtime/Metasound/Source/MetasoundFrontend/Private/MetasoundFrontendDocumentCache.cpp
@@ -363,11 +363,11 @@ namespace Metasound::Frontend
 
 	bool FDocumentGraphEdgeCache::ContainsEdge(const FMetasoundFrontendEdge& InEdge) const
 	{
-		const FMetasoundFrontendVertexHandle InputPair = { InEdge.FromNodeID, InEdge.FromVertexID };
+		const FMetasoundFrontendVertexHandle InputPair = { InEdge.ToNodeID, InEdge.ToVertexID };
 		if (const int32* Index = InputToEdgeIndex.Find(InputPair))
 		{
 			const FMetasoundFrontendEdge& Edge = GetDocument().RootGraph.Graph.Edges[*Index];
-			return Edge.ToNodeID == InEdge.ToNodeID && Edge.ToVertexID == InEdge.ToVertexID;
+			return Edge.FromNodeID == InEdge.FromNodeID && Edge.FromVertexID == InEdge.FromVertexID;
 		}
 
 		return false;
diff --git a/Engine/Source/Runtime/AudioExtensions/Private/IAudioParameterInterfaceRegistry.cpp b/Engine/Source/Runtime/AudioExtensions/Private/IAudioParameterInterfaceRegistry.cpp
--- a/Engine/Source/Runtime/AudioExtensions/Private/IAudioParameterInterfaceRegistry.cpp
+++ b/Engine/Source/Runtime/AudioExtensions/Private/IAudioParameterInterfaceRegistry.cpp
@@ -42,7 +42,7 @@ namespace Audio
 	)
 		: NamePrivate(InName)
 		, VersionPrivate(InVersion)
-		, SupportedUClassNames({ InType.GetFullName() })
+		, SupportedUClassNames({ InType.GetPathName() })
 	{
 	}
 
@@ -59,7 +59,7 @@ namespace Audio
 		Algo::Transform(InClasses, SupportedUClassNames, [](const UClass* Class)
 		{
 			check(Class);
-			return Class->GetFullName();
+			return Class->GetPathName();
 		});
 	}
 
@@ -98,7 +98,7 @@ namespace Audio
 		TArray<const UClass*> SupportedUClasses;
 		for (const FString& Name : SupportedUClassNames)
 		{
-			if (const UClass* Class = FindObject<const UClass>(nullptr, *Name))
+			if (const UClass* Class = FindFirstObject<const UClass>(*Name, EFindFirstObjectOptions::ExactClass, ELogVerbosity::Warning, TEXT("FParameterInterface::FindSupportedUClasses")))
 			{
 				SupportedUClasses.Add(Class);
 			}
