Interface Registry Tidying
1.) Move Bindings to own registry
2.) Nuke source files & left over cruft relating to archetypes
3.) Move deprecated interface registration next to respective interfaces for clarity
4.) Rename "FMetasoundFrontendInterfaceBindingConnections" to "FMetasoundFrontendInterfaceVertexBindings"

diff --git a/Engine/Plugins/Runtime/Metasound/Source/MetasoundEditor/Private/MetasoundDetailCustomization.cpp b/Engine/Plugins/Runtime/Metasound/Source/MetasoundEditor/Private/MetasoundDetailCustomization.cpp
--- a/Engine/Plugins/Runtime/Metasound/Source/MetasoundEditor/Private/MetasoundDetailCustomization.cpp
+++ b/Engine/Plugins/Runtime/Metasound/Source/MetasoundEditor/Private/MetasoundDetailCustomization.cpp
@@ -12,12 +12,12 @@
 #include "IAudioParameterTransmitter.h"
 #include "IDetailGroup.h"
 #include "Input/Events.h"
+#include "Interfaces/MetasoundFrontendInterfaceRegistry.h"
 #include "MetasoundAssetBase.h"
 #include "MetasoundEditor.h"
 #include "MetasoundEditorGraphBuilder.h"
 #include "MetasoundEditorSettings.h"
 #include "MetasoundFrontend.h"
-#include "MetasoundFrontendArchetypeRegistry.h"
 #include "MetasoundFrontendController.h"
 #include "MetasoundFrontendSearchEngine.h"
 #include "MetasoundSource.h"
diff --git a/Engine/Plugins/Runtime/Metasound/Source/MetasoundEditor/Private/SMetasoundGraphNode.cpp b/Engine/Plugins/Runtime/Metasound/Source/MetasoundEditor/Private/SMetasoundGraphNode.cpp
--- a/Engine/Plugins/Runtime/Metasound/Source/MetasoundEditor/Private/SMetasoundGraphNode.cpp
+++ b/Engine/Plugins/Runtime/Metasound/Source/MetasoundEditor/Private/SMetasoundGraphNode.cpp
@@ -3,10 +3,10 @@
 
 #include "AudioParameterControllerInterface.h"
 #include "Components/AudioComponent.h"
-#include "Styling/AppStyle.h"
 #include "GraphEditorSettings.h"
 #include "IAudioParameterTransmitter.h"
 #include "IDocumentation.h"
+#include "Interfaces/MetasoundFrontendInterfaceRegistry.h"
 #include "KismetPins/SGraphPinBool.h"
 #include "KismetPins/SGraphPinExec.h"
 #include "KismetPins/SGraphPinInteger.h"
@@ -20,7 +20,6 @@
 #include "MetasoundEditorGraphNode.h"
 #include "MetasoundEditorGraphSchema.h"
 #include "MetasoundEditorModule.h"
-#include "MetasoundFrontendArchetypeRegistry.h"
 #include "MetasoundFrontendNodeTemplateRegistry.h"
 #include "MetasoundFrontendRegistries.h"
 #include "MetasoundTrace.h"
@@ -34,6 +33,7 @@
 #include "ScopedTransaction.h"
 #include "SGraphNode.h"
 #include "SGraphPinComboBox.h"
+#include "Styling/AppStyle.h"
 #include "SLevelOfDetailBranchNode.h"
 #include "SMetasoundGraphEnumPin.h"
 #include "SMetasoundGraphPin.h"
diff --git a/Engine/Plugins/Runtime/Metasound/Source/MetasoundEngine/Private/Interfaces/MetasoundDeprecatedInterfaces.cpp b/Engine/Plugins/Runtime/Metasound/Source/MetasoundEngine/Private/Interfaces/MetasoundDeprecatedInterfaces.cpp
new file mode 100644
--- /dev/null
+++ b/Engine/Plugins/Runtime/Metasound/Source/MetasoundEngine/Private/Interfaces/MetasoundDeprecatedInterfaces.cpp
@@ -0,0 +1,631 @@
+// Copyright Epic Games, Inc. All Rights Reserved.
+#include "Interfaces/MetasoundDeprecatedInterfaces.h"
+
+#include "IAudioParameterInterfaceRegistry.h"
+#include "Interfaces/MetasoundFrontendInterfaceRegistry.h"
+#include "Interfaces/MetasoundInterface.h"
+#include "Interfaces/MetasoundOutputFormatInterfaces.h"
+#include "MetasoundAudioFormats.h"
+#include "MetasoundDataReference.h"
+#include "MetasoundFrontendTransform.h"
+#include "MetasoundInterface.h"
+#include "MetasoundRouter.h"
+#include "MetasoundSource.h"
+#include "MetasoundUObjectRegistry.h"
+#include "UObject/Class.h"
+#include "UObject/NoExportTypes.h"
+
+#define LOCTEXT_NAMESPACE "MetasoundEngine"
+
+namespace Metasound::Engine
+{
+	namespace DeprecatedInterfacesPrivate
+	{
+		auto MatchMemberNamesIgnoreSpaces = [] (FName InNameA, FName InNameB)
+		{
+			FName ParamA;
+			FName ParamB;
+			FName Namespace;
+			Audio::FParameterPath::SplitName(InNameA, Namespace, ParamA);
+			Audio::FParameterPath::SplitName(InNameB, Namespace, ParamB);
+
+			auto GetStringNoWhitespace = [](FName InName) { return InName.ToString().Replace(TEXT(" "), TEXT("")); };
+			return GetStringNoWhitespace(ParamA) == GetStringNoWhitespace(ParamB);
+		};
+
+		const FVertexName& GetOnPlayInputName()
+		{
+			static const FVertexName OnPlayInputName = "On Play";
+			return OnPlayInputName;
+		}
+
+		const FVertexName& GetIsFinishedOutputName()
+		{
+			static const FVertexName OnFinishedOutputName = "On Finished";
+			return OnFinishedOutputName;
+		}
+
+		const FVertexName& GetAudioDeviceIDVariableName()
+		{
+			static const FVertexName AudioDeviceIDVarName = "AudioDeviceID";
+			return AudioDeviceIDVarName;
+		}
+
+		const FVertexName& GetGraphName()
+		{
+			static const FVertexName SoundGraphName = "GraphName";
+			return SoundGraphName;
+		}
+	}
+
+	// MetasoundV1_0 is a metasound without any required inputs or outputs.
+	namespace MetasoundV1_0
+	{
+		const FMetasoundFrontendVersion& GetVersion()
+		{
+			static const FMetasoundFrontendVersion Version { "MetaSound", FMetasoundFrontendVersionNumber {1, 0} };
+			return Version;
+		}
+
+		FMetasoundFrontendInterface GetInterface()
+		{
+			FMetasoundFrontendInterface Interface;
+			Interface.Version = GetVersion();
+			return Interface;
+		}
+	}
+
+	namespace MetasoundOutputFormatMonoV1_0
+	{
+		const FMetasoundFrontendVersion& GetVersion()
+		{
+			static const FMetasoundFrontendVersion Version { "MonoSource", FMetasoundFrontendVersionNumber {1, 0} };
+			return Version;
+		}
+
+		const FVertexName& GetAudioOutputName()
+		{
+			static const FVertexName AudioOutputName = "Generated Audio";
+			return AudioOutputName;
+		}
+
+		FMetasoundFrontendInterface GetInterface()
+		{
+			FMetasoundFrontendInterface Interface;
+			Interface.Version = GetVersion();
+				
+			// Inputs
+			FMetasoundFrontendClassVertex OnPlayTrigger;
+			OnPlayTrigger.Name = DeprecatedInterfacesPrivate::GetOnPlayInputName();
+			OnPlayTrigger.TypeName = GetMetasoundDataTypeName<FTrigger>();
+			OnPlayTrigger.VertexID = FGuid::NewGuid();
+
+#if WITH_EDITOR
+			OnPlayTrigger.Metadata.SetDisplayName(LOCTEXT("OnPlay", "On Play"));
+			OnPlayTrigger.Metadata.SetDescription(LOCTEXT("OnPlayTriggerToolTip", "Trigger executed when this source is played."));
+#endif // WITH_EDITOR
+
+			Interface.Inputs.Add(OnPlayTrigger);
+
+			// Outputs 
+			FMetasoundFrontendClassVertex OnFinished;
+			OnFinished.Name = DeprecatedInterfacesPrivate::GetIsFinishedOutputName();
+			OnFinished.TypeName = GetMetasoundDataTypeName<FTrigger>();
+			OnFinished.VertexID = FGuid::NewGuid();
+
+#if WITH_EDITOR
+			OnFinished.Metadata.SetDisplayName(LOCTEXT("OnFinished", "On Finished"));
+			OnFinished.Metadata.SetDescription(LOCTEXT("OnFinishedToolTip", "Trigger executed to initiate stopping the source."));
+#endif // WITH_EDITOR
+
+			Interface.Outputs.Add(OnFinished);
+
+			FMetasoundFrontendClassVertex GeneratedAudio;
+			GeneratedAudio.Name = GetAudioOutputName();
+			GeneratedAudio.TypeName = GetMetasoundDataTypeName<FMonoAudioFormat>();
+			GeneratedAudio.VertexID = FGuid::NewGuid();
+
+#if WITH_EDITOR
+			GeneratedAudio.Metadata.SetDisplayName(LOCTEXT("GeneratedMono", "Audio"));
+			GeneratedAudio.Metadata.SetDescription(LOCTEXT("GeneratedAudioToolTip", "The resulting output audio from this source."));
+#endif // WITH_EDITOR
+
+			Interface.Outputs.Add(GeneratedAudio);
+
+			// Environment 
+			FMetasoundFrontendClassEnvironmentVariable AudioDeviceID;
+			AudioDeviceID.Name = DeprecatedInterfacesPrivate::GetAudioDeviceIDVariableName();
+
+			Interface.Environment.Add(AudioDeviceID);
+
+			return Interface;
+		}
+	}
+
+	namespace MetasoundOutputFormatStereoV1_0
+	{
+		const FMetasoundFrontendVersion& GetVersion()
+		{
+			static const FMetasoundFrontendVersion Version { "StereoSource", FMetasoundFrontendVersionNumber { 1, 0 } };
+			return Version;
+		}
+
+		const FVertexName& GetAudioOutputName()
+		{
+			static const FVertexName AudioOutputName = "Generated Audio";
+			return AudioOutputName;
+		}
+
+		FMetasoundFrontendInterface GetInterface()
+		{
+			FMetasoundFrontendInterface Interface;
+			Interface.Version = GetVersion();
+				
+			// Inputs
+			FMetasoundFrontendClassVertex OnPlayTrigger;
+			OnPlayTrigger.Name = DeprecatedInterfacesPrivate::GetOnPlayInputName();
+			OnPlayTrigger.TypeName = GetMetasoundDataTypeName<FTrigger>();
+			OnPlayTrigger.VertexID = FGuid::NewGuid();
+
+#if WITH_EDITOR
+			OnPlayTrigger.Metadata.SetDisplayName(LOCTEXT("OnPlay", "On Play"));
+			OnPlayTrigger.Metadata.SetDescription(LOCTEXT("OnPlayTriggerToolTip", "Trigger executed when this source is played."));
+#endif // WITH_EDITOR
+
+			Interface.Inputs.Add(OnPlayTrigger);
+
+			// Outputs
+			FMetasoundFrontendClassVertex OnFinished;
+			OnFinished.Name = DeprecatedInterfacesPrivate::GetIsFinishedOutputName();
+			OnFinished.TypeName = GetMetasoundDataTypeName<FTrigger>();
+			OnFinished.VertexID = FGuid::NewGuid();
+
+#if WITH_EDITOR
+			OnFinished.Metadata.SetDisplayName(LOCTEXT("OnFinished", "On Finished"));
+			OnFinished.Metadata.SetDescription(LOCTEXT("OnFinishedToolTip", "Trigger executed to initiate stopping the source."));
+#endif // WITH_EDITOR
+
+			Interface.Outputs.Add(OnFinished);
+
+			FMetasoundFrontendClassVertex GeneratedAudio;
+			GeneratedAudio.Name = GetAudioOutputName();
+			GeneratedAudio.TypeName = GetMetasoundDataTypeName<FStereoAudioFormat>();
+			GeneratedAudio.VertexID = FGuid::NewGuid();
+
+#if WITH_EDITOR
+			GeneratedAudio.Metadata.SetDisplayName(LOCTEXT("GeneratedStereo", "Audio"));
+			GeneratedAudio.Metadata.SetDescription(LOCTEXT("GeneratedAudioToolTip", "The resulting output audio from this source."));
+#endif // WITH_EDITOR
+
+			Interface.Outputs.Add(GeneratedAudio);
+
+			// Environment
+			FMetasoundFrontendClassEnvironmentVariable AudioDeviceID;
+			AudioDeviceID.Name = DeprecatedInterfacesPrivate::GetAudioDeviceIDVariableName();
+
+			Interface.Environment.Add(AudioDeviceID);
+
+			return Interface;
+		}
+	}
+
+	namespace MetasoundOutputFormatMonoV1_1
+	{
+		const FMetasoundFrontendVersion& GetVersion()
+		{
+			static const FMetasoundFrontendVersion Version { "MonoSource", FMetasoundFrontendVersionNumber { 1, 1} };
+			return Version;
+		}
+
+		const FVertexName& GetAudioOutputName()
+		{
+			static const FVertexName AudioOutputName = "Audio:0";
+			return AudioOutputName;
+		}
+
+		FMetasoundFrontendClassVertex GetClassAudioOutput()
+		{
+			FMetasoundFrontendClassVertex GeneratedAudio;
+			GeneratedAudio.Name = GetAudioOutputName();
+			GeneratedAudio.TypeName = GetMetasoundDataTypeName<FAudioBuffer>();
+			GeneratedAudio.VertexID = FGuid::NewGuid();
+
+#if WITH_EDITOR
+			GeneratedAudio.Metadata.SetDisplayName(LOCTEXT("GeneratedMono", "Audio"));
+			GeneratedAudio.Metadata.SetDescription(LOCTEXT("GeneratedAudioToolTip", "The resulting output audio from this source."));
+#endif // WITH_EDITOR
+
+			return GeneratedAudio;
+		}
+
+		FMetasoundFrontendInterface GetInterface()
+		{
+			FMetasoundFrontendInterface Interface;
+			Interface.Version = GetVersion();
+
+			// Inputs
+			FMetasoundFrontendClassVertex OnPlayTrigger;
+			OnPlayTrigger.Name = DeprecatedInterfacesPrivate::GetOnPlayInputName();
+			OnPlayTrigger.TypeName = GetMetasoundDataTypeName<FTrigger>();
+			OnPlayTrigger.VertexID = FGuid::NewGuid();
+
+#if WITH_EDITOR
+			OnPlayTrigger.Metadata.SetDisplayName(LOCTEXT("OnPlay", "On Play"));
+			OnPlayTrigger.Metadata.SetDescription(LOCTEXT("OnPlayTriggerToolTip", "Trigger executed when this source is played."));
+#endif // WITH_EDITOR
+
+			Interface.Inputs.Add(OnPlayTrigger);
+
+			// Outputs
+			FMetasoundFrontendClassVertex OnFinished;
+			OnFinished.Name = DeprecatedInterfacesPrivate::GetIsFinishedOutputName();
+			OnFinished.TypeName = GetMetasoundDataTypeName<FTrigger>();
+			OnFinished.VertexID = FGuid::NewGuid();
+
+#if WITH_EDITOR
+			OnFinished.Metadata.SetDisplayName(LOCTEXT("OnFinished", "On Finished"));
+			OnFinished.Metadata.SetDescription(LOCTEXT("OnFinishedToolTip", "Trigger executed to initiate stopping the source."));
+#endif // WITH_EDITOR
+
+			Interface.Outputs.Add(OnFinished);
+
+			FMetasoundFrontendClassVertex GeneratedAudio = GetClassAudioOutput();
+			Interface.Outputs.Add(GeneratedAudio);
+
+			// Environment
+			FMetasoundFrontendClassEnvironmentVariable AudioDeviceID;
+			AudioDeviceID.Name = DeprecatedInterfacesPrivate::GetAudioDeviceIDVariableName();
+
+			Interface.Environment.Add(AudioDeviceID);
+
+			return Interface;
+		}
+
+		bool FUpdateInterface::Transform(Frontend::FDocumentHandle InDocument) const
+			{
+				// Swap FMonoAudioFormat output node to an FAudioBuffer output node.
+				using namespace Frontend;
+
+				FGraphHandle Graph = InDocument->GetRootGraph();
+				if (!Graph->IsValid())
+				{
+					return false;
+				}
+
+				InDocument->RemoveInterfaceVersion(MetasoundOutputFormatMonoV1_0::GetInterface().Version);
+				InDocument->AddInterfaceVersion(MetasoundOutputFormatMonoV1_1::GetInterface().Version);
+
+				FNodeHandle MonoFormatOutput = Graph->GetOutputNodeWithName(MetasoundOutputFormatMonoV1_0::GetAudioOutputName());
+				FVector2D MonoFormatLocation;
+
+				FOutputHandle OutputToReconnect = IOutputController::GetInvalidHandle();
+				if (MonoFormatOutput->IsValid())
+				{
+#if WITH_EDITOR
+					// Get the first location.
+					for (auto Location : MonoFormatOutput->GetNodeStyle().Display.Locations)
+					{
+						MonoFormatLocation = Location.Value;
+					}
+#endif // WITH_EDITOR
+
+					// Get connections
+					TArray<FInputHandle> Inputs = MonoFormatOutput->GetInputs();
+					if (ensure(Inputs.Num() == 1))
+					{
+						OutputToReconnect = Inputs[0]->GetConnectedOutput();
+					}
+
+					Graph->RemoveOutputVertex(MetasoundOutputFormatMonoV1_0::GetAudioOutputName());
+				}
+
+				// Create output
+				FNodeHandle BufferOutput = Graph->AddOutputVertex(GetClassAudioOutput());
+				if (ensure(BufferOutput->IsValid()))
+				{
+#if WITH_EDITOR
+					FMetasoundFrontendNodeStyle Style = BufferOutput->GetNodeStyle();
+					Style.Display.Locations.Add(FGuid(), MonoFormatLocation);
+					BufferOutput->SetNodeStyle(Style);
+#endif // WITH_EDITOR
+
+					if (OutputToReconnect->IsValid())
+					{
+						// Reconnect
+						TArray<FInputHandle> Inputs = BufferOutput->GetInputs();
+						if (ensure(Inputs.Num() == 1))
+						{
+							ensure(OutputToReconnect->Connect(*Inputs[0]));
+						}
+					}
+				}
+
+				return true;
+			}
+	}
+
+	namespace MetasoundOutputFormatStereoV1_1
+	{
+		const FMetasoundFrontendVersion& GetVersion()
+		{
+			static const FMetasoundFrontendVersion Version { "StereoSource", FMetasoundFrontendVersionNumber { 1, 1 } };
+			return Version;
+		}
+
+		const FVertexName& GetOnPlayInputName()
+		{
+			static const FVertexName TriggerInputName = "On Play";
+			return TriggerInputName;
+		}
+
+		const FVertexName& GetLeftAudioOutputName()
+		{
+			static const FVertexName AudioOutputName = "Audio:0";
+			return AudioOutputName;
+		}
+
+		const FVertexName& GetRightAudioOutputName()
+		{
+			static const FVertexName AudioOutputName = "Audio:1";
+			return AudioOutputName;
+		}
+
+		FMetasoundFrontendClassVertex GetClassLeftAudioOutput()
+		{
+			FMetasoundFrontendClassVertex GeneratedLeftAudio;
+			GeneratedLeftAudio.Name = GetLeftAudioOutputName();
+			GeneratedLeftAudio.TypeName = GetMetasoundDataTypeName<FAudioBuffer>();
+			GeneratedLeftAudio.VertexID = FGuid::NewGuid();
+
+#if WITH_EDITOR
+			GeneratedLeftAudio.Metadata.SetDisplayName(LOCTEXT("GeneratedStereoLeft", "Left Audio"));
+			GeneratedLeftAudio.Metadata.SetDescription(LOCTEXT("GeneratedLeftAudioToolTip", "The resulting output audio from this source."));
+#endif // WITH_EDITOR
+
+			return GeneratedLeftAudio;
+		}
+
+		FMetasoundFrontendClassVertex GetClassRightAudioOutput()
+		{
+			FMetasoundFrontendClassVertex GeneratedRightAudio;
+			GeneratedRightAudio.Name = GetRightAudioOutputName();
+			GeneratedRightAudio.TypeName = GetMetasoundDataTypeName<FAudioBuffer>();
+
+#if WITH_EDITOR
+			GeneratedRightAudio.Metadata.SetDisplayName(LOCTEXT("GeneratedStereoRight", "Right Audio"));
+			GeneratedRightAudio.Metadata.SetDescription(LOCTEXT("GeneratedRightAudioToolTip", "The resulting output audio from this source."));
+#endif // WITH_EDITOR
+
+			GeneratedRightAudio.VertexID = FGuid::NewGuid();
+
+			return GeneratedRightAudio;
+		}
+
+		FMetasoundFrontendInterface GetInterface()
+		{
+			FMetasoundFrontendInterface Interface;
+			Interface.Version = GetVersion();
+				
+			// Inputs
+			FMetasoundFrontendClassVertex OnPlayTrigger;
+			OnPlayTrigger.Name = DeprecatedInterfacesPrivate::GetOnPlayInputName();
+			OnPlayTrigger.TypeName = GetMetasoundDataTypeName<FTrigger>();
+
+#if WITH_EDITOR
+			OnPlayTrigger.Metadata.SetDisplayName(LOCTEXT("OnPlay", "On Play"));
+			OnPlayTrigger.Metadata.SetDescription(LOCTEXT("OnPlayTriggerToolTip", "Trigger executed when this source is played."));
+#endif // WITH_EDITOR
+
+			OnPlayTrigger.VertexID = FGuid::NewGuid();
+
+			Interface.Inputs.Add(OnPlayTrigger);
+
+			// Outputs
+			FMetasoundFrontendClassVertex OnFinished;
+			OnFinished.Name = DeprecatedInterfacesPrivate::GetIsFinishedOutputName();
+
+			OnFinished.TypeName = GetMetasoundDataTypeName<FTrigger>();
+
+#if WITH_EDITOR
+			OnFinished.Metadata.SetDisplayName(LOCTEXT("OnFinished", "On Finished"));
+			OnFinished.Metadata.SetDescription(LOCTEXT("OnFinishedToolTip", "Trigger executed to initiate stopping the source."));
+#endif // WITH_EDITOR
+
+			OnFinished.VertexID = FGuid::NewGuid();
+
+			Interface.Outputs.Add(OnFinished);
+
+			FMetasoundFrontendClassVertex GeneratedLeftAudio = GetClassLeftAudioOutput();
+			Interface.Outputs.Add(GeneratedLeftAudio);
+
+			FMetasoundFrontendClassVertex GeneratedRightAudio = GetClassRightAudioOutput();
+			Interface.Outputs.Add(GeneratedRightAudio);
+
+			// Environment
+			FMetasoundFrontendClassEnvironmentVariable AudioDeviceID;
+			AudioDeviceID.Name = DeprecatedInterfacesPrivate::GetAudioDeviceIDVariableName();
+			Interface.Environment.Add(AudioDeviceID);
+
+			return Interface;
+		};
+
+		bool FUpdateInterface::Transform(Frontend::FDocumentHandle InDocument) const
+			{
+				using namespace Frontend;
+
+				FGraphHandle Graph = InDocument->GetRootGraph();
+
+				if (!Graph->IsValid())
+				{
+					return false;
+				}
+
+				InDocument->RemoveInterfaceVersion(MetasoundOutputFormatStereoV1_0::GetInterface().Version);
+				InDocument->AddInterfaceVersion(MetasoundOutputFormatStereoV1_1::GetInterface().Version);
+
+				FNodeHandle StereoFormatOutput = Graph->GetOutputNodeWithName(MetasoundOutputFormatStereoV1_0::GetAudioOutputName());
+				FOutputHandle LeftOutputToReconnect = IOutputController::GetInvalidHandle();
+				FOutputHandle RightOutputToReconnect = IOutputController::GetInvalidHandle();
+
+				FVector2D StereoFormatLocation;
+
+				if (StereoFormatOutput->IsValid())
+				{
+#if WITH_EDITOR
+					// Get the first location.
+					for (auto Location : StereoFormatOutput->GetNodeStyle().Display.Locations)
+					{
+						StereoFormatLocation = Location.Value;
+					}
+#endif // WITH_EDITOR
+
+					FInputHandle LeftInput = StereoFormatOutput->GetInputWithVertexName("Left");
+					LeftOutputToReconnect = LeftInput->GetConnectedOutput();
+
+					FInputHandle RightInput = StereoFormatOutput->GetInputWithVertexName("Right");
+					RightOutputToReconnect = RightInput->GetConnectedOutput();
+
+					Graph->RemoveOutputVertex(MetasoundOutputFormatStereoV1_0::GetAudioOutputName());
+				}
+
+				FNodeHandle LeftBufferOutput = Graph->AddOutputVertex(GetClassLeftAudioOutput());
+				if (ensure(LeftBufferOutput->IsValid()))
+				{
+#if WITH_EDITOR
+					FMetasoundFrontendNodeStyle Style = LeftBufferOutput->GetNodeStyle();
+					Style.Display.Locations.Add(FGuid(), StereoFormatLocation);
+					LeftBufferOutput->SetNodeStyle(Style);
+#endif // WITH_EDITOR
+
+					if (LeftOutputToReconnect->IsValid())
+					{
+						TArray<FInputHandle> Inputs = LeftBufferOutput->GetInputs();
+						if (ensure(Inputs.Num() == 1))
+						{
+							ensure(LeftOutputToReconnect->Connect(*Inputs[0]));
+						}
+					}
+				}
+
+				FNodeHandle RightBufferOutput = Graph->AddOutputVertex(GetClassRightAudioOutput());
+				if (ensure(RightBufferOutput->IsValid()))
+				{
+					if (RightOutputToReconnect->IsValid())
+					{
+#if WITH_EDITOR
+						FMetasoundFrontendNodeStyle Style = RightBufferOutput->GetNodeStyle();
+						Style.Display.Locations.Add(FGuid(), StereoFormatLocation + FVector2D{0, 100.f}); // TODO: How should this new output be placed?
+						RightBufferOutput->SetNodeStyle(Style);
+#endif // WITH_EDITOR
+
+						TArray<FInputHandle> Inputs = RightBufferOutput->GetInputs();
+						if (ensure(Inputs.Num() == 1))
+						{
+							ensure(RightOutputToReconnect->Connect(*Inputs[0]));
+						}
+					}
+				}
+
+				return true;
+			}
+	}
+
+	namespace MetasoundOutputFormatMonoV1_2
+	{
+		const FMetasoundFrontendVersion& GetVersion()
+		{
+			const static FMetasoundFrontendVersion Version = { "MonoSource", FMetasoundFrontendVersionNumber { 1, 2 } };
+			return Version;
+		}
+
+		FMetasoundFrontendInterface GetInterface()
+		{
+			FMetasoundFrontendInterface Interface;
+			Interface.Version = GetVersion();
+			return Interface;
+		}
+
+		bool FUpdateInterface::Transform(Frontend::FDocumentHandle InDocument) const
+		{
+			using namespace DeprecatedInterfacesPrivate;
+			using namespace Frontend;
+
+			const TArray<FMetasoundFrontendVersion> InterfacesToRemove
+			{
+				MetasoundOutputFormatMonoV1_1::GetVersion()
+			};
+
+			const TArray<FMetasoundFrontendVersion> InterfacesToAdd
+			{
+				SourceInterfaceV1_0::GetVersion(),
+				OutputFormatMonoInterface::GetVersion()
+			};
+
+			FModifyRootGraphInterfaces InterfaceTransform(InterfacesToRemove, InterfacesToAdd);
+			InterfaceTransform.SetNamePairingFunction(MatchMemberNamesIgnoreSpaces);
+			return InterfaceTransform.Transform(InDocument);
+		}
+	}
+
+	namespace MetasoundOutputFormatStereoV1_2
+	{
+		const FMetasoundFrontendVersion& GetVersion()
+		{
+			const static FMetasoundFrontendVersion Version = { "StereoSource", FMetasoundFrontendVersionNumber { 1, 2 } };
+			return Version;
+		}
+
+		FMetasoundFrontendInterface GetInterface()
+		{
+			FMetasoundFrontendInterface Interface;
+			Interface.Version = GetVersion();
+			return Interface;
+		}
+
+		bool FUpdateInterface::Transform(Frontend::FDocumentHandle InDocument) const
+		{
+			using namespace DeprecatedInterfacesPrivate;
+			using namespace Frontend;
+
+			const TArray<FMetasoundFrontendVersion> InterfacesToRemove
+			{
+				MetasoundOutputFormatStereoV1_1::GetVersion()
+			};
+
+			const TArray<FMetasoundFrontendVersion> InterfacesToAdd
+			{
+				SourceInterfaceV1_0::GetVersion(),
+				OutputFormatStereoInterface::GetVersion()
+			};
+
+			FModifyRootGraphInterfaces InterfaceTransform(InterfacesToRemove, InterfacesToAdd);
+			InterfaceTransform.SetNamePairingFunction(MatchMemberNamesIgnoreSpaces);
+			return InterfaceTransform.Transform(InDocument);
+		}
+	}
+
+	void RegisterDeprecatedInterfaces()
+	{
+		using namespace Frontend;
+
+		const UClass& PatchClass = *UMetaSoundPatch::StaticClass();
+		const UClass& SourceClass = *UMetaSoundSource::StaticClass();
+
+		// Default Patch Interfaces (legacy)
+		RegisterInterfaceForSingleClass(PatchClass, MetasoundV1_0::GetInterface(), nullptr, true, false);
+
+		RegisterInterfaceForSingleClass(SourceClass, SourceInterfaceV1_0::CreateInterface(), nullptr);
+
+		// Set default interface with unset version to use base UMetaSoundPatch class implementation (legacy requirement for 5.0 alpha).
+		RegisterInterfaceForSingleClass(PatchClass, FMetasoundFrontendInterface(), nullptr);
+
+		RegisterInterfaceForSingleClass(SourceClass, MetasoundOutputFormatStereoV1_0::GetInterface(), nullptr);
+		RegisterInterfaceForSingleClass(SourceClass, MetasoundOutputFormatStereoV1_1::GetInterface(), MakeUnique<MetasoundOutputFormatStereoV1_1::FUpdateInterface>());
+		RegisterInterfaceForSingleClass(SourceClass, MetasoundOutputFormatStereoV1_2::GetInterface(), MakeUnique<MetasoundOutputFormatStereoV1_2::FUpdateInterface>());
+
+		RegisterInterfaceForSingleClass(SourceClass, MetasoundOutputFormatMonoV1_0::GetInterface(), nullptr);
+		RegisterInterfaceForSingleClass(SourceClass, MetasoundOutputFormatMonoV1_1::GetInterface(), MakeUnique<MetasoundOutputFormatMonoV1_1::FUpdateInterface>());
+		RegisterInterfaceForSingleClass(SourceClass, MetasoundOutputFormatMonoV1_2::GetInterface(), MakeUnique<MetasoundOutputFormatMonoV1_2::FUpdateInterface>());
+	}
+} // namespace Metasound::Engine
+#undef LOCTEXT_NAMESPACE // MetasoundEngineDeprecatedInterfaces
diff --git a/Engine/Plugins/Runtime/Metasound/Source/MetasoundEngine/Private/Interfaces/MetasoundDeprecatedInterfaces.h b/Engine/Plugins/Runtime/Metasound/Source/MetasoundEngine/Private/Interfaces/MetasoundDeprecatedInterfaces.h
new file mode 100644
--- /dev/null
+++ b/Engine/Plugins/Runtime/Metasound/Source/MetasoundEngine/Private/Interfaces/MetasoundDeprecatedInterfaces.h
@@ -0,0 +1,104 @@
+// Copyright Epic Games, Inc. All Rights Reserved.
+// Copyright Epic Games, Inc. All Rights Reserved.
+// Copyright Epic Games, Inc. All Rights Reserved.
+// Copyright Epic Games, Inc. All Rights Reserved.
+// Copyright Epic Games, Inc. All Rights Reserved.
+// Copyright Epic Games, Inc. All Rights Reserved.
+// Copyright Epic Games, Inc. All Rights Reserved.
+// Copyright Epic Games, Inc. All Rights Reserved.
+// Copyright Epic Games, Inc. All Rights Reserved.
+#pragma once
+
+#include "MetasoundFrontendDocument.h"
+#include "MetasoundFrontendTransform.h"
+#include "MetasoundVertex.h"
+
+
+// DEPRECATED: DeprecatedInterfaces that corrispond with archetypes
+// that preceded the Interface System.  Only remain for purposes of
+// upgrading content generated prior to 5.0 release.
+namespace Metasound::Engine
+{
+	// Base interface without any required inputs or outputs (TODO: Version & Remove)
+	namespace MetasoundV1_0
+	{
+		const FMetasoundFrontendVersion& GetVersion();
+		FMetasoundFrontendInterface GetInterface();
+	}
+
+	// V1.0 of a Metasound Mono output format. Uses FMonoAudioFormat as output.
+	namespace MetasoundOutputFormatMonoV1_0
+	{
+		const FMetasoundFrontendVersion& GetVersion();
+		const FVertexName& GetAudioOutputName();
+		FMetasoundFrontendInterface GetInterface();
+	}
+
+	// V1.0 of a Metasound Stereo output format. Uses FStereoAudioFormat as output.
+	namespace MetasoundOutputFormatStereoV1_0
+	{
+		const FMetasoundFrontendVersion& GetVersion();
+		const FVertexName& GetAudioOutputName();
+		FMetasoundFrontendInterface GetInterface();
+	}
+
+	// V1.1 of a Metasound Mono output format. Uses FAudioBuffer as output.
+	namespace MetasoundOutputFormatMonoV1_1
+	{
+		const FMetasoundFrontendVersion& GetVersion();
+		const FVertexName& GetAudioOutputName();
+		FMetasoundFrontendInterface GetInterface();
+
+		class FUpdateInterface : public Frontend::IDocumentTransform
+		{
+		public:
+			virtual bool Transform(Frontend::FDocumentHandle InDocument) const override;
+		};
+	}
+
+	// V1.2 of a Metasound Mono output format.
+	// Deprecate path from MetasoundOutputFormatStereoV1_1 to SourceInterface & respective OutputFormatMonoInterface
+	namespace MetasoundOutputFormatMonoV1_2
+	{
+		const FMetasoundFrontendVersion& GetVersion();
+		FMetasoundFrontendInterface GetInterface();
+
+		class FUpdateInterface : public Frontend::IDocumentTransform
+		{
+		public:
+			virtual bool Transform(Frontend::FDocumentHandle InDocument) const override;
+		};
+	}
+
+	// V1.1 of a Metasound Stereo output format. Uses two FAudioBuffers as outputs.
+	namespace MetasoundOutputFormatStereoV1_1
+	{
+		const FMetasoundFrontendVersion& GetVersion();
+		const FVertexName& GetLeftAudioOutputName();
+		const FVertexName& GetRightAudioOutputName();
+		FMetasoundFrontendInterface GetInterface();
+
+		class FUpdateInterface : public Frontend::IDocumentTransform
+		{
+		public:
+			virtual bool Transform(Frontend::FDocumentHandle InDocument) const override;
+		};
+	}
+
+	// V1.2 of a Metasound Stereo output format. Uses two FAudioBuffers as outputs.
+	// Deprecate path from MetasoundOutputFormatStereoV1_1 to SourceInterface & respective OutputFormatStereoInterface
+	namespace MetasoundOutputFormatStereoV1_2
+	{
+		const FMetasoundFrontendVersion& GetVersion();
+		FMetasoundFrontendInterface GetInterface();
+
+		// Update from MetasoundOutputFormatStereoV1_0 to MetasoundOutputFormatStereoV1_1
+		class FUpdateInterface : public Frontend::IDocumentTransform
+		{
+		public:
+			virtual bool Transform(Frontend::FDocumentHandle InDocument) const override;
+		};
+	}
+
+	void RegisterDeprecatedInterfaces();
+} // namespace Metasound::Engine
diff --git a/Engine/Plugins/Runtime/Metasound/Source/MetasoundEngine/Private/Interfaces/MetasoundInputFormatInterfaces.cpp b/Engine/Plugins/Runtime/Metasound/Source/MetasoundEngine/Private/Interfaces/MetasoundInputFormatInterfaces.cpp
--- a/Engine/Plugins/Runtime/Metasound/Source/MetasoundEngine/Private/Interfaces/MetasoundInputFormatInterfaces.cpp
+++ b/Engine/Plugins/Runtime/Metasound/Source/MetasoundEngine/Private/Interfaces/MetasoundInputFormatInterfaces.cpp
@@ -52,13 +52,14 @@ namespace Metasound::Engine
 			return CreateInterface();
 		}
 
-		TArray<FMetasoundFrontendInterfaceBinding> CreateOutputBindings()
+		TArray<FMetasoundFrontendInterfaceBinding> CreateBindings()
 		{
 			return TArray<FMetasoundFrontendInterfaceBinding>
 			{
 				{
-					OutputFormatMonoInterface::GetVersion(), // Mono-to-mono takes priority
-					0,
+					OutputFormatMonoInterface::GetVersion(),
+					InputFormatMonoInterface::GetVersion(),
+					0, // Mono-to-mono takes priority
 					{
 						{
 							OutputFormatMonoInterface::Outputs::MonoOut, InputFormatMonoInterface::Inputs::MonoIn,
@@ -67,6 +68,7 @@ namespace Metasound::Engine
 				},
 				{
 					OutputFormatStereoInterface::GetVersion(),
+					InputFormatMonoInterface::GetVersion(),
 					10,
 					{
 						{
@@ -76,6 +78,7 @@ namespace Metasound::Engine
 				},
 				{
 					OutputFormatQuadInterface::GetVersion(),
+					InputFormatMonoInterface::GetVersion(),
 					20,
 					{
 						{
@@ -85,6 +88,7 @@ namespace Metasound::Engine
 				},
 				{
 					OutputFormatFiveDotOneInterface::GetVersion(),
+					InputFormatMonoInterface::GetVersion(),
 					30,
 					{
 						{
@@ -94,6 +98,7 @@ namespace Metasound::Engine
 				},
 				{
 					OutputFormatSevenDotOneInterface::GetVersion(),
+					InputFormatMonoInterface::GetVersion(),
 					40,
 					{
 						{
@@ -158,12 +163,13 @@ namespace Metasound::Engine
 			return CreateInterface();
 		}
 
-		TArray<FMetasoundFrontendInterfaceBinding> CreateOutputBindings()
+		TArray<FMetasoundFrontendInterfaceBinding> CreateBindings()
 		{
 			return TArray<FMetasoundFrontendInterfaceBinding>
 			{
 				{
 					OutputFormatMonoInterface::GetVersion(),
+					InputFormatStereoInterface::GetVersion(),
 					10,
 					{
 						{
@@ -173,6 +179,7 @@ namespace Metasound::Engine
 				},
 				{
 					OutputFormatStereoInterface::GetVersion(),
+					InputFormatStereoInterface::GetVersion(),
 					0,	// Stereo to stereo takes priority
 					{
 						{
@@ -185,6 +192,7 @@ namespace Metasound::Engine
 				},
 				{
 					OutputFormatQuadInterface::GetVersion(),
+					InputFormatStereoInterface::GetVersion(),
 					30,
 					{
 						{
@@ -197,6 +205,7 @@ namespace Metasound::Engine
 				},
 				{
 					OutputFormatFiveDotOneInterface::GetVersion(),
+					InputFormatStereoInterface::GetVersion(),
 					40,
 					{
 						{
@@ -209,6 +218,7 @@ namespace Metasound::Engine
 				},
 				{
 					OutputFormatSevenDotOneInterface::GetVersion(),
+					InputFormatStereoInterface::GetVersion(),
 					50,
 					{
 						{
diff --git a/Engine/Plugins/Runtime/Metasound/Source/MetasoundEngine/Private/Interfaces/MetasoundInterface.cpp b/Engine/Plugins/Runtime/Metasound/Source/MetasoundEngine/Private/Interfaces/MetasoundInterface.cpp
new file mode 100644
--- /dev/null
+++ b/Engine/Plugins/Runtime/Metasound/Source/MetasoundEngine/Private/Interfaces/MetasoundInterface.cpp
@@ -0,0 +1,410 @@
+// Copyright Epic Games, Inc. All Rights Reserved.
+
+#include "Interfaces/MetasoundInterface.h"
+
+#include "Algo/Transform.h"
+#include "AudioParameterControllerInterface.h"
+#include "IAudioParameterInterfaceRegistry.h"
+#include "IAudioParameterTransmitter.h"
+#include "Interfaces/MetasoundFrontendInterfaceRegistry.h"
+#include "Interfaces/MetasoundFrontendSourceInterface.h"
+#include "Interfaces/MetasoundInputFormatInterfaces.h"
+#include "Interfaces/MetasoundOutputFormatInterfaces.h"
+#include "Metasound.h"
+#include "MetasoundFrontendDataTypeRegistry.h"
+#include "MetasoundFrontendDocument.h"
+#include "MetasoundFrontendTransform.h"
+#include "MetasoundLog.h"
+#include "MetasoundParameterTransmitter.h"
+#include "MetasoundSource.h"
+#include "MetasoundUObjectRegistry.h"
+#include "Templates/SharedPointer.h"
+#include "Templates/UniquePtr.h"
+#include "UObject/Class.h"
+#include "UObject/NoExportTypes.h"
+
+
+namespace Metasound::Engine
+{
+	// Entry for registered interface.
+	class FInterfaceRegistryEntry : public Frontend::IInterfaceRegistryEntry
+	{
+		const FInterfaceRegistryUClassOptions* FindOptionsByClassName(FName ClassName) const
+		{
+			auto NameIsEqual = [&ClassName](const FInterfaceRegistryUClassOptions& InOptions)
+			{
+				return InOptions.ClassName == ClassName;
+			};
+
+			return Options.UClassOptions.FindByPredicate(NameIsEqual);
+		}
+
+	public:
+		FInterfaceRegistryEntry(
+			const FMetasoundFrontendInterface& InInterface,
+			TUniquePtr<Frontend::IDocumentTransform>&& InUpdateTransform,
+			FInterfaceRegistryOptions&& InOptions
+		)
+			: Interface(InInterface)
+			, UpdateTransform(MoveTemp(InUpdateTransform))
+			, Options(MoveTemp(InOptions))
+		{
+		}
+
+		virtual bool EditorCanAddOrRemove(FName InUClassName) const override
+		{
+			if (const FInterfaceRegistryUClassOptions* UClassOptions = FindOptionsByClassName(InUClassName))
+			{
+				return UClassOptions->bEditorCanAddOrRemove;
+			}
+
+			return false;
+		}
+
+		virtual bool UClassIsSupported(FName InUClassName) const override
+		{
+			if (const FInterfaceRegistryUClassOptions* UClassOptions = FindOptionsByClassName(InUClassName))
+			{
+				return true;
+			}
+
+			// TODO: Support child asset class types.
+			return false;
+		}
+
+		virtual bool IsDefault(FName InUClassName) const override
+		{
+			if (const FInterfaceRegistryUClassOptions* UClassOptions = FindOptionsByClassName(InUClassName))
+			{
+				return UClassOptions->bIsDefault;
+			}
+
+			return false;
+		}
+
+		virtual FName GetRouterName() const override
+		{
+			return Options.InputSystemName;
+		}
+
+		virtual const FMetasoundFrontendInterface& GetInterface() const override
+		{
+			return Interface;
+		}
+
+		virtual bool UpdateRootGraphInterface(Frontend::FDocumentHandle InDocument) const override
+		{
+			if (UpdateTransform.IsValid())
+			{
+				return UpdateTransform->Transform(InDocument);
+			}
+			return false;
+		}
+
+	private:
+		FMetasoundFrontendInterface Interface;
+		TArray<FMetasoundFrontendInterfaceBinding> OutputBindings;
+		TUniquePtr<Frontend::IDocumentTransform> UpdateTransform;
+		FInterfaceRegistryOptions Options;
+	};
+
+	FMetasoundFrontendInterface ConvertParameterToFrontendInterface(const Audio::FParameterInterface& InInterface)
+	{
+		using namespace Frontend;
+
+		auto ResolveMemberDataType = [](FName DataType, EAudioParameterType ParamType)
+		{
+			if (!DataType.IsNone())
+			{
+				const bool bIsRegisteredType = Frontend::IDataTypeRegistry::Get().IsRegistered(DataType);
+				if (ensureAlwaysMsgf(bIsRegisteredType, TEXT("Attempting to register Interface member with unregistered DataType '%s'."), *DataType.ToString()))
+				{
+					return DataType;
+				}
+			}
+
+			return ConvertParameterToDataType(ParamType);
+		};
+
+		FMetasoundFrontendInterface FrontendInterface;
+		FrontendInterface.Version = { InInterface.GetName(), FMetasoundFrontendVersionNumber { InInterface.GetVersion().Major, InInterface.GetVersion().Minor } };
+
+		// Transfer all input data from AudioExtension interface struct to FrontendInterface
+		{
+			Algo::Transform(InInterface.GetInputs(), FrontendInterface.Inputs, [&](const Audio::FParameterInterface::FInput& Input)
+			{
+				FMetasoundFrontendClassInput ClassInput;
+				ClassInput.Name = Input.InitValue.ParamName;
+				ClassInput.DefaultLiteral = FMetasoundFrontendLiteral(Input.InitValue);
+				ClassInput.TypeName = ResolveMemberDataType(Input.DataType, Input.InitValue.ParamType);
+
+#if WITH_EDITOR
+				// Interfaces should never serialize text to avoid desync between
+				// copied versions serialized in assets and those defined in code.
+				ClassInput.Metadata.SetSerializeText(false);
+				ClassInput.Metadata.SetDisplayName(Input.DisplayName);
+				ClassInput.Metadata.SetDescription(Input.Description);
+				ClassInput.Metadata.SortOrderIndex = Input.SortOrderIndex;
+					
+				FrontendInterface.AddSortOrderToInputStyle(Input.SortOrderIndex);
+
+				// Setup required inputs by telling the style that the input is required
+				// This will later be validated against.
+				if (!Input.RequiredText.IsEmpty())
+				{
+					FrontendInterface.AddRequiredInputToStyle(Input.InitValue.ParamName, Input.RequiredText);
+				}
+#endif // WITH_EDITOR
+
+				ClassInput.VertexID = FGuid::NewGuid();
+
+				return ClassInput;
+			});
+		}
+
+		// Transfer all output data from AudioExtension interface struct to FrontendInterface
+		{
+			Algo::Transform(InInterface.GetOutputs(), FrontendInterface.Outputs, [&](const Audio::FParameterInterface::FOutput& Output)
+			{
+				FMetasoundFrontendClassOutput ClassOutput;
+				ClassOutput.Name = Output.ParamName;
+				ClassOutput.TypeName = ResolveMemberDataType(Output.DataType, Output.ParamType);
+
+#if WITH_EDITOR
+				// Interfaces should never serialize text to avoid desync between
+				// copied versions serialized in assets and those defined in code.
+				ClassOutput.Metadata.SetSerializeText(false);
+				ClassOutput.Metadata.SetDisplayName(Output.DisplayName);
+				ClassOutput.Metadata.SetDescription(Output.Description);
+				ClassOutput.Metadata.SortOrderIndex = Output.SortOrderIndex;
+					
+				FrontendInterface.AddSortOrderToOutputStyle(Output.SortOrderIndex);
+
+				// Setup required outputs by telling the style that the output is required
+				// This will later be validated against.
+				if (!Output.RequiredText.IsEmpty())
+				{
+					FrontendInterface.AddRequiredOutputToStyle(Output.ParamName, Output.RequiredText);
+				}
+#endif // WITH_EDITOR
+
+				ClassOutput.VertexID = FGuid::NewGuid();
+
+				return ClassOutput;
+			});
+		}
+
+		Algo::Transform(InInterface.GetEnvironment(), FrontendInterface.Environment, [&](const Audio::FParameterInterface::FEnvironmentVariable& Environment)
+		{
+			FMetasoundFrontendClassEnvironmentVariable EnvironmentVariable;
+			EnvironmentVariable.Name = Environment.ParamName;
+
+			// Disabled as it isn't used to infer type when getting/setting at a lower level.
+			// TODO: Either remove type info for environment variables all together or enforce type.
+			// EnvironmentVariable.TypeName = ResolveMemberDataType(Environment.DataType, Environment.ParamType);
+
+			return EnvironmentVariable;
+		});
+
+		return FrontendInterface;
+	}
+
+	void RegisterAudioFormatInterfaces()
+	{
+		using namespace Frontend;
+
+		const FName SourceClassName = UMetaSoundSource::StaticClass()->GetFName();
+		const FName PatchClassName = UMetaSoundPatch::StaticClass()->GetFName();
+
+		// Input Formats
+		{
+			RegisterInterface(InputFormatMonoInterface::CreateInterface(), nullptr, FInterfaceRegistryOptions
+				{
+					IDataReference::RouterName,
+					{
+						{ PatchClassName, false /* bIsDefault */, true /* bEditorCanAddOrRemove */ },
+						{ SourceClassName },
+					}
+				});
+
+			RegisterInterface(InputFormatStereoInterface::CreateInterface(), nullptr, FInterfaceRegistryOptions
+				{
+					IDataReference::RouterName,
+					{
+						{ PatchClassName, false /* bIsDefault */, true /* bEditorCanAddOrRemove */ },
+						{ SourceClassName },
+					}
+				});
+		}
+
+		// Output Formats
+		{
+			RegisterInterface(OutputFormatMonoInterface::CreateInterface(), nullptr, FInterfaceRegistryOptions
+				{
+					IDataReference::RouterName,
+					{
+						{ PatchClassName, false /* bIsDefault */, true /* bEditorCanAddOrRemove */ },
+						{ SourceClassName, true /* bIsDefault */, false /* bEditorCanAddOrRemove */ },
+					}
+				});
+
+			RegisterInterface(OutputFormatStereoInterface::CreateInterface(), nullptr, FInterfaceRegistryOptions
+				{
+					IDataReference::RouterName,
+					{
+						{ PatchClassName, false /* bIsDefault */, true /* bEditorCanAddOrRemove */ },
+						{ SourceClassName },
+					}
+				});
+
+			RegisterInterface(OutputFormatQuadInterface::CreateInterface(), nullptr, FInterfaceRegistryOptions
+				{
+					IDataReference::RouterName,
+					{
+						{ PatchClassName, false /* bIsDefault */, true /* bEditorCanAddOrRemove */ },
+						{ SourceClassName },
+					}
+				});
+
+			RegisterInterface(OutputFormatFiveDotOneInterface::CreateInterface(), nullptr, FInterfaceRegistryOptions
+				{
+					IDataReference::RouterName,
+					{
+						{ PatchClassName, false /* bIsDefault */, true /* bEditorCanAddOrRemove */ },
+						{ SourceClassName },
+					}
+				});
+
+			RegisterInterface(OutputFormatSevenDotOneInterface::CreateInterface(), nullptr, FInterfaceRegistryOptions
+				{
+					IDataReference::RouterName,
+					{
+						{ PatchClassName, false /* bIsDefault */, true /* bEditorCanAddOrRemove */ },
+						{ SourceClassName },
+					}
+				});
+		}
+	}
+
+	void RegisterExternalInterfaces()
+	{
+		using namespace Frontend;
+
+		// Register External Interfaces (Interfaces defined externally & can be managed directly by end-user).
+		auto RegisterExternalInterface = [](Audio::FParameterInterfacePtr Interface)
+		{
+			FInterfaceRegistryOptions Options{ Audio::IParameterTransmitter::RouterName };
+
+			// Currently, no externally defined interfaces can be added as default for protection against undesired
+			// interfaces automatically being added when creating a new MetaSound asset through the editor. Also,
+			// all parameter interfaces are enabled in the editor for addition/removal.
+			constexpr bool bIsDefault = false;
+			constexpr bool bEditorCanAddOrRemove = true;
+
+			bool bInterfaceSupportsPatch = true;
+			bool bInterfaceSupportsSource = true;
+
+			const TArray<const UClass*> SupportedUClasses = Interface->FindSupportedUClasses();
+			if (!SupportedUClasses.IsEmpty())
+			{
+				bInterfaceSupportsPatch = false;
+				bInterfaceSupportsSource = false;
+				for (const UClass* SupportedUClass : SupportedUClasses)
+				{
+					const UClass* PatchClass = UMetaSoundPatch::StaticClass();
+					check(PatchClass);
+					bInterfaceSupportsPatch |= PatchClass->IsChildOf(SupportedUClass);
+
+					const UClass* SourceClass = UMetaSoundSource::StaticClass();
+					check(SourceClass);
+					bInterfaceSupportsSource |= SourceClass->IsChildOf(SupportedUClass);
+				}
+			}
+
+			const bool bSupported = bInterfaceSupportsSource || bInterfaceSupportsPatch;
+			if (ensureAlwaysMsgf(bSupported, TEXT("Parameter interface '%s' type not supported by MetaSounds"), *Interface->GetName().ToString()))
+			{
+				const FMetasoundFrontendInterface FrontendInterface = ConvertParameterToFrontendInterface(*Interface);
+				if (bInterfaceSupportsPatch)
+				{
+					const FName ClassName = UMetaSoundPatch::StaticClass()->GetFName();
+					Options.UClassOptions.Add(FInterfaceRegistryUClassOptions{ ClassName, bIsDefault, bEditorCanAddOrRemove });
+				}
+
+				if (bInterfaceSupportsSource)
+				{
+					const FName ClassName = UMetaSoundSource::StaticClass()->GetFName();
+					Options.UClassOptions.Add(FInterfaceRegistryUClassOptions{ ClassName, bIsDefault, bEditorCanAddOrRemove });
+				}
+
+				IInterfaceRegistry::Get().RegisterInterface(MakeUnique<FInterfaceRegistryEntry>(FrontendInterface, nullptr, MoveTemp(Options)));
+			}
+		};
+
+		Audio::IAudioParameterInterfaceRegistry::Get().IterateInterfaces(RegisterExternalInterface);
+		Audio::IAudioParameterInterfaceRegistry::Get().OnRegistration(MoveTemp(RegisterExternalInterface));
+	}
+
+	void RegisterInterface(const FMetasoundFrontendInterface& InInterface, TUniquePtr<Frontend::IDocumentTransform>&& InUpdateTransform, FInterfaceRegistryOptions&& InOptions)
+	{
+		using namespace Frontend;
+		IInterfaceRegistry::Get().RegisterInterface(MakeUnique<FInterfaceRegistryEntry>(InInterface, MoveTemp(InUpdateTransform), MoveTemp(InOptions)));
+	}
+
+	void RegisterInterface(Audio::FParameterInterfacePtr Interface, TUniquePtr<Frontend::IDocumentTransform>&& InUpdateTransform, FInterfaceRegistryOptions&& InOptions)
+	{
+		using namespace Frontend;
+
+		const FMetasoundFrontendInterface FrontendInterface = ConvertParameterToFrontendInterface(*Interface);
+		IInterfaceRegistry::Get().RegisterInterface(MakeUnique<FInterfaceRegistryEntry>(FrontendInterface, MoveTemp(InUpdateTransform), MoveTemp(InOptions)));
+	}
+
+	void RegisterInterfaceForSingleClass(const UClass& InClass, Audio::FParameterInterfacePtr Interface, TUniquePtr<Frontend::IDocumentTransform>&& InUpdateTransform, bool bInIsDefault, bool bInEditorCanAddOrRemove, FName InRouterName)
+	{
+		using namespace Frontend;
+
+		const FMetasoundFrontendInterface FrontendInterface = ConvertParameterToFrontendInterface(*Interface);
+		RegisterInterface(FrontendInterface, MoveTemp(InUpdateTransform), FInterfaceRegistryOptions
+			{
+				InRouterName,
+				{
+					{
+						InClass.GetFName(),
+						bInIsDefault,
+						bInEditorCanAddOrRemove
+					}
+				}
+			});
+	}
+
+	void RegisterInterfaceForSingleClass(const UClass& InClass, const FMetasoundFrontendInterface& InInterface, TUniquePtr<Frontend::IDocumentTransform>&& InUpdateTransform, bool bInIsDefault, bool bInEditorCanAddOrRemove, FName InRouterName)
+	{
+		using namespace Frontend;
+
+		RegisterInterface(InInterface, MoveTemp(InUpdateTransform), FInterfaceRegistryOptions
+			{
+				InRouterName,
+				{
+					{
+						InClass.GetFName(),
+						bInIsDefault,
+						bInEditorCanAddOrRemove
+					}
+				}
+			});
+	}
+
+	void RegisterInterfaces()
+	{
+		using namespace Frontend;
+
+		// Default Source Interfaces
+		{
+			RegisterInterfaceForSingleClass(*UMetaSoundSource::StaticClass(), SourceInterface::CreateInterface(), MakeUnique<SourceInterface::FUpdateInterface>(), true, false);
+			RegisterInterfaceForSingleClass(*UMetaSoundSource::StaticClass(), SourceOneShotInterface::CreateInterface(), nullptr, true, true);
+		}
+
+		RegisterAudioFormatInterfaces();
+		RegisterExternalInterfaces();
+	}
+} // namespace Metasound::Engine
diff --git a/Engine/Plugins/Runtime/Metasound/Source/MetasoundEngine/Private/Interfaces/MetasoundInterface.h b/Engine/Plugins/Runtime/Metasound/Source/MetasoundEngine/Private/Interfaces/MetasoundInterface.h
new file mode 100644
--- /dev/null
+++ b/Engine/Plugins/Runtime/Metasound/Source/MetasoundEngine/Private/Interfaces/MetasoundInterface.h
@@ -0,0 +1,31 @@
+// Copyright Epic Games, Inc. All Rights Reserved.
+#pragma once
+
+#include "IAudioParameterInterfaceRegistry.h"
+#include "MetasoundFrontendDocument.h"
+#include "MetasoundFrontendTransform.h"
+
+
+namespace Metasound::Engine
+{
+	struct FInterfaceRegistryUClassOptions
+	{
+		FName ClassName;
+		bool bIsDefault = false;
+		bool bEditorCanAddOrRemove = false;
+	};
+
+	struct FInterfaceRegistryOptions
+	{
+		FName InputSystemName;
+		TArray<FInterfaceRegistryUClassOptions> UClassOptions;
+	};
+
+	FMetasoundFrontendInterface ConvertParameterToFrontendInterface(const Audio::FParameterInterface& InInterface);
+
+	void RegisterInterfaces();
+	void RegisterInterface(Audio::FParameterInterfacePtr Interface, TUniquePtr<Frontend::IDocumentTransform>&& InUpdateTransform, FInterfaceRegistryOptions&& InOptions);
+	void RegisterInterface(const FMetasoundFrontendInterface& InInterface, TUniquePtr<Frontend::IDocumentTransform>&& InUpdateTransform, FInterfaceRegistryOptions&& InOptions);
+	void RegisterInterfaceForSingleClass(const UClass& InClass, Audio::FParameterInterfacePtr Interface, TUniquePtr<Frontend::IDocumentTransform>&& InUpdateTransform, bool bInIsDefault = false, bool bInEditorCanAddOrRemove = false, FName InRouterName = IDataReference::RouterName);
+	void RegisterInterfaceForSingleClass(const UClass& InClass, const FMetasoundFrontendInterface& InInterface, TUniquePtr<Frontend::IDocumentTransform>&& InUpdateTransform, bool bInIsDefault = false, bool bInEditorCanAddOrRemove = false, FName InRouterName = IDataReference::RouterName);
+} // namespace Metasound::Engine
diff --git a/Engine/Plugins/Runtime/Metasound/Source/MetasoundEngine/Private/Interfaces/MetasoundInterfaceBindings.cpp b/Engine/Plugins/Runtime/Metasound/Source/MetasoundEngine/Private/Interfaces/MetasoundInterfaceBindings.cpp
new file mode 100644
--- /dev/null
+++ b/Engine/Plugins/Runtime/Metasound/Source/MetasoundEngine/Private/Interfaces/MetasoundInterfaceBindings.cpp
@@ -0,0 +1,40 @@
+// Copyright Epic Games, Inc. All Rights Reserved.
+#include "Interfaces/MetasoundInterfaceBindings.h"
+
+#include "Algo/Transform.h"
+#include "Interfaces/MetasoundFrontendInterfaceBindingRegistry.h"
+#include "Interfaces/MetasoundInputFormatInterfaces.h"
+#include "MetasoundFrontendDocument.h"
+
+
+namespace Metasound::Engine
+{
+	void RegisterInternalInterfaceBindings()
+	{
+		RegisterInterfaceBindings(InputFormatMonoInterface::CreateBindings());
+		RegisterInterfaceBindings(InputFormatStereoInterface::CreateBindings());
+	}
+
+	void RegisterInterfaceBindings(const TArray<FMetasoundFrontendInterfaceBinding>& InBindings)
+	{
+		using namespace Metasound::Frontend;
+
+		for (const FMetasoundFrontendInterfaceBinding& Binding : InBindings)
+		{
+			TUniquePtr<FInterfaceBindingRegistryEntry> Entry = MakeUnique<FInterfaceBindingRegistryEntry>(Binding);
+			IInterfaceBindingRegistry::Get().RegisterInterfaceBinding(Binding.InputInterfaceVersion, MoveTemp(Entry));
+		}
+	}
+
+	bool UnregisterInterfaceBinding(const FMetasoundFrontendVersion& InInputInterfaceVersion, const FMetasoundFrontendVersion& InOutputInterfaceVersion)
+	{
+		using namespace Metasound::Frontend;
+		return IInterfaceBindingRegistry::Get().UnregisterInterfaceBinding(InInputInterfaceVersion, InOutputInterfaceVersion);
+	}
+
+	bool UnregisterAllInterfaceBindings(const FMetasoundFrontendVersion& InInputInterfaceVersion)
+	{
+		using namespace Metasound::Frontend;
+		return IInterfaceBindingRegistry::Get().UnregisterAllInterfaceBindings(InInputInterfaceVersion);
+	}
+} // namespace Metasound::Engine
diff --git a/Engine/Plugins/Runtime/Metasound/Source/MetasoundEngine/Private/Interfaces/MetasoundInterfaceBindingsPrivate.h b/Engine/Plugins/Runtime/Metasound/Source/MetasoundEngine/Private/Interfaces/MetasoundInterfaceBindingsPrivate.h
new file mode 100644
--- /dev/null
+++ b/Engine/Plugins/Runtime/Metasound/Source/MetasoundEngine/Private/Interfaces/MetasoundInterfaceBindingsPrivate.h
@@ -0,0 +1,8 @@
+// Copyright Epic Games, Inc. All Rights Reserved.
+#pragma once
+
+
+namespace Metasound::Engine
+{
+	void RegisterInternalInterfaceBindings();
+} // namespace Metasound::Engine
diff --git a/Engine/Plugins/Runtime/Metasound/Source/MetasoundEngine/Private/Metasound.cpp b/Engine/Plugins/Runtime/Metasound/Source/MetasoundEngine/Private/Metasound.cpp
--- a/Engine/Plugins/Runtime/Metasound/Source/MetasoundEngine/Private/Metasound.cpp
+++ b/Engine/Plugins/Runtime/Metasound/Source/MetasoundEngine/Private/Metasound.cpp
@@ -6,7 +6,6 @@
 #include "Logging/TokenizedMessage.h"
 #include "MetasoundAssetBase.h"
 #include "MetasoundAudioFormats.h"
-#include "MetasoundEngineArchetypes.h"
 #include "MetasoundEngineAsset.h"
 #include "MetasoundEngineEnvironment.h"
 #include "MetasoundEnvironment.h"
diff --git a/Engine/Plugins/Runtime/Metasound/Source/MetasoundEngine/Private/MetasoundEngineArchetypes.cpp b/Engine/Plugins/Runtime/Metasound/Source/MetasoundEngine/Private/MetasoundEngineArchetypes.cpp
deleted file mode 100644
diff --git a/Engine/Plugins/Runtime/Metasound/Source/MetasoundEngine/Private/MetasoundEngineArchetypes.h b/Engine/Plugins/Runtime/Metasound/Source/MetasoundEngine/Private/MetasoundEngineArchetypes.h
deleted file mode 100644
diff --git a/Engine/Plugins/Runtime/Metasound/Source/MetasoundEngine/Private/MetasoundEngineModule.cpp b/Engine/Plugins/Runtime/Metasound/Source/MetasoundEngine/Private/MetasoundEngineModule.cpp
--- a/Engine/Plugins/Runtime/Metasound/Source/MetasoundEngine/Private/MetasoundEngineModule.cpp
+++ b/Engine/Plugins/Runtime/Metasound/Source/MetasoundEngine/Private/MetasoundEngineModule.cpp
@@ -2,12 +2,13 @@
 #include "MetasoundEngineModule.h"
 
 #include "Analysis/MetasoundFrontendAnalyzerRegistry.h"
+#include "Interfaces/MetasoundDeprecatedInterfaces.h"
+#include "Interfaces/MetasoundInterface.h"
+#include "Interfaces/MetasoundInterfaceBindingsPrivate.h"
 #include "Metasound.h"
 #include "MetasoundDataReference.h"
 #include "MetasoundDataTypeRegistrationMacro.h"
-#include "MetasoundEngineArchetypes.h"
 #include "MetasoundFrontendDataTypeTraits.h"
-#include "MetasoundInterface.h"
 #include "MetasoundSource.h"
 #include "MetasoundTrace.h"
 #include "MetasoundUObjectRegistry.h"
@@ -30,6 +31,9 @@ class FMetasoundEngineModule : public IMetasoundEngineModule
 {
 	virtual void StartupModule() override
 	{
+		using namespace Metasound;
+		using namespace Metasound::Engine;
+
 		METASOUND_LLM_SCOPE;
 		FModuleManager::Get().LoadModuleChecked("MetasoundGraphCore");
 		FModuleManager::Get().LoadModuleChecked("MetasoundFrontend");
@@ -44,9 +48,12 @@ class FMetasoundEngineModule : public IMetasoundEngineModule
 		check(AudioSettings);
 		AudioSettings->RegisterParameterInterfaces();
 
-		// Register supported MetaSound UClasses & interfaces
-		Metasound::Engine::RegisterUClasses();
-		Metasound::Engine::RegisterInterfaces();
+		IMetasoundUObjectRegistry::Get().RegisterUClass(MakeUnique<TMetasoundUObjectRegistryEntry<UMetaSoundPatch>>());
+		IMetasoundUObjectRegistry::Get().RegisterUClass(MakeUnique<TMetasoundUObjectRegistryEntry<UMetaSoundSource>>());
+
+		Engine::RegisterDeprecatedInterfaces();
+		Engine::RegisterInterfaces();
+		Engine::RegisterInternalInterfaceBindings();
 
 		// Flush node registration queue
 		FMetasoundFrontendRegistryContainer::Get()->RegisterPendingNodes();
diff --git a/Engine/Plugins/Runtime/Metasound/Source/MetasoundEngine/Private/MetasoundInterface.cpp b/Engine/Plugins/Runtime/Metasound/Source/MetasoundEngine/Private/MetasoundInterface.cpp
deleted file mode 100644
diff --git a/Engine/Plugins/Runtime/Metasound/Source/MetasoundEngine/Private/MetasoundInterface.h b/Engine/Plugins/Runtime/Metasound/Source/MetasoundEngine/Private/MetasoundInterface.h
deleted file mode 100644
diff --git a/Engine/Plugins/Runtime/Metasound/Source/MetasoundEngine/Private/MetasoundSource.cpp b/Engine/Plugins/Runtime/Metasound/Source/MetasoundEngine/Private/MetasoundSource.cpp
--- a/Engine/Plugins/Runtime/Metasound/Source/MetasoundEngine/Private/MetasoundSource.cpp
+++ b/Engine/Plugins/Runtime/Metasound/Source/MetasoundEngine/Private/MetasoundSource.cpp
@@ -13,7 +13,6 @@
 #include "MetasoundAssetManager.h"
 #include "MetasoundAudioFormats.h"
 #include "MetasoundBuilderSubsystem.h"
-#include "MetasoundEngineArchetypes.h"
 #include "MetasoundEngineAsset.h"
 #include "MetasoundEngineEnvironment.h"
 #include "MetasoundEnvironment.h"
diff --git a/Engine/Plugins/Runtime/Metasound/Source/MetasoundEngine/Private/MetasoundUObjectRegistry.cpp b/Engine/Plugins/Runtime/Metasound/Source/MetasoundEngine/Private/MetasoundUObjectRegistry.cpp
--- a/Engine/Plugins/Runtime/Metasound/Source/MetasoundEngine/Private/MetasoundUObjectRegistry.cpp
+++ b/Engine/Plugins/Runtime/Metasound/Source/MetasoundEngine/Private/MetasoundUObjectRegistry.cpp
@@ -10,7 +10,6 @@
 #include "Engine/AssetManager.h"
 #include "Metasound.h"
 #include "MetasoundAssetBase.h"
-#include "MetasoundFrontendArchetypeRegistry.h"
 #include "MetasoundFrontendRegistries.h"
 #include "MetasoundFrontendTransform.h"
 #include "MetasoundSettings.h"
@@ -18,6 +17,7 @@
 #include "MetasoundTrace.h"
 #include "UObject/Object.h"
 
+
 namespace Metasound
 {
 	class FMetasoundUObjectRegistry : public IMetasoundUObjectRegistry
@@ -115,5 +115,4 @@ namespace Metasound
 		static FMetasoundUObjectRegistry Registry;
 		return Registry;
 	}
-}
-
+} // namespace Metasound
diff --git a/Engine/Plugins/Runtime/Metasound/Source/MetasoundEngine/Public/Interfaces/MetasoundInputFormatInterfaces.h b/Engine/Plugins/Runtime/Metasound/Source/MetasoundEngine/Public/Interfaces/MetasoundInputFormatInterfaces.h
--- a/Engine/Plugins/Runtime/Metasound/Source/MetasoundEngine/Public/Interfaces/MetasoundInputFormatInterfaces.h
+++ b/Engine/Plugins/Runtime/Metasound/Source/MetasoundEngine/Public/Interfaces/MetasoundInputFormatInterfaces.h
@@ -17,7 +17,7 @@ namespace Metasound::Engine
 
 		METASOUNDENGINE_API const FMetasoundFrontendVersion& GetVersion();
 		METASOUNDENGINE_API Audio::FParameterInterfacePtr CreateInterface();
-		METASOUNDENGINE_API TArray<FMetasoundFrontendInterfaceBinding> CreateOutputBindings();
+		METASOUNDENGINE_API TArray<FMetasoundFrontendInterfaceBinding> CreateBindings();
 	} // namespace InputFormatMonoInterface
 
 	namespace InputFormatStereoInterface
@@ -30,6 +30,6 @@ namespace Metasound::Engine
 
 		METASOUNDENGINE_API const FMetasoundFrontendVersion& GetVersion();
 		METASOUNDENGINE_API Audio::FParameterInterfacePtr CreateInterface();
-		METASOUNDENGINE_API TArray<FMetasoundFrontendInterfaceBinding> CreateOutputBindings();
+		METASOUNDENGINE_API TArray<FMetasoundFrontendInterfaceBinding> CreateBindings();
 	} // namespace InputFormatStereoInterface
 } // namespace Metasound::Engine
diff --git a/Engine/Plugins/Runtime/Metasound/Source/MetasoundEngine/Public/Interfaces/MetasoundInterfaceBindings.h b/Engine/Plugins/Runtime/Metasound/Source/MetasoundEngine/Public/Interfaces/MetasoundInterfaceBindings.h
new file mode 100644
--- /dev/null
+++ b/Engine/Plugins/Runtime/Metasound/Source/MetasoundEngine/Public/Interfaces/MetasoundInterfaceBindings.h
@@ -0,0 +1,12 @@
+// Copyright Epic Games, Inc. All Rights Reserved.
+#pragma once
+
+#include "MetasoundFrontendDocument.h"
+
+
+namespace Metasound::Engine
+{
+	METASOUNDENGINE_API void RegisterInterfaceBindings(const TArray<FMetasoundFrontendInterfaceBinding>& InBindings);
+	METASOUNDENGINE_API bool UnregisterInterfaceBinding(const FMetasoundFrontendVersion& InInputInterfaceVersion, const FMetasoundFrontendVersion& InOutputInterfaceVersion);
+	METASOUNDENGINE_API bool UnregisterAllInterfaceBindings(const FMetasoundFrontendVersion& InInputInterfaceVersion);
+} // namespace Metasound::Engine
diff --git a/Engine/Plugins/Runtime/Metasound/Source/MetasoundFrontend/Private/Interfaces/MetasoundFrontendInterfaceBindingRegistry.cpp b/Engine/Plugins/Runtime/Metasound/Source/MetasoundFrontend/Private/Interfaces/MetasoundFrontendInterfaceBindingRegistry.cpp
new file mode 100644
--- /dev/null
+++ b/Engine/Plugins/Runtime/Metasound/Source/MetasoundFrontend/Private/Interfaces/MetasoundFrontendInterfaceBindingRegistry.cpp
@@ -0,0 +1,90 @@
+// Copyright Epic Games, Inc. All Rights Reserved.
+#include "Interfaces/MetasoundFrontendInterfaceBindingRegistry.h"
+
+#include "HAL/PlatformTime.h"
+#include "Interfaces/MetasoundFrontendInterfaceRegistry.h"
+#include "Interfaces/MetasoundFrontendInterfaceBindingRegistryPrivate.h"
+#include "MetasoundFrontendRegistryTransaction.h"
+#include "MetasoundTrace.h"
+
+
+namespace Metasound::Frontend
+{
+	FInterfaceBindingRegistryEntry::FInterfaceBindingRegistryEntry(const FMetasoundFrontendInterfaceBinding& InBinding)
+		: OutputInterfaceVersion(InBinding.OutputInterfaceVersion)
+		, VertexBindings(InBinding.VertexBindings)
+		, BindingPriority(InBinding.BindingPriority)
+	{
+	}
+
+	const FMetasoundFrontendVersion& FInterfaceBindingRegistryEntry::GetOutputInterfaceVersion() const
+	{
+		return OutputInterfaceVersion;
+	}
+
+	const TArray<FMetasoundFrontendInterfaceVertexBinding>& FInterfaceBindingRegistryEntry::GetVertexBindings() const
+	{
+		return VertexBindings;
+	}
+
+	int32 FInterfaceBindingRegistryEntry::GetBindingPriority() const
+	{
+		return BindingPriority;
+	}
+
+	bool FInterfaceBindingRegistry::FindInterfaceBindingEntries(const FMetasoundFrontendVersion& InInputInterfaceVersion, TArray<const FInterfaceBindingRegistryEntry*>& OutEntries) const
+	{
+		TArray<const TUniquePtr<FInterfaceBindingRegistryEntry>*> FoundEntries;
+		const FInterfaceRegistryKey Key = GetInterfaceRegistryKey(InInputInterfaceVersion);
+		Entries.MultiFindPointer(Key, FoundEntries);
+
+		OutEntries.Reset();
+		Algo::Transform(FoundEntries, OutEntries, [](const TUniquePtr<FInterfaceBindingRegistryEntry>* FoundEntry)
+		{
+			const FInterfaceBindingRegistryEntry* EntryPtr = FoundEntry->Get();
+			return EntryPtr;
+		});
+
+		return !OutEntries.IsEmpty();
+	}
+
+	void FInterfaceBindingRegistry::RegisterInterfaceBinding(const FMetasoundFrontendVersion& InInputInterfaceVersion, TUniquePtr<FInterfaceBindingRegistryEntry>&& InEntry)
+	{
+		FInterfaceRegistryKey Key = GetInterfaceRegistryKey(InInputInterfaceVersion);
+		Entries.Add(MoveTemp(Key), MoveTemp(InEntry));
+	}
+
+	bool FInterfaceBindingRegistry::UnregisterInterfaceBinding(const FMetasoundFrontendVersion& InInputInterfaceVersion, const FMetasoundFrontendVersion& InOutputInterfaceVersion)
+	{
+		TArray<const TUniquePtr<FInterfaceBindingRegistryEntry>*> FoundEntries;
+		const FInterfaceRegistryKey Key = GetInterfaceRegistryKey(InInputInterfaceVersion);
+		Entries.MultiFindPointer(Key, FoundEntries);
+
+		for (const TUniquePtr<FInterfaceBindingRegistryEntry>* EntryPtr : FoundEntries)
+		{
+			const FInterfaceBindingRegistryEntry* Entry = (*EntryPtr).Get();
+			check(Entry);
+			if (Entry->GetOutputInterfaceVersion() == InOutputInterfaceVersion)
+			{
+				if (Entries.RemoveSingle(Key, *EntryPtr) > 0)
+				{
+					return true;
+				}
+			}
+		}
+
+		return false;
+	}
+
+	bool FInterfaceBindingRegistry::UnregisterAllInterfaceBindings(const FMetasoundFrontendVersion& InInputInterfaceVersion)
+	{
+		FInterfaceRegistryKey Key = GetInterfaceRegistryKey(InInputInterfaceVersion);
+		return Entries.Remove(Key) > 0;
+	}
+
+	IInterfaceBindingRegistry& IInterfaceBindingRegistry::Get()
+	{
+		static FInterfaceBindingRegistry Registry;
+		return Registry;
+	}
+} // namespace Metasound::Frontend
diff --git a/Engine/Plugins/Runtime/Metasound/Source/MetasoundFrontend/Private/Interfaces/MetasoundFrontendInterfaceBindingRegistryPrivate.h b/Engine/Plugins/Runtime/Metasound/Source/MetasoundFrontend/Private/Interfaces/MetasoundFrontendInterfaceBindingRegistryPrivate.h
new file mode 100644
--- /dev/null
+++ b/Engine/Plugins/Runtime/Metasound/Source/MetasoundFrontend/Private/Interfaces/MetasoundFrontendInterfaceBindingRegistryPrivate.h
@@ -0,0 +1,31 @@
+// Copyright Epic Games, Inc. All Rights Reserved.
+#pragma once
+
+
+#include "Interfaces/MetasoundFrontendInterfaceRegistry.h"
+#include "MetasoundFrontendDocument.h"
+#include "MetasoundFrontendRegistryTransaction.h"
+#include "UObject/NoExportTypes.h"
+
+
+namespace Metasound::Frontend
+{
+	class FInterfaceBindingRegistry : public IInterfaceBindingRegistry
+	{
+	public:
+		FInterfaceBindingRegistry() = default;
+
+		virtual ~FInterfaceBindingRegistry() = default;
+
+		virtual bool FindInterfaceBindingEntries(const FMetasoundFrontendVersion& InInputInterfaceVersion, TArray<const FInterfaceBindingRegistryEntry*>& OutEntries) const override;
+
+		virtual void RegisterInterfaceBinding(const FMetasoundFrontendVersion& InInputInterfaceVersion, TUniquePtr<FInterfaceBindingRegistryEntry>&& InEntry) override;
+
+		virtual bool UnregisterInterfaceBinding(const FMetasoundFrontendVersion& InInputInterfaceVersion, const FMetasoundFrontendVersion& InOutputInterfaceVersion) override;
+
+		virtual bool UnregisterAllInterfaceBindings(const FMetasoundFrontendVersion& InInputInterfaceVersion) override;
+
+		// MultiMap of interface bindings keyed by the shared input interface registry key for an array of connection-priority sorted bindings
+		TMultiMap<FInterfaceRegistryKey, TUniquePtr<FInterfaceBindingRegistryEntry>> Entries;
+	};
+} // namespace Metasound::Frontend
diff --git a/Engine/Plugins/Runtime/Metasound/Source/MetasoundFrontend/Private/Interfaces/MetasoundFrontendInterfaceRegistry.cpp b/Engine/Plugins/Runtime/Metasound/Source/MetasoundFrontend/Private/Interfaces/MetasoundFrontendInterfaceRegistry.cpp
new file mode 100644
--- /dev/null
+++ b/Engine/Plugins/Runtime/Metasound/Source/MetasoundFrontend/Private/Interfaces/MetasoundFrontendInterfaceRegistry.cpp
@@ -0,0 +1,123 @@
+// Copyright Epic Games, Inc. All Rights Reserved.
+#include "Interfaces/MetasoundFrontendInterfaceRegistry.h"
+
+#include "HAL/PlatformTime.h"
+#include "Interfaces/MetasoundFrontendInterfaceBindingRegistry.h"
+#include "MetasoundFrontendInterfaceRegistryPrivate.h"
+#include "MetasoundFrontendRegistryTransaction.h"
+#include "MetasoundTrace.h"
+
+
+namespace Metasound::Frontend
+{
+	FInterfaceRegistry::FInterfaceRegistry()
+	: TransactionBuffer(MakeShared<TTransactionBuffer<FInterfaceRegistryTransaction>>())
+	{
+	}
+
+	void FInterfaceRegistry::RegisterInterface(TUniquePtr<IInterfaceRegistryEntry>&& InEntry)
+	{
+		METASOUND_LLM_SCOPE;
+
+		FInterfaceRegistryTransaction::FTimeType TransactionTime = FPlatformTime::Cycles64();
+		if (InEntry.IsValid())
+		{
+			FInterfaceRegistryKey Key = GetInterfaceRegistryKey(InEntry->GetInterface());
+			if (IsValidInterfaceRegistryKey(Key))
+			{
+				if (const IInterfaceRegistryEntry* Entry = FindInterfaceRegistryEntry(Key))
+				{
+					UE_LOG(LogMetaSound, Warning, TEXT("Registration of interface overwriting previously registered interface [RegistryKey: %s]"), *Key);
+						
+					FInterfaceRegistryTransaction Transaction{FInterfaceRegistryTransaction::ETransactionType::InterfaceUnregistration, Key, Entry->GetInterface().Version, TransactionTime};
+					TransactionBuffer->AddTransaction(MoveTemp(Transaction));
+				}
+					
+				FInterfaceRegistryTransaction Transaction{FInterfaceRegistryTransaction::ETransactionType::InterfaceRegistration, Key, InEntry->GetInterface().Version, TransactionTime};
+				TransactionBuffer->AddTransaction(MoveTemp(Transaction));
+
+				Entries.Add(Key, MoveTemp(InEntry)).Get();
+			}
+		}
+	}
+
+	const IInterfaceRegistryEntry* FInterfaceRegistry::FindInterfaceRegistryEntry(const FInterfaceRegistryKey& InKey) const
+	{
+		if (const TUniquePtr<IInterfaceRegistryEntry>* Entry = Entries.Find(InKey))
+		{
+			return Entry->Get();
+		}
+		return nullptr;
+	}
+
+	bool FInterfaceRegistry::FindInterface(const FInterfaceRegistryKey& InKey, FMetasoundFrontendInterface& OutInterface) const
+	{
+		if (const IInterfaceRegistryEntry* Entry = FindInterfaceRegistryEntry(InKey))
+		{
+			OutInterface = Entry->GetInterface();
+			return true;
+		}
+
+		return false;
+	}
+
+	TUniquePtr<FInterfaceTransactionStream> FInterfaceRegistry::CreateTransactionStream()
+	{
+		return MakeUnique<FInterfaceTransactionStream>(TransactionBuffer);
+	}
+
+	bool IsValidInterfaceRegistryKey(const FInterfaceRegistryKey& InKey)
+	{
+		return !InKey.IsEmpty();
+	}
+
+	FInterfaceRegistryKey GetInterfaceRegistryKey(const FMetasoundFrontendVersion& InInterfaceVersion)
+	{
+		return FString::Format(TEXT("{0}_{1}.{2}"), { InInterfaceVersion.Name.ToString(), InInterfaceVersion.Number.Major, InInterfaceVersion.Number.Minor });
+
+	}
+
+	FInterfaceRegistryKey GetInterfaceRegistryKey(const FMetasoundFrontendInterface& InInterface)
+	{
+		return GetInterfaceRegistryKey(InInterface.Version);
+	}
+
+	FInterfaceRegistryTransaction::FInterfaceRegistryTransaction(ETransactionType InType, const FInterfaceRegistryKey& InKey, const FMetasoundFrontendVersion& InInterfaceVersion, FInterfaceRegistryTransaction::FTimeType InTimestamp)
+	: Type(InType)
+	, Key(InKey)
+	, InterfaceVersion(InInterfaceVersion)
+	, Timestamp(InTimestamp)
+	{
+	}
+
+	FInterfaceRegistryTransaction::ETransactionType FInterfaceRegistryTransaction::GetTransactionType() const
+	{
+		return Type;
+	}
+
+	const FMetasoundFrontendVersion& FInterfaceRegistryTransaction::GetInterfaceVersion() const
+	{
+		return InterfaceVersion;
+	}
+
+	const FInterfaceRegistryKey& FInterfaceRegistryTransaction::GetInterfaceRegistryKey() const
+	{
+		return Key;
+	}
+
+	FInterfaceRegistryTransaction::FTimeType FInterfaceRegistryTransaction::GetTimestamp() const
+	{
+		return Timestamp;
+	}
+
+	FInterfaceRegistry& FInterfaceRegistry::Get()
+	{
+		static FInterfaceRegistry Registry;
+		return Registry;
+	}
+
+	IInterfaceRegistry& IInterfaceRegistry::Get()
+	{
+		return FInterfaceRegistry::Get();
+	}
+} // namespace Metasound::Frontend
diff --git a/Engine/Plugins/Runtime/Metasound/Source/MetasoundFrontend/Private/Interfaces/MetasoundFrontendInterfaceRegistryPrivate.h b/Engine/Plugins/Runtime/Metasound/Source/MetasoundFrontend/Private/Interfaces/MetasoundFrontendInterfaceRegistryPrivate.h
new file mode 100644
--- /dev/null
+++ b/Engine/Plugins/Runtime/Metasound/Source/MetasoundFrontend/Private/Interfaces/MetasoundFrontendInterfaceRegistryPrivate.h
@@ -0,0 +1,39 @@
+// Copyright Epic Games, Inc. All Rights Reserved.
+#pragma once
+
+#include "Interfaces/MetasoundFrontendInterfaceRegistry.h"
+#include "MetasoundFrontendController.h"
+#include "MetasoundFrontendDocument.h"
+#include "MetasoundFrontendRegistryTransaction.h"
+#include "UObject/NoExportTypes.h"
+
+
+namespace Metasound::Frontend
+{
+	using FInterfaceTransactionStream = TTransactionStream<FInterfaceRegistryTransaction>;
+
+	class FInterfaceRegistry : public IInterfaceRegistry
+	{
+	public:
+		static FInterfaceRegistry& Get();
+
+		FInterfaceRegistry();
+
+		virtual ~FInterfaceRegistry() = default;
+
+		virtual void RegisterInterface(TUniquePtr<IInterfaceRegistryEntry>&& InEntry) override;
+
+		virtual const IInterfaceRegistryEntry* FindInterfaceRegistryEntry(const FInterfaceRegistryKey& InKey) const override;
+
+		virtual bool FindInterface(const FInterfaceRegistryKey& InKey, FMetasoundFrontendInterface& OutInterface) const override;
+
+		TUniquePtr<FInterfaceTransactionStream> CreateTransactionStream();
+
+	private:
+		using FInterfaceTransactionBuffer = TTransactionBuffer<FInterfaceRegistryTransaction>;
+
+		TMap<FInterfaceRegistryKey, TUniquePtr<IInterfaceRegistryEntry>> Entries;
+
+		TSharedRef<FInterfaceTransactionBuffer> TransactionBuffer;
+	};
+} // namespace Metasound::Frontend
diff --git a/Engine/Plugins/Runtime/Metasound/Source/MetasoundFrontend/Private/MetasoundAssetBase.cpp b/Engine/Plugins/Runtime/Metasound/Source/MetasoundFrontend/Private/MetasoundAssetBase.cpp
--- a/Engine/Plugins/Runtime/Metasound/Source/MetasoundFrontend/Private/MetasoundAssetBase.cpp
+++ b/Engine/Plugins/Runtime/Metasound/Source/MetasoundFrontend/Private/MetasoundAssetBase.cpp
@@ -11,11 +11,11 @@
 #include "HAL/IConsoleManager.h"
 #include "IAudioParameterTransmitter.h"
 #include "Interfaces/MetasoundFrontendInterface.h"
+#include "Interfaces/MetasoundFrontendInterfaceRegistry.h"
 #include "Internationalization/Text.h"
 #include "IStructSerializerBackend.h"
 #include "Logging/LogMacros.h"
 #include "MetasoundAssetManager.h"
-#include "MetasoundFrontendArchetypeRegistry.h"
 #include "MetasoundFrontendController.h"
 #include "MetasoundFrontendDocument.h"
 #include "MetasoundFrontendDocumentBuilder.h"
diff --git a/Engine/Plugins/Runtime/Metasound/Source/MetasoundFrontend/Private/MetasoundFrontendArchetypeRegistry.cpp b/Engine/Plugins/Runtime/Metasound/Source/MetasoundFrontend/Private/MetasoundFrontendArchetypeRegistry.cpp
deleted file mode 100644
diff --git a/Engine/Plugins/Runtime/Metasound/Source/MetasoundFrontend/Private/MetasoundFrontendDocumentBuilder.cpp b/Engine/Plugins/Runtime/Metasound/Source/MetasoundFrontend/Private/MetasoundFrontendDocumentBuilder.cpp
--- a/Engine/Plugins/Runtime/Metasound/Source/MetasoundFrontend/Private/MetasoundFrontendDocumentBuilder.cpp
+++ b/Engine/Plugins/Runtime/Metasound/Source/MetasoundFrontend/Private/MetasoundFrontendDocumentBuilder.cpp
@@ -4,11 +4,13 @@
 #include "Algo/AnyOf.h"
 #include "Algo/Find.h"
 #include "Algo/NoneOf.h"
+#include "Algo/Sort.h"
 #include "Algo/Transform.h"
+#include "Interfaces/MetasoundFrontendInterfaceBindingRegistry.h"
+#include "Interfaces/MetasoundFrontendInterfaceRegistry.h"
 #include "MetasoundAssetBase.h"
 #include "MetasoundAssetManager.h"
 #include "MetasoundDocumentInterface.h"
-#include "MetasoundFrontendArchetypeRegistry.h"
 #include "MetasoundFrontendDocumentCache.h"
 #include "MetasoundFrontendDocument.h"
 #include "MetasoundFrontendRegistries.h"
@@ -23,6 +25,51 @@ namespace Metasound::Frontend
 {
 	namespace DocumentBuilderPrivate
 	{
+		bool TryGetInterfaceBoundEdges(
+			const FGuid& InFromNodeID,
+			const FMetasoundFrontendDocument& InFromNodeDocument,
+			const FGuid& InToNodeID,
+			const FMetasoundFrontendDocument& InToNodeDocument,
+			TSet<FNamedEdge>& OutNamedEdges)
+		{
+			OutNamedEdges.Reset();
+			TSet<FName> InputNames;
+			const TSet<FMetasoundFrontendVersion>& OutputInterfaceVersions = InFromNodeDocument.Interfaces;
+			for (const FMetasoundFrontendVersion& InputInterfaceVersion : InToNodeDocument.Interfaces)
+			{
+				TArray<const FInterfaceBindingRegistryEntry*> BindingEntries;
+				if (IInterfaceBindingRegistry::Get().FindInterfaceBindingEntries(InputInterfaceVersion, BindingEntries))
+				{
+					Algo::Sort(BindingEntries, [](const FInterfaceBindingRegistryEntry* A, const FInterfaceBindingRegistryEntry* B)
+					{
+						check(A);
+						check(B);
+						return A->GetBindingPriority() < B->GetBindingPriority();
+					});
+
+					// Bindings are sorted in registry with earlier entries being higher priority to apply connections,
+					// so earlier listed connections are selected over potential collisions with later entries.
+					for (const FInterfaceBindingRegistryEntry* BindingEntry : BindingEntries)
+					{
+						check(BindingEntry);
+						if (OutputInterfaceVersions.Contains(BindingEntry->GetOutputInterfaceVersion()))
+						{
+							for (const FMetasoundFrontendInterfaceVertexBinding& VertexBinding : BindingEntry->GetVertexBindings())
+							{
+								if (!InputNames.Contains(VertexBinding.InputName))
+								{
+									InputNames.Add(VertexBinding.InputName);
+									OutNamedEdges.Add(FNamedEdge { InFromNodeID, VertexBinding.OutputName, InToNodeID, VertexBinding.InputName });
+								}
+							}
+						}
+					}
+				}
+			};
+
+			return true;
+		}
+
 		void FinalizeGraphVertexNode(FMetasoundFrontendNode& InOutNode, const FMetasoundFrontendClassVertex& InVertex)
 		{
 			InOutNode.Name = InVertex.Name;
@@ -156,7 +203,7 @@ const FMetasoundFrontendEdge* FMetaSoundFrontendDocumentBuilder::AddEdge(FMetaso
 	return nullptr;
 }
 
-bool FMetaSoundFrontendDocumentBuilder::AddEdgesByNamedConnections(const TSet<Metasound::Frontend::FNamedEdge>& ConnectionsToMake, TArray<const FMetasoundFrontendEdge*>* OutNewEdges)
+bool FMetaSoundFrontendDocumentBuilder::AddNamedEdges(const TSet<Metasound::Frontend::FNamedEdge>& EdgesToMake, TArray<const FMetasoundFrontendEdge*>* OutNewEdges)
 {
 	using namespace Metasound::Frontend;
 
@@ -171,15 +218,15 @@ bool FMetaSoundFrontendDocumentBuilder::AddEdgesByNamedConnections(const TSet<Me
 	bool bSuccess = true;
 
 	TArray<FMetasoundFrontendEdge> EdgesToAdd;
-	for (const FNamedEdge& Connection : ConnectionsToMake)
+	for (const FNamedEdge& Edge : EdgesToMake)
 	{
 		const FMetasoundFrontendGraph& Graph = Document.RootGraph.Graph;
-		const FMetasoundFrontendVertex* OutputVertex = NodeCache.FindOutputVertex(Connection.OutputNodeID, Connection.OutputName);
-		const FMetasoundFrontendVertex* InputVertex = NodeCache.FindInputVertex(Connection.InputNodeID, Connection.InputName);
+		const FMetasoundFrontendVertex* OutputVertex = NodeCache.FindOutputVertex(Edge.OutputNodeID, Edge.OutputName);
+		const FMetasoundFrontendVertex* InputVertex = NodeCache.FindInputVertex(Edge.InputNodeID, Edge.InputName);
 
 		if (OutputVertex && InputVertex)
 		{
-			FMetasoundFrontendEdge NewEdge = { Connection.OutputNodeID, OutputVertex->VertexID, Connection.InputNodeID, InputVertex->VertexID };
+			FMetasoundFrontendEdge NewEdge = { Edge.OutputNodeID, OutputVertex->VertexID, Edge.InputNodeID, InputVertex->VertexID };
 			if (CanAddEdge(NewEdge))
 			{
 				EdgesToAdd.Add(MoveTemp(NewEdge));
@@ -187,7 +234,7 @@ bool FMetaSoundFrontendDocumentBuilder::AddEdgesByNamedConnections(const TSet<Me
 			else
 			{
 				bSuccess = false;
-				UE_LOG(LogMetaSound, Warning, TEXT("Failed to add connections between MetaSound output '%s' and input '%s': Vertex types either incompatable (eg. construction pin to non-construction pin) or input already connected."), *Connection.OutputName.ToString(), *Connection.InputName.ToString());
+				UE_LOG(LogMetaSound, Warning, TEXT("Failed to add connections between MetaSound output '%s' and input '%s': Vertex types either incompatable (eg. construction pin to non-construction pin) or input already connected."), *Edge.OutputName.ToString(), *Edge.InputName.ToString());
 			}
 		}
 	}
@@ -216,56 +263,19 @@ bool FMetaSoundFrontendDocumentBuilder::AddEdgesByNodeClassInterfaceBindings(con
 	using namespace Metasound;
 	using namespace Metasound::Frontend;
 
-	auto FindNodeClassDoc = [this](const FGuid & InNodeID) -> const FMetasoundFrontendDocument*
+	const FMetasoundFrontendDocument* FromNodeDocument = FindOrLoadNodeClassDocument(InFromNodeID);
+	const FMetasoundFrontendDocument* ToNodeDocument = FindOrLoadNodeClassDocument(InToNodeID);
+	if (FromNodeDocument && ToNodeDocument)
 	{
-		const TScriptInterface<IMetaSoundDocumentInterface> NodeDocumentUInterface = FindOrLoadNodeClassDocument(InNodeID);
-		if (NodeDocumentUInterface)
+		TSet<FNamedEdge> NamedEdges;
+		if (DocumentBuilderPrivate::TryGetInterfaceBoundEdges(InFromNodeID, *FromNodeDocument, InToNodeID, *ToNodeDocument, NamedEdges))
 		{
-			return &NodeDocumentUInterface->GetDocument();
+			return AddNamedEdges(NamedEdges);
 		}
-
-		return nullptr;
-	};
-
-	const FMetasoundFrontendDocument* OutputNodeDocument = FindNodeClassDoc(InFromNodeID);
-	const FMetasoundFrontendDocument* InputNodeDocument = FindNodeClassDoc(InToNodeID);
-
-	if (!InputNodeDocument || !OutputNodeDocument)
-	{
-		return false;
 	}
 
-	const TSet<FMetasoundFrontendVersion>& OutputInterfaceVersions = OutputNodeDocument->Interfaces;
-
-	TArray<const IInterfaceRegistryEntry*> InputNodeInterfaces;
-	FindDeclaredInterfaces(*InputNodeDocument, InputNodeInterfaces);
-
-	TSet<FNamedEdge> ConnectionsToMake;
-	TSet<FName> InputsToConnect;
-	for (const IInterfaceRegistryEntry* InputInterfaceEntry : InputNodeInterfaces)
-	{
-		const FMetasoundFrontendInterface& InputInterface = InputInterfaceEntry->GetInterface();
-		const TArray<FMetasoundFrontendInterfaceBinding>& OutputBindings = InputInterfaceEntry->GetOutputBindings();
-
-		// Bindings are sorted in registry with earlier entries being higher priority to apply connections,
-		// so earlier listed connections are selected over potential collisions with later entries.
-		for (const FMetasoundFrontendInterfaceBinding& OutputBinding : OutputBindings)
-		{
-			if (OutputInterfaceVersions.Contains(OutputBinding.Version))
-			{
-				for (const FMetasoundFrontendInterfaceBindingConnection& Connection : OutputBinding.Connections)
-				{
-					if (!InputsToConnect.Contains(Connection.InputName))
-					{
-						InputsToConnect.Add(Connection.InputName);
-						ConnectionsToMake.Add(FNamedEdge { InFromNodeID, Connection.OutputName, InToNodeID, Connection.InputName });
-					}
-				}
-			}
-		}
-	}
+	return false;
 
-	return AddEdgesByNamedConnections(ConnectionsToMake);
 }
 
 bool FMetaSoundFrontendDocumentBuilder::AddEdgesFromMatchingInterfaceNodeOutputsToGraphOutputs(const FGuid& InNodeID, TArray<const FMetasoundFrontendEdge*>& OutEdgesCreated)
@@ -276,24 +286,24 @@ bool FMetaSoundFrontendDocumentBuilder::AddEdgesFromMatchingInterfaceNodeOutputs
 
 	OutEdgesCreated.Reset();
 
-	const TScriptInterface<IMetaSoundDocumentInterface> NodeDocumentInterface = FindOrLoadNodeClassDocument(InNodeID);
-	if (!NodeDocumentInterface)
+	const FMetasoundFrontendDocument* NodeDocument = FindOrLoadNodeClassDocument(InNodeID);
+	if (!NodeDocument)
 	{
 		return false;
 	}
 
 	const IDocumentGraphNodeCache& NodeCache = DocumentCache->GetNodeCache();
 
-	const TSet<FMetasoundFrontendVersion>& NodeInterfaces = NodeDocumentInterface->GetDocument().Interfaces;
+	const TSet<FMetasoundFrontendVersion>& NodeInterfaces = NodeDocument->Interfaces;
 	const TSet<FMetasoundFrontendVersion> CommonInterfaces = NodeInterfaces.Intersect(GetDocument().Interfaces);
 
-	TSet<FNamedEdge> ConnectionsToMake;
+	TSet<FNamedEdge> EdgesToMake;
 	for (const FMetasoundFrontendVersion& Version : CommonInterfaces)
 	{
 		const FInterfaceRegistryKey InterfaceKey = GetInterfaceRegistryKey(Version);
 		if (const IInterfaceRegistryEntry* RegistryEntry = IInterfaceRegistry::Get().FindInterfaceRegistryEntry(InterfaceKey))
 		{
-			Algo::Transform(RegistryEntry->GetInterface().Outputs, ConnectionsToMake, [this, &NodeCache, InNodeID](const FMetasoundFrontendClassOutput& Output)
+			Algo::Transform(RegistryEntry->GetInterface().Outputs, EdgesToMake, [this, &NodeCache, InNodeID](const FMetasoundFrontendClassOutput& Output)
 			{
 				const FMetasoundFrontendGraph& Graph = GetDocument().RootGraph.Graph;
 				const FMetasoundFrontendVertex* NodeVertex = NodeCache.FindOutputVertex(InNodeID, Output.Name);
@@ -307,7 +317,7 @@ bool FMetaSoundFrontendDocumentBuilder::AddEdgesFromMatchingInterfaceNodeOutputs
 		}
 	}
 
-	return AddEdgesByNamedConnections(ConnectionsToMake, &OutEdgesCreated);
+	return AddNamedEdges(EdgesToMake, &OutEdgesCreated);
 }
 
 bool FMetaSoundFrontendDocumentBuilder::AddEdgesFromMatchingInterfaceNodeInputsToGraphInputs(const FGuid& InNodeID, TArray<const FMetasoundFrontendEdge*>& OutEdgesCreated)
@@ -318,24 +328,24 @@ bool FMetaSoundFrontendDocumentBuilder::AddEdgesFromMatchingInterfaceNodeInputsT
 
 	OutEdgesCreated.Reset();
 
-	const TScriptInterface<IMetaSoundDocumentInterface> NodeDocumentInterface = FindOrLoadNodeClassDocument(InNodeID);
-	if (!NodeDocumentInterface)
+	const FMetasoundFrontendDocument* NodeDocument = FindOrLoadNodeClassDocument(InNodeID);
+	if (!NodeDocument)
 	{
 		return false;
 	}
 
 	const IDocumentGraphNodeCache& NodeCache = DocumentCache->GetNodeCache();
 
-	const TSet<FMetasoundFrontendVersion>& NodeInterfaces = NodeDocumentInterface->GetDocument().Interfaces;
+	const TSet<FMetasoundFrontendVersion>& NodeInterfaces = NodeDocument->Interfaces;
 	const TSet<FMetasoundFrontendVersion> CommonInterfaces = NodeInterfaces.Intersect(GetDocument().Interfaces);
 
-	TSet<FNamedEdge> ConnectionsToMake;
+	TSet<FNamedEdge> EdgesToMake;
 	for (const FMetasoundFrontendVersion& Version : CommonInterfaces)
 	{
 		const FInterfaceRegistryKey InterfaceKey = GetInterfaceRegistryKey(Version);
 		if (const IInterfaceRegistryEntry* RegistryEntry = IInterfaceRegistry::Get().FindInterfaceRegistryEntry(InterfaceKey))
 		{
-			Algo::Transform(RegistryEntry->GetInterface().Inputs, ConnectionsToMake, [this, &NodeCache, InNodeID](const FMetasoundFrontendClassInput& Input)
+			Algo::Transform(RegistryEntry->GetInterface().Inputs, EdgesToMake, [this, &NodeCache, InNodeID](const FMetasoundFrontendClassInput& Input)
 			{
 				const FMetasoundFrontendGraph& Graph = GetDocument().RootGraph.Graph;
 				const FMetasoundFrontendVertex* NodeVertex = NodeCache.FindOutputVertex(InNodeID, Input.Name);
@@ -349,7 +359,7 @@ bool FMetaSoundFrontendDocumentBuilder::AddEdgesFromMatchingInterfaceNodeInputsT
 		}
 	}
 
-	return AddEdgesByNamedConnections(ConnectionsToMake, &OutEdgesCreated);
+	return AddNamedEdges(EdgesToMake, &OutEdgesCreated);
 }
 
 const FMetasoundFrontendNode* FMetaSoundFrontendDocumentBuilder::AddGraphInput(const FMetasoundFrontendClassInput& InClassInput)
@@ -853,7 +863,7 @@ TArray<const FMetasoundFrontendVertex*> FMetaSoundFrontendDocumentBuilder::FindN
 	return DocumentCache->GetNodeCache().FindNodeOutputs(InNodeID, TypeName);
 }
 
-const TScriptInterface<IMetaSoundDocumentInterface> FMetaSoundFrontendDocumentBuilder::FindOrLoadNodeClassDocument(const FGuid& InNodeID) const
+const FMetasoundFrontendDocument* FMetaSoundFrontendDocumentBuilder::FindOrLoadNodeClassDocument(const FGuid& InNodeID) const
 {
 	using namespace Metasound;
 	using namespace Metasound::Frontend;
@@ -868,8 +878,9 @@ const TScriptInterface<IMetaSoundDocumentInterface> FMetaSoundFrontendDocumentBu
 			if (FMetasoundAssetBase* Asset = IMetaSoundAssetManager::GetChecked().TryLoadAssetFromKey(NodeClassRegistryKey))
 			{
 				UObject* Object = Asset->GetOwningAsset();
-				check(Object);
-				return Object;
+				TScriptInterface<IMetaSoundDocumentInterface> Interface = Object;
+				check(Interface);
+				return &Interface->GetDocument();
 			}
 		}
 	}
@@ -1076,7 +1087,7 @@ bool FMetaSoundFrontendDocumentBuilder::RemoveEdge(const FMetasoundFrontendEdge&
 	return false;
 }
 
-bool FMetaSoundFrontendDocumentBuilder::RemoveEdgesByNamedConnections(const TSet<Metasound::Frontend::FNamedEdge>& ConnectionsToRemove, TArray<FMetasoundFrontendEdge>* OutRemovedEdges)
+bool FMetaSoundFrontendDocumentBuilder::RemoveNamedEdges(const TSet<Metasound::Frontend::FNamedEdge>& InNamedEdgesToRemove, TArray<FMetasoundFrontendEdge>* OutRemovedEdges)
 {
 	using namespace Metasound::Frontend;
 
@@ -1091,15 +1102,15 @@ bool FMetaSoundFrontendDocumentBuilder::RemoveEdgesByNamedConnections(const TSet
 	bool bSuccess = true;
 
 	TArray<FMetasoundFrontendEdge> EdgesToRemove;
-	for (const FNamedEdge& Connection : ConnectionsToRemove)
+	for (const FNamedEdge& NamedEdge : InNamedEdgesToRemove)
 	{
 		const FMetasoundFrontendGraph& Graph = Document.RootGraph.Graph;
-		const FMetasoundFrontendVertex* OutputVertex = NodeCache.FindOutputVertex(Connection.OutputNodeID, Connection.OutputName);
-		const FMetasoundFrontendVertex* InputVertex = NodeCache.FindInputVertex(Connection.InputNodeID, Connection.InputName);
+		const FMetasoundFrontendVertex* OutputVertex = NodeCache.FindOutputVertex(NamedEdge.OutputNodeID, NamedEdge.OutputName);
+		const FMetasoundFrontendVertex* InputVertex = NodeCache.FindInputVertex(NamedEdge.InputNodeID, NamedEdge.InputName);
 
 		if (OutputVertex && InputVertex)
 		{
-			FMetasoundFrontendEdge NewEdge = { Connection.OutputNodeID, OutputVertex->VertexID, Connection.InputNodeID, InputVertex->VertexID };
+			FMetasoundFrontendEdge NewEdge = { NamedEdge.OutputNodeID, OutputVertex->VertexID, NamedEdge.InputNodeID, InputVertex->VertexID };
 			if (ContainsEdge(NewEdge))
 			{
 				EdgesToRemove.Add(MoveTemp(NewEdge));
@@ -1107,7 +1118,7 @@ bool FMetaSoundFrontendDocumentBuilder::RemoveEdgesByNamedConnections(const TSet
 			else
 			{
 				bSuccess = false;
-				UE_LOG(LogMetaSound, Warning, TEXT("Failed to remove connection between MetaSound node output '%s' and input '%s': No connection found."), *Connection.OutputName.ToString(), *Connection.InputName.ToString());
+				UE_LOG(LogMetaSound, Warning, TEXT("Failed to remove connection between MetaSound node output '%s' and input '%s': No connection found."), *NamedEdge.OutputName.ToString(), *NamedEdge.InputName.ToString());
 			}
 		}
 	}
@@ -1136,58 +1147,18 @@ bool FMetaSoundFrontendDocumentBuilder::RemoveEdgesByNodeClassInterfaceBindings(
 	using namespace Metasound;
 	using namespace Metasound::Frontend;
 
-	auto FindNodeClassDoc = [this](const FGuid& InNodeID) -> const FMetasoundFrontendDocument*
-	{
-		const TScriptInterface<IMetaSoundDocumentInterface> NodeDocumentUInterface = FindOrLoadNodeClassDocument(InNodeID);
-		if (NodeDocumentUInterface)
-		{
-			return &NodeDocumentUInterface->GetDocument();
-		}
-
-		return nullptr;
-	};
-
-	const FMetasoundFrontendDocument* OutputNodeDocument = FindNodeClassDoc(InFromNodeID);
-	const FMetasoundFrontendDocument* InputNodeDocument = FindNodeClassDoc(InToNodeID);
-
-	if (!InputNodeDocument || !OutputNodeDocument)
-	{
-		return false;
-	}
-
-	const TSet<FMetasoundFrontendVersion>& OutputInterfaceVersions = OutputNodeDocument->Interfaces;
-
-	TArray<const IInterfaceRegistryEntry*> InputNodeInterfaces;
-	FindDeclaredInterfaces(*InputNodeDocument, InputNodeInterfaces);
-
-	bool bSuccess = true;
-
-	TSet<FNamedEdge> ConnectionsToRemove;
-	TSet<FName> InputsToDisconnect;
-	for (const IInterfaceRegistryEntry* InputInterfaceEntry : InputNodeInterfaces)
+	const FMetasoundFrontendDocument* FromNodeDocument = FindOrLoadNodeClassDocument(InFromNodeID);
+	const FMetasoundFrontendDocument* ToNodeDocument = FindOrLoadNodeClassDocument(InToNodeID);
+	if (FromNodeDocument && ToNodeDocument)
 	{
-		const FMetasoundFrontendInterface& InputInterface = InputInterfaceEntry->GetInterface();
-		const TArray<FMetasoundFrontendInterfaceBinding>& OutputBindings = InputInterfaceEntry->GetOutputBindings();
-
-		// Bindings are sorted in registry with earlier entries being higher priority to apply connections,
-		// so earlier listed connections are selected over potential collisions with later entries.
-		for (const FMetasoundFrontendInterfaceBinding& OutputBinding : OutputBindings)
+		TSet<FNamedEdge> NamedEdges;
+		if (DocumentBuilderPrivate::TryGetInterfaceBoundEdges(InFromNodeID, *FromNodeDocument, InToNodeID, *ToNodeDocument, NamedEdges))
 		{
-			if (OutputInterfaceVersions.Contains(OutputBinding.Version))
-			{
-				for (const FMetasoundFrontendInterfaceBindingConnection& Connection : OutputBinding.Connections)
-				{
-					if (!InputsToDisconnect.Contains(Connection.InputName))
-					{
-						InputsToDisconnect.Add(Connection.InputName);
-						ConnectionsToRemove.Add(FNamedEdge{ InFromNodeID, Connection.OutputName, InToNodeID, Connection.InputName });
-					}
-				}
-			}
+			return RemoveNamedEdges(NamedEdges);
 		}
 	}
 
-	return RemoveEdgesByNamedConnections(ConnectionsToRemove);
+	return false;
 }
 
 bool FMetaSoundFrontendDocumentBuilder::RemoveEdgesFromNodeOutput(const FGuid& InNodeID, const FGuid& InVertexID)
diff --git a/Engine/Plugins/Runtime/Metasound/Source/MetasoundFrontend/Private/MetasoundFrontendDocumentVersioning.cpp b/Engine/Plugins/Runtime/Metasound/Source/MetasoundFrontend/Private/MetasoundFrontendDocumentVersioning.cpp
--- a/Engine/Plugins/Runtime/Metasound/Source/MetasoundFrontend/Private/MetasoundFrontendDocumentVersioning.cpp
+++ b/Engine/Plugins/Runtime/Metasound/Source/MetasoundFrontend/Private/MetasoundFrontendDocumentVersioning.cpp
@@ -3,8 +3,8 @@
 
 #include "Algo/Transform.h"
 #include "Interfaces/MetasoundFrontendInterface.h"
+#include "Interfaces/MetasoundFrontendInterfaceRegistry.h"
 #include "MetasoundAccessPtr.h"
-#include "MetasoundFrontendArchetypeRegistry.h"
 #include "MetasoundFrontendDocumentController.h"
 #include "MetasoundFrontendRegistries.h"
 #include "MetasoundFrontendSearchEngine.h"
diff --git a/Engine/Plugins/Runtime/Metasound/Source/MetasoundFrontend/Private/MetasoundFrontendInterfaceRegistryPrivate.h b/Engine/Plugins/Runtime/Metasound/Source/MetasoundFrontend/Private/MetasoundFrontendInterfaceRegistryPrivate.h
deleted file mode 100644
diff --git a/Engine/Plugins/Runtime/Metasound/Source/MetasoundFrontend/Private/MetasoundFrontendSearchEngine.cpp b/Engine/Plugins/Runtime/Metasound/Source/MetasoundFrontend/Private/MetasoundFrontendSearchEngine.cpp
--- a/Engine/Plugins/Runtime/Metasound/Source/MetasoundFrontend/Private/MetasoundFrontendSearchEngine.cpp
+++ b/Engine/Plugins/Runtime/Metasound/Source/MetasoundFrontend/Private/MetasoundFrontendSearchEngine.cpp
@@ -5,7 +5,7 @@
 #include "MetasoundFrontendSearchEngineEditorOnly.h"
 
 #include "Algo/MaxElement.h"
-#include "MetasoundFrontendArchetypeRegistry.h"
+#include "Interfaces/MetasoundFrontendInterfaceRegistry.h"
 #include "MetasoundFrontendQuery.h"
 #include "MetasoundFrontendQuerySteps.h"
 #include "MetasoundFrontendRegistryTransaction.h"
diff --git a/Engine/Plugins/Runtime/Metasound/Source/MetasoundFrontend/Private/MetasoundFrontendSearchEngineCore.cpp b/Engine/Plugins/Runtime/Metasound/Source/MetasoundFrontend/Private/MetasoundFrontendSearchEngineCore.cpp
--- a/Engine/Plugins/Runtime/Metasound/Source/MetasoundFrontend/Private/MetasoundFrontendSearchEngineCore.cpp
+++ b/Engine/Plugins/Runtime/Metasound/Source/MetasoundFrontend/Private/MetasoundFrontendSearchEngineCore.cpp
@@ -5,8 +5,8 @@
 #include "Algo/MaxElement.h"
 #include "Containers/Array.h"
 #include "Containers/UnrealString.h"
-#include "MetasoundFrontendArchetypeRegistry.h"
-#include "MetasoundFrontendInterfaceRegistryPrivate.h"
+#include "Interfaces/MetasoundFrontendInterfaceRegistry.h"
+#include "Interfaces/MetasoundFrontendInterfaceRegistryPrivate.h"
 #include "MetasoundFrontendQuery.h"
 #include "MetasoundFrontendQuerySteps.h"
 #include "MetasoundFrontendRegistryTransaction.h"
diff --git a/Engine/Plugins/Runtime/Metasound/Source/MetasoundFrontend/Private/MetasoundFrontendSearchEngineCore.h b/Engine/Plugins/Runtime/Metasound/Source/MetasoundFrontend/Private/MetasoundFrontendSearchEngineCore.h
--- a/Engine/Plugins/Runtime/Metasound/Source/MetasoundFrontend/Private/MetasoundFrontendSearchEngineCore.h
+++ b/Engine/Plugins/Runtime/Metasound/Source/MetasoundFrontend/Private/MetasoundFrontendSearchEngineCore.h
@@ -4,8 +4,8 @@
 
 #include "Containers/Array.h"
 #include "HAL/CriticalSection.h"
-#include "MetasoundFrontendArchetypeRegistry.h"
-#include "MetasoundFrontendInterfaceRegistryPrivate.h"
+#include "Interfaces/MetasoundFrontendInterfaceRegistry.h"
+#include "Interfaces/MetasoundFrontendInterfaceRegistryPrivate.h"
 #include "MetasoundFrontendQuery.h"
 #include "MetasoundFrontendQuerySteps.h"
 #include "MetasoundFrontendRegistries.h"
diff --git a/Engine/Plugins/Runtime/Metasound/Source/MetasoundFrontend/Private/MetasoundFrontendSearchEngineEditorOnly.cpp b/Engine/Plugins/Runtime/Metasound/Source/MetasoundFrontend/Private/MetasoundFrontendSearchEngineEditorOnly.cpp
--- a/Engine/Plugins/Runtime/Metasound/Source/MetasoundFrontend/Private/MetasoundFrontendSearchEngineEditorOnly.cpp
+++ b/Engine/Plugins/Runtime/Metasound/Source/MetasoundFrontend/Private/MetasoundFrontendSearchEngineEditorOnly.cpp
@@ -4,7 +4,7 @@
 #include "MetasoundFrontendSearchEngineEditorOnly.h"
 
 #include "Containers/Array.h"
-#include "MetasoundFrontendArchetypeRegistry.h"
+#include "Interfaces/MetasoundFrontendInterfaceRegistry.h"
 #include "MetasoundFrontendDocument.h"
 #include "MetasoundFrontendQuery.h"
 #include "MetasoundFrontendQuerySteps.h"
diff --git a/Engine/Plugins/Runtime/Metasound/Source/MetasoundFrontend/Private/MetasoundFrontendTransform.cpp b/Engine/Plugins/Runtime/Metasound/Source/MetasoundFrontend/Private/MetasoundFrontendTransform.cpp
--- a/Engine/Plugins/Runtime/Metasound/Source/MetasoundFrontend/Private/MetasoundFrontendTransform.cpp
+++ b/Engine/Plugins/Runtime/Metasound/Source/MetasoundFrontend/Private/MetasoundFrontendTransform.cpp
@@ -3,9 +3,9 @@
 
 #include "Algo/Transform.h"
 #include "Interfaces/MetasoundFrontendInterface.h"
+#include "Interfaces/MetasoundFrontendInterfaceRegistry.h"
 #include "MetasoundAccessPtr.h"
 #include "MetasoundAssetBase.h"
-#include "MetasoundFrontendArchetypeRegistry.h"
 #include "MetasoundFrontendDocument.h"
 #include "MetasoundFrontendDocumentController.h"
 #include "MetasoundFrontendRegistries.h"
diff --git a/Engine/Plugins/Runtime/Metasound/Source/MetasoundFrontend/Public/Interfaces/MetasoundFrontendInterfaceBindingRegistry.h b/Engine/Plugins/Runtime/Metasound/Source/MetasoundFrontend/Public/Interfaces/MetasoundFrontendInterfaceBindingRegistry.h
new file mode 100644
--- /dev/null
+++ b/Engine/Plugins/Runtime/Metasound/Source/MetasoundFrontend/Public/Interfaces/MetasoundFrontendInterfaceBindingRegistry.h
@@ -0,0 +1,55 @@
+// Copyright Epic Games, Inc. All Rights Reserved.
+#pragma once
+
+#include "MetasoundFrontendController.h"
+#include "MetasoundFrontendDocument.h"
+#include "UObject/NoExportTypes.h"
+
+
+namespace Metasound::Frontend
+{
+	class METASOUNDFRONTEND_API FInterfaceBindingRegistryEntry
+	{
+		const FMetasoundFrontendVersion OutputInterfaceVersion;
+		const TArray<FMetasoundFrontendInterfaceVertexBinding> VertexBindings;
+		const int32 BindingPriority = 0;
+
+	public:
+		FInterfaceBindingRegistryEntry(const FMetasoundFrontendInterfaceBinding& InBinding);
+		virtual ~FInterfaceBindingRegistryEntry() = default;
+
+		// Returns the output interface version corresponding with the given binding entry.
+		const FMetasoundFrontendVersion& GetOutputInterfaceVersion() const;
+
+		// Returns array of FName pairs that describe how the output and input interfaces are to bind
+		const TArray<FMetasoundFrontendInterfaceVertexBinding>& GetVertexBindings() const;
+
+		// If two binding registry entries match output & input interfaces versions, integer corresponds
+		// to which binding takes priority for a given bind operation. Lower numbers are higher priority.
+		int32 GetBindingPriority() const;
+	};
+
+
+	// Registry of interface bindings, a way to characterize how MetaSound node instances that implement
+	// a give set of interfaces "bind" to other nodes that implement a second set set of interfaces. Binding
+	// is an abstracted way to connect two nodes together using multiple edges as defined through this registry.
+	// Accessing this registry is currently not thread safe and only expected to be occur on the Game Thread.
+	class METASOUNDFRONTEND_API IInterfaceBindingRegistry
+	{
+	public:
+		// Returns binding registry singleton.
+		static IInterfaceBindingRegistry& Get();
+
+		// Find all bindings associated with the given input interface version. Returns true if interface binding is found, false if not.
+		virtual bool FindInterfaceBindingEntries(const FMetasoundFrontendVersion& InInputInterfaceVersion, TArray<const FInterfaceBindingRegistryEntry*>& OutEntries) const = 0;
+
+		// Add interface bindings to the registry.
+		virtual void RegisterInterfaceBinding(const FMetasoundFrontendVersion& InInputInterfaceVersion, TUniquePtr<FInterfaceBindingRegistryEntry>&& InEntry) = 0;
+
+		// Removes an interface binding from registry.
+		virtual bool UnregisterInterfaceBinding(const FMetasoundFrontendVersion& InInputInterfaceVersion, const FMetasoundFrontendVersion& InOutputInterfaceVersion) = 0;
+
+		// Removes all interface bindings associated with the given input interface if found. Returns true if found and removed, false if not found.
+		virtual bool UnregisterAllInterfaceBindings(const FMetasoundFrontendVersion& InInputInterfaceVersion) = 0;
+	};
+} // namespace Metasound::Frontend
diff --git a/Engine/Plugins/Runtime/Metasound/Source/MetasoundFrontend/Public/Interfaces/MetasoundFrontendInterfaceRegistry.h b/Engine/Plugins/Runtime/Metasound/Source/MetasoundFrontend/Public/Interfaces/MetasoundFrontendInterfaceRegistry.h
new file mode 100644
--- /dev/null
+++ b/Engine/Plugins/Runtime/Metasound/Source/MetasoundFrontend/Public/Interfaces/MetasoundFrontendInterfaceRegistry.h
@@ -0,0 +1,93 @@
+// Copyright Epic Games, Inc. All Rights Reserved.
+
+#pragma once
+
+#include "MetasoundFrontendController.h"
+#include "MetasoundFrontendDocument.h"
+#include "UObject/NoExportTypes.h"
+
+namespace Metasound::Frontend
+{
+	using FInterfaceRegistryKey = FString;
+	using FRegistryTransactionID = int32;
+
+	METASOUNDFRONTEND_API bool IsValidInterfaceRegistryKey(const FInterfaceRegistryKey& InKey);
+	METASOUNDFRONTEND_API FInterfaceRegistryKey GetInterfaceRegistryKey(const FMetasoundFrontendVersion& InInterfaceVersion);
+	METASOUNDFRONTEND_API FInterfaceRegistryKey GetInterfaceRegistryKey(const FMetasoundFrontendInterface& InInterface);
+
+	class METASOUNDFRONTEND_API IInterfaceRegistryEntry
+	{
+	public:
+		virtual ~IInterfaceRegistryEntry() = default;
+
+		// MetaSound Interface definition
+		virtual const FMetasoundFrontendInterface& GetInterface() const = 0;
+
+		// Whether or not to immediately apply the given interface
+		// when a MetaSound asset of the supported UClass is created.
+		virtual bool IsDefault(FName InUClassName) const = 0;
+
+		UE_DEPRECATED(5.3, "UClass name must now be provided to determine if interface is to be apply by default.")
+		virtual bool IsDefault() const { return false; }
+
+		// Whether or not the interface can be added or removed by the editor.
+		virtual bool EditorCanAddOrRemove(FName InUClassName) const = 0;
+
+		UE_DEPRECATED(5.3, "UClass name must now be provided to determine if editor can directly add or remove interface.")
+		virtual bool EditorCanAddOrRemove() const { return false; }
+
+		// Name of routing system used to update interface inputs (ex. ParameterInterface or DataReference).
+		virtual FName GetRouterName() const = 0;
+
+		// Name of UClass supported by the given interface.
+		virtual bool UClassIsSupported(FName InUClassName) const = 0;
+
+		// How to update a given document if versioning is required to this interface from a deprecated version.
+		virtual bool UpdateRootGraphInterface(FDocumentHandle InDocument) const = 0;
+	};
+
+	class METASOUNDFRONTEND_API FInterfaceRegistryTransaction
+	{
+	public:
+		using FTimeType = uint64;
+
+		/** Describes the type of transaction. */
+		enum class ETransactionType : uint8
+		{
+			InterfaceRegistration,     //< Something was added to the registry.
+			InterfaceUnregistration,  //< Something was removed from the registry.
+			Invalid
+		};
+
+		FInterfaceRegistryTransaction(ETransactionType InType, const FInterfaceRegistryKey& InKey, const FMetasoundFrontendVersion& InInterfaceVersion, FTimeType InTimestamp);
+
+		ETransactionType GetTransactionType() const;
+		const FMetasoundFrontendVersion& GetInterfaceVersion() const;
+		const FInterfaceRegistryKey& GetInterfaceRegistryKey() const;
+		FTimeType GetTimestamp() const;
+
+	private:
+
+		ETransactionType Type;
+		FInterfaceRegistryKey Key;
+		FMetasoundFrontendVersion InterfaceVersion;
+		FTimeType Timestamp;
+	};
+
+	class METASOUNDFRONTEND_API IInterfaceRegistry
+	{
+	public:
+		static IInterfaceRegistry& Get();
+
+		virtual ~IInterfaceRegistry() = default;
+
+		// Register an interface
+		virtual void RegisterInterface(TUniquePtr<IInterfaceRegistryEntry>&& InEntry) = 0;
+
+		// Find an interface entry with the given key. Returns null if entry not found with given key.
+		virtual const IInterfaceRegistryEntry* FindInterfaceRegistryEntry(const FInterfaceRegistryKey& InKey) const = 0;
+
+		// Find an interface with the given key. Returns true if interface is found, false if not.
+		virtual bool FindInterface(const FInterfaceRegistryKey& InKey, FMetasoundFrontendInterface& OutInterface) const = 0;
+	};
+} // namespace Metasound::Frontend
diff --git a/Engine/Plugins/Runtime/Metasound/Source/MetasoundFrontend/Public/MetasoundFrontendArchetypeRegistry.h b/Engine/Plugins/Runtime/Metasound/Source/MetasoundFrontend/Public/MetasoundFrontendArchetypeRegistry.h
--- a/Engine/Plugins/Runtime/Metasound/Source/MetasoundFrontend/Public/MetasoundFrontendArchetypeRegistry.h
+++ b/Engine/Plugins/Runtime/Metasound/Source/MetasoundFrontend/Public/MetasoundFrontendArchetypeRegistry.h
@@ -1,121 +1,11 @@
 // Copyright Epic Games, Inc. All Rights Reserved.
-
+// DEPRECATED MetasoundFrontendArchetypeRegistry.h. Use Interfaces/MetasoundFrontendInterfaceRegistry.h instead
 #pragma once
 
-#include "MetasoundFrontendController.h"
-#include "MetasoundFrontendDocument.h"
-#include "UObject/NoExportTypes.h"
-
-namespace Metasound
-{
-	namespace Frontend
-	{
-		using FInterfaceRegistryKey = FString;
-		using FRegistryTransactionID = int32;
-
-		METASOUNDFRONTEND_API bool IsValidInterfaceRegistryKey(const FInterfaceRegistryKey& InKey);
-		METASOUNDFRONTEND_API FInterfaceRegistryKey GetInterfaceRegistryKey(const FMetasoundFrontendVersion& InInterfaceVersion);
-		METASOUNDFRONTEND_API FInterfaceRegistryKey GetInterfaceRegistryKey(const FMetasoundFrontendInterface& InInterface);
-
-		class METASOUNDFRONTEND_API IInterfaceRegistryEntry
-		{
-		public:
-			virtual ~IInterfaceRegistryEntry() = default;
-
-			// MetaSound Interface definition
-			virtual const FMetasoundFrontendInterface& GetInterface() const = 0;
-
-			// Whether or not to immediately apply the given interface
-			// when a MetaSound asset of the supported UClass is created.
-			virtual bool IsDefault(FName InUClassName) const = 0;
-
-			UE_DEPRECATED(5.3, "UClass name must now be provided to determine if interface is to be apply by default.")
-			virtual bool IsDefault() const { return false; }
-
-			// Whether or not the interface can be added or removed by the editor.
-			virtual bool EditorCanAddOrRemove(FName InUClassName) const = 0;
-
-			UE_DEPRECATED(5.3, "UClass name must now be provided to determine if editor can directly add or remove interface.")
-			virtual bool EditorCanAddOrRemove() const { return false; }
-
-			// Name of routing system used to update interface inputs (ex. ParameterInterface or DataReference).
-			virtual FName GetRouterName() const = 0;
-
-			// Name of UClass supported by the given interface.
-			virtual bool UClassIsSupported(FName InUClassName) const = 0;
-
-			// How to update a given document if versioning is required to this interface from a deprecated version.
-			virtual bool UpdateRootGraphInterface(FDocumentHandle InDocument) const = 0;
-
-			// Adds bindings that informs how this interface can be used to connect its inputs to node that implements another interface's outputs. Faster than adding iteratively.
-			virtual void AddOutputBindings(TArray<FMetasoundFrontendInterfaceBinding>&& InBindings) = 0;
-
-			// Returns all bindings that inform how this interface can be used to connect its inputs to a node that implements another interface's outputs.
-			virtual const TArray<FMetasoundFrontendInterfaceBinding>& GetOutputBindings() const = 0;
-
-			// Remove binding that informs how this interface can be used to connect its inputs to node that implements another interface's outputs.
-			virtual bool RemoveOutputBinding(const FMetasoundFrontendVersion& InVersion) = 0;
-
-			// Resets output bindings collection.
-			virtual void ResetOutputBindings() = 0;
-		};
-
-		class METASOUNDFRONTEND_API FInterfaceRegistryTransaction
-		{
-		public:
-			using FTimeType = uint64;
-
-			/** Describes the type of transaction. */
-			enum class ETransactionType : uint8
-			{
-				InterfaceRegistration,     //< Something was added to the registry.
-				InterfaceUnregistration,  //< Something was removed from the registry.
-				Invalid
-			};
-
-			FInterfaceRegistryTransaction(ETransactionType InType, const FInterfaceRegistryKey& InKey, const FMetasoundFrontendVersion& InInterfaceVersion, FTimeType InTimestamp);
-
-			ETransactionType GetTransactionType() const;
-			const FMetasoundFrontendVersion& GetInterfaceVersion() const;
-			const FInterfaceRegistryKey& GetInterfaceRegistryKey() const;
-			FTimeType GetTimestamp() const;
-
-		private:
-
-			ETransactionType Type;
-			FInterfaceRegistryKey Key;
-			FMetasoundFrontendVersion InterfaceVersion;
-			FTimeType Timestamp;
-		};
-
-		class METASOUNDFRONTEND_API IInterfaceRegistry
-		{
-		public:
-			static IInterfaceRegistry& Get();
-
-			virtual ~IInterfaceRegistry() = default;
-
-			// Register an interface
-			virtual void RegisterInterface(TUniquePtr<IInterfaceRegistryEntry>&& InEntry) = 0;
-
-			// Find an interface entry with the given key. Returns null if entry not found with given key.
-			virtual const IInterfaceRegistryEntry* FindInterfaceRegistryEntry(const FInterfaceRegistryKey& InKey) const = 0;
-
-			// Find an interface with the given key. Returns true if interface is found, false if not.
-			virtual bool FindInterface(const FInterfaceRegistryKey& InKey, FMetasoundFrontendInterface& OutInterface) const = 0;
-
-			// Add interface output bindings to the interface with the given key. Orders the bindings using the given sort order, where lower values are evaluated earlier if multiple bindings are provided.
-			virtual void AddInterfaceOutputBindings(const FMetasoundFrontendVersion& InInputInterfaceVersion, TArray<FMetasoundFrontendInterfaceBinding>&& InOutputBindings) = 0;
-
-			// Removes an interface binding to the input interface with the given key from the interface with the given key. Returns true if interface binding was removed, false if not.
-			virtual bool RemoveInterfaceOutputBinding(const FMetasoundFrontendVersion& InInputInterfaceVersion, const FMetasoundFrontendVersion& InOutputInterfaceVersion) = 0;
-
-			// Removes all output interface bindings registered with the given interface version. Returns true if found and removed, false if not found.
-			virtual bool RemoveAllInterfaceOutputBindings(const FMetasoundFrontendVersion& InInputInterfaceVersion) = 0;
+#ifdef _MSC_VER
+#pragma message(__FILE__"(9): warning: use Interfaces/MetasoundFrontendInterfaceRegistry.h instead instead of MetasoundFrontendArchetypeRegistry.h")
+#else
+#pragma message("#include Interfaces/MetasoundFrontendInterfaceRegistry.h instead of MetasoundFrontendArchetypeRegistry.h")
+#endif
 
-			// Executes a function against all transactions since a provided ID and provides the current transaction ID.
-			UE_DEPRECATED(5.1, "ForEachRegistryTransactionSince is no longer supported. Use Metasound::Frontend::ISearchEngine to query for all interfaces.")
-			virtual void ForEachRegistryTransactionSince(FRegistryTransactionID InSince, FRegistryTransactionID* OutCurrentRegistryTransactionID, TFunctionRef<void(const FInterfaceRegistryTransaction&)> InFunc) const = 0;
-		};
-	}
-}
+#include "Interfaces/MetasoundFrontendInterfaceRegistry.h"
diff --git a/Engine/Plugins/Runtime/Metasound/Source/MetasoundFrontend/Public/MetasoundFrontendDocument.h b/Engine/Plugins/Runtime/Metasound/Source/MetasoundFrontend/Public/MetasoundFrontendDocument.h
--- a/Engine/Plugins/Runtime/Metasound/Source/MetasoundFrontend/Public/MetasoundFrontendDocument.h
+++ b/Engine/Plugins/Runtime/Metasound/Source/MetasoundFrontend/Public/MetasoundFrontendDocument.h
@@ -1089,7 +1089,7 @@ public:
 };
 
 USTRUCT()
-struct METASOUNDFRONTEND_API FMetasoundFrontendInterfaceBindingConnection
+struct METASOUNDFRONTEND_API FMetasoundFrontendInterfaceVertexBinding
 {
 	GENERATED_BODY()
 
@@ -1099,12 +1099,12 @@ struct METASOUNDFRONTEND_API FMetasoundFrontendInterfaceBindingConnection
 	UPROPERTY()
 	FName InputName;
 
-	friend bool operator==(const FMetasoundFrontendInterfaceBindingConnection& InLHS, const FMetasoundFrontendInterfaceBindingConnection& InRHS)
+	friend bool operator==(const FMetasoundFrontendInterfaceVertexBinding& InLHS, const FMetasoundFrontendInterfaceVertexBinding& InRHS)
 	{
 		return InLHS.OutputName == InRHS.OutputName && InLHS.InputName == InRHS.InputName;
 	}
 
-	friend bool operator!=(const FMetasoundFrontendInterfaceBindingConnection& InLHS, const FMetasoundFrontendInterfaceBindingConnection& InRHS)
+	friend bool operator!=(const FMetasoundFrontendInterfaceVertexBinding& InLHS, const FMetasoundFrontendInterfaceVertexBinding& InRHS)
 	{
 		return InLHS.OutputName != InRHS.OutputName || InLHS.InputName != InRHS.InputName;
 	}
@@ -1114,7 +1114,7 @@ struct METASOUNDFRONTEND_API FMetasoundFrontendInterfaceBindingConnection
 		return FString::Format(TEXT("{0}->{1}"), { OutputName.ToString(), InputName.ToString() });
 	}
 
-	friend FORCEINLINE uint32 GetTypeHash(const FMetasoundFrontendInterfaceBindingConnection& InBinding)
+	friend FORCEINLINE uint32 GetTypeHash(const FMetasoundFrontendInterfaceVertexBinding& InBinding)
 	{
 		return HashCombineFast(GetTypeHash(InBinding.OutputName), GetTypeHash(InBinding.InputName));
 	}
@@ -1126,16 +1126,23 @@ struct METASOUNDFRONTEND_API FMetasoundFrontendInterfaceBinding
 {
 	GENERATED_BODY()
 
+	// Version of interface to bind from (the corresponding output vertices)
 	UPROPERTY()
-	FMetasoundFrontendVersion Version;
+	FMetasoundFrontendVersion OutputInterfaceVersion;
+
+	// Version of interface to bind to (the corresponding input vertices)
+	UPROPERTY()
+	FMetasoundFrontendVersion InputInterfaceVersion;
 
-	// Priority to determine if binding connections are higher or lower than another interface
-	// that may be shared between nodes attempting to be connected via shared interfaces.
+	// Value describing if interface binding priority is higher or lower than another interface
+	// binding that may be shared between vertices attempting to be connected via binding functionality.
 	UPROPERTY()
-	int32 ConnectionPriority = 0;
+	int32 BindingPriority = 0;
 
+	// Array of named pairs (output & input names) that describe what edges to create if binding functionality
+	// is executed between two nodes.
 	UPROPERTY()
-	TArray<FMetasoundFrontendInterfaceBindingConnection> Connections;
+	TArray<FMetasoundFrontendInterfaceVertexBinding> VertexBindings;
 };
 
 
diff --git a/Engine/Plugins/Runtime/Metasound/Source/MetasoundFrontend/Public/MetasoundFrontendDocumentBuilder.h b/Engine/Plugins/Runtime/Metasound/Source/MetasoundFrontend/Public/MetasoundFrontendDocumentBuilder.h
--- a/Engine/Plugins/Runtime/Metasound/Source/MetasoundFrontend/Public/MetasoundFrontendDocumentBuilder.h
+++ b/Engine/Plugins/Runtime/Metasound/Source/MetasoundFrontend/Public/MetasoundFrontendDocumentBuilder.h
@@ -2,11 +2,11 @@
 #pragma once
 
 #include "Containers/Map.h"
+#include "Interfaces/MetasoundFrontendInterfaceRegistry.h"
+#include "MetasoundFrontendDocumentCacheInterface.h"
 #include "MetasoundDocumentInterface.h"
-#include "MetasoundFrontendArchetypeRegistry.h"
 #include "MetasoundFrontendController.h"
 #include "MetasoundFrontendDocument.h"
-#include "MetasoundFrontendDocumentCacheInterface.h"
 #include "MetasoundFrontendRegistries.h"
 #include "MetasoundFrontendTransform.h"
 #include "MetasoundVertex.h"
@@ -22,7 +22,7 @@ class FMetasoundAssetBase;
 
 namespace Metasound::Frontend
 {
-	struct FNamedEdge
+	struct METASOUNDFRONTEND_API FNamedEdge
 	{
 		const FGuid OutputNodeID;
 		const FName OutputName;
@@ -101,7 +101,7 @@ public:
 	const FMetasoundFrontendClass* AddGraphDependency(const FMetasoundFrontendGraphClass& InClass);
 	const FMetasoundFrontendClass* AddNativeDependency(const FMetasoundFrontendClassMetadata& InClassMetadata);
 	const FMetasoundFrontendEdge* AddEdge(FMetasoundFrontendEdge&& InNewEdge);
-	bool AddEdgesByNamedConnections(const TSet<Metasound::Frontend::FNamedEdge>& ConnectionsToMake, TArray<const FMetasoundFrontendEdge*>* OutNewEdges = nullptr);
+	bool AddNamedEdges(const TSet<Metasound::Frontend::FNamedEdge>& ConnectionsToMake, TArray<const FMetasoundFrontendEdge*>* OutNewEdges = nullptr);
 	bool AddEdgesByNodeClassInterfaceBindings(const FGuid& InFromNodeID, const FGuid& InToNodeID);
 	bool AddEdgesFromMatchingInterfaceNodeOutputsToGraphOutputs(const FGuid& InNodeID, TArray<const FMetasoundFrontendEdge*>& OutEdgesCreated);
 	bool AddEdgesFromMatchingInterfaceNodeInputsToGraphInputs(const FGuid& InNodeID, TArray<const FMetasoundFrontendEdge*>& OutEdgesCreated);
@@ -154,7 +154,7 @@ public:
 
 	bool RemoveDependency(EMetasoundFrontendClassType ClassType, const FMetasoundFrontendClassName& InClassName, const FMetasoundFrontendVersionNumber& InClassVersionNumber);
 	bool RemoveEdge(const FMetasoundFrontendEdge& EdgeToRemove);
-	bool RemoveEdgesByNamedConnections(const TSet<Metasound::Frontend::FNamedEdge>& ConnectionsToRemove, TArray<FMetasoundFrontendEdge>* OutRemovedEdges = nullptr);
+	bool RemoveNamedEdges(const TSet<Metasound::Frontend::FNamedEdge>& InNamedEdgesToRemove, TArray<FMetasoundFrontendEdge>* OutRemovedEdges = nullptr);
 	bool RemoveEdgesByNodeClassInterfaceBindings(const FGuid& InOutputNodeID, const FGuid& InInputNodeID);
 	bool RemoveEdgesFromNodeOutput(const FGuid& InNodeID, const FGuid& InVertexID);
 	bool RemoveEdgeToNodeInput(const FGuid& InNodeID, const FGuid& InVertexID);
@@ -179,7 +179,7 @@ private:
 
 	FMetasoundFrontendNode* AddNodeInternal(const FMetasoundFrontendClassMetadata& InClassMetadata, FFinalizeNodeFunctionRef FinalizeNode, FGuid InNodeID = FGuid::NewGuid());
 
-	const TScriptInterface<IMetaSoundDocumentInterface> FindOrLoadNodeClassDocument(const FGuid& InNodeID) const;
+	const FMetasoundFrontendDocument* FindOrLoadNodeClassDocument(const FGuid& InNodeID) const;
 	const FMetasoundFrontendClassInput* FindNodeInputClassInput(const FGuid& InNodeID, const FGuid& InVertexID) const;
 	const FMetasoundFrontendClassOutput* FindNodeOutputClassOutput(const FGuid& InNodeID, const FGuid& InVertexID) const;
 
diff --git a/Engine/Plugins/Runtime/Metasound/Source/MetasoundFrontend/Public/MetasoundFrontendGraph.h b/Engine/Plugins/Runtime/Metasound/Source/MetasoundFrontend/Public/MetasoundFrontendGraph.h
--- a/Engine/Plugins/Runtime/Metasound/Source/MetasoundFrontend/Public/MetasoundFrontendGraph.h
+++ b/Engine/Plugins/Runtime/Metasound/Source/MetasoundFrontend/Public/MetasoundFrontendGraph.h
@@ -91,7 +91,7 @@ namespace Metasound
 			TArray<TSharedPtr<const INode>> NodeStorage;
 	};
 
-	/** FFrontendGraphBuilder builds a FFrontendGraph from a FMetasoundDoucment
+	/** FFrontendGraphBuilder builds a FFrontendGraph from a FMetasoundDocument
 	 * or FMetasoundFrontendClass.
 	 */
 	class METASOUNDFRONTEND_API FFrontendGraphBuilder
diff --git a/Engine/Plugins/Runtime/Metasound/Source/MetasoundFrontend/Public/MetasoundFrontendQuery.h b/Engine/Plugins/Runtime/Metasound/Source/MetasoundFrontend/Public/MetasoundFrontendQuery.h
--- a/Engine/Plugins/Runtime/Metasound/Source/MetasoundFrontend/Public/MetasoundFrontendQuery.h
+++ b/Engine/Plugins/Runtime/Metasound/Source/MetasoundFrontend/Public/MetasoundFrontendQuery.h
@@ -4,10 +4,10 @@
 
 #include "Containers/Array.h"
 #include "Containers/SortedMap.h"
+#include "Interfaces/MetasoundFrontendInterfaceRegistry.h"
 #include "Math/NumericLimits.h"
 #include "MetasoundFrontendDocument.h"
 #include "MetasoundFrontendRegistries.h"
-#include "MetasoundFrontendArchetypeRegistry.h"
 #include "Misc/TVariant.h"
 
 /** MetaSound Frontend Query
diff --git a/Engine/Plugins/Runtime/Metasound/Source/MetasoundFrontend/Public/MetasoundFrontendSearchEngine.h b/Engine/Plugins/Runtime/Metasound/Source/MetasoundFrontend/Public/MetasoundFrontendSearchEngine.h
--- a/Engine/Plugins/Runtime/Metasound/Source/MetasoundFrontend/Public/MetasoundFrontendSearchEngine.h
+++ b/Engine/Plugins/Runtime/Metasound/Source/MetasoundFrontend/Public/MetasoundFrontendSearchEngine.h
@@ -1,58 +1,56 @@
 // Copyright Epic Games, Inc. All Rights Reserved.
-
 #pragma once
 
-#include "MetasoundFrontendArchetypeRegistry.h"
+
+#include "Interfaces/MetasoundFrontendInterfaceRegistry.h"
 #include "MetasoundFrontendDocument.h"
 #include "MetasoundNodeInterface.h"
 
-namespace Metasound
+
+namespace Metasound::Frontend
 {
-	namespace Frontend
+	/** Interface for frontend search engine. A frontend search engine provides
+		* a simple interface for common frontend queries. It also serves as an
+		* opportunity to cache queries to reduce CPU load. 
+		*/
+	class METASOUNDFRONTEND_API ISearchEngine
 	{
-		/** Interface for frontend search engine. A frontend search engine provides
-		 * a simple interface for common frontend queries. It also serves as an
-		 * opportunity to cache queries to reduce CPU load. 
-		 */
-		class METASOUNDFRONTEND_API ISearchEngine
-		{
-			public:
-				/** Return an instance of a search engine. */
-				static ISearchEngine& Get();
+		public:
+			/** Return an instance of a search engine. */
+			static ISearchEngine& Get();
 
-				virtual ~ISearchEngine() = default;
+			virtual ~ISearchEngine() = default;
 
-				/** Updates internal state to speed up queries. */
-				virtual void Prime() = 0;
+			/** Updates internal state to speed up queries. */
+			virtual void Prime() = 0;
 
-				/** Find the class with the given ClassName & Major Version. Returns false if not found, true if found. */
-				virtual bool FindClassWithHighestMinorVersion(const FMetasoundFrontendClassName& InName, int32 InMajorVersion, FMetasoundFrontendClass& OutClass) = 0;
+			/** Find the class with the given ClassName & Major Version. Returns false if not found, true if found. */
+			virtual bool FindClassWithHighestMinorVersion(const FMetasoundFrontendClassName& InName, int32 InMajorVersion, FMetasoundFrontendClass& OutClass) = 0;
 
-				virtual TArray<FMetasoundFrontendVersion> FindAllRegisteredInterfacesWithName(FName InInterfaceName) = 0;
+			virtual TArray<FMetasoundFrontendVersion> FindAllRegisteredInterfacesWithName(FName InInterfaceName) = 0;
 
-				virtual bool FindInterfaceWithHighestVersion(FName InInterfaceName, FMetasoundFrontendInterface& OutInterface) = 0;
+			virtual bool FindInterfaceWithHighestVersion(FName InInterfaceName, FMetasoundFrontendInterface& OutInterface) = 0;
 
-				virtual TArray<FMetasoundFrontendInterface> FindUClassDefaultInterfaces(FName InUClassName) = 0;
+			virtual TArray<FMetasoundFrontendInterface> FindUClassDefaultInterfaces(FName InUClassName) = 0;
 
-				virtual TArray<FMetasoundFrontendInterface> FindAllInterfaces(bool bInIncludeAllVersions = false) = 0;
+			virtual TArray<FMetasoundFrontendInterface> FindAllInterfaces(bool bInIncludeAllVersions = false) = 0;
 
 #if WITH_EDITORONLY_DATA
-				/** Find all FMetasoundFrontendClasses.
-				  * (Optional) Include all versions (i.e. deprecated classes and versions of classes that are not the highest major version).
-				  */
-				virtual TArray<FMetasoundFrontendClass> FindAllClasses(bool bInIncludeAllVersions) = 0;
-
-				/** Find all classes with the given ClassName.
-				  * (Optional) Sort matches based on version.
-				  */
-				virtual TArray<FMetasoundFrontendClass> FindClassesWithName(const FMetasoundFrontendClassName& InName, bool bInSortByVersion) = 0;
-
-				/** Find the highest version of a class with the given ClassName. Returns false if not found, true if found. */
-				virtual bool FindClassWithHighestVersion(const FMetasoundFrontendClassName& InName, FMetasoundFrontendClass& OutClass) = 0;
+			/** Find all FMetasoundFrontendClasses.
+				* (Optional) Include all versions (i.e. deprecated classes and versions of classes that are not the highest major version).
+				*/
+			virtual TArray<FMetasoundFrontendClass> FindAllClasses(bool bInIncludeAllVersions) = 0;
+
+			/** Find all classes with the given ClassName.
+				* (Optional) Sort matches based on version.
+				*/
+			virtual TArray<FMetasoundFrontendClass> FindClassesWithName(const FMetasoundFrontendClassName& InName, bool bInSortByVersion) = 0;
+
+			/** Find the highest version of a class with the given ClassName. Returns false if not found, true if found. */
+			virtual bool FindClassWithHighestVersion(const FMetasoundFrontendClassName& InName, FMetasoundFrontendClass& OutClass) = 0;
 #endif
 
-			protected:
-				ISearchEngine() = default;
-		};
-	}
-}
+		protected:
+			ISearchEngine() = default;
+	};
+} // namespace Metasound::Frontend
diff --git a/Engine/Plugins/Runtime/Metasound/Source/MetasoundFrontend/Public/MetasoundFrontendTransform.h b/Engine/Plugins/Runtime/Metasound/Source/MetasoundFrontend/Public/MetasoundFrontendTransform.h
--- a/Engine/Plugins/Runtime/Metasound/Source/MetasoundFrontend/Public/MetasoundFrontendTransform.h
+++ b/Engine/Plugins/Runtime/Metasound/Source/MetasoundFrontend/Public/MetasoundFrontendTransform.h
@@ -2,8 +2,8 @@
 #pragma once
 
 #include "Internationalization/Text.h"
+#include "Interfaces/MetasoundFrontendInterfaceRegistry.h"
 #include "MetasoundAssetBase.h"
-#include "MetasoundFrontendArchetypeRegistry.h"
 #include "MetasoundFrontendController.h"
 #include "MetasoundFrontendDocument.h"
 #include "MetasoundVertex.h"
