[Metasound] TVertexAnalyzerForwardValue - initialize LastValue and always mark dirty on first execute

Original review: https://p4-swarm.epicgames.net/reviews/24636511

diff --git a/Engine/Plugins/Runtime/Metasound/Source/MetasoundFrontend/Public/Analysis/MetasoundFrontendVertexAnalyzerForwardValue.h b/Engine/Plugins/Runtime/Metasound/Source/MetasoundFrontend/Public/Analysis/MetasoundFrontendVertexAnalyzerForwardValue.h
--- a/Engine/Plugins/Runtime/Metasound/Source/MetasoundFrontend/Public/Analysis/MetasoundFrontendVertexAnalyzerForwardValue.h
+++ b/Engine/Plugins/Runtime/Metasound/Source/MetasoundFrontend/Public/Analysis/MetasoundFrontendVertexAnalyzerForwardValue.h
@@ -24,6 +24,7 @@ namespace Metasound
 		class TVertexAnalyzerForwardValue : public FVertexAnalyzerBase
 		{
 			DataType LastValue;
+			bool HasExecuted = false;
 
 		public:
 			static const FName& GetDataType()
@@ -67,12 +68,24 @@ namespace Metasound
 			virtual void Execute() override
 			{
 				const DataType& Value = GetVertexData<DataType>();
-				if (LastValue != Value)
+
+				if (!HasExecuted)
+				{
+					HasExecuted = true;
+					UpdateValue(Value);
+				}
+				else if (LastValue != Value)
 				{
-					LastValue = Value;
-					MarkOutputDirty();
+					UpdateValue(Value);
 				}
 			}
+
+		private:
+			void UpdateValue(const DataType& Value)
+			{
+				LastValue = Value;
+				MarkOutputDirty();
+			}
 		};
 
 		#define METASOUND_DECLARE_FORWARD_VALUE_ANALYZER(ClassType, AnalyzerFName, DataType) \
