Add render timing feedback to metasound editor.

diff --git a/Engine/Plugins/Runtime/Metasound/Source/MetasoundEditor/Private/MetasoundEditor.cpp b/Engine/Plugins/Runtime/Metasound/Source/MetasoundEditor/Private/MetasoundEditor.cpp
--- a/Engine/Plugins/Runtime/Metasound/Source/MetasoundEditor/Private/MetasoundEditor.cpp
+++ b/Engine/Plugins/Runtime/Metasound/Source/MetasoundEditor/Private/MetasoundEditor.cpp
@@ -46,7 +46,9 @@
 #include "MetasoundFrontendRegistries.h"
 #include "MetasoundFrontendSearchEngine.h"
 #include "MetasoundFrontendTransform.h"
+#include "MetasoundGenerator.h"
 #include "MetasoundLog.h"
+#include "MetasoundSource.h"
 #include "MetasoundUObjectRegistry.h"
 #include "Misc/Attribute.h"
 #include "Modules/ModuleManager.h"
@@ -604,7 +606,7 @@ namespace Metasound
 
 			FAssetEditorToolkit::RegisterTabSpawners(InTabManager);
 
-			InTabManager->RegisterTabSpawner(TabFactory::Names::GraphCanvas, FOnSpawnTab::CreateLambda([InPlayTimeWidget = PlayTimeWidget, InMetasoundGraphEditor = MetasoundGraphEditor](const FSpawnTabArgs& Args)
+			InTabManager->RegisterTabSpawner(TabFactory::Names::GraphCanvas, FOnSpawnTab::CreateLambda([InPlayTimeWidget = PlayTimeWidget, InCPUCoreUtilizationWidget = CPUCoreUtilizationWidget, InMetasoundGraphEditor = MetasoundGraphEditor](const FSpawnTabArgs& Args)
 			{
 				return TabFactory::CreateGraphCanvasTab(SNew(SOverlay)
 					+ SOverlay::Slot()
@@ -615,6 +617,11 @@ namespace Metasound
 					[
 						InPlayTimeWidget.ToSharedRef()
 					]
+					+ SOverlay::Slot()
+					.VAlign(VAlign_Bottom)
+					[
+						InCPUCoreUtilizationWidget.ToSharedRef()
+					]
 					.Padding(5.0f, 5.0f)
 				, Args);
 			}))
@@ -1564,6 +1571,33 @@ namespace Metasound
 								PlayTimeString.ReplaceInline(TEXT("+"), TEXT(""));
 								PlayTimeWidget->SetText(FText::FromString(PlayTimeString));
 							}
+							if (CPUCoreUtilizationWidget.IsValid())
+							{
+								// TODO: Need to protect against shipping builds of editor
+								check(IsInGameThread());
+								double CPUCoreUtilization = 0;
+								if (UMetaSoundSource* Source = Cast<UMetaSoundSource>(Metasound))
+								{
+									if (const UAudioComponent* PreviewComponent = GEditor->GetPreviewAudioComponent())
+									{
+										TSharedPtr<FMetasoundGenerator> Generator = Source->GetGeneratorForAudioComponent(PreviewComponent->GetAudioComponentID()).Pin();
+										if (Generator.IsValid())
+										{
+											CPUCoreUtilization = Generator->GetCPUCoreUtilization();
+										}
+									}
+								}
+								if (CPUCoreUtilization > 0)
+								{
+									FString CPUCoreUtilizationString = FString::Printf(TEXT("%5.2f %% CPU Core"), 100. * CPUCoreUtilization);
+									CPUCoreUtilizationWidget->SetText(FText::FromString(CPUCoreUtilizationString));
+								}
+								else
+								{
+									CPUCoreUtilizationWidget->SetText(FText::GetEmpty());
+								}
+							}
+
 							return EActiveTimerReturnType::Continue;
 						}
 						else
@@ -1571,6 +1605,10 @@ namespace Metasound
 							SetPreviewID(INDEX_NONE);
 							PlayTime = 0.0;
 							PlayTimeWidget->SetText(FText::GetEmpty());
+							if (CPUCoreUtilizationWidget.IsValid())
+							{
+								CPUCoreUtilizationWidget->SetText(FText::GetEmpty());
+							}
 							GraphConnectionManager = MakeUnique<FGraphConnectionManager>();
 
 							return EActiveTimerReturnType::Stop;
@@ -1912,6 +1950,11 @@ namespace Metasound
 				.Visibility(EVisibility::HitTestInvisible)
 				.TextStyle(FAppStyle::Get(), "Graph.ZoomText")
 				.ColorAndOpacity(FLinearColor(1, 1, 1, 0.30f));
+
+			SAssignNew(CPUCoreUtilizationWidget, STextBlock)
+				.Visibility(EVisibility::HitTestInvisible)
+				.TextStyle(FAppStyle::Get(), "Graph.ZoomText")
+				.ColorAndOpacity(FLinearColor(1, 1, 1, 0.30f));
 		}
 
 		FGraphAppearanceInfo FEditor::GetGraphAppearance() const
diff --git a/Engine/Plugins/Runtime/Metasound/Source/MetasoundEditor/Private/MetasoundEditor.h b/Engine/Plugins/Runtime/Metasound/Source/MetasoundEditor/Private/MetasoundEditor.h
--- a/Engine/Plugins/Runtime/Metasound/Source/MetasoundEditor/Private/MetasoundEditor.h
+++ b/Engine/Plugins/Runtime/Metasound/Source/MetasoundEditor/Private/MetasoundEditor.h
@@ -427,6 +427,9 @@ namespace Metasound
 			TSharedPtr<STextBlock> PlayTimeWidget;
 			double PlayTime = 0.0;
 
+			/** Widget showing render speed that overlays the graph when previewing */
+			TSharedPtr<STextBlock> CPUCoreUtilizationWidget;
+
 			TUniquePtr<FGraphConnectionManager> GraphConnectionManager;
 
 			/** Command list for this editor */
diff --git a/Engine/Plugins/Runtime/Metasound/Source/MetasoundGenerator/Private/MetasoundGenerator.cpp b/Engine/Plugins/Runtime/Metasound/Source/MetasoundGenerator/Private/MetasoundGenerator.cpp
--- a/Engine/Plugins/Runtime/Metasound/Source/MetasoundGenerator/Private/MetasoundGenerator.cpp
+++ b/Engine/Plugins/Runtime/Metasound/Source/MetasoundGenerator/Private/MetasoundGenerator.cpp
@@ -17,8 +17,8 @@
 #include "DSP/FloatArrayMath.h"
 #include "Interfaces/MetasoundFrontendSourceInterface.h"
 
-#ifndef METASOUNDGENERATOR_ENABLE_INVALID_SAMPLE_VALUE_LOGGING 
-#define METASOUNDGENERATOR_ENABLE_INVALID_SAMPLE_VALUE_LOGGING !UE_BUILD_SHIPPING
+#ifndef ENABLE_METASOUNDGENERATOR_INVALID_SAMPLE_VALUE_LOGGING 
+#define ENABLE_METASOUNDGENERATOR_INVALID_SAMPLE_VALUE_LOGGING !UE_BUILD_SHIPPING
 #endif
 
 namespace Metasound
@@ -26,16 +26,16 @@ namespace Metasound
 	namespace ConsoleVariables
 	{
 		static bool bEnableAsyncMetaSoundGeneratorBuilder = true;
-#if METASOUNDGENERATOR_ENABLE_INVALID_SAMPLE_VALUE_LOGGING
+#if ENABLE_METASOUNDGENERATOR_INVALID_SAMPLE_VALUE_LOGGING
 		static bool bEnableMetaSoundGeneratorNonFiniteLogging = false;
 		static bool bEnableMetaSoundGeneratorInvalidSampleValueLogging = false;
 		static float MetasoundGeneratorSampleValueThreshold = 2.f;
-#endif // #if METASOUNDGENERATOR_ENABLE_INVALID_SAMPLE_VALUE_LOGGING
+#endif // if ENABLE_METASOUNDGENERATOR_INVALID_SAMPLE_VALUE_LOGGING
 	}
 
-#if METASOUNDGENERATOR_ENABLE_INVALID_SAMPLE_VALUE_LOGGING
 	namespace MetasoundGeneratorPrivate
 	{
+#if ENABLE_METASOUNDGENERATOR_INVALID_SAMPLE_VALUE_LOGGING
 		void LogInvalidAudioSampleValues(const FString& InMetaSoundName, const TArray<FAudioBufferReadRef>& InAudioBuffers)
 		{
 			if (ConsoleVariables::bEnableMetaSoundGeneratorNonFiniteLogging || ConsoleVariables::bEnableMetaSoundGeneratorInvalidSampleValueLogging)
@@ -72,8 +72,73 @@ namespace Metasound
 				}
 			}
 		}
+#endif // if ENABLE_METASOUNDGENERATOR_INVALID_SAMPLE_VALUE_LOGGING
+
+#if ENABLE_METASOUND_GENERATOR_RENDER_TIMING
+		struct FRenderTimer
+		{
+			FRenderTimer(const FOperatorSettings& InSettings, double InAnalysisDuration)
+			{
+				// Use single pole IIR filter to smooth data.
+				
+				const double AnalysisAudioFrameCount = FMath::Max(1.0, InAnalysisDuration * InSettings.GetSampleRate());
+				const double AnalysisRenderBlockCount = FMath::Max(1.0, AnalysisAudioFrameCount / FMath::Max(1, InSettings.GetNumFramesPerBlock()));
+				const double DigitalCutoff = 1. / AnalysisRenderBlockCount;
+
+				SmoothingAlpha = 1. - FMath::Exp(-UE_PI * DigitalCutoff);
+				SmoothingAlpha = FMath::Clamp(SmoothingAlpha, 0.0, 1.0 - UE_DOUBLE_SMALL_NUMBER);
+				SecondsOfAudioProducedPerBlock = static_cast<double>(InSettings.GetNumFramesPerBlock()) / FMath::Max(1., static_cast<double>(InSettings.GetSampleRate()));
+			}
+
+			double GetCPUCoreUtilization() const
+			{
+				return CPUCoreUtilization;
+			}
+
+			void UpdateCPUCoreUtilization(double InCPUSecondsToRenderBlock)
+			{
+				if (InCPUSecondsToRenderBlock > 0.0)
+				{
+					double NewCPUUtil = InCPUSecondsToRenderBlock / SecondsOfAudioProducedPerBlock;
+					if (CPUCoreUtilization >= 0.0)
+					{
+						CPUCoreUtilization = SmoothingAlpha * NewCPUUtil + (1. - SmoothingAlpha) * CPUCoreUtilization;
+					}
+					else
+					{
+						CPUCoreUtilization = NewCPUUtil;
+					}
+				}
+			}
+
+		private:
+			double CPUCoreUtilization = -1.0;
+			double SmoothingAlpha = 1.0;
+			double SecondsOfAudioProducedPerBlock = 0.0;
+		};
+
+		struct FBlockRenderScope
+		{
+			FBlockRenderScope(FRenderTimer& InTimer)
+			: Timer(&InTimer)
+			{
+				StartCycle = FPlatformTime::Cycles64();
+			}
+
+			~FBlockRenderScope()
+			{
+				uint64 EndCycle = FPlatformTime::Cycles64();
+				if (EndCycle > StartCycle)
+				{
+					Timer->UpdateCPUCoreUtilization(FPlatformTime::ToSeconds64(EndCycle - StartCycle));
+				}
+			}
+		private:
+			uint64 StartCycle = 0;
+			FRenderTimer* Timer = nullptr;
+		};
+#endif // if ENABLE_METASOUND_GENERATOR_RENDER_TIMING
 	}
-#endif // #if METASOUNDGENERATOR_ENABLE_INVALID_SAMPLE_VALUE_LOGGING
 }
 
 FAutoConsoleVariableRef CVarMetaSoundEnableAsyncGeneratorBuilder(
@@ -83,7 +148,7 @@ FAutoConsoleVariableRef CVarMetaSoundEnableAsyncGeneratorBuilder(
 	TEXT("Default: true"),
 	ECVF_Default);
 
-#if !UE_BUILD_SHIPPING
+#if ENABLE_METASOUNDGENERATOR_INVALID_SAMPLE_VALUE_LOGGING
 
 FAutoConsoleVariableRef CVarMetaSoundEnableGeneratorNonFiniteLogging(
 	TEXT("au.MetaSound.EnableGeneratorNonFiniteLogging"),
@@ -106,7 +171,7 @@ FAutoConsoleVariableRef CVarMetaSoundGeneratorSampleValueThrehshold(
 	TEXT("Default: 2.0"),
 	ECVF_Default);
 
-#endif // #if !UE_BUILD_SHIPPING
+#endif // if !UE_BUILD_SHIPPING
 
 namespace Metasound
 {
@@ -318,6 +383,10 @@ namespace Metasound
 		, bPendingGraphTrigger(true)
 		, bIsNewGraphPending(false)
 		, bIsWaitingForFirstGraph(true)
+#if ENABLE_METASOUND_GENERATOR_RENDER_TIMING
+		, RenderTimer(MakeUnique<MetasoundGeneratorPrivate::FRenderTimer>(InParams.OperatorSettings, 1. /* AnalysisPeriod */))
+#endif // if ENABLE_METASOUND_GENERATOR_RENDER_TIMING
+
 	{
 		NumChannels = InParams.AudioOutputNames.Num();
 		NumFramesPerExecute = InParams.OperatorSettings.GetNumFramesPerBlock();
@@ -560,8 +629,15 @@ namespace Metasound
 		{
 			UnpackAndTransmitUpdatedParameters();
 
-			// Call metasound graph operator.
-			RootExecuter.Execute();
+			{
+#if ENABLE_METASOUND_GENERATOR_RENDER_TIMING
+				// Time how long the root executer takes.
+				MetasoundGeneratorPrivate::FBlockRenderScope BlockRenderScope(*RenderTimer);
+#endif // if ENABLE_METASOUND_GENERATOR_RENDER_TIMING
+
+				// Call metasound graph operator.
+				RootExecuter.Execute();
+			}
 
 			if (GraphAnalyzer.IsValid())
 			{
@@ -616,6 +692,13 @@ namespace Metasound
 		return bIsFinished;
 	}
 
+#if ENABLE_METASOUND_GENERATOR_RENDER_TIMING
+	double FMetasoundGenerator::GetCPUCoreUtilization() const
+	{
+		return RenderTimer->GetCPUCoreUtilization();
+	}
+#endif // if ENABLE_METASOUND_GENERATOR_RENDER_TIMING
+
 	int32 FMetasoundGenerator::FillWithBuffer(const Audio::FAlignedFloatBuffer& InBuffer, float* OutAudio, int32 MaxNumOutputSamples)
 	{
 		int32 InNum = InBuffer.Num();
@@ -651,10 +734,9 @@ namespace Metasound
 			InterleavedAudioBuffer.AddUninitialized(NumSamplesPerExecute);
 		}
 
-#if METASOUNDGENERATOR_ENABLE_INVALID_SAMPLE_VALUE_LOGGING
+#if ENABLE_METASOUNDGENERATOR_INVALID_SAMPLE_VALUE_LOGGING
 		MetasoundGeneratorPrivate::LogInvalidAudioSampleValues(MetasoundName, GraphOutputAudio);
-#endif // #if METASOUNDGENERATOR_ENABLE_INVALID_SAMPLE_VALUE_LOGGING
-		
+#endif // if ENABLE_METASOUNDGENERATOR_INVALID_SAMPLE_VALUE_LOGGING
 
 		// Iterate over channels
 		for (int32 ChannelIndex = 0; ChannelIndex < NumChannels; ChannelIndex++)
diff --git a/Engine/Plugins/Runtime/Metasound/Source/MetasoundGenerator/Public/MetasoundGenerator.h b/Engine/Plugins/Runtime/Metasound/Source/MetasoundGenerator/Public/MetasoundGenerator.h
--- a/Engine/Plugins/Runtime/Metasound/Source/MetasoundGenerator/Public/MetasoundGenerator.h
+++ b/Engine/Plugins/Runtime/Metasound/Source/MetasoundGenerator/Public/MetasoundGenerator.h
@@ -15,8 +15,19 @@
 #include "Containers/MpscQueue.h"
 #include "Sound/SoundGenerator.h"
 
+#ifndef ENABLE_METASOUND_GENERATOR_RENDER_TIMING
+#define ENABLE_METASOUND_GENERATOR_RENDER_TIMING !UE_BUILD_SHIPPING && WITH_EDITOR
+#endif // ifndef ENABLE_METASOUND_GENERATOR_RENDER_TIMING
+
 namespace Metasound
 {
+	namespace MetasoundGeneratorPrivate
+	{
+#if ENABLE_METASOUND_GENERATOR_RENDER_TIMING
+		struct FRenderTimer;
+#endif // if ENABLE_METASOUND_GENERATOR_RENDER_TIMING
+	}
+
 	// Struct needed for building the metasound graph
 	struct METASOUNDGENERATOR_API FMetasoundGeneratorInitParams
 	{
@@ -204,6 +215,12 @@ namespace Metasound
 		int32 GetDesiredNumSamplesToRenderPerCallback() const override;
 		bool IsFinished() const override;
 		//~ End FSoundGenerator
+
+#if ENABLE_METASOUND_GENERATOR_RENDER_TIMING
+		/** Fraction of a single CPU core used to render audio on a scale of 0.0 to 1.0 */
+		double GetCPUCoreUtilization() const;
+
+#endif // if ENABLE_METASOUND_GENERATOR_RENDER_TIMING
 		
 		/** Update the current graph operator with a new graph operator. The number of channels
 		 * of InGraphOutputAudioRef must match the existing number of channels reported by
@@ -216,6 +233,7 @@ namespace Metasound
 		void SetPendingGraphBuildFailed();
 
 	private:
+
 		bool UpdateGraphIfPending();
 
 		// Internal set graph after checking compatibility.
@@ -282,5 +300,8 @@ namespace Metasound
 
 		TMpscQueue<TUniqueFunction<void()>> OutputAnalyzerModificationQueue;
 		TMap<FName, TUniquePtr<Frontend::IVertexAnalyzer>> OutputAnalyzers;
+#if ENABLE_METASOUND_GENERATOR_RENDER_TIMING
+		TUniquePtr<MetasoundGeneratorPrivate::FRenderTimer> RenderTimer;
+#endif // if ENABLE_METASOUND_GENERATOR_RENDER_TIMING
 	};
 }
