Build fix. Move metasound enum definition to fix incompatible dll linkage.
#rnx

diff --git a/Engine/Plugins/Runtime/Metasound/Source/MetasoundStandardNodes/Private/MetasoundMidiScaleToArrayNode.cpp b/Engine/Plugins/Runtime/Metasound/Source/MetasoundStandardNodes/Private/MetasoundMidiScaleToArrayNode.cpp
--- a/Engine/Plugins/Runtime/Metasound/Source/MetasoundStandardNodes/Private/MetasoundMidiScaleToArrayNode.cpp
+++ b/Engine/Plugins/Runtime/Metasound/Source/MetasoundStandardNodes/Private/MetasoundMidiScaleToArrayNode.cpp
@@ -21,6 +21,48 @@
 
 namespace Metasound
 {
+	DEFINE_METASOUND_ENUM_BEGIN(Audio::EMusicalScale::Scale, FEnumEMusicalScale, "MusicalScale")
+
+	// modes
+	DEFINE_METASOUND_ENUM_ENTRY(Audio::EMusicalScale::Scale::Major, "MajorDescription", "Major Scale", "MajorDescriptionTT", "Major (Ionian)"),
+	DEFINE_METASOUND_ENUM_ENTRY(Audio::EMusicalScale::Scale::Minor_Dorian, "Minor_DorianDescription", "Minor (Dorian)", "Minor_DorianDescriptionTT", "Dorian Minor"),
+	DEFINE_METASOUND_ENUM_ENTRY(Audio::EMusicalScale::Scale::Phrygian, "PhrygianDescription", "Phrygian ", "PhrygianDescriptionTT", "Phrygian"),
+	DEFINE_METASOUND_ENUM_ENTRY(Audio::EMusicalScale::Scale::Lydian, "LydianDescription", "Lydian", "LydianDescriptionTT", "Lydian (sharp-4)"),
+	DEFINE_METASOUND_ENUM_ENTRY(Audio::EMusicalScale::Scale::Dominant7th_Mixolydian, "Dominant7th_MixolydianDescription", "Dominant 7th (Mixolydian)", "Dominant7th_MixolydianDescriptionTT", "Mioxlydian (Dominant 7)"),
+	DEFINE_METASOUND_ENUM_ENTRY(Audio::EMusicalScale::Scale::NaturalMinor_Aeolian, "NaturalMinor_AeolianDescription", "Natural Minor (Aeolian)", "NaturalMinor_AeolianDescriptionTT", "Natural Minor (Aeolian)"),
+	DEFINE_METASOUND_ENUM_ENTRY(Audio::EMusicalScale::Scale::HalfDiminished_Locrian, "HalfDiminished_LocrianDescription", "Half Diminished (Locrian)", "HalfDiminished_LocrianDescriptionTT", "Half-Diminished (Locrian)"),
+	DEFINE_METASOUND_ENUM_ENTRY(Audio::EMusicalScale::Scale::Diminished, "DiminishedDescription", "Diminished ", "DiminishedDescriptionTT", "Diminished"),
+	// non-diatonic
+	DEFINE_METASOUND_ENUM_ENTRY(Audio::EMusicalScale::Scale::Chromatic, "ChromaticDescription", "Chromatic", "ChromaticDescriptionTT", "Chromatic"),
+	DEFINE_METASOUND_ENUM_ENTRY(Audio::EMusicalScale::Scale::WholeTone, "WholeToneDescription", "Whole-Tone", "WholeToneDescriptionTT", "Whole Tone"),
+	DEFINE_METASOUND_ENUM_ENTRY(Audio::EMusicalScale::Scale::DiminishedWholeTone, "DiminishedWholeToneDescription", "Diminished Whole-Tone", "DiminishedWholeToneDescriptionTT", "Diminished Whole Tone"),
+	// petantonic
+	DEFINE_METASOUND_ENUM_ENTRY(Audio::EMusicalScale::Scale::MajorPentatonic, "MajorPentatonicDescription", "Major Pentatonic ", "MajorPentatonicDescriptionTT", "Major Pentatonic"),
+	DEFINE_METASOUND_ENUM_ENTRY(Audio::EMusicalScale::Scale::MinorPentatonic, "MinorPentatonicDescription", "Minor Pentatonic ", "MinorPentatonicDescriptionTT", "Minor Pentatonic"),
+	DEFINE_METASOUND_ENUM_ENTRY(Audio::EMusicalScale::Scale::Blues, "BluesDescription", "Blues ", "BluesDescriptionTT", "Blues"),
+	// bebop
+	DEFINE_METASOUND_ENUM_ENTRY(Audio::EMusicalScale::Scale::Bebop_Major, "Bebop_MajorDescription", "Bebop (Major)", "Bebop_MajorDescriptionTT", "Bebop Major"),
+	DEFINE_METASOUND_ENUM_ENTRY(Audio::EMusicalScale::Scale::Bebop_Minor, "Bebop_MinorDescription", "Bebop (Minor)", "Bebop_MinorDescriptionTT", "Bebop Minor"),
+	DEFINE_METASOUND_ENUM_ENTRY(Audio::EMusicalScale::Scale::Bebop_MinorNumber2, "Bebop_MinorNumber2Description", "Bebop (Minor) #2", "Bebop_MinorNumber2DescriptionTT", "Bebop Minor #2"),
+	DEFINE_METASOUND_ENUM_ENTRY(Audio::EMusicalScale::Scale::Bebop_Dominant, "Bebop_DominantDescription", "Bebop (Dominant)", "Bebop_DominantDescriptionTT", "Bebop Dominant"),
+	// common major/minors
+	DEFINE_METASOUND_ENUM_ENTRY(Audio::EMusicalScale::Scale::HarmonicMajor, "HarmonicMajorDescription", "Harmonic Major", "HarmonicMajorDescriptionTT", "Harmonic Major"),
+	DEFINE_METASOUND_ENUM_ENTRY(Audio::EMusicalScale::Scale::HarmonicMinor, "HarmonicMinorDescription", "Harmonic Minor ", "HarmonicMinorDescriptionTT", "Harmonic Minor"),
+	DEFINE_METASOUND_ENUM_ENTRY(Audio::EMusicalScale::Scale::MelodicMinor, "MelodicMinorDescription", "Melodic Minor ", "MelodicMinorDescriptionTT", "Melodic Minor"),
+	DEFINE_METASOUND_ENUM_ENTRY(Audio::EMusicalScale::Scale::SixthModeOfHarmonicMinor, "SixthModeOfHarmonicMinorDescription", "Sixth Mode of Harmonic Minor", "SixthModeOfHarmonicMinorDescriptionTT", "Sixth Mode of Harmonic Minor"),
+	// lydian/augmented
+	DEFINE_METASOUND_ENUM_ENTRY(Audio::EMusicalScale::Scale::LydianAugmented, "LydianAugmentedDescription", "Lydian Augmented", "LydianAugmentedDescriptionTT", "Lydian Augmented"),
+	DEFINE_METASOUND_ENUM_ENTRY(Audio::EMusicalScale::Scale::LydianDominant, "LydianDominantDescription", "Lydian Dominant ", "LydianDominantDescriptionTT", "Lydian Dominant"),
+	DEFINE_METASOUND_ENUM_ENTRY(Audio::EMusicalScale::Scale::Augmented, "AugmentedDescription", "Augmented", "AugmentedDescriptionTT", "Augmented"),
+	// diminished
+	DEFINE_METASOUND_ENUM_ENTRY(Audio::EMusicalScale::Scale::Diminished_BeginWithHalfStep, "Diminished_BeginWithHalfStepDescription", "Diminished (Begin With Half-Step)", "Diminished_BeginWithHalfStepDescriptionTT", "Diminished (begins with Half Step)"),
+	DEFINE_METASOUND_ENUM_ENTRY(Audio::EMusicalScale::Scale::Diminished_BeginWithWholeStep, "Diminished_BeginWithWholeStepDescription", "Diminished (Begin With Whole-Step", "Diminished_BeginWithWholeStepDescriptionTT", "Diminished (begins with Whole Step)"),
+	DEFINE_METASOUND_ENUM_ENTRY(Audio::EMusicalScale::Scale::HalfDiminished_LocrianNumber2, "HalfDiminished_LocrianNumber2Description", "Half-Diminished (Locrian #2)", "HalfDiminished_LocrianNumber2DescriptionTT", "Half Diminished Locrian (#2)"),
+	// other
+	DEFINE_METASOUND_ENUM_ENTRY(Audio::EMusicalScale::Scale::Spanish_or_Jewish, "Spanish_or_JewishDescription", "Spanish or Jewish Scale", "Spanish_or_JewishDescriptionTT", "Spanish/Jewish"),
+	DEFINE_METASOUND_ENUM_ENTRY(Audio::EMusicalScale::Scale::Hindu, "HinduDescription", "Hindu ", "HinduDescriptionTT", "Hindu")
+
+	DEFINE_METASOUND_ENUM_END()
 #pragma region Parameter Names
 	namespace MusicalScaleToNoteArrayParameterNames
 	{
diff --git a/Engine/Plugins/Runtime/Metasound/Source/MetasoundStandardNodes/Private/MetasoundTriggerCompareNode.cpp b/Engine/Plugins/Runtime/Metasound/Source/MetasoundStandardNodes/Private/MetasoundTriggerCompareNode.cpp
--- a/Engine/Plugins/Runtime/Metasound/Source/MetasoundStandardNodes/Private/MetasoundTriggerCompareNode.cpp
+++ b/Engine/Plugins/Runtime/Metasound/Source/MetasoundStandardNodes/Private/MetasoundTriggerCompareNode.cpp
@@ -30,6 +30,15 @@ namespace Metasound
 		}
 	}
 
+	DEFINE_METASOUND_ENUM_BEGIN(ETriggerComparisonType, FEnumTriggerComparisonType, "TriggerComparisonType")
+		DEFINE_METASOUND_ENUM_ENTRY(ETriggerComparisonType::Equals, "EqualsDescription", "Equals", "EqualsDescriptionTT", "True if A and B are equal."),
+		DEFINE_METASOUND_ENUM_ENTRY(ETriggerComparisonType::NotEquals, "NotEqualsDescriptioin", "Not Equals", "NotEqualsTT", "True if A and B are not equal."),
+		DEFINE_METASOUND_ENUM_ENTRY(ETriggerComparisonType::LessThan, "LessThanDescription", "Less Than", "LessThanTT", "True if A is less than B."),
+		DEFINE_METASOUND_ENUM_ENTRY(ETriggerComparisonType::GreaterThan, "GreaterThanDescription", "Greater Than", "GreaterThanTT", "True if A is greater than B."),
+		DEFINE_METASOUND_ENUM_ENTRY(ETriggerComparisonType::LessThanOrEquals, "LessThanOrEqualsDescription", "Less Than Or Equals", "LessThanOrEqualsTT", "True if A is less than or equal to B."),
+		DEFINE_METASOUND_ENUM_ENTRY(ETriggerComparisonType::GreaterThanOrEquals, "GreaterThanOrEqualsDescription", "Greater Than Or Equals", "GreaterThanOrEqualsTT", "True if A is greater than or equal to B."),
+	DEFINE_METASOUND_ENUM_END()
+
 	using FTriggerCompareNodeInt32 = TTriggerCompareNode<int32>;
 	METASOUND_REGISTER_NODE(FTriggerCompareNodeInt32)
 
diff --git a/Engine/Plugins/Runtime/Metasound/Source/MetasoundStandardNodes/Public/MetasoundScaleEnums.h b/Engine/Plugins/Runtime/Metasound/Source/MetasoundStandardNodes/Public/MetasoundScaleEnums.h
--- a/Engine/Plugins/Runtime/Metasound/Source/MetasoundStandardNodes/Public/MetasoundScaleEnums.h
+++ b/Engine/Plugins/Runtime/Metasound/Source/MetasoundStandardNodes/Public/MetasoundScaleEnums.h
@@ -11,55 +11,11 @@
 
 namespace Metasound
 {
-#define LOCTEXT_NAMESPACE "MidiScaleDefinitions"
 	// Any desired additions to this Enum/list need to first be added to the EMusicalScale enum in MidiNoteQuantizer.h
 	// and defined in MidiNoteQuantizer.cpp in the TMap<EMusicalScale::Scale, ScaleDegreeSet> ScaleDegreeSetMap static init
 
 	// Metasound enum
 	DECLARE_METASOUND_ENUM(Audio::EMusicalScale::Scale, Audio::EMusicalScale::Scale::Major,
 	METASOUNDSTANDARDNODES_API, FEnumEMusicalScale, FEnumMusicalScaleTypeInfo, FEnumMusicalScaleReadRef, FEnumMusicalScaleWriteRef);
-	DEFINE_METASOUND_ENUM_BEGIN(Audio::EMusicalScale::Scale, FEnumEMusicalScale, "MusicalScale")
 
-	// modes
-	DEFINE_METASOUND_ENUM_ENTRY(Audio::EMusicalScale::Scale::Major, "MajorDescription", "Major Scale", "MajorDescriptionTT", "Major (Ionian)"),
-	DEFINE_METASOUND_ENUM_ENTRY(Audio::EMusicalScale::Scale::Minor_Dorian, "Minor_DorianDescription", "Minor (Dorian)", "Minor_DorianDescriptionTT", "Dorian Minor"),
-	DEFINE_METASOUND_ENUM_ENTRY(Audio::EMusicalScale::Scale::Phrygian, "PhrygianDescription", "Phrygian ", "PhrygianDescriptionTT", "Phrygian"),
-	DEFINE_METASOUND_ENUM_ENTRY(Audio::EMusicalScale::Scale::Lydian, "LydianDescription", "Lydian", "LydianDescriptionTT", "Lydian (sharp-4)"),
-	DEFINE_METASOUND_ENUM_ENTRY(Audio::EMusicalScale::Scale::Dominant7th_Mixolydian, "Dominant7th_MixolydianDescription", "Dominant 7th (Mixolydian)", "Dominant7th_MixolydianDescriptionTT", "Mioxlydian (Dominant 7)"),
-	DEFINE_METASOUND_ENUM_ENTRY(Audio::EMusicalScale::Scale::NaturalMinor_Aeolian, "NaturalMinor_AeolianDescription", "Natural Minor (Aeolian)", "NaturalMinor_AeolianDescriptionTT", "Natural Minor (Aeolian)"),
-	DEFINE_METASOUND_ENUM_ENTRY(Audio::EMusicalScale::Scale::HalfDiminished_Locrian, "HalfDiminished_LocrianDescription", "Half Diminished (Locrian)", "HalfDiminished_LocrianDescriptionTT", "Half-Diminished (Locrian)"),
-	DEFINE_METASOUND_ENUM_ENTRY(Audio::EMusicalScale::Scale::Diminished, "DiminishedDescription", "Diminished ", "DiminishedDescriptionTT", "Diminished"),
-	// non-diatonic
-	DEFINE_METASOUND_ENUM_ENTRY(Audio::EMusicalScale::Scale::Chromatic, "ChromaticDescription", "Chromatic", "ChromaticDescriptionTT", "Chromatic"),
-	DEFINE_METASOUND_ENUM_ENTRY(Audio::EMusicalScale::Scale::WholeTone, "WholeToneDescription", "Whole-Tone", "WholeToneDescriptionTT", "Whole Tone"),
-	DEFINE_METASOUND_ENUM_ENTRY(Audio::EMusicalScale::Scale::DiminishedWholeTone, "DiminishedWholeToneDescription", "Diminished Whole-Tone", "DiminishedWholeToneDescriptionTT", "Diminished Whole Tone"),
-	// petantonic
-	DEFINE_METASOUND_ENUM_ENTRY(Audio::EMusicalScale::Scale::MajorPentatonic, "MajorPentatonicDescription", "Major Pentatonic ", "MajorPentatonicDescriptionTT", "Major Pentatonic"),
-	DEFINE_METASOUND_ENUM_ENTRY(Audio::EMusicalScale::Scale::MinorPentatonic, "MinorPentatonicDescription", "Minor Pentatonic ", "MinorPentatonicDescriptionTT", "Minor Pentatonic"),
-	DEFINE_METASOUND_ENUM_ENTRY(Audio::EMusicalScale::Scale::Blues, "BluesDescription", "Blues ", "BluesDescriptionTT", "Blues"),
-	// bebop
-	DEFINE_METASOUND_ENUM_ENTRY(Audio::EMusicalScale::Scale::Bebop_Major, "Bebop_MajorDescription", "Bebop (Major)", "Bebop_MajorDescriptionTT", "Bebop Major"),
-	DEFINE_METASOUND_ENUM_ENTRY(Audio::EMusicalScale::Scale::Bebop_Minor, "Bebop_MinorDescription", "Bebop (Minor)", "Bebop_MinorDescriptionTT", "Bebop Minor"),
-	DEFINE_METASOUND_ENUM_ENTRY(Audio::EMusicalScale::Scale::Bebop_MinorNumber2, "Bebop_MinorNumber2Description", "Bebop (Minor) #2", "Bebop_MinorNumber2DescriptionTT", "Bebop Minor #2"),
-	DEFINE_METASOUND_ENUM_ENTRY(Audio::EMusicalScale::Scale::Bebop_Dominant, "Bebop_DominantDescription", "Bebop (Dominant)", "Bebop_DominantDescriptionTT", "Bebop Dominant"),
-	// common major/minors
-	DEFINE_METASOUND_ENUM_ENTRY(Audio::EMusicalScale::Scale::HarmonicMajor, "HarmonicMajorDescription", "Harmonic Major", "HarmonicMajorDescriptionTT", "Harmonic Major"),
-	DEFINE_METASOUND_ENUM_ENTRY(Audio::EMusicalScale::Scale::HarmonicMinor, "HarmonicMinorDescription", "Harmonic Minor ", "HarmonicMinorDescriptionTT", "Harmonic Minor"),
-	DEFINE_METASOUND_ENUM_ENTRY(Audio::EMusicalScale::Scale::MelodicMinor, "MelodicMinorDescription", "Melodic Minor ", "MelodicMinorDescriptionTT", "Melodic Minor"),
-	DEFINE_METASOUND_ENUM_ENTRY(Audio::EMusicalScale::Scale::SixthModeOfHarmonicMinor, "SixthModeOfHarmonicMinorDescription", "Sixth Mode of Harmonic Minor", "SixthModeOfHarmonicMinorDescriptionTT", "Sixth Mode of Harmonic Minor"),
-	// lydian/augmented
-	DEFINE_METASOUND_ENUM_ENTRY(Audio::EMusicalScale::Scale::LydianAugmented, "LydianAugmentedDescription", "Lydian Augmented", "LydianAugmentedDescriptionTT", "Lydian Augmented"),
-	DEFINE_METASOUND_ENUM_ENTRY(Audio::EMusicalScale::Scale::LydianDominant, "LydianDominantDescription", "Lydian Dominant ", "LydianDominantDescriptionTT", "Lydian Dominant"),
-	DEFINE_METASOUND_ENUM_ENTRY(Audio::EMusicalScale::Scale::Augmented, "AugmentedDescription", "Augmented", "AugmentedDescriptionTT", "Augmented"),
-	// diminished
-	DEFINE_METASOUND_ENUM_ENTRY(Audio::EMusicalScale::Scale::Diminished_BeginWithHalfStep, "Diminished_BeginWithHalfStepDescription", "Diminished (Begin With Half-Step)", "Diminished_BeginWithHalfStepDescriptionTT", "Diminished (begins with Half Step)"),
-	DEFINE_METASOUND_ENUM_ENTRY(Audio::EMusicalScale::Scale::Diminished_BeginWithWholeStep, "Diminished_BeginWithWholeStepDescription", "Diminished (Begin With Whole-Step", "Diminished_BeginWithWholeStepDescriptionTT", "Diminished (begins with Whole Step)"),
-	DEFINE_METASOUND_ENUM_ENTRY(Audio::EMusicalScale::Scale::HalfDiminished_LocrianNumber2, "HalfDiminished_LocrianNumber2Description", "Half-Diminished (Locrian #2)", "HalfDiminished_LocrianNumber2DescriptionTT", "Half Diminished Locrian (#2)"),
-	// other
-	DEFINE_METASOUND_ENUM_ENTRY(Audio::EMusicalScale::Scale::Spanish_or_Jewish, "Spanish_or_JewishDescription", "Spanish or Jewish Scale", "Spanish_or_JewishDescriptionTT", "Spanish/Jewish"),
-	DEFINE_METASOUND_ENUM_ENTRY(Audio::EMusicalScale::Scale::Hindu, "HinduDescription", "Hindu ", "HinduDescriptionTT", "Hindu")
-
-	DEFINE_METASOUND_ENUM_END()
-
-#undef LOCTEXT_NAMESPACE
 } // namespace Metasound
diff --git a/Engine/Plugins/Runtime/Metasound/Source/MetasoundStandardNodes/Public/MetasoundTriggerCompareNode.h b/Engine/Plugins/Runtime/Metasound/Source/MetasoundStandardNodes/Public/MetasoundTriggerCompareNode.h
--- a/Engine/Plugins/Runtime/Metasound/Source/MetasoundStandardNodes/Public/MetasoundTriggerCompareNode.h
+++ b/Engine/Plugins/Runtime/Metasound/Source/MetasoundStandardNodes/Public/MetasoundTriggerCompareNode.h
@@ -50,16 +50,6 @@ namespace Metasound
 	DECLARE_METASOUND_ENUM(ETriggerComparisonType, ETriggerComparisonType::Equals, METASOUNDSTANDARDNODES_API,
 	FEnumTriggerComparisonType, FEnumTriggerComparisonTypeInfo, FEnumTriggerComparisonTypeReadRef, FEnumTriggerComparisonTypeWriteRef);
 
-	DEFINE_METASOUND_ENUM_BEGIN(ETriggerComparisonType, FEnumTriggerComparisonType, "TriggerComparisonType")
-		DEFINE_METASOUND_ENUM_ENTRY(ETriggerComparisonType::Equals, "EqualsDescription", "Equals", "EqualsDescriptionTT", "True if A and B are equal."),
-		DEFINE_METASOUND_ENUM_ENTRY(ETriggerComparisonType::NotEquals, "NotEqualsDescriptioin", "Not Equals", "NotEqualsTT", "True if A and B are not equal."),
-		DEFINE_METASOUND_ENUM_ENTRY(ETriggerComparisonType::LessThan, "LessThanDescription", "Less Than", "LessThanTT", "True if A is less than B."),
-		DEFINE_METASOUND_ENUM_ENTRY(ETriggerComparisonType::GreaterThan, "GreaterThanDescription", "Greater Than", "GreaterThanTT", "True if A is greater than B."),
-		DEFINE_METASOUND_ENUM_ENTRY(ETriggerComparisonType::LessThanOrEquals, "LessThanOrEqualsDescription", "Less Than Or Equals", "LessThanOrEqualsTT", "True if A is less than or equal to B."),
-		DEFINE_METASOUND_ENUM_ENTRY(ETriggerComparisonType::GreaterThanOrEquals, "GreaterThanOrEqualsDescription", "Greater Than Or Equals", "GreaterThanOrEqualsTT", "True if A is greater than or equal to B."),
-	DEFINE_METASOUND_ENUM_END()
-
-
 	template<typename ValueType>
 	class TTriggerCompareNodeOperator : public TExecutableOperator<TTriggerCompareNodeOperator<ValueType>>
 	{
