Fix metasound mid side node decode init condition
#rnx

diff --git a/Engine/Plugins/Runtime/Metasound/Source/MetasoundStandardNodes/Private/MetasoundMidSideNodes.cpp b/Engine/Plugins/Runtime/Metasound/Source/MetasoundStandardNodes/Private/MetasoundMidSideNodes.cpp
--- a/Engine/Plugins/Runtime/Metasound/Source/MetasoundStandardNodes/Private/MetasoundMidSideNodes.cpp
+++ b/Engine/Plugins/Runtime/Metasound/Source/MetasoundStandardNodes/Private/MetasoundMidSideNodes.cpp
@@ -62,6 +62,8 @@ namespace Metasound
 
 				FNodeClassMetadata Metadata
 				{
+					// This node needs to maintain this node-class-name in order to avoid breaking linkage. 
+					// Note that the class name is "Mid-Side Decode" when this is in fact the "Encoder". 
 					FNodeClassName { StandardNodes::Namespace, "Mid-Side Decode", StandardNodes::AudioVariant },
 					1, // Major Version
 					0, // Minor Version
@@ -272,7 +274,7 @@ namespace Metasound
 	{
 	public:
 
-		FMidSideDecodeOperator(const FOperatorSettings& InSettings,
+		FMidSideDecodeOperator(const FCreateOperatorParams& InParams,
 			const FAudioBufferReadRef& InLeftAudioInput,
 			const FAudioBufferReadRef& InRightAudioInput,
 			const FFloatReadRef& InSpreadAmount,
@@ -281,19 +283,21 @@ namespace Metasound
 			, AudioInputSide(InRightAudioInput)
 			, SpreadAmount(InSpreadAmount)
 			, bEqualPower(InEqualPower)
-			, AudioOutputLeft(FAudioBufferWriteRef::CreateNew(InSettings))
-			, AudioOutputRight(FAudioBufferWriteRef::CreateNew(InSettings))
+			, AudioOutputLeft(FAudioBufferWriteRef::CreateNew(InParams.OperatorSettings))
+			, AudioOutputRight(FAudioBufferWriteRef::CreateNew(InParams.OperatorSettings))
 			, MidScale(0.0f)
 			, SideScale(0.0f)
 			, PrevMidScale(0.0f)
 			, PrevSideScale(0.0f)
 			, PrevSpreadAmount(*SpreadAmount)
 		{
-			MidInBuffer.AddUninitialized(InSettings.GetNumFramesPerBlock());
+			MidInBuffer.AddUninitialized(InParams.OperatorSettings.GetNumFramesPerBlock());
 			FMemory::Memset(MidInBuffer.GetData(), 0, sizeof(float) * MidInBuffer.Num());
 
-			SideInBuffer.AddUninitialized(InSettings.GetNumFramesPerBlock());
+			SideInBuffer.AddUninitialized(InParams.OperatorSettings.GetNumFramesPerBlock());
 			FMemory::Memset(SideInBuffer.GetData(), 0, sizeof(float) * SideInBuffer.Num());
+
+			Reset(InParams);
 		}
 
 		static const FNodeClassMetadata& GetNodeInfo()
@@ -304,6 +308,8 @@ namespace Metasound
 
 				FNodeClassMetadata Metadata
 				{
+					// This node needs to maintain this node-class-name in order to avoid breaking linkage. 
+					// Note that the class name is "Mid-Side Encode" when this is in fact the "Decoder". 
 					FNodeClassName { StandardNodes::Namespace, "Mid-Side Encode", StandardNodes::AudioVariant },
 					1, // Major Version
 					0, // Minor Version
@@ -382,7 +388,7 @@ namespace Metasound
 			FFloatReadRef SpreadAmountIn = Inputs.GetDataReadReferenceOrConstructWithVertexDefault<float>(InputInterface, METASOUND_GET_PARAM_NAME(InputSpreadAmount), InParams.OperatorSettings);
 			FBoolReadRef bEqualPowerIn = Inputs.GetDataReadReferenceOrConstructWithVertexDefault<bool>(InputInterface, METASOUND_GET_PARAM_NAME(InputEqualPower), InParams.OperatorSettings);
 
-			return MakeUnique<FMidSideDecodeOperator>(InParams.OperatorSettings, MidAudioIn, SideAudioIn, SpreadAmountIn, bEqualPowerIn);
+			return MakeUnique<FMidSideDecodeOperator>(InParams, MidAudioIn, SideAudioIn, SpreadAmountIn, bEqualPowerIn);
 		}
 
 		void Execute()
