Fix build break in MetaSound editor

diff --git a/Engine/Plugins/Runtime/Metasound/Source/MetasoundGenerator/Public/MetasoundGenerator.h b/Engine/Plugins/Runtime/Metasound/Source/MetasoundGenerator/Public/MetasoundGenerator.h
--- a/Engine/Plugins/Runtime/Metasound/Source/MetasoundGenerator/Public/MetasoundGenerator.h
+++ b/Engine/Plugins/Runtime/Metasound/Source/MetasoundGenerator/Public/MetasoundGenerator.h
@@ -16,7 +16,7 @@
 #include "Sound/SoundGenerator.h"
 
 #ifndef ENABLE_METASOUND_GENERATOR_RENDER_TIMING
-#define ENABLE_METASOUND_GENERATOR_RENDER_TIMING !UE_BUILD_SHIPPING && WITH_EDITOR
+#define ENABLE_METASOUND_GENERATOR_RENDER_TIMING WITH_EDITOR
 #endif // ifndef ENABLE_METASOUND_GENERATOR_RENDER_TIMING
 
 namespace Metasound
