Force the metasound builder audition to create a new sound each time

diff --git a/Engine/Plugins/Runtime/Metasound/Source/MetasoundEngine/Private/MetasoundBuilderSubsystem.cpp b/Engine/Plugins/Runtime/Metasound/Source/MetasoundEngine/Private/MetasoundBuilderSubsystem.cpp
--- a/Engine/Plugins/Runtime/Metasound/Source/MetasoundEngine/Private/MetasoundBuilderSubsystem.cpp
+++ b/Engine/Plugins/Runtime/Metasound/Source/MetasoundEngine/Private/MetasoundBuilderSubsystem.cpp
@@ -622,7 +622,9 @@ void UMetaSoundSourceBuilder::Audition(UObject* Parent, UAudioComponent* AudioCo
 		return;
 	}
 
-	if (AuditionSound.IsValid())
+	// kevin.cavanaugh [3/13/2023] - Rebuilding the audition sound currently isn't working. 
+	// Forcing a rebuild each time for now.
+	/*if (AuditionSound.IsValid())
 	{
 		const FMetasoundFrontendDocument& Doc = GetConstBuilder().GetDocument();
 		AuditionSound->SetDocument(Doc);
@@ -643,7 +645,7 @@ void UMetaSoundSourceBuilder::Audition(UObject* Parent, UAudioComponent* AudioCo
 			AudioComponent->SetSound(AuditionSound.Get());
 		}
 	}
-	else
+	else*/
 	{
 		FMetaSoundBuilderOptions BuildOptions;
 		BuildOptions.Name = MakeUniqueObjectName(nullptr, UMetaSoundSource::StaticClass(), FName(GetName() + TEXT("_Audition")));
