[Metasound] Watch outputs on a FMetasoundGenerator

diff --git a/Engine/Plugins/Runtime/Metasound/Source/MetasoundEngine/Private/MetasoundEngineModule.cpp b/Engine/Plugins/Runtime/Metasound/Source/MetasoundEngine/Private/MetasoundEngineModule.cpp
--- a/Engine/Plugins/Runtime/Metasound/Source/MetasoundEngine/Private/MetasoundEngineModule.cpp
+++ b/Engine/Plugins/Runtime/Metasound/Source/MetasoundEngine/Private/MetasoundEngineModule.cpp
@@ -1,20 +1,20 @@
 // Copyright Epic Games, Inc. All Rights Reserved.
 #include "MetasoundEngineModule.h"
 
-#include "Analysis/MetasoundFrontendAnalyzerRegistry.h"
-#include "Interfaces/MetasoundDeprecatedInterfaces.h"
-#include "Interfaces/MetasoundInterface.h"
-#include "Interfaces/MetasoundInterfaceBindingsPrivate.h"
 #include "Metasound.h"
+#include "MetasoundAudioBus.h"
 #include "MetasoundDataReference.h"
 #include "MetasoundDataTypeRegistrationMacro.h"
-#include "MetasoundFrontendDataTypeTraits.h"
+#include "MetasoundGeneratorOutputWatcher.h"
 #include "MetasoundSource.h"
 #include "MetasoundTrace.h"
 #include "MetasoundUObjectRegistry.h"
 #include "MetasoundWave.h"
 #include "MetasoundWaveTable.h"
-#include "MetasoundAudioBus.h"
+#include "Analysis/MetasoundFrontendAnalyzerRegistry.h"
+#include "Interfaces/MetasoundDeprecatedInterfaces.h"
+#include "Interfaces/MetasoundInterface.h"
+#include "Interfaces/MetasoundInterfaceBindingsPrivate.h"
 #include "Modules/ModuleManager.h"
 #include "Sound/AudioSettings.h"
 
@@ -61,6 +61,12 @@ class FMetasoundEngineModule : public IMetasoundEngineModule
 		// Register Analyzers
 		Metasound::Frontend::IVertexAnalyzerRegistry::Get().RegisterAnalyzerFactories();
 
+		// Register output watcher ops
+		METASOUND_PRIVATE_REGISTER_GENERATOR_OUTPUT_WATCHER_TYPE_OPERATIONS(float, 0.0f)
+		METASOUND_PRIVATE_REGISTER_GENERATOR_OUTPUT_WATCHER_TYPE_OPERATIONS(int32, 0)
+		METASOUND_PRIVATE_REGISTER_GENERATOR_OUTPUT_WATCHER_TYPE_OPERATIONS(bool, false)
+		METASOUND_PRIVATE_REGISTER_GENERATOR_OUTPUT_WATCHER_TYPE_OPERATIONS(FString, "")
+
 		UE_LOG(LogMetasoundEngine, Log, TEXT("MetaSound Engine Initialized"));
 	}
 };
diff --git a/Engine/Plugins/Runtime/Metasound/Source/MetasoundEngine/Private/MetasoundGeneratorOutput.cpp b/Engine/Plugins/Runtime/Metasound/Source/MetasoundEngine/Private/MetasoundGeneratorOutput.cpp
new file mode 100644
--- /dev/null
+++ b/Engine/Plugins/Runtime/Metasound/Source/MetasoundEngine/Private/MetasoundGeneratorOutput.cpp
@@ -0,0 +1,80 @@
+// Copyright Epic Games, Inc. All Rights Reserved.
+
+#include "MetasoundGeneratorOutput.h"
+
+#include "MetasoundPrimitives.h"
+
+FMetasoundGeneratorOutput::FMetasoundGeneratorOutput(const FMetasoundGeneratorOutput& Other)
+{
+	Name = Other.Name;
+	DataReference = MakeUnique<Metasound::FAnyDataReference>(*Other.DataReference);
+}
+
+FMetasoundGeneratorOutput& FMetasoundGeneratorOutput::operator=(const FMetasoundGeneratorOutput& Other)
+{
+	if (this != &Other)
+	{
+		this->Name = Other.Name;
+		this->DataReference = MakeUnique<Metasound::FAnyDataReference>(*Other.DataReference);
+	}
+
+	return *this;
+}
+
+bool FMetasoundGeneratorOutput::IsValid() const
+{
+	return nullptr != DataReference;
+}
+
+FName FMetasoundGeneratorOutput::GetTypeName() const
+{
+	return IsValid() ? DataReference->GetDataTypeName() : FName();
+}
+
+bool UMetasoundGeneratorOutputBlueprintAccess::IsFloat(const FMetasoundGeneratorOutput& Output)
+{
+	return Output.IsType<float>();
+}
+
+float UMetasoundGeneratorOutputBlueprintAccess::GetFloat(const FMetasoundGeneratorOutput& Output, bool& Success)
+{
+	float Value = 0;
+	Success = Output.Get(Value);
+	return Value;
+}
+
+bool UMetasoundGeneratorOutputBlueprintAccess::IsInt32(const FMetasoundGeneratorOutput& Output)
+{
+	return Output.IsType<int32>();
+}
+
+int32 UMetasoundGeneratorOutputBlueprintAccess::GetInt32(const FMetasoundGeneratorOutput& Output, bool& Success)
+{
+	int32 Value = 0;
+	Success = Output.Get(Value);
+	return Value;
+}
+
+bool UMetasoundGeneratorOutputBlueprintAccess::IsBool(const FMetasoundGeneratorOutput& Output)
+{
+	return Output.IsType<bool>();
+}
+
+bool UMetasoundGeneratorOutputBlueprintAccess::GetBool(const FMetasoundGeneratorOutput& Output, bool& Success)
+{
+	bool Value = false;
+	Success = Output.Get(Value);
+	return Value;
+}
+
+bool UMetasoundGeneratorOutputBlueprintAccess::IsString(const FMetasoundGeneratorOutput& Output)
+{
+	return Output.IsType<FString>();
+}
+
+FString UMetasoundGeneratorOutputBlueprintAccess::GetString(const FMetasoundGeneratorOutput& Output, bool& Success)
+{
+	FString Value;
+	Success = Output.Get(Value);
+	return Value;
+}
diff --git a/Engine/Plugins/Runtime/Metasound/Source/MetasoundEngine/Private/MetasoundGeneratorOutputWatcher.cpp b/Engine/Plugins/Runtime/Metasound/Source/MetasoundEngine/Private/MetasoundGeneratorOutputWatcher.cpp
new file mode 100644
--- /dev/null
+++ b/Engine/Plugins/Runtime/Metasound/Source/MetasoundEngine/Private/MetasoundGeneratorOutputWatcher.cpp
@@ -0,0 +1,56 @@
+// Copyright Epic Games, Inc. All Rights Reserved.
+
+#include "MetasoundGeneratorOutputWatcher.h"
+
+namespace Metasound::Private
+{
+	TMap<FName, TUniquePtr<FMetasoundGeneratorOutputWatcher::IOutputTypeOperations>> FMetasoundGeneratorOutputWatcher::OutputTypeOperationMap{};
+	
+	FMetasoundGeneratorOutputWatcher::FMetasoundGeneratorOutputWatcher(
+		Frontend::FAnalyzerAddress&& Address,
+		const FOperatorSettings& OperatorSettings)
+		: Name(Address.OutputName)
+		, View(MoveTemp(Address))
+	{
+		View.BindToAllOutputs(OperatorSettings);
+		
+		TArray<Frontend::FMetasoundAnalyzerView::FBoundOutputDescription> OutputDescriptions = View.GetBoundOutputDescriptions();
+		for (Frontend::FMetasoundAnalyzerView::FBoundOutputDescription& OutputDescription : OutputDescriptions)
+		{
+			if (!ensure(OutputTypeOperationMap.Contains(OutputDescription.TypeName)))
+			{
+				continue;
+			}
+
+			FMetasoundGeneratorOutput Output;
+			Output.Name = OutputDescription.Name;
+			OutputTypeOperationMap[OutputDescription.TypeName]->Init(Output);
+			Outputs.Emplace(MoveTemp(Output));
+		}
+	}
+
+	void FMetasoundGeneratorOutputWatcher::Update(const TFunctionRef<void(FName, const FMetasoundGeneratorOutput&)> OnOutputChanged)
+	{
+		for (FMetasoundGeneratorOutput& Output : Outputs)
+		{
+			const FName TypeName = Output.GetTypeName();
+			check(!TypeName.IsNone());
+			
+			if (ensure(OutputTypeOperationMap.Contains(TypeName)))
+			{
+				if (OutputTypeOperationMap[TypeName]->Update(*this, Output))
+				{
+					OnOutputChanged(Name, Output);
+				}
+			}
+		}
+	}
+
+	void FMetasoundGeneratorOutputWatcher::RegisterOutputTypeOperations(FName TypeName,
+		TUniquePtr<IOutputTypeOperations>&& OutputTypeOperations)
+	{
+		check(!OutputTypeOperationMap.Contains(TypeName));
+		check(OutputTypeOperations.IsValid());
+		OutputTypeOperationMap.Emplace(TypeName, MoveTemp(OutputTypeOperations));
+	}
+}
diff --git a/Engine/Plugins/Runtime/Metasound/Source/MetasoundEngine/Public/MetasoundGeneratorOutput.h b/Engine/Plugins/Runtime/Metasound/Source/MetasoundEngine/Public/MetasoundGeneratorOutput.h
new file mode 100644
--- /dev/null
+++ b/Engine/Plugins/Runtime/Metasound/Source/MetasoundEngine/Public/MetasoundGeneratorOutput.h
@@ -0,0 +1,131 @@
+// Copyright Epic Games, Inc. All Rights Reserved.
+
+#pragma once
+
+#include "MetasoundDataReference.h"
+#include "SoundGeneratorOutput.h"
+#include "Kismet/BlueprintFunctionLibrary.h"
+
+#include "MetasoundGeneratorOutput.generated.h"
+
+/**
+ * Exposes the value of an output on a FMetasoundGenerator.
+ */
+USTRUCT(BlueprintType)
+struct METASOUNDENGINE_API FMetasoundGeneratorOutput : public FSoundGeneratorOutput
+{
+	GENERATED_BODY()
+
+	FMetasoundGeneratorOutput() = default;
+
+	FMetasoundGeneratorOutput(const FMetasoundGeneratorOutput& Other);
+
+	FMetasoundGeneratorOutput& operator=(const FMetasoundGeneratorOutput& Other);
+
+	/**
+	 * Has this output been initialized?
+	 *
+	 * @returns true if the output has been initialized, false otherwise
+	 */
+	bool IsValid() const;
+
+	/**
+	 * Get the type name of the output
+	 *
+	 * @returns The type name, or none if invalid
+	 */
+	FName GetTypeName() const;
+	
+	/**
+	 * Initialize the output with some data.
+	 *
+	 * @tparam DataType - The type to use for initialization
+	 * @param InitialValue - The initial value to use
+	 */
+	template<typename DataType>
+	void Init(const DataType& InitialValue)
+	{
+		DataReference = MakeUnique<Metasound::FAnyDataReference>(Metasound::TDataWriteReference<DataType>::CreateNew());
+		*DataReference->GetDataWriteReference<DataType>() = InitialValue;
+	}
+
+	/**
+	 * Check if this output is of the given type
+	 *
+	 * @tparam DataType - The data type to check
+	 * @returns true if it's the given type, false otherwise
+	 */
+	template<typename DataType>
+	bool IsType() const
+	{
+		return IsValid() && Metasound::IsDataReferenceOfType<DataType>(*DataReference);
+	}
+	
+	/**
+	 * Get the value, for copyable registered Metasound data types
+	 *
+	 * @tparam DataType - The expected data type of the output
+	 * @param Value - The value to use
+	 * @returns true if the value was retrieved, false otherwise
+	 */
+	template<typename DataType>
+	bool Get(DataType& Value) const
+	{
+		if (!IsType<DataType>())
+		{
+			return false;
+		}
+
+		Value = *DataReference->GetDataReadReference<DataType>();
+		return true;
+	}
+	
+	/**
+	 * Set the value, for copyable registered Metasound data types
+	 *
+	 * @tparam DataType - The expected data type of the output
+	 * @param Value - The value to use
+	 * @returns true if the value was set, false otherwise
+	 */
+	template<typename DataType>
+	bool Set(const DataType& Value)
+	{
+		if (!IsType<DataType>())
+		{
+			return false;
+		}
+
+		*DataReference->GetDataWriteReference<DataType>() = Value;
+		return true;
+	}
+
+private:
+	TUniquePtr<Metasound::FAnyDataReference> DataReference;
+};
+
+/**
+ * Blueprint support for core types. If you want to support more core types, add them here.
+ * If you want to support types introduced in other plugins, create a blueprint library in that plugin.
+ */
+UCLASS()
+class UMetasoundGeneratorOutputBlueprintAccess final : public UBlueprintFunctionLibrary
+{
+	GENERATED_BODY()
+
+	UFUNCTION(BlueprintCallable, Category="MetasoundGeneratorOutput")
+	static bool IsFloat(const FMetasoundGeneratorOutput& Output);
+	UFUNCTION(BlueprintCallable, Category="MetasoundGeneratorOutput")
+	static float GetFloat(const FMetasoundGeneratorOutput& Output, bool& Success);
+	UFUNCTION(BlueprintCallable, Category="MetasoundGeneratorOutput")
+	static bool IsInt32(const FMetasoundGeneratorOutput& Output);
+	UFUNCTION(BlueprintCallable, Category="MetasoundGeneratorOutput")
+	static int32 GetInt32(const FMetasoundGeneratorOutput& Output, bool& Success);
+	UFUNCTION(BlueprintCallable, Category="MetasoundGeneratorOutput")
+	static bool IsBool(const FMetasoundGeneratorOutput& Output);
+	UFUNCTION(BlueprintCallable, Category="MetasoundGeneratorOutput")
+	static bool GetBool(const FMetasoundGeneratorOutput& Output, bool& Success);
+	UFUNCTION(BlueprintCallable, Category="MetasoundGeneratorOutput")
+	static bool IsString(const FMetasoundGeneratorOutput& Output);
+	UFUNCTION(BlueprintCallable, Category="MetasoundGeneratorOutput")
+	static FString GetString(const FMetasoundGeneratorOutput& Output, bool& Success);
+};
diff --git a/Engine/Plugins/Runtime/Metasound/Source/MetasoundEngine/Public/MetasoundGeneratorOutputWatcher.h b/Engine/Plugins/Runtime/Metasound/Source/MetasoundEngine/Public/MetasoundGeneratorOutputWatcher.h
new file mode 100644
--- /dev/null
+++ b/Engine/Plugins/Runtime/Metasound/Source/MetasoundEngine/Public/MetasoundGeneratorOutputWatcher.h
@@ -0,0 +1,114 @@
+// Copyright Epic Games, Inc. All Rights Reserved.
+
+#pragma once
+
+#include "MetasoundGeneratorOutput.h"
+#include "Analysis/MetasoundFrontendAnalyzerView.h"
+#include "Containers/Array.h"
+#include "Containers/Map.h"
+#include "Templates/Function.h"
+
+namespace Metasound::Private
+{
+	/**
+	 * Provides a mechanism to check an output and notify subscribers when its value(s) change
+	 */
+	class METASOUNDENGINE_API FMetasoundGeneratorOutputWatcher
+	{
+	public:
+		FMetasoundGeneratorOutputWatcher(Frontend::FAnalyzerAddress&& Address, const FOperatorSettings& OperatorSettings);
+
+		const FName Name;
+		
+		/**
+		 * Update the watched outputs
+		 */
+		void Update(TFunctionRef<void(FName, const FMetasoundGeneratorOutput&)> OnOutputChanged);
+
+		using FOutputInitFn = TFunction<void(FMetasoundGeneratorOutput&)>;
+		using FOutputUpdateFn = TFunction<bool(FMetasoundGeneratorOutputWatcher&, FMetasoundGeneratorOutput&)>;
+
+		/**
+		 * Stores the typed operations for a specific supported type
+		 */
+		class IOutputTypeOperations
+		{
+		public:
+			virtual ~IOutputTypeOperations() = default;
+			
+			virtual void Init(FMetasoundGeneratorOutput& Output) = 0;
+			virtual bool Update(FMetasoundGeneratorOutputWatcher& Watcher, FMetasoundGeneratorOutput& Output) = 0;
+		};
+
+		template<typename DataType>
+		class TOutputTypeOperations final : public IOutputTypeOperations
+		{
+		public:
+			explicit TOutputTypeOperations(DataType InDefaultValue)
+				: DefaultValue(InDefaultValue)
+			{}
+			
+			const DataType DefaultValue;
+			
+			virtual void Init(FMetasoundGeneratorOutput& Output) override
+			{
+				Output.Init<DataType>(DefaultValue);
+			}
+			
+			virtual bool Update(FMetasoundGeneratorOutputWatcher& Watcher, FMetasoundGeneratorOutput& Output) override
+			{
+				return Watcher.UpdateOutput<DataType>(Output);
+			}
+
+		private:
+		};
+
+		/**
+		 * Register operations for a given type.
+		 * NOTE: This is exposed to enable external registration of types.
+		 * You shouldn't need to use this directly.
+		 *
+		 * @param TypeName - The name of the type, usually returned from GetMetasoundDataTypeName()
+		 * @param OutputTypeOperations - The operations to use
+		 */
+		static void RegisterOutputTypeOperations(FName TypeName, TUniquePtr<IOutputTypeOperations>&& OutputTypeOperations);
+	
+	private:
+		template<typename DataType>
+		bool UpdateOutput(FMetasoundGeneratorOutput& Output)
+		{
+			// check the type
+			if (!Output.IsType<DataType>())
+			{
+				return false;
+			}
+
+			// try to get the latest value
+			DataType Value;
+
+			if (View.TryGetOutputData(Output.Name, Value))
+			{
+				Output.Set(Value);
+				return true;
+			}
+			
+			return false;
+		}
+		
+		static TMap<FName, TUniquePtr<IOutputTypeOperations>> OutputTypeOperationMap;
+		Frontend::FMetasoundAnalyzerView View;
+		TArray<FMetasoundGeneratorOutput> Outputs;
+	};
+}
+
+
+/** Helper to register typed operations on outputs */
+#define METASOUND_PRIVATE_REGISTER_GENERATOR_OUTPUT_WATCHER_TYPE_OPERATIONS(TYPE, DEFAULT_VALUE) \
+	{ \
+		using FTypeOps = Metasound::Private::FMetasoundGeneratorOutputWatcher::TOutputTypeOperations<TYPE>; \
+		TUniquePtr<FTypeOps> Ops = MakeUnique<FTypeOps>(DEFAULT_VALUE); \
+		Metasound::Private::FMetasoundGeneratorOutputWatcher::RegisterOutputTypeOperations( \
+			Metasound::GetMetasoundDataTypeName<TYPE>(), \
+			MoveTemp(Ops) \
+			); \
+	} \
diff --git a/Engine/Plugins/Runtime/Metasound/Source/MetasoundEngineTest/Private/MetasoundGeneratorOutputTests.cpp b/Engine/Plugins/Runtime/Metasound/Source/MetasoundEngineTest/Private/MetasoundGeneratorOutputTests.cpp
new file mode 100644
--- /dev/null
+++ b/Engine/Plugins/Runtime/Metasound/Source/MetasoundEngineTest/Private/MetasoundGeneratorOutputTests.cpp
@@ -0,0 +1,121 @@
+// Copyright Epic Games, Inc. All Rights Reserved.
+
+#include "MetasoundGeneratorOutput.h"
+#include "MetasoundPrimitives.h"
+#include "Misc/AutomationTest.h"
+
+#if WITH_DEV_AUTOMATION_TESTS
+
+namespace Metasound::Test::GeneratorOutput
+{
+	IMPLEMENT_SIMPLE_AUTOMATION_TEST(
+	FMetasoundGeneratorOutputIsTypeTest,
+	"Audio.Metasound.GeneratorOutput.IsType",
+	EAutomationTestFlags::EditorContext | EAutomationTestFlags::EngineFilter)
+	bool FMetasoundGeneratorOutputIsTypeTest::RunTest(const FString&)
+	{
+		FMetasoundGeneratorOutput Output;
+		
+		Output.Init(float{ 0.0f });
+		UTEST_TRUE("float is float", Output.IsType<float>());
+		UTEST_FALSE("float is not int32", Output.IsType<int32>());
+		UTEST_FALSE("float is not bool", Output.IsType<bool>());
+		UTEST_FALSE("float is not FString", Output.IsType<FString>());
+
+		Output.Init(int32{ 0 });
+		UTEST_FALSE("int32 is not float", Output.IsType<float>());
+		UTEST_TRUE("int32 is int32", Output.IsType<int32>());
+		UTEST_FALSE("int32 is not bool", Output.IsType<bool>());
+		UTEST_FALSE("int32 is not FString", Output.IsType<FString>());
+		
+		Output.Init(bool{ false });
+		UTEST_FALSE("bool is not float", Output.IsType<float>());
+		UTEST_FALSE("bool is not int32", Output.IsType<int32>());
+		UTEST_TRUE("bool is bool", Output.IsType<bool>());
+		UTEST_FALSE("bool is not FString", Output.IsType<FString>());
+		
+		Output.Init(FString{ "hi!" });
+		UTEST_FALSE("FString is not float", Output.IsType<float>());
+		UTEST_FALSE("FString is not int32", Output.IsType<int32>());
+		UTEST_FALSE("FString is not bool", Output.IsType<bool>());
+		UTEST_TRUE("FString is FString", Output.IsType<FString>());
+		
+		return true;
+	}
+
+	template<typename DataType>
+	bool RunOutputGetSetTest(FAutomationTestBase& Test, const DataType& InitialValue, const DataType& ExpectedValue)
+	{
+		FMetasoundGeneratorOutput Output;
+		Output.Init<DataType>(InitialValue);
+		DataType Value;
+		
+		if (!Test.TestTrue("Got value", Output.Get<DataType>(Value)))
+		{
+			return false;
+		}
+		if (!Test.TestEqual("Value equals initial", Value, InitialValue))
+		{
+			return false;
+		}
+		if (!Test.TestTrue("Set value", Output.Set<DataType>(ExpectedValue)))
+		{
+			return false;
+		}
+		if (!Test.TestTrue("Got value", Output.Get<DataType>(Value)))
+		{
+			return false;
+		}
+		return Test.TestEqual("Value equals expected", Value, ExpectedValue);
+	}
+
+	IMPLEMENT_SIMPLE_AUTOMATION_TEST(
+	FMetasoundGeneratorOutputGetSetFloatTest,
+	"Audio.Metasound.GeneratorOutput.GetSet.Float",
+	EAutomationTestFlags::EditorContext | EAutomationTestFlags::EngineFilter)
+	bool FMetasoundGeneratorOutputGetSetFloatTest::RunTest(const FString&)
+	{
+		using DataType = float;
+		constexpr DataType InitialValue = 321.654f;
+		constexpr DataType ExpectedValue = 123.456f;
+		return RunOutputGetSetTest<DataType>(*this, InitialValue, ExpectedValue);
+	}
+
+	IMPLEMENT_SIMPLE_AUTOMATION_TEST(
+	FMetasoundGeneratorOutputGetSetInt32Test,
+	"Audio.Metasound.GeneratorOutput.GetSet.Int32",
+	EAutomationTestFlags::EditorContext | EAutomationTestFlags::EngineFilter)
+	bool FMetasoundGeneratorOutputGetSetInt32Test::RunTest(const FString&)
+	{
+		using DataType = int32;
+		constexpr DataType InitialValue = 321654;
+		constexpr DataType ExpectedValue = 123456;
+		return RunOutputGetSetTest<DataType>(*this, InitialValue, ExpectedValue);
+	}
+
+	IMPLEMENT_SIMPLE_AUTOMATION_TEST(
+	FMetasoundGeneratorOutputGetSetBoolTest,
+	"Audio.Metasound.GeneratorOutput.GetSet.Bool",
+	EAutomationTestFlags::EditorContext | EAutomationTestFlags::EngineFilter)
+	bool FMetasoundGeneratorOutputGetSetBoolTest::RunTest(const FString&)
+	{
+		using DataType = bool;
+		constexpr DataType InitialValue = false;
+		constexpr DataType ExpectedValue = true;
+		return RunOutputGetSetTest<DataType>(*this, InitialValue, ExpectedValue);
+	}
+
+	IMPLEMENT_SIMPLE_AUTOMATION_TEST(
+	FMetasoundGeneratorOutputGetSetStringTest,
+	"Audio.Metasound.GeneratorOutput.GetSet.String",
+	EAutomationTestFlags::EditorContext | EAutomationTestFlags::EngineFilter)
+	bool FMetasoundGeneratorOutputGetSetStringTest::RunTest(const FString&)
+	{
+		using FDataType = FString;
+		const FDataType InitialValue{ "hello" };
+		const FDataType ExpectedValue{ "goodbye" };
+		return RunOutputGetSetTest<FDataType>(*this, InitialValue, ExpectedValue);
+	}
+}
+
+#endif
\ No newline at end of file
diff --git a/Engine/Plugins/Runtime/Metasound/Source/MetasoundEngineTest/Private/MetasoundGeneratorOutputWatcherTests.cpp b/Engine/Plugins/Runtime/Metasound/Source/MetasoundEngineTest/Private/MetasoundGeneratorOutputWatcherTests.cpp
new file mode 100644
--- /dev/null
+++ b/Engine/Plugins/Runtime/Metasound/Source/MetasoundEngineTest/Private/MetasoundGeneratorOutputWatcherTests.cpp
@@ -0,0 +1,175 @@
+// Copyright Epic Games, Inc. All Rights Reserved.
+
+#include "MetasoundGenerator.h"
+#include "MetasoundGeneratorOutputWatcher.h"
+#include "NodeTestGraphBuilder.h"
+#include "Misc/AutomationTest.h"
+
+#if WITH_DEV_AUTOMATION_TESTS
+
+namespace Metasound::Test::GeneratorOutputWatcher
+{
+	using GraphBuilder = FNodeTestGraphBuilder;
+	using namespace Frontend;
+	using namespace Private;
+	
+	template<typename DataType>
+	TUniquePtr<FMetasoundGenerator> BuildPassthroughGraph(
+		FAutomationTestBase& Test,
+		const FName& InputName,
+		const FName& OutputName,
+		const FSampleRate SampleRate,
+		const int32 NumSamplesPerBlock,
+		FGuid* OutputGuid = nullptr)
+	{
+		GraphBuilder Builder;
+		const FNodeHandle InputNode = Builder.AddInput(InputName, Metasound::GetMetasoundDataTypeName<DataType>());
+		const FNodeHandle OutputNode = Builder.AddOutput(OutputName, Metasound::GetMetasoundDataTypeName<DataType>());
+		const FOutputHandle OutputToConnect = InputNode->GetOutputWithVertexName(InputName);
+		const FInputHandle InputToConnect = OutputNode->GetInputWithVertexName(OutputName);
+
+		if (!Test.TestTrue("Connected input to output", InputToConnect->Connect(*OutputToConnect)))
+		{
+			return nullptr;
+		}
+
+		if (nullptr != OutputGuid)
+		{
+			*OutputGuid = OutputNode->GetID();
+		}
+
+		// have to add an audio output for the generator to render
+		Builder.AddOutput("Audio", Metasound::GetMetasoundDataTypeName<FAudioBuffer>());
+		
+		return Builder.BuildGenerator(SampleRate, NumSamplesPerBlock);
+	}
+
+	template<typename DataType>
+	bool RunSimpleOutputWatcherTest(FAutomationTestBase& Test, DataType ExpectedValue, const FName AnalyzerName, const FName ValueOutputName)
+	{
+		const FName InputName = "MyInput";
+		const FName OutputName = "MyOutput";
+		constexpr FSampleRate SampleRate = 48000;
+		constexpr int32 NumSamplesPerBlock = 128;
+
+		// Build a passthrough graph
+		FGuid OutputNodeId;
+		const TUniquePtr<FMetasoundGenerator> Generator =
+			BuildPassthroughGraph<DataType>(Test, InputName, OutputName, SampleRate, NumSamplesPerBlock, &OutputNodeId);
+		if (!Test.TestNotNull("Generator built", Generator.Get()))
+		{
+			return false;
+		}
+
+		// Add an analyzer to the output
+		FAnalyzerAddress AnalyzerAddress;
+		AnalyzerAddress.DataType = GetMetasoundDataTypeName<DataType>();
+		AnalyzerAddress.InstanceID = 1234;
+		AnalyzerAddress.OutputName = OutputName;
+		AnalyzerAddress.AnalyzerName = AnalyzerName;
+		AnalyzerAddress.AnalyzerInstanceID = FGuid::NewGuid();
+		AnalyzerAddress.NodeID = OutputNodeId;
+		Generator->AddOutputVertexAnalyzer(AnalyzerAddress);
+
+		// Add an output watcher
+		FMetasoundGeneratorOutputWatcher Watcher{ MoveTemp(AnalyzerAddress), Generator->OperatorSettings };
+
+		// listen for changes to the output
+		DataType WatchedValue;
+		bool ReceivedValue = false;
+		const TFunction<void(FName, const FMetasoundGeneratorOutput&)> OnChanged =
+			[&WatchedValue, &ReceivedValue](FName, const FMetasoundGeneratorOutput& Output)
+			{
+				ReceivedValue = Output.Get<DataType>(WatchedValue);
+			};
+		
+		// Set the input value
+		Generator->SetInputValue(InputName, ExpectedValue);
+
+		// Render a block and update the watcher
+		{
+			TArray<float> Buffer;
+			Buffer.Reserve(NumSamplesPerBlock);
+			if (!Test.TestEqual(
+				"Generated the right number of samples.",
+				Generator->OnGenerateAudio(Buffer.GetData(), NumSamplesPerBlock),
+				NumSamplesPerBlock))
+			{
+				return false;
+			}
+			Watcher.Update(OnChanged);
+		}
+
+		// Check the output
+		if (!Test.TestTrue("Received a value", ReceivedValue)
+			|| !Test.TestEqual("Received the expected value", WatchedValue, ExpectedValue))
+		{
+			return false;
+		}
+
+		// don't set the input and expect no callbacks
+		ReceivedValue = false;
+		
+		// Render a block and update the watcher
+		{
+			TArray<float> Buffer;
+			Buffer.Reserve(NumSamplesPerBlock);
+			if (!Test.TestEqual(
+				"Generated the right number of samples.",
+				Generator->OnGenerateAudio(Buffer.GetData(), NumSamplesPerBlock),
+				NumSamplesPerBlock))
+			{
+				return false;
+			}
+			Watcher.Update(OnChanged);
+		}
+
+		return Test.TestFalse("No more updates", ReceivedValue);
+	}
+
+	IMPLEMENT_SIMPLE_AUTOMATION_TEST(
+	FMetasoundGeneratorOutputWatcherForwardValueFloatTest,
+	"Audio.Metasound.GeneratorOutputWatcher.ForwardValue.Float",
+	EAutomationTestFlags::EditorContext | EAutomationTestFlags::EngineFilter)
+	bool FMetasoundGeneratorOutputWatcherForwardValueFloatTest::RunTest(const FString&)
+	{
+		using DataType = float;
+		constexpr DataType ExpectedValue = 123.456f;
+		return RunSimpleOutputWatcherTest<DataType>(*this, ExpectedValue, "UE.Forward.Float", "Value");
+	}
+	
+	IMPLEMENT_SIMPLE_AUTOMATION_TEST(
+	FMetasoundGeneratorOutputWatcherForwardValueIntTest,
+	"Audio.Metasound.GeneratorOutputWatcher.ForwardValue.Int32",
+	EAutomationTestFlags::EditorContext | EAutomationTestFlags::EngineFilter)
+	bool FMetasoundGeneratorOutputWatcherForwardValueIntTest::RunTest(const FString&)
+	{
+		using DataType = int32;
+		constexpr DataType ExpectedValue = 123456;
+		return RunSimpleOutputWatcherTest<DataType>(*this, ExpectedValue, "UE.Forward.Int32", "Value");
+	}
+
+	IMPLEMENT_SIMPLE_AUTOMATION_TEST(
+	FMetasoundGeneratorOutputWatcherForwardValueBoolTest,
+	"Audio.Metasound.GeneratorOutputWatcher.ForwardValue.Bool",
+	EAutomationTestFlags::EditorContext | EAutomationTestFlags::EngineFilter)
+	bool FMetasoundGeneratorOutputWatcherForwardValueBoolTest::RunTest(const FString&)
+	{
+		using DataType = bool;
+		constexpr DataType ExpectedValue = true;
+		return RunSimpleOutputWatcherTest<DataType>(*this, ExpectedValue, "UE.Forward.Bool", "Value");
+	}
+
+	IMPLEMENT_SIMPLE_AUTOMATION_TEST(
+	FMetasoundGeneratorOutputWatcherForwardValueStringTest,
+	"Audio.Metasound.GeneratorOutputWatcher.ForwardValue.String",
+	EAutomationTestFlags::EditorContext | EAutomationTestFlags::EngineFilter)
+	bool FMetasoundGeneratorOutputWatcherForwardValueStringTest::RunTest(const FString&)
+	{
+		using FDataType = FString;
+		const FDataType ExpectedValue{ "you're awesome" };
+		return RunSimpleOutputWatcherTest<FDataType>(*this, ExpectedValue, "UE.Forward.String", "Value");
+	}
+}
+
+#endif
\ No newline at end of file
diff --git a/Engine/Plugins/Runtime/Metasound/Source/MetasoundEngineTest/Private/MetasoundGeneratorTests.cpp b/Engine/Plugins/Runtime/Metasound/Source/MetasoundEngineTest/Private/MetasoundGeneratorTests.cpp
--- a/Engine/Plugins/Runtime/Metasound/Source/MetasoundEngineTest/Private/MetasoundGeneratorTests.cpp
+++ b/Engine/Plugins/Runtime/Metasound/Source/MetasoundEngineTest/Private/MetasoundGeneratorTests.cpp
@@ -375,6 +375,17 @@ namespace Metasound::Test::Generator
 		constexpr DataType ExpectedValue = true;
 		return RunOutputAnalyzerForwardValueTest<DataType>(*this, ExpectedValue, "UE.Forward.Bool");
 	}
+
+	IMPLEMENT_SIMPLE_AUTOMATION_TEST(
+	FMetasoundGeneratorOututAnalyzerForwardValueStringTest,
+	"Audio.Metasound.Generator.OutputAnalyzer.ForwardValue.String",
+	EAutomationTestFlags::EditorContext | EAutomationTestFlags::EngineFilter)
+	bool FMetasoundGeneratorOututAnalyzerForwardValueStringTest::RunTest(const FString&)
+	{
+		using FDataType = FString;
+		const FDataType ExpectedValue { "unexpected value" };
+		return RunOutputAnalyzerForwardValueTest<FDataType>(*this, ExpectedValue, "UE.Forward.String");
+	}
 }
 
 #endif
\ No newline at end of file
diff --git a/Engine/Plugins/Runtime/Metasound/Source/MetasoundFrontend/Private/Analysis/MetasoundFrontendAnalyzerView.cpp b/Engine/Plugins/Runtime/Metasound/Source/MetasoundFrontend/Private/Analysis/MetasoundFrontendAnalyzerView.cpp
--- a/Engine/Plugins/Runtime/Metasound/Source/MetasoundFrontend/Private/Analysis/MetasoundFrontendAnalyzerView.cpp
+++ b/Engine/Plugins/Runtime/Metasound/Source/MetasoundFrontend/Private/Analysis/MetasoundFrontendAnalyzerView.cpp
@@ -43,5 +43,22 @@ namespace Metasound
 		{
 			return OutputReceivers.Remove(InOutputName) > 0;
 		}
+
+		TArray<FMetasoundAnalyzerView::FBoundOutputDescription> FMetasoundAnalyzerView::GetBoundOutputDescriptions() const
+		{
+			TArray<FBoundOutputDescription> OutputDescriptions;
+			
+			for (auto& Receiver : OutputReceivers)
+			{
+				if (!ensure(Receiver.Value.IsValid()))
+				{
+					continue;
+				}
+
+				OutputDescriptions.Add({ Receiver.Key, Receiver.Value->GetDataType() });
+			}
+
+			return OutputDescriptions;
+		}
 	} // namespace Frontend
 } // namespace Metasound
diff --git a/Engine/Plugins/Runtime/Metasound/Source/MetasoundFrontend/Public/Analysis/MetasoundFrontendAnalyzerView.h b/Engine/Plugins/Runtime/Metasound/Source/MetasoundFrontend/Public/Analysis/MetasoundFrontendAnalyzerView.h
--- a/Engine/Plugins/Runtime/Metasound/Source/MetasoundFrontend/Public/Analysis/MetasoundFrontendAnalyzerView.h
+++ b/Engine/Plugins/Runtime/Metasound/Source/MetasoundFrontend/Public/Analysis/MetasoundFrontendAnalyzerView.h
@@ -51,6 +51,14 @@ namespace Metasound
 
 				return false;
 			}
+
+			struct FBoundOutputDescription
+			{
+				FName Name;
+				FName TypeName;
+			};
+
+			TArray<FBoundOutputDescription> GetBoundOutputDescriptions() const;
 		};
 	}
 }
diff --git a/Engine/Plugins/Runtime/Metasound/Source/MetasoundGenerator/Private/MetasoundGenerator.cpp b/Engine/Plugins/Runtime/Metasound/Source/MetasoundGenerator/Private/MetasoundGenerator.cpp
--- a/Engine/Plugins/Runtime/Metasound/Source/MetasoundGenerator/Private/MetasoundGenerator.cpp
+++ b/Engine/Plugins/Runtime/Metasound/Source/MetasoundGenerator/Private/MetasoundGenerator.cpp
@@ -540,6 +540,11 @@ namespace Metasound
 	{
 		OutputAnalyzerModificationQueue.Enqueue([this, AnalyzerAddress]
 		{
+			if (OutputAnalyzers.Contains(AnalyzerAddress.OutputName))
+			{
+				return;
+			}
+			
 			const FAnyDataReference* OutputReference = VertexInterfaceData.GetOutputs().FindDataReference(AnalyzerAddress.OutputName);
 			
 			if (nullptr == OutputReference)
diff --git a/Engine/Source/Runtime/AudioExtensions/Public/SoundGeneratorOutput.h b/Engine/Source/Runtime/AudioExtensions/Public/SoundGeneratorOutput.h
new file mode 100644
--- /dev/null
+++ b/Engine/Source/Runtime/AudioExtensions/Public/SoundGeneratorOutput.h
@@ -0,0 +1,20 @@
+// Copyright Epic Games, Inc. All Rights Reserved.
+
+#pragma once
+
+#include "SoundGeneratorOutput.generated.h"
+
+/**
+ * Base class for generators that have outputs that can be exposed to other game code.
+ *
+ * NOTE: This is not widely supported and should be considered experimental.
+ */
+USTRUCT(BlueprintType)
+struct FSoundGeneratorOutput
+{
+	GENERATED_BODY()
+
+	/** The output's name */
+	UPROPERTY(BlueprintReadOnly, Category = SoundGeneratorOutput)
+	FName Name;
+};
