MusicClockComponent & MusicTempometerComponent for the new Harmonix plugin. Includes massive refactor for Measure -> Bar and Pulse -> Beat.

diff --git a/Engine/Plugins/Runtime/Metasound/Source/MetasoundEngine/Private/MetasoundGeneratorHandle.cpp b/Engine/Plugins/Runtime/Metasound/Source/MetasoundEngine/Private/MetasoundGeneratorHandle.cpp
--- a/Engine/Plugins/Runtime/Metasound/Source/MetasoundEngine/Private/MetasoundGeneratorHandle.cpp
+++ b/Engine/Plugins/Runtime/Metasound/Source/MetasoundEngine/Private/MetasoundGeneratorHandle.cpp
@@ -136,14 +136,26 @@ bool UMetasoundGeneratorHandle::ApplyParameterPack(UMetasoundParameterPack* Pack
 	return true;
 }
 
+TSharedPtr<Metasound::FMetasoundGenerator> UMetasoundGeneratorHandle::GetGenerator()
+{
+	// Attach if we aren't attached, check for changes, etc...
+	CacheMetasoundSource();
+
+	return PinGenerator();
+}
+
 void UMetasoundGeneratorHandle::OnSourceCreatedAGenerator(uint64 InAudioComponentId, TSharedPtr<Metasound::FMetasoundGenerator> InGenerator)
 {
 	if (InAudioComponentId == AudioComponentId)
 	{
 		CachedGeneratorPtr = InGenerator;
-		if (InGenerator && CachedParameterPack)
+		if (InGenerator)
 		{
-			InGenerator->QueueParameterPack(CachedParameterPack);
+			if (CachedParameterPack)
+			{
+				InGenerator->QueueParameterPack(CachedParameterPack);
+			}
+			OnGeneratorHandleAttached.Broadcast();
 		}
 	}
 }
@@ -152,6 +164,10 @@ void UMetasoundGeneratorHandle::OnSourceDestroyedAGenerator(uint64 InAudioCompon
 {
 	if (InAudioComponentId == AudioComponentId)
 	{
+		if (CachedGeneratorPtr.IsValid())
+		{
+			OnGeneratorHandleDetached.Broadcast();
+		}
 		CachedGeneratorPtr = nullptr;
 	}
 }
diff --git a/Engine/Plugins/Runtime/Metasound/Source/MetasoundEngine/Public/MetasoundGeneratorHandle.h b/Engine/Plugins/Runtime/Metasound/Source/MetasoundEngine/Public/MetasoundGeneratorHandle.h
--- a/Engine/Plugins/Runtime/Metasound/Source/MetasoundEngine/Public/MetasoundGeneratorHandle.h
+++ b/Engine/Plugins/Runtime/Metasound/Source/MetasoundEngine/Public/MetasoundGeneratorHandle.h
@@ -14,6 +14,9 @@ namespace Metasound
 	class FMetasoundGenerator;
 }
 
+DECLARE_TS_MULTICAST_DELEGATE(FOnGeneratorHandleAttached);
+DECLARE_TS_MULTICAST_DELEGATE(FOnGeneratorHandleDetached);
+
 UCLASS(BlueprintType,Category="MetaSound")
 class METASOUNDENGINE_API UMetasoundGeneratorHandle : public UObject
 {
@@ -33,6 +36,14 @@ public:
 	UFUNCTION(BlueprintCallable, Category="MetaSoundParameterPack")
 	bool ApplyParameterPack(UMetasoundParameterPack* Pack);
 
+	TSharedPtr<Metasound::FMetasoundGenerator> GetGenerator();
+
+	/**
+	 * NOTE: These delegates are called on the audio render thread. Careful!
+	 */
+	FOnGeneratorHandleAttached OnGeneratorHandleAttached;
+	FOnGeneratorHandleDetached OnGeneratorHandleDetached;
+
 private:
 
 	void SetAudioComponent(UAudioComponent* InAudioComponent);
diff --git a/Engine/Plugins/Runtime/Metasound/Source/MetasoundGenerator/Private/MetasoundGenerator.cpp b/Engine/Plugins/Runtime/Metasound/Source/MetasoundGenerator/Private/MetasoundGenerator.cpp
--- a/Engine/Plugins/Runtime/Metasound/Source/MetasoundGenerator/Private/MetasoundGenerator.cpp
+++ b/Engine/Plugins/Runtime/Metasound/Source/MetasoundGenerator/Private/MetasoundGenerator.cpp
@@ -524,6 +524,8 @@ namespace Metasound
 		{
 			OnPlayTriggerRef->TriggerFrame(0);
 		}
+
+		OnSetGraph.Broadcast();
 	}
 
 	void FMetasoundGenerator::QueueParameterPack(TSharedPtr<FMetasoundParameterPackStorage> ParameterPack)
diff --git a/Engine/Plugins/Runtime/Metasound/Source/MetasoundGenerator/Public/MetasoundGenerator.h b/Engine/Plugins/Runtime/Metasound/Source/MetasoundGenerator/Public/MetasoundGenerator.h
--- a/Engine/Plugins/Runtime/Metasound/Source/MetasoundGenerator/Public/MetasoundGenerator.h
+++ b/Engine/Plugins/Runtime/Metasound/Source/MetasoundGenerator/Public/MetasoundGenerator.h
@@ -14,6 +14,7 @@
 #include "Async/AsyncWork.h"
 #include "Containers/MpscQueue.h"
 #include "Sound/SoundGenerator.h"
+#include "Delegates/Delegate.h"
 
 #ifndef ENABLE_METASOUND_GENERATOR_RENDER_TIMING
 #define ENABLE_METASOUND_GENERATOR_RENDER_TIMING WITH_EDITOR
@@ -100,6 +101,8 @@ namespace Metasound
 		bool bTriggerGenerator;
 	};
 
+	DECLARE_TS_MULTICAST_DELEGATE(FOnSetGraph);
+
 	/** FMetasoundGenerator generates audio from a given metasound IOperator
 	 * which produces a multichannel audio output.
 	 */
@@ -232,6 +235,9 @@ namespace Metasound
 		void SetPendingGraph(FMetasoundGeneratorData&& InData, bool bTriggerGraph);
 		void SetPendingGraphBuildFailed();
 
+		// Called when a new graph has been "compiled" and set up as this generator's graph.
+		FOnSetGraph OnSetGraph;
+
 	private:
 
 		bool UpdateGraphIfPending();
