[Metasound] Blueprint access to watch outputs on a playing Metasound source

diff --git a/Engine/Plugins/Runtime/Metasound/Source/MetasoundEngine/Private/MetasoundEngineModule.cpp b/Engine/Plugins/Runtime/Metasound/Source/MetasoundEngine/Private/MetasoundEngineModule.cpp
--- a/Engine/Plugins/Runtime/Metasound/Source/MetasoundEngine/Private/MetasoundEngineModule.cpp
+++ b/Engine/Plugins/Runtime/Metasound/Source/MetasoundEngine/Private/MetasoundEngineModule.cpp
@@ -5,13 +5,15 @@
 #include "MetasoundAudioBus.h"
 #include "MetasoundDataReference.h"
 #include "MetasoundDataTypeRegistrationMacro.h"
-#include "MetasoundGeneratorOutputWatcher.h"
+#include "MetasoundOutputSubsystem.h"
+#include "MetasoundOutputWatcher.h"
 #include "MetasoundSource.h"
 #include "MetasoundTrace.h"
 #include "MetasoundUObjectRegistry.h"
 #include "MetasoundWave.h"
 #include "MetasoundWaveTable.h"
 #include "Analysis/MetasoundFrontendAnalyzerRegistry.h"
+#include "Analysis/MetasoundFrontendVertexAnalyzerForwardValue.h"
 #include "Interfaces/MetasoundDeprecatedInterfaces.h"
 #include "Interfaces/MetasoundInterface.h"
 #include "Interfaces/MetasoundInterfaceBindingsPrivate.h"
@@ -67,6 +69,20 @@ class FMetasoundEngineModule : public IMetasoundEngineModule
 		METASOUND_PRIVATE_REGISTER_GENERATOR_OUTPUT_WATCHER_TYPE_OPERATIONS(bool, false)
 		METASOUND_PRIVATE_REGISTER_GENERATOR_OUTPUT_WATCHER_TYPE_OPERATIONS(FString, "")
 
+		// Register passthrough output analyzers
+		UMetaSoundOutputSubsystem::RegisterPassthroughAnalyzerForType(
+			GetMetasoundDataTypeName<float>(),
+			Frontend::FVertexAnalyzerForwardFloat::GetAnalyzerName());
+		UMetaSoundOutputSubsystem::RegisterPassthroughAnalyzerForType(
+			GetMetasoundDataTypeName<int32>(),
+			Frontend::FVertexAnalyzerForwardInt::GetAnalyzerName());
+		UMetaSoundOutputSubsystem::RegisterPassthroughAnalyzerForType(
+			GetMetasoundDataTypeName<bool>(),
+			Frontend::FVertexAnalyzerForwardBool::GetAnalyzerName());
+		UMetaSoundOutputSubsystem::RegisterPassthroughAnalyzerForType(
+			GetMetasoundDataTypeName<FString>(),
+			Frontend::FVertexAnalyzerForwardString::GetAnalyzerName());
+
 		UE_LOG(LogMetasoundEngine, Log, TEXT("MetaSound Engine Initialized"));
 	}
 };
diff --git a/Engine/Plugins/Runtime/Metasound/Source/MetasoundEngine/Private/MetasoundGeneratorOutput.cpp b/Engine/Plugins/Runtime/Metasound/Source/MetasoundEngine/Private/MetasoundGeneratorOutput.cpp
deleted file mode 100644
diff --git a/Engine/Plugins/Runtime/Metasound/Source/MetasoundEngine/Private/MetasoundGeneratorOutputWatcher.cpp b/Engine/Plugins/Runtime/Metasound/Source/MetasoundEngine/Private/MetasoundGeneratorOutputWatcher.cpp
deleted file mode 100644
diff --git a/Engine/Plugins/Runtime/Metasound/Source/MetasoundEngine/Private/MetasoundOutput.cpp b/Engine/Plugins/Runtime/Metasound/Source/MetasoundEngine/Private/MetasoundOutput.cpp
new file mode 100644
--- /dev/null
+++ b/Engine/Plugins/Runtime/Metasound/Source/MetasoundEngine/Private/MetasoundOutput.cpp
@@ -0,0 +1,80 @@
+// Copyright Epic Games, Inc. All Rights Reserved.
+
+#include "MetasoundOutput.h"
+
+#include "MetasoundPrimitives.h"
+
+FMetaSoundOutput::FMetaSoundOutput(const FMetaSoundOutput& Other)
+{
+	Name = Other.Name;
+	DataReference = MakeUnique<Metasound::FAnyDataReference>(*Other.DataReference);
+}
+
+FMetaSoundOutput& FMetaSoundOutput::operator=(const FMetaSoundOutput& Other)
+{
+	if (this != &Other)
+	{
+		this->Name = Other.Name;
+		this->DataReference = MakeUnique<Metasound::FAnyDataReference>(*Other.DataReference);
+	}
+
+	return *this;
+}
+
+bool FMetaSoundOutput::IsValid() const
+{
+	return nullptr != DataReference;
+}
+
+FName FMetaSoundOutput::GetTypeName() const
+{
+	return IsValid() ? DataReference->GetDataTypeName() : FName();
+}
+
+bool UMetasoundOutputBlueprintAccess::IsFloat(const FMetaSoundOutput& Output)
+{
+	return Output.IsType<float>();
+}
+
+float UMetasoundOutputBlueprintAccess::GetFloat(const FMetaSoundOutput& Output, bool& Success)
+{
+	float Value = 0;
+	Success = Output.Get(Value);
+	return Value;
+}
+
+bool UMetasoundOutputBlueprintAccess::IsInt32(const FMetaSoundOutput& Output)
+{
+	return Output.IsType<int32>();
+}
+
+int32 UMetasoundOutputBlueprintAccess::GetInt32(const FMetaSoundOutput& Output, bool& Success)
+{
+	int32 Value = 0;
+	Success = Output.Get(Value);
+	return Value;
+}
+
+bool UMetasoundOutputBlueprintAccess::IsBool(const FMetaSoundOutput& Output)
+{
+	return Output.IsType<bool>();
+}
+
+bool UMetasoundOutputBlueprintAccess::GetBool(const FMetaSoundOutput& Output, bool& Success)
+{
+	bool Value = false;
+	Success = Output.Get(Value);
+	return Value;
+}
+
+bool UMetasoundOutputBlueprintAccess::IsString(const FMetaSoundOutput& Output)
+{
+	return Output.IsType<FString>();
+}
+
+FString UMetasoundOutputBlueprintAccess::GetString(const FMetaSoundOutput& Output, bool& Success)
+{
+	FString Value;
+	Success = Output.Get(Value);
+	return Value;
+}
diff --git a/Engine/Plugins/Runtime/Metasound/Source/MetasoundEngine/Private/MetasoundOutputSubsystem.cpp b/Engine/Plugins/Runtime/Metasound/Source/MetasoundEngine/Private/MetasoundOutputSubsystem.cpp
new file mode 100644
--- /dev/null
+++ b/Engine/Plugins/Runtime/Metasound/Source/MetasoundEngine/Private/MetasoundOutputSubsystem.cpp
@@ -0,0 +1,269 @@
+// Copyright Epic Games, Inc. All Rights Reserved.
+
+#include "MetasoundOutputSubsystem.h"
+
+#include "MetasoundGenerator.h"
+#include "MetasoundSource.h"
+#include "MetasoundTrace.h"
+#include "Analysis/MetasoundFrontendAnalyzerAddress.h"
+#include "Components/AudioComponent.h"
+#include "ProfilingDebugging/CpuProfilerTrace.h"
+
+TMap<FName, FName> UMetaSoundOutputSubsystem::PassthroughAnalyzers{};
+
+bool UMetaSoundOutputSubsystem::WatchOutput(
+	UAudioComponent* AudioComponent,
+	const FName OutputName,
+	const FOnMetasoundOutputValueChanged& OnOutputValueChanged,
+	FName AnalyzerName,
+	FName AnalyzerOutputName)
+{
+	METASOUND_TRACE_CPUPROFILER_EVENT_SCOPE(UMetasoundOutputSubsystem::WatchOutput);
+	
+	FGeneratorInfo* GeneratorInfo = FindOrAddGeneratorInfo(AudioComponent);
+
+	if (nullptr == GeneratorInfo)
+	{
+		return false;
+	}
+	
+	check(GeneratorInfo->Source.IsValid());
+
+	// Find the node id and type name
+	const Metasound::Frontend::FNodeHandle Node = GeneratorInfo->Source->GetRootGraphHandle()->GetOutputNodeWithName(OutputName);
+
+	if (!Node->IsValid())
+	{
+		return false;
+	}
+
+	const FGuid NodeId = Node->GetID();
+
+	// We expect an output node to only have one output
+	if (Node->GetNumOutputs() != 1)
+	{
+		return false;
+	}
+
+	const FName TypeName = Node->GetOutputs()[0]->GetDataType();
+
+	// Make the analyzer address
+	Metasound::Frontend::FAnalyzerAddress AnalyzerAddress;
+	AnalyzerAddress.DataType = TypeName;
+	AnalyzerAddress.InstanceID = AudioComponent->GetAudioComponentID();
+	AnalyzerAddress.OutputName = OutputName;
+
+	if (AnalyzerName.IsNone())
+	{
+		if (!PassthroughAnalyzers.Contains(TypeName))
+		{
+			return false;
+		}
+
+		AnalyzerName = PassthroughAnalyzers[TypeName];
+		AnalyzerOutputName = "Value";
+	}
+	AnalyzerAddress.AnalyzerName = AnalyzerName;
+	
+	AnalyzerAddress.AnalyzerInstanceID = FGuid::NewGuid();
+	AnalyzerAddress.NodeID = NodeId;
+
+	// if we already have a generator, go ahead and make the analyzer and watcher
+	if (const TSharedPtr<Metasound::FMetasoundGenerator, ESPMode::ThreadSafe> PinnedGenerator = GeneratorInfo->Generator.Pin())
+	{
+		CreateAnalyzerAndWatcher(PinnedGenerator, MoveTemp(AnalyzerAddress));
+	}
+	// otherwise enqueue it for later
+	else
+	{
+		GeneratorInfo->AnalyzersToCreate.Enqueue(MoveTemp(AnalyzerAddress));
+	}
+
+	// either way, add the delegate to the map
+	GeneratorInfo->OutputChangedDelegates.FindOrAdd(OutputName).FindOrAdd(AnalyzerOutputName).AddUnique(OnOutputValueChanged);
+
+	// set up the output handler
+	GeneratorInfo->HandleOutputChanged = [GeneratorInfo](const FName GeneratorOutputName, const FMetaSoundOutput& AnalyzerOutput)
+	{
+		if (TMap<FName, FOnOutputValueChangedMulticast>* Map = GeneratorInfo->OutputChangedDelegates.Find(GeneratorOutputName))
+		{
+			if (const FOnOutputValueChangedMulticast* Delegate = Map->Find(AnalyzerOutput.Name))
+			{
+				Delegate->Broadcast(AnalyzerOutput);
+			}
+		}
+	};
+	
+	return true;
+}
+
+bool UMetaSoundOutputSubsystem::IsTickable() const
+{
+	return TrackedGenerators.Num() > 0;
+}
+
+void UMetaSoundOutputSubsystem::Tick(float DeltaTime)
+{
+	METASOUND_TRACE_CPUPROFILER_EVENT_SCOPE(UMetasoundOutputSubsystem::Tick);
+	 
+	for (auto GeneratorIt = TrackedGenerators.CreateIterator(); GeneratorIt; ++GeneratorIt)
+	{
+		if (!GeneratorIt->IsValid())
+		{
+			GeneratorIt.RemoveCurrent();
+		}
+		else
+		{
+			for (Metasound::Private::FMetasoundOutputWatcher& Watcher : GeneratorIt->OutputWatchers)
+			{
+				Watcher.Update(GeneratorIt->HandleOutputChanged);
+			}
+		}
+	}
+}
+
+TStatId UMetaSoundOutputSubsystem::GetStatId() const
+{
+	RETURN_QUICK_DECLARE_CYCLE_STAT(UMetasoundGeneratorAccessSubsystem, STATGROUP_Tickables);
+}
+
+void UMetaSoundOutputSubsystem::RegisterPassthroughAnalyzerForType(const FName TypeName, const FName AnalyzerName)
+{
+	check(!PassthroughAnalyzers.Contains(TypeName));
+	PassthroughAnalyzers.Add(TypeName, AnalyzerName);
+}
+
+bool UMetaSoundOutputSubsystem::FGeneratorInfo::IsValid() const
+{
+	return AudioComponent.IsValid() && Source.IsValid();
+}
+
+UMetaSoundOutputSubsystem::FGeneratorInfo* UMetaSoundOutputSubsystem::FindOrAddGeneratorInfo(UAudioComponent* AudioComponent)
+{
+	if (!ensure(nullptr != AudioComponent))
+	{
+		return  nullptr;
+	}
+	
+	const int64 AudioComponentId = AudioComponent->GetAudioComponentID();
+
+	// already in the list...
+	if (FGeneratorInfo* Info = TrackedGenerators.FindByPredicate([AudioComponentId](const FGeneratorInfo& GeneratorInfo)
+	{
+		return GeneratorInfo.AudioComponent.IsValid()
+		&& GeneratorInfo.AudioComponent->GetAudioComponentID() == AudioComponentId;
+	}))
+	{
+		return Info;
+	}
+
+	UMetaSoundSource* Source = Cast<UMetaSoundSource>(AudioComponent->GetSound());
+	
+	if (nullptr == Source)
+	{
+		return nullptr;
+	}
+
+	// set up a new entry
+	TrackedGenerators.AddDefaulted();
+	FGeneratorInfo& GeneratorInfo = TrackedGenerators.Last();
+	GeneratorInfo.AudioComponent = AudioComponent;
+	GeneratorInfo.Source = Source;
+
+	// attempt to get the generator
+	GeneratorInfo.Generator = GeneratorInfo.Source->GetGeneratorForAudioComponent(AudioComponentId);
+
+	// if that didn't work, set up the generator creation callback
+	if (nullptr == GeneratorInfo.Generator)
+	{
+		GeneratorInfo.OnCreatedHandle = GeneratorInfo.Source->OnGeneratorInstanceCreated.AddUObject(
+			this,
+			&UMetaSoundOutputSubsystem::OnGeneratorCreated);
+	}
+	
+	// always set up the generator destruction callback
+	GeneratorInfo.OnDestroyedHandle = GeneratorInfo.Source->OnGeneratorInstanceDestroyed.AddUObject(
+		this,
+		&UMetaSoundOutputSubsystem::OnGeneratorDestroyed);
+
+	return &GeneratorInfo;
+}
+
+void UMetaSoundOutputSubsystem::OnGeneratorCreated(
+	const uint64 InAudioComponentId,
+	const TSharedPtr<Metasound::FMetasoundGenerator> Generator)
+{
+	for (auto GeneratorIt = TrackedGenerators.CreateIterator(); GeneratorIt; ++GeneratorIt)
+	{
+		if (!GeneratorIt->IsValid())
+		{
+			GeneratorIt.RemoveCurrent();
+		}
+		else if (GeneratorIt->AudioComponent->GetAudioComponentID() == InAudioComponentId)
+		{
+			GeneratorIt->Generator = Generator;
+			GeneratorIt->Source->OnGeneratorInstanceCreated.Remove(GeneratorIt->OnCreatedHandle);
+
+			// create analyzers
+			Metasound::Frontend::FAnalyzerAddress Address;
+			while (GeneratorIt->AnalyzersToCreate.Dequeue(Address))
+			{
+				CreateAnalyzerAndWatcher(Generator, MoveTemp(Address));
+			}
+		}
+	}
+}
+
+void UMetaSoundOutputSubsystem::OnGeneratorDestroyed(
+	const uint64 InAudioComponentId,
+	TSharedPtr<Metasound::FMetasoundGenerator>)
+{
+	for (auto GeneratorIt = TrackedGenerators.CreateIterator(); GeneratorIt; ++GeneratorIt)
+	{
+		if (!GeneratorIt->IsValid())
+		{
+			GeneratorIt.RemoveCurrent();
+		}
+		else if (GeneratorIt->AudioComponent.IsValid() && GeneratorIt->AudioComponent->GetAudioComponentID() == InAudioComponentId)
+		{
+			GeneratorIt->Source->OnGeneratorInstanceDestroyed.Remove(GeneratorIt->OnDestroyedHandle);
+			GeneratorIt.RemoveCurrent();
+			return;
+		}
+	}
+}
+
+void UMetaSoundOutputSubsystem::CreateAnalyzerAndWatcher(
+	const TSharedPtr<Metasound::FMetasoundGenerator>& Generator,
+	Metasound::Frontend::FAnalyzerAddress&& AnalyzerAddress)
+{
+	if (!Generator.IsValid())
+	{
+		return;
+	}
+	
+	// Find the generator info and create the watcher
+	{
+		FGeneratorInfo* GeneratorInfo = TrackedGenerators.FindByPredicate([&Generator](const FGeneratorInfo& Info)
+		{
+			return Info.Generator.HasSameObject(Generator.Get());
+		});
+
+		if (nullptr == GeneratorInfo)
+		{
+			return;
+		}
+
+		if (!GeneratorInfo->OutputWatchers.ContainsByPredicate(
+			[&AnalyzerAddress](const Metasound::Private::FMetasoundOutputWatcher& Watcher)
+			{
+				return AnalyzerAddress.OutputName == Watcher.Name;
+			}))
+		{
+			GeneratorInfo->OutputWatchers.Emplace(MoveTemp(AnalyzerAddress), Generator->OperatorSettings);
+		}
+	}
+
+	// Create the analyzer (will skip if there's already one)
+	Generator->AddOutputVertexAnalyzer(AnalyzerAddress);
+}
diff --git a/Engine/Plugins/Runtime/Metasound/Source/MetasoundEngine/Private/MetasoundOutputWatcher.cpp b/Engine/Plugins/Runtime/Metasound/Source/MetasoundEngine/Private/MetasoundOutputWatcher.cpp
new file mode 100644
--- /dev/null
+++ b/Engine/Plugins/Runtime/Metasound/Source/MetasoundEngine/Private/MetasoundOutputWatcher.cpp
@@ -0,0 +1,59 @@
+// Copyright Epic Games, Inc. All Rights Reserved.
+
+#include "MetasoundOutputWatcher.h"
+#include "MetasoundTrace.h"
+
+namespace Metasound::Private
+{
+	TMap<FName, TUniquePtr<FMetasoundOutputWatcher::IOutputTypeOperations>> FMetasoundOutputWatcher::OutputTypeOperationMap{};
+	
+	FMetasoundOutputWatcher::FMetasoundOutputWatcher(
+		Frontend::FAnalyzerAddress&& Address,
+		const FOperatorSettings& OperatorSettings)
+		: Name(Address.OutputName)
+		, View(MoveTemp(Address))
+	{
+		View.BindToAllOutputs(OperatorSettings);
+		
+		TArray<Frontend::FMetasoundAnalyzerView::FBoundOutputDescription> OutputDescriptions = View.GetBoundOutputDescriptions();
+		for (Frontend::FMetasoundAnalyzerView::FBoundOutputDescription& OutputDescription : OutputDescriptions)
+		{
+			if (!ensure(OutputTypeOperationMap.Contains(OutputDescription.TypeName)))
+			{
+				continue;
+			}
+
+			FMetaSoundOutput Output;
+			Output.Name = OutputDescription.Name;
+			OutputTypeOperationMap[OutputDescription.TypeName]->Init(Output);
+			Outputs.Emplace(MoveTemp(Output));
+		}
+	}
+
+	void FMetasoundOutputWatcher::Update(const TFunctionRef<void(FName, const FMetaSoundOutput&)> OnOutputChanged)
+	{
+		METASOUND_TRACE_CPUPROFILER_EVENT_SCOPE(FMetasoundOutputWatcher::Update);
+		
+		for (FMetaSoundOutput& Output : Outputs)
+		{
+			const FName TypeName = Output.GetTypeName();
+			check(!TypeName.IsNone());
+			
+			if (ensure(OutputTypeOperationMap.Contains(TypeName)))
+			{
+				if (OutputTypeOperationMap[TypeName]->Update(*this, Output))
+				{
+					OnOutputChanged(Name, Output);
+				}
+			}
+		}
+	}
+
+	void FMetasoundOutputWatcher::RegisterOutputTypeOperations(FName TypeName,
+		TUniquePtr<IOutputTypeOperations>&& OutputTypeOperations)
+	{
+		check(!OutputTypeOperationMap.Contains(TypeName));
+		check(OutputTypeOperations.IsValid());
+		OutputTypeOperationMap.Emplace(TypeName, MoveTemp(OutputTypeOperations));
+	}
+}
diff --git a/Engine/Plugins/Runtime/Metasound/Source/MetasoundEngine/Public/MetasoundGeneratorOutput.h b/Engine/Plugins/Runtime/Metasound/Source/MetasoundEngine/Public/MetasoundGeneratorOutput.h
deleted file mode 100644
diff --git a/Engine/Plugins/Runtime/Metasound/Source/MetasoundEngine/Public/MetasoundGeneratorOutputWatcher.h b/Engine/Plugins/Runtime/Metasound/Source/MetasoundEngine/Public/MetasoundGeneratorOutputWatcher.h
deleted file mode 100644
diff --git a/Engine/Plugins/Runtime/Metasound/Source/MetasoundEngine/Public/MetasoundOutput.h b/Engine/Plugins/Runtime/Metasound/Source/MetasoundEngine/Public/MetasoundOutput.h
new file mode 100644
--- /dev/null
+++ b/Engine/Plugins/Runtime/Metasound/Source/MetasoundEngine/Public/MetasoundOutput.h
@@ -0,0 +1,131 @@
+// Copyright Epic Games, Inc. All Rights Reserved.
+
+#pragma once
+
+#include "MetasoundDataReference.h"
+#include "SoundGeneratorOutput.h"
+#include "Kismet/BlueprintFunctionLibrary.h"
+
+#include "MetasoundOutput.generated.h"
+
+/**
+ * Exposes the value of an output on a FMetasoundGenerator.
+ */
+USTRUCT(BlueprintType)
+struct METASOUNDENGINE_API FMetaSoundOutput : public FSoundGeneratorOutput
+{
+	GENERATED_BODY()
+
+	FMetaSoundOutput() = default;
+
+	FMetaSoundOutput(const FMetaSoundOutput& Other);
+
+	FMetaSoundOutput& operator=(const FMetaSoundOutput& Other);
+
+	/**
+	 * Has this output been initialized?
+	 *
+	 * @returns true if the output has been initialized, false otherwise
+	 */
+	bool IsValid() const;
+
+	/**
+	 * Get the type name of the output
+	 *
+	 * @returns The type name, or none if invalid
+	 */
+	FName GetTypeName() const;
+	
+	/**
+	 * Initialize the output with some data.
+	 *
+	 * @tparam DataType - The type to use for initialization
+	 * @param InitialValue - The initial value to use
+	 */
+	template<typename DataType>
+	void Init(const DataType& InitialValue)
+	{
+		DataReference = MakeUnique<Metasound::FAnyDataReference>(Metasound::TDataWriteReference<DataType>::CreateNew());
+		*DataReference->GetDataWriteReference<DataType>() = InitialValue;
+	}
+
+	/**
+	 * Check if this output is of the given type
+	 *
+	 * @tparam DataType - The data type to check
+	 * @returns true if it's the given type, false otherwise
+	 */
+	template<typename DataType>
+	bool IsType() const
+	{
+		return IsValid() && Metasound::IsDataReferenceOfType<DataType>(*DataReference);
+	}
+	
+	/**
+	 * Get the value, for copyable registered Metasound data types
+	 *
+	 * @tparam DataType - The expected data type of the output
+	 * @param Value - The value to use
+	 * @returns true if the value was retrieved, false otherwise
+	 */
+	template<typename DataType>
+	bool Get(DataType& Value) const
+	{
+		if (!IsType<DataType>())
+		{
+			return false;
+		}
+
+		Value = *DataReference->GetDataReadReference<DataType>();
+		return true;
+	}
+	
+	/**
+	 * Set the value, for copyable registered Metasound data types
+	 *
+	 * @tparam DataType - The expected data type of the output
+	 * @param Value - The value to use
+	 * @returns true if the value was set, false otherwise
+	 */
+	template<typename DataType>
+	bool Set(const DataType& Value)
+	{
+		if (!IsType<DataType>())
+		{
+			return false;
+		}
+
+		*DataReference->GetDataWriteReference<DataType>() = Value;
+		return true;
+	}
+
+private:
+	TUniquePtr<Metasound::FAnyDataReference> DataReference;
+};
+
+/**
+ * Blueprint support for core types. If you want to support more core types, add them here.
+ * If you want to support types introduced in other plugins, create a blueprint library in that plugin.
+ */
+UCLASS()
+class UMetasoundOutputBlueprintAccess final : public UBlueprintFunctionLibrary
+{
+	GENERATED_BODY()
+
+	UFUNCTION(BlueprintCallable, Category="MetaSoundOutput")
+	static bool IsFloat(const FMetaSoundOutput& Output);
+	UFUNCTION(BlueprintCallable, Category="MetaSoundOutput")
+	static float GetFloat(const FMetaSoundOutput& Output, bool& Success);
+	UFUNCTION(BlueprintCallable, Category="MetaSoundOutput")
+	static bool IsInt32(const FMetaSoundOutput& Output);
+	UFUNCTION(BlueprintCallable, Category="MetaSoundOutput")
+	static int32 GetInt32(const FMetaSoundOutput& Output, bool& Success);
+	UFUNCTION(BlueprintCallable, Category="MetaSoundOutput")
+	static bool IsBool(const FMetaSoundOutput& Output);
+	UFUNCTION(BlueprintCallable, Category="MetaSoundOutput")
+	static bool GetBool(const FMetaSoundOutput& Output, bool& Success);
+	UFUNCTION(BlueprintCallable, Category="MetaSoundOutput")
+	static bool IsString(const FMetaSoundOutput& Output);
+	UFUNCTION(BlueprintCallable, Category="MetaSoundOutput")
+	static FString GetString(const FMetaSoundOutput& Output, bool& Success);
+};
diff --git a/Engine/Plugins/Runtime/Metasound/Source/MetasoundEngine/Public/MetasoundOutputSubsystem.h b/Engine/Plugins/Runtime/Metasound/Source/MetasoundEngine/Public/MetasoundOutputSubsystem.h
new file mode 100644
--- /dev/null
+++ b/Engine/Plugins/Runtime/Metasound/Source/MetasoundEngine/Public/MetasoundOutputSubsystem.h
@@ -0,0 +1,98 @@
+// Copyright Epic Games, Inc. All Rights Reserved.
+
+#pragma once
+
+#include "MetasoundOutput.h"
+#include "MetasoundOutputWatcher.h"
+#include "Containers/Queue.h"
+#include "Containers/Map.h"
+#include "Subsystems/WorldSubsystem.h"
+
+#include "MetasoundOutputSubsystem.generated.h"
+
+namespace Metasound
+{
+	namespace Frontend
+	{
+		class FAnalyzerAddress;
+	}
+
+	class FMetasoundGenerator;
+}
+
+class UAudioComponent;
+class UMetaSoundSource;
+
+DECLARE_DYNAMIC_DELEGATE_OneParam(FOnMetasoundOutputValueChanged, const FMetaSoundOutput&, Output);
+
+/**
+ * Provides access to a playing Metasound generator's outputs
+ */
+UCLASS()
+class UMetaSoundOutputSubsystem : public UTickableWorldSubsystem
+{
+	GENERATED_BODY()
+
+public:
+	virtual ~UMetaSoundOutputSubsystem() override = default;
+	
+	/**
+	 * Watch an output on a Metasound playing on a given audio component.
+	 *
+	 * @param AudioComponent - The audio component
+	 * @param OutputName - The user-specified name of the output in the Metasound
+	 * @param OnOutputValueChanged - The event to fire when the output's value changes
+	 * @param AnalyzerName - (optional) The name of the analyzer to use on the output, defaults to a passthrough
+	 * @param AnalyzerOutputName - (optional) The name of the output on the analyzer to watch, defaults to the passthrough output
+	 * @returns true if the watch setup succeeded, false otherwise
+	 */
+	UFUNCTION(BlueprintCallable, Category="MetaSoundOutput", meta=(AdvancedDisplay = "3"))
+	bool WatchOutput(
+		UAudioComponent* AudioComponent,
+		FName OutputName,
+		const FOnMetasoundOutputValueChanged& OnOutputValueChanged,
+		FName AnalyzerName = NAME_None,
+		FName AnalyzerOutputName = NAME_None);
+	
+	/** Begin UTickableWorldSubsystem */
+	virtual bool IsTickable() const override;
+	virtual void Tick(float DeltaTime) override;
+	virtual TStatId GetStatId() const override;
+	/** End UTickableWorldSubsystem */
+
+	/**
+	 * Map a type name to a passthrough analyzer name to use as a default for UMetasoundOutputSubsystem::WatchOutput()
+	 *
+	 * @param TypeName - The type name returned from GetMetasoundDataTypeName()
+	 * @param AnalyzerName - The name of the analyzer to use
+	 */
+	static void RegisterPassthroughAnalyzerForType(FName TypeName, FName AnalyzerName);
+
+private:
+	static TMap<FName, FName> PassthroughAnalyzers;
+
+	DECLARE_DYNAMIC_MULTICAST_DELEGATE_OneParam(FOnOutputValueChangedMulticast, const FMetaSoundOutput&, Output);
+	
+	struct FGeneratorInfo
+	{
+		TWeakObjectPtr<UAudioComponent> AudioComponent;
+		TWeakObjectPtr<UMetaSoundSource> Source;
+		TWeakPtr<Metasound::FMetasoundGenerator, ESPMode::ThreadSafe> Generator;
+		FDelegateHandle OnCreatedHandle;
+		FDelegateHandle OnDestroyedHandle;
+		TQueue<Metasound::Frontend::FAnalyzerAddress> AnalyzersToCreate;
+		TMap<FName, TMap<FName, FOnOutputValueChangedMulticast>> OutputChangedDelegates;
+		TArray<Metasound::Private::FMetasoundOutputWatcher> OutputWatchers;
+		TFunction<void(FName, const FMetaSoundOutput&)> HandleOutputChanged;
+		bool IsValid() const;
+	};
+
+	FGeneratorInfo* FindOrAddGeneratorInfo(UAudioComponent* AudioComponent);
+	void OnGeneratorCreated(uint64 InAudioComponentId, TSharedPtr<Metasound::FMetasoundGenerator> Generator);
+	void OnGeneratorDestroyed(uint64 InAudioComponentId, TSharedPtr<Metasound::FMetasoundGenerator> Generator);
+	void CreateAnalyzerAndWatcher(
+		const TSharedPtr<Metasound::FMetasoundGenerator>& Generator,
+		Metasound::Frontend::FAnalyzerAddress&& AnalyzerAddress);
+
+	TArray<FGeneratorInfo> TrackedGenerators;
+};
diff --git a/Engine/Plugins/Runtime/Metasound/Source/MetasoundEngine/Public/MetasoundOutputWatcher.h b/Engine/Plugins/Runtime/Metasound/Source/MetasoundEngine/Public/MetasoundOutputWatcher.h
new file mode 100644
--- /dev/null
+++ b/Engine/Plugins/Runtime/Metasound/Source/MetasoundEngine/Public/MetasoundOutputWatcher.h
@@ -0,0 +1,114 @@
+// Copyright Epic Games, Inc. All Rights Reserved.
+
+#pragma once
+
+#include "MetasoundOutput.h"
+#include "Analysis/MetasoundFrontendAnalyzerView.h"
+#include "Containers/Array.h"
+#include "Containers/Map.h"
+#include "Templates/Function.h"
+
+namespace Metasound::Private
+{
+	/**
+	 * Provides a mechanism to check an output and notify subscribers when its value(s) change
+	 */
+	class METASOUNDENGINE_API FMetasoundOutputWatcher
+	{
+	public:
+		FMetasoundOutputWatcher(Frontend::FAnalyzerAddress&& Address, const FOperatorSettings& OperatorSettings);
+
+		const FName Name;
+		
+		/**
+		 * Update the watched outputs
+		 */
+		void Update(TFunctionRef<void(FName, const FMetaSoundOutput&)> OnOutputChanged);
+
+		using FOutputInitFn = TFunction<void(FMetaSoundOutput&)>;
+		using FOutputUpdateFn = TFunction<bool(FMetasoundOutputWatcher&, FMetaSoundOutput&)>;
+
+		/**
+		 * Stores the typed operations for a specific supported type
+		 */
+		class IOutputTypeOperations
+		{
+		public:
+			virtual ~IOutputTypeOperations() = default;
+			
+			virtual void Init(FMetaSoundOutput& Output) = 0;
+			virtual bool Update(FMetasoundOutputWatcher& Watcher, FMetaSoundOutput& Output) = 0;
+		};
+
+		template<typename DataType>
+		class TOutputTypeOperations final : public IOutputTypeOperations
+		{
+		public:
+			explicit TOutputTypeOperations(DataType InDefaultValue)
+				: DefaultValue(InDefaultValue)
+			{}
+			
+			const DataType DefaultValue;
+			
+			virtual void Init(FMetaSoundOutput& Output) override
+			{
+				Output.Init<DataType>(DefaultValue);
+			}
+			
+			virtual bool Update(FMetasoundOutputWatcher& Watcher, FMetaSoundOutput& Output) override
+			{
+				return Watcher.UpdateOutput<DataType>(Output);
+			}
+
+		private:
+		};
+
+		/**
+		 * Register operations for a given type.
+		 * NOTE: This is exposed to enable external registration of types.
+		 * You shouldn't need to use this directly.
+		 *
+		 * @param TypeName - The name of the type, usually returned from GetMetasoundDataTypeName()
+		 * @param OutputTypeOperations - The operations to use
+		 */
+		static void RegisterOutputTypeOperations(FName TypeName, TUniquePtr<IOutputTypeOperations>&& OutputTypeOperations);
+	
+	private:
+		template<typename DataType>
+		bool UpdateOutput(FMetaSoundOutput& Output)
+		{
+			// check the type
+			if (!Output.IsType<DataType>())
+			{
+				return false;
+			}
+
+			// try to get the latest value
+			DataType Value;
+
+			if (View.TryGetOutputData(Output.Name, Value))
+			{
+				Output.Set(Value);
+				return true;
+			}
+			
+			return false;
+		}
+		
+		static TMap<FName, TUniquePtr<IOutputTypeOperations>> OutputTypeOperationMap;
+		Frontend::FMetasoundAnalyzerView View;
+		TArray<FMetaSoundOutput> Outputs;
+	};
+}
+
+
+/** Helper to register typed operations on outputs */
+#define METASOUND_PRIVATE_REGISTER_GENERATOR_OUTPUT_WATCHER_TYPE_OPERATIONS(TYPE, DEFAULT_VALUE) \
+	{ \
+		using FTypeOps = Metasound::Private::FMetasoundOutputWatcher::TOutputTypeOperations<TYPE>; \
+		TUniquePtr<FTypeOps> Ops = MakeUnique<FTypeOps>(DEFAULT_VALUE); \
+		Metasound::Private::FMetasoundOutputWatcher::RegisterOutputTypeOperations( \
+			Metasound::GetMetasoundDataTypeName<TYPE>(), \
+			MoveTemp(Ops) \
+			); \
+	} \
diff --git a/Engine/Plugins/Runtime/Metasound/Source/MetasoundEngineTest/Private/MetasoundGeneratorOutputTests.cpp b/Engine/Plugins/Runtime/Metasound/Source/MetasoundEngineTest/Private/MetasoundGeneratorOutputTests.cpp
deleted file mode 100644
diff --git a/Engine/Plugins/Runtime/Metasound/Source/MetasoundEngineTest/Private/MetasoundGeneratorOutputWatcherTests.cpp b/Engine/Plugins/Runtime/Metasound/Source/MetasoundEngineTest/Private/MetasoundGeneratorOutputWatcherTests.cpp
deleted file mode 100644
diff --git a/Engine/Plugins/Runtime/Metasound/Source/MetasoundEngineTest/Private/MetasoundOutputTests.cpp b/Engine/Plugins/Runtime/Metasound/Source/MetasoundEngineTest/Private/MetasoundOutputTests.cpp
new file mode 100644
--- /dev/null
+++ b/Engine/Plugins/Runtime/Metasound/Source/MetasoundEngineTest/Private/MetasoundOutputTests.cpp
@@ -0,0 +1,121 @@
+// Copyright Epic Games, Inc. All Rights Reserved.
+
+#include "MetasoundOutput.h"
+#include "MetasoundPrimitives.h"
+#include "Misc/AutomationTest.h"
+
+#if WITH_DEV_AUTOMATION_TESTS
+
+namespace Metasound::Test::GeneratorOutput
+{
+	IMPLEMENT_SIMPLE_AUTOMATION_TEST(
+	FMetasoundOutputIsTypeTest,
+	"Audio.Metasound.Output.IsType",
+	EAutomationTestFlags::EditorContext | EAutomationTestFlags::EngineFilter)
+	bool FMetasoundOutputIsTypeTest::RunTest(const FString&)
+	{
+		FMetaSoundOutput Output;
+		
+		Output.Init(float{ 0.0f });
+		UTEST_TRUE("float is float", Output.IsType<float>());
+		UTEST_FALSE("float is not int32", Output.IsType<int32>());
+		UTEST_FALSE("float is not bool", Output.IsType<bool>());
+		UTEST_FALSE("float is not FString", Output.IsType<FString>());
+
+		Output.Init(int32{ 0 });
+		UTEST_FALSE("int32 is not float", Output.IsType<float>());
+		UTEST_TRUE("int32 is int32", Output.IsType<int32>());
+		UTEST_FALSE("int32 is not bool", Output.IsType<bool>());
+		UTEST_FALSE("int32 is not FString", Output.IsType<FString>());
+		
+		Output.Init(bool{ false });
+		UTEST_FALSE("bool is not float", Output.IsType<float>());
+		UTEST_FALSE("bool is not int32", Output.IsType<int32>());
+		UTEST_TRUE("bool is bool", Output.IsType<bool>());
+		UTEST_FALSE("bool is not FString", Output.IsType<FString>());
+		
+		Output.Init(FString{ "hi!" });
+		UTEST_FALSE("FString is not float", Output.IsType<float>());
+		UTEST_FALSE("FString is not int32", Output.IsType<int32>());
+		UTEST_FALSE("FString is not bool", Output.IsType<bool>());
+		UTEST_TRUE("FString is FString", Output.IsType<FString>());
+		
+		return true;
+	}
+
+	template<typename DataType>
+	bool RunOutputGetSetTest(FAutomationTestBase& Test, const DataType& InitialValue, const DataType& ExpectedValue)
+	{
+		FMetaSoundOutput Output;
+		Output.Init<DataType>(InitialValue);
+		DataType Value;
+		
+		if (!Test.TestTrue("Got value", Output.Get<DataType>(Value)))
+		{
+			return false;
+		}
+		if (!Test.TestEqual("Value equals initial", Value, InitialValue))
+		{
+			return false;
+		}
+		if (!Test.TestTrue("Set value", Output.Set<DataType>(ExpectedValue)))
+		{
+			return false;
+		}
+		if (!Test.TestTrue("Got value", Output.Get<DataType>(Value)))
+		{
+			return false;
+		}
+		return Test.TestEqual("Value equals expected", Value, ExpectedValue);
+	}
+
+	IMPLEMENT_SIMPLE_AUTOMATION_TEST(
+	FMetasoundOutputGetSetFloatTest,
+	"Audio.Metasound.Output.GetSet.Float",
+	EAutomationTestFlags::EditorContext | EAutomationTestFlags::EngineFilter)
+	bool FMetasoundOutputGetSetFloatTest::RunTest(const FString&)
+	{
+		using DataType = float;
+		constexpr DataType InitialValue = 321.654f;
+		constexpr DataType ExpectedValue = 123.456f;
+		return RunOutputGetSetTest<DataType>(*this, InitialValue, ExpectedValue);
+	}
+
+	IMPLEMENT_SIMPLE_AUTOMATION_TEST(
+	FMetasoundOutputGetSetInt32Test,
+	"Audio.Metasound.Output.GetSet.Int32",
+	EAutomationTestFlags::EditorContext | EAutomationTestFlags::EngineFilter)
+	bool FMetasoundOutputGetSetInt32Test::RunTest(const FString&)
+	{
+		using DataType = int32;
+		constexpr DataType InitialValue = 321654;
+		constexpr DataType ExpectedValue = 123456;
+		return RunOutputGetSetTest<DataType>(*this, InitialValue, ExpectedValue);
+	}
+
+	IMPLEMENT_SIMPLE_AUTOMATION_TEST(
+	FMetasoundOutputGetSetBoolTest,
+	"Audio.Metasound.Output.GetSet.Bool",
+	EAutomationTestFlags::EditorContext | EAutomationTestFlags::EngineFilter)
+	bool FMetasoundOutputGetSetBoolTest::RunTest(const FString&)
+	{
+		using DataType = bool;
+		constexpr DataType InitialValue = false;
+		constexpr DataType ExpectedValue = true;
+		return RunOutputGetSetTest<DataType>(*this, InitialValue, ExpectedValue);
+	}
+
+	IMPLEMENT_SIMPLE_AUTOMATION_TEST(
+	FMetasoundOutputGetSetStringTest,
+	"Audio.Metasound.Output.GetSet.String",
+	EAutomationTestFlags::EditorContext | EAutomationTestFlags::EngineFilter)
+	bool FMetasoundOutputGetSetStringTest::RunTest(const FString&)
+	{
+		using FDataType = FString;
+		const FDataType InitialValue{ "hello" };
+		const FDataType ExpectedValue{ "goodbye" };
+		return RunOutputGetSetTest<FDataType>(*this, InitialValue, ExpectedValue);
+	}
+}
+
+#endif
\ No newline at end of file
diff --git a/Engine/Plugins/Runtime/Metasound/Source/MetasoundEngineTest/Private/MetasoundOutputWatcherTests.cpp b/Engine/Plugins/Runtime/Metasound/Source/MetasoundEngineTest/Private/MetasoundOutputWatcherTests.cpp
new file mode 100644
--- /dev/null
+++ b/Engine/Plugins/Runtime/Metasound/Source/MetasoundEngineTest/Private/MetasoundOutputWatcherTests.cpp
@@ -0,0 +1,175 @@
+// Copyright Epic Games, Inc. All Rights Reserved.
+
+#include "MetasoundGenerator.h"
+#include "MetasoundOutputWatcher.h"
+#include "NodeTestGraphBuilder.h"
+#include "Misc/AutomationTest.h"
+
+#if WITH_DEV_AUTOMATION_TESTS
+
+namespace Metasound::Test::GeneratorOutputWatcher
+{
+	using GraphBuilder = FNodeTestGraphBuilder;
+	using namespace Frontend;
+	using namespace Private;
+	
+	template<typename DataType>
+	TUniquePtr<FMetasoundGenerator> BuildPassthroughGraph(
+		FAutomationTestBase& Test,
+		const FName& InputName,
+		const FName& OutputName,
+		const FSampleRate SampleRate,
+		const int32 NumSamplesPerBlock,
+		FGuid* OutputGuid = nullptr)
+	{
+		GraphBuilder Builder;
+		const FNodeHandle InputNode = Builder.AddInput(InputName, Metasound::GetMetasoundDataTypeName<DataType>());
+		const FNodeHandle OutputNode = Builder.AddOutput(OutputName, Metasound::GetMetasoundDataTypeName<DataType>());
+		const FOutputHandle OutputToConnect = InputNode->GetOutputWithVertexName(InputName);
+		const FInputHandle InputToConnect = OutputNode->GetInputWithVertexName(OutputName);
+
+		if (!Test.TestTrue("Connected input to output", InputToConnect->Connect(*OutputToConnect)))
+		{
+			return nullptr;
+		}
+
+		if (nullptr != OutputGuid)
+		{
+			*OutputGuid = OutputNode->GetID();
+		}
+
+		// have to add an audio output for the generator to render
+		Builder.AddOutput("Audio", Metasound::GetMetasoundDataTypeName<FAudioBuffer>());
+		
+		return Builder.BuildGenerator(SampleRate, NumSamplesPerBlock);
+	}
+
+	template<typename DataType>
+	bool RunSimpleOutputWatcherTest(FAutomationTestBase& Test, DataType ExpectedValue, const FName AnalyzerName, const FName ValueOutputName)
+	{
+		const FName InputName = "MyInput";
+		const FName OutputName = "MyOutput";
+		constexpr FSampleRate SampleRate = 48000;
+		constexpr int32 NumSamplesPerBlock = 128;
+
+		// Build a passthrough graph
+		FGuid OutputNodeId;
+		const TUniquePtr<FMetasoundGenerator> Generator =
+			BuildPassthroughGraph<DataType>(Test, InputName, OutputName, SampleRate, NumSamplesPerBlock, &OutputNodeId);
+		if (!Test.TestNotNull("Generator built", Generator.Get()))
+		{
+			return false;
+		}
+
+		// Add an analyzer to the output
+		FAnalyzerAddress AnalyzerAddress;
+		AnalyzerAddress.DataType = GetMetasoundDataTypeName<DataType>();
+		AnalyzerAddress.InstanceID = 1234;
+		AnalyzerAddress.OutputName = OutputName;
+		AnalyzerAddress.AnalyzerName = AnalyzerName;
+		AnalyzerAddress.AnalyzerInstanceID = FGuid::NewGuid();
+		AnalyzerAddress.NodeID = OutputNodeId;
+		Generator->AddOutputVertexAnalyzer(AnalyzerAddress);
+
+		// Add an output watcher
+		FMetasoundOutputWatcher Watcher{ MoveTemp(AnalyzerAddress), Generator->OperatorSettings };
+
+		// listen for changes to the output
+		DataType WatchedValue;
+		bool ReceivedValue = false;
+		const TFunction<void(FName, const FMetaSoundOutput&)> OnChanged =
+			[&WatchedValue, &ReceivedValue](FName, const FMetaSoundOutput& Output)
+			{
+				ReceivedValue = Output.Get<DataType>(WatchedValue);
+			};
+		
+		// Set the input value
+		Generator->SetInputValue(InputName, ExpectedValue);
+
+		// Render a block and update the watcher
+		{
+			TArray<float> Buffer;
+			Buffer.Reserve(NumSamplesPerBlock);
+			if (!Test.TestEqual(
+				"Generated the right number of samples.",
+				Generator->OnGenerateAudio(Buffer.GetData(), NumSamplesPerBlock),
+				NumSamplesPerBlock))
+			{
+				return false;
+			}
+			Watcher.Update(OnChanged);
+		}
+
+		// Check the output
+		if (!Test.TestTrue("Received a value", ReceivedValue)
+			|| !Test.TestEqual("Received the expected value", WatchedValue, ExpectedValue))
+		{
+			return false;
+		}
+
+		// don't set the input and expect no callbacks
+		ReceivedValue = false;
+		
+		// Render a block and update the watcher
+		{
+			TArray<float> Buffer;
+			Buffer.Reserve(NumSamplesPerBlock);
+			if (!Test.TestEqual(
+				"Generated the right number of samples.",
+				Generator->OnGenerateAudio(Buffer.GetData(), NumSamplesPerBlock),
+				NumSamplesPerBlock))
+			{
+				return false;
+			}
+			Watcher.Update(OnChanged);
+		}
+
+		return Test.TestFalse("No more updates", ReceivedValue);
+	}
+
+	IMPLEMENT_SIMPLE_AUTOMATION_TEST(
+	FMetasoundOutputWatcherForwardValueFloatTest,
+	"Audio.Metasound.OutputWatcher.ForwardValue.Float",
+	EAutomationTestFlags::EditorContext | EAutomationTestFlags::EngineFilter)
+	bool FMetasoundOutputWatcherForwardValueFloatTest::RunTest(const FString&)
+	{
+		using DataType = float;
+		constexpr DataType ExpectedValue = 123.456f;
+		return RunSimpleOutputWatcherTest<DataType>(*this, ExpectedValue, "UE.Forward.Float", "Value");
+	}
+	
+	IMPLEMENT_SIMPLE_AUTOMATION_TEST(
+	FMetasoundOutputWatcherForwardValueIntTest,
+	"Audio.Metasound.OutputWatcher.ForwardValue.Int32",
+	EAutomationTestFlags::EditorContext | EAutomationTestFlags::EngineFilter)
+	bool FMetasoundOutputWatcherForwardValueIntTest::RunTest(const FString&)
+	{
+		using DataType = int32;
+		constexpr DataType ExpectedValue = 123456;
+		return RunSimpleOutputWatcherTest<DataType>(*this, ExpectedValue, "UE.Forward.Int32", "Value");
+	}
+
+	IMPLEMENT_SIMPLE_AUTOMATION_TEST(
+	FMetasoundOutputWatcherForwardValueBoolTest,
+	"Audio.Metasound.OutputWatcher.ForwardValue.Bool",
+	EAutomationTestFlags::EditorContext | EAutomationTestFlags::EngineFilter)
+	bool FMetasoundOutputWatcherForwardValueBoolTest::RunTest(const FString&)
+	{
+		using DataType = bool;
+		constexpr DataType ExpectedValue = true;
+		return RunSimpleOutputWatcherTest<DataType>(*this, ExpectedValue, "UE.Forward.Bool", "Value");
+	}
+
+	IMPLEMENT_SIMPLE_AUTOMATION_TEST(
+	FMetasoundOutputWatcherForwardValueStringTest,
+	"Audio.Metasound.OutputWatcher.ForwardValue.String",
+	EAutomationTestFlags::EditorContext | EAutomationTestFlags::EngineFilter)
+	bool FMetasoundOutputWatcherForwardValueStringTest::RunTest(const FString&)
+	{
+		using FDataType = FString;
+		const FDataType ExpectedValue{ "you're awesome" };
+		return RunSimpleOutputWatcherTest<FDataType>(*this, ExpectedValue, "UE.Forward.String", "Value");
+	}
+}
+
+#endif
\ No newline at end of file
