- Fix for builder subsystem not able to reuse audition sound properly due to the sound not being unregistered from the MetaSound Frontend Class Registry before mutation and attempt to play/re-register. Default to register when building and don't expose option to BP (referencing unregistered assets can lead to silent failures when attempting playback later)
- Fix for const correctness causing BP functions to not provide execution path even though they're set to BlueprintCallable :(

diff --git a/Engine/Plugins/Runtime/Metasound/Source/MetasoundEditor/Private/MetasoundEditorSubsystem.cpp b/Engine/Plugins/Runtime/Metasound/Source/MetasoundEditor/Private/MetasoundEditorSubsystem.cpp
--- a/Engine/Plugins/Runtime/Metasound/Source/MetasoundEditor/Private/MetasoundEditorSubsystem.cpp
+++ b/Engine/Plugins/Runtime/Metasound/Source/MetasoundEditor/Private/MetasoundEditorSubsystem.cpp
@@ -27,7 +27,7 @@ TScriptInterface<IMetaSoundDocumentInterface> UMetaSoundEditorSubsystem::BuildTo
 			InBuilder->SetAuthor(Author);
 
 			FMetaSoundBuilderOptions BuilderOptions { FName(*AssetName) };
-			BuilderOptions.ExistingAsset = NewMetaSound;
+			BuilderOptions.ExistingMetaSound = NewMetaSound;
 			TScriptInterface<IMetaSoundDocumentInterface> DocInterface = InBuilder->Build(Parent, BuilderOptions);
 
 			FMetasoundAssetBase* MetaSoundAsset = IMetasoundUObjectRegistry::Get().GetObjectAsAssetBase(NewMetaSound);
diff --git a/Engine/Plugins/Runtime/Metasound/Source/MetasoundEngine/Private/MetasoundBuilderSubsystem.cpp b/Engine/Plugins/Runtime/Metasound/Source/MetasoundEngine/Private/MetasoundBuilderSubsystem.cpp
--- a/Engine/Plugins/Runtime/Metasound/Source/MetasoundEngine/Private/MetasoundBuilderSubsystem.cpp
+++ b/Engine/Plugins/Runtime/Metasound/Source/MetasoundEngine/Private/MetasoundBuilderSubsystem.cpp
@@ -9,6 +9,7 @@
 #include "Interfaces/MetasoundFrontendSourceInterface.h"
 #include "Metasound.h"
 #include "MetasoundFrontendDocument.h"
+#include "MetasoundFrontendRegistries.h"
 #include "MetasoundFrontendSearchEngine.h"
 #include "MetasoundFrontendTransform.h"
 #include "MetasoundGeneratorHandle.h"
@@ -298,7 +299,7 @@ void UMetaSoundBuilderBase::DisconnectNodesByInterfaceBindings(const FMetaSoundN
 	OutResult = bEdgesRemoved ? EMetaSoundBuilderResult::Succeeded : EMetaSoundBuilderResult::Failed;
 }
 
-FMetaSoundBuilderNodeInputHandle UMetaSoundBuilderBase::FindNodeInputByName(const FMetaSoundNodeHandle& NodeHandle, FName InputName, EMetaSoundBuilderResult& OutResult) const
+FMetaSoundBuilderNodeInputHandle UMetaSoundBuilderBase::FindNodeInputByName(const FMetaSoundNodeHandle& NodeHandle, FName InputName, EMetaSoundBuilderResult& OutResult)
 {
 	if (const FMetasoundFrontendNode* Node = Builder.FindNode(NodeHandle.NodeID))
 	{
@@ -311,20 +312,30 @@ FMetaSoundBuilderNodeInputHandle UMetaSoundBuilderBase::FindNodeInputByName(cons
 			return FMetaSoundBuilderNodeInputHandle { Node->GetID(), Input->VertexID };
 		}
 
-		UE_LOG(LogMetaSound, Display, TEXT("Failed to find node input by name with builder '%s'. Node of class '%s' has no input named '%s'"), *GetName(), *InputName.ToString());
+		FString NodeClassName = TEXT("N/A");
+		if (const FMetasoundFrontendClass* Class = Builder.FindDependency(Node->ClassID))
+		{
+			NodeClassName = Class->Metadata.GetClassName().ToString();
+		}
+
+		UE_LOG(LogMetaSound, Display, TEXT("Builder '%s' failed to find node input '%s': Node class '%s' contains no such input"), *GetName(), *InputName.ToString(), *NodeClassName);
 	}
+	else
+	{
+		UE_LOG(LogMetaSound, Display, TEXT("Builder '%s' failed to find node input '%s': Node with ID '%s' not found"), *GetName(), *InputName.ToString(), *NodeHandle.NodeID.ToString());
+	}
+
 
-	UE_LOG(LogMetaSound, Display, TEXT("Failed to find node input by name with builder '%s'. Node of with ID '%s' not found"), *GetName(), *NodeHandle.NodeID.ToString());
 	OutResult = EMetaSoundBuilderResult::Failed;
 	return { };
 }
 
-TArray<FMetaSoundBuilderNodeInputHandle> UMetaSoundBuilderBase::FindNodeInputs(const FMetaSoundNodeHandle& NodeHandle, EMetaSoundBuilderResult& OutResult) const
+TArray<FMetaSoundBuilderNodeInputHandle> UMetaSoundBuilderBase::FindNodeInputs(const FMetaSoundNodeHandle& NodeHandle, EMetaSoundBuilderResult& OutResult)
 {
 	return FindNodeInputsByDataType(NodeHandle, OutResult, { });
 }
 
-TArray<FMetaSoundBuilderNodeInputHandle> UMetaSoundBuilderBase::FindNodeInputsByDataType(const FMetaSoundNodeHandle& NodeHandle, EMetaSoundBuilderResult& OutResult, FName DataType) const
+TArray<FMetaSoundBuilderNodeInputHandle> UMetaSoundBuilderBase::FindNodeInputsByDataType(const FMetaSoundNodeHandle& NodeHandle, EMetaSoundBuilderResult& OutResult, FName DataType)
 {
 	TArray<FMetaSoundBuilderNodeInputHandle> FoundVertices;
 	if (Builder.ContainsNode(NodeHandle.NodeID))
@@ -345,7 +356,7 @@ TArray<FMetaSoundBuilderNodeInputHandle> UMetaSoundBuilderBase::FindNodeInputsBy
 	return FoundVertices;
 }
 
-FMetaSoundBuilderNodeOutputHandle UMetaSoundBuilderBase::FindNodeOutputByName(const FMetaSoundNodeHandle& NodeHandle, FName OutputName, EMetaSoundBuilderResult& OutResult) const
+FMetaSoundBuilderNodeOutputHandle UMetaSoundBuilderBase::FindNodeOutputByName(const FMetaSoundNodeHandle& NodeHandle, FName OutputName, EMetaSoundBuilderResult& OutResult)
 {
 	if (const FMetasoundFrontendNode* Node = Builder.FindNode(NodeHandle.NodeID))
 	{
@@ -357,20 +368,30 @@ FMetaSoundBuilderNodeOutputHandle UMetaSoundBuilderBase::FindNodeOutputByName(co
 			OutResult = EMetaSoundBuilderResult::Succeeded;
 			return FMetaSoundBuilderNodeOutputHandle { Node->GetID(), Output->VertexID };
 		}
-		UE_LOG(LogMetaSound, Display, TEXT("Failed to find node output by name with builder '%s'. Node of class '%s' has no output named '%s'"), *GetName(), *OutputName.ToString());
+
+		FString NodeClassName = TEXT("N/A");
+		if (const FMetasoundFrontendClass* Class = Builder.FindDependency(Node->ClassID))
+		{
+			NodeClassName = Class->Metadata.GetClassName().ToString();
+		}
+
+		UE_LOG(LogMetaSound, Display, TEXT("Builder '%s' failed to find node output '%s': Node class '%s' contains no such output"), *GetName(), *OutputName.ToString(), *NodeClassName);
+	}
+	else
+	{
+		UE_LOG(LogMetaSound, Display, TEXT("Builder '%s' failed to find node output '%s': Node with ID '%s' not found"), *GetName(), *OutputName.ToString(), *NodeHandle.NodeID.ToString());
 	}
 
-	UE_LOG(LogMetaSound, Display, TEXT("Failed to find node output by name with builder '%s'. Node of with ID '%s' not found"), *GetName(), *NodeHandle.NodeID.ToString());
 	OutResult = EMetaSoundBuilderResult::Failed;
 	return { };
 }
 
-TArray<FMetaSoundBuilderNodeOutputHandle> UMetaSoundBuilderBase::FindNodeOutputs(const FMetaSoundNodeHandle& NodeHandle, EMetaSoundBuilderResult& OutResult) const
+TArray<FMetaSoundBuilderNodeOutputHandle> UMetaSoundBuilderBase::FindNodeOutputs(const FMetaSoundNodeHandle& NodeHandle, EMetaSoundBuilderResult& OutResult)
 {
 	return FindNodeOutputsByDataType(NodeHandle, OutResult, { });
 }
 
-TArray<FMetaSoundBuilderNodeOutputHandle> UMetaSoundBuilderBase::FindNodeOutputsByDataType(const FMetaSoundNodeHandle& NodeHandle, EMetaSoundBuilderResult& OutResult, FName DataType) const
+TArray<FMetaSoundBuilderNodeOutputHandle> UMetaSoundBuilderBase::FindNodeOutputsByDataType(const FMetaSoundNodeHandle& NodeHandle, EMetaSoundBuilderResult& OutResult, FName DataType)
 {
 	TArray<FMetaSoundBuilderNodeOutputHandle> FoundVertices;
 	if (Builder.ContainsNode(NodeHandle.NodeID))
@@ -391,7 +412,7 @@ TArray<FMetaSoundBuilderNodeOutputHandle> UMetaSoundBuilderBase::FindNodeOutputs
 	return FoundVertices;
 }
 
-TArray<FMetaSoundNodeHandle> UMetaSoundBuilderBase::FindInterfaceInputNodes(FName InterfaceName, EMetaSoundBuilderResult& OutResult) const
+TArray<FMetaSoundNodeHandle> UMetaSoundBuilderBase::FindInterfaceInputNodes(FName InterfaceName, EMetaSoundBuilderResult& OutResult)
 {
 	TArray<FMetaSoundNodeHandle> NodeHandles;
 
@@ -414,7 +435,7 @@ TArray<FMetaSoundNodeHandle> UMetaSoundBuilderBase::FindInterfaceInputNodes(FNam
 	return NodeHandles;
 }
 
-TArray<FMetaSoundNodeHandle> UMetaSoundBuilderBase::FindInterfaceOutputNodes(FName InterfaceName, EMetaSoundBuilderResult& OutResult) const
+TArray<FMetaSoundNodeHandle> UMetaSoundBuilderBase::FindInterfaceOutputNodes(FName InterfaceName, EMetaSoundBuilderResult& OutResult)
 {
 	TArray<FMetaSoundNodeHandle> NodeHandles;
 
@@ -436,7 +457,7 @@ TArray<FMetaSoundNodeHandle> UMetaSoundBuilderBase::FindInterfaceOutputNodes(FNa
 	return NodeHandles;
 }
 
-FMetaSoundNodeHandle UMetaSoundBuilderBase::FindGraphInputNode(FName InputName, EMetaSoundBuilderResult& OutResult) const
+FMetaSoundNodeHandle UMetaSoundBuilderBase::FindGraphInputNode(FName InputName, EMetaSoundBuilderResult& OutResult)
 {
 	if (const FMetasoundFrontendNode* GraphInputNode = Builder.FindGraphInputNode(InputName))
 	{
@@ -448,7 +469,7 @@ FMetaSoundNodeHandle UMetaSoundBuilderBase::FindGraphInputNode(FName InputName,
 	return { };
 }
 
-FMetaSoundNodeHandle UMetaSoundBuilderBase::FindGraphOutputNode(FName OutputName, EMetaSoundBuilderResult& OutResult) const
+FMetaSoundNodeHandle UMetaSoundBuilderBase::FindGraphOutputNode(FName OutputName, EMetaSoundBuilderResult& OutResult)
 {
 	if (const FMetasoundFrontendNode* GraphOutputNode = Builder.FindGraphOutputNode(OutputName))
 	{
@@ -524,7 +545,7 @@ void UMetaSoundBuilderBase::SetNodeInputDefault(const FMetaSoundBuilderNodeInput
 	OutResult = bInputDefaultSet ? EMetaSoundBuilderResult::Succeeded : EMetaSoundBuilderResult::Failed;
 }
 
-FMetaSoundNodeHandle UMetaSoundBuilderBase::FindNodeInputParent(const FMetaSoundBuilderNodeInputHandle& InputHandle, EMetaSoundBuilderResult& OutResult) const
+FMetaSoundNodeHandle UMetaSoundBuilderBase::FindNodeInputParent(const FMetaSoundBuilderNodeInputHandle& InputHandle, EMetaSoundBuilderResult& OutResult)
 {
 	if (Builder.ContainsNode(InputHandle.NodeID))
 	{
@@ -536,7 +557,7 @@ FMetaSoundNodeHandle UMetaSoundBuilderBase::FindNodeInputParent(const FMetaSound
 	return { };
 }
 
-FMetaSoundNodeHandle UMetaSoundBuilderBase::FindNodeOutputParent(const FMetaSoundBuilderNodeOutputHandle& OutputHandle, EMetaSoundBuilderResult& OutResult) const
+FMetaSoundNodeHandle UMetaSoundBuilderBase::FindNodeOutputParent(const FMetaSoundBuilderNodeOutputHandle& OutputHandle, EMetaSoundBuilderResult& OutResult)
 {
 	if (Builder.ContainsNode(OutputHandle.NodeID))
 	{
@@ -548,7 +569,7 @@ FMetaSoundNodeHandle UMetaSoundBuilderBase::FindNodeOutputParent(const FMetaSoun
 	return { };
 }
 
-FMetasoundFrontendVersion UMetaSoundBuilderBase::FindNodeClassVersion(const FMetaSoundNodeHandle& NodeHandle, EMetaSoundBuilderResult& OutResult) const
+FMetasoundFrontendVersion UMetaSoundBuilderBase::FindNodeClassVersion(const FMetaSoundNodeHandle& NodeHandle, EMetaSoundBuilderResult& OutResult)
 {
 	if (const FMetasoundFrontendNode* Node = Builder.FindNode(NodeHandle.NodeID))
 	{
@@ -563,7 +584,7 @@ FMetasoundFrontendVersion UMetaSoundBuilderBase::FindNodeClassVersion(const FMet
 	return FMetasoundFrontendVersion::GetInvalid();
 }
 
-void UMetaSoundBuilderBase::GetNodeInputData(const FMetaSoundBuilderNodeInputHandle& InputHandle, FName& Name, FName& DataType, EMetaSoundBuilderResult& OutResult) const
+void UMetaSoundBuilderBase::GetNodeInputData(const FMetaSoundBuilderNodeInputHandle& InputHandle, FName& Name, FName& DataType, EMetaSoundBuilderResult& OutResult)
 {
 	if (const FMetasoundFrontendVertex* Vertex = Builder.FindNodeInput(InputHandle.NodeID, InputHandle.VertexID))
 	{
@@ -579,9 +600,9 @@ void UMetaSoundBuilderBase::GetNodeInputData(const FMetaSoundBuilderNodeInputHan
 	}
 }
 
-void UMetaSoundBuilderBase::GetNodeOutputData(const FMetaSoundBuilderNodeOutputHandle& InputHandle, FName& Name, FName& DataType, EMetaSoundBuilderResult& OutResult) const
+void UMetaSoundBuilderBase::GetNodeOutputData(const FMetaSoundBuilderNodeOutputHandle& OutputHandle, FName& Name, FName& DataType, EMetaSoundBuilderResult& OutResult)
 {
-	if (const FMetasoundFrontendVertex* Vertex = Builder.FindNodeInput(InputHandle.NodeID, InputHandle.VertexID))
+	if (const FMetasoundFrontendVertex* Vertex = Builder.FindNodeOutput(OutputHandle.NodeID, OutputHandle.VertexID))
 	{
 		Name = Vertex->Name;
 		DataType = Vertex->TypeName;
@@ -595,9 +616,9 @@ void UMetaSoundBuilderBase::GetNodeOutputData(const FMetaSoundBuilderNodeOutputH
 	}
 }
 
-TScriptInterface<IMetaSoundDocumentInterface> UMetaSoundPatchBuilder::Build(UObject* Parent, const FMetaSoundBuilderOptions& BuildOptions) const
+TScriptInterface<IMetaSoundDocumentInterface> UMetaSoundPatchBuilder::Build(UObject* Parent, const FMetaSoundBuilderOptions& InBuilderOptions) const
 {
-	return BuildInternal<UMetaSoundPatch>(Parent, BuildOptions);
+	return BuildInternal<UMetaSoundPatch>(Parent, InBuilderOptions);
 }
 
 UClass& UMetaSoundPatchBuilder::GetBuilderUClass() const
@@ -612,24 +633,23 @@ UClass& UMetaSoundSourceBuilder::GetBuilderUClass() const
 
 void UMetaSoundSourceBuilder::Audition(UObject* Parent, UAudioComponent* AudioComponent, FOnCreateAuditionGeneratorHandleDelegate CreateGenerator)
 {
-	METASOUND_TRACE_CPUPROFILER_EVENT_SCOPE(UMetaSoundSourceBuilder::Audition);
-
+	using namespace Metasound;
 	using namespace Metasound::Frontend;
 
+	METASOUND_TRACE_CPUPROFILER_EVENT_SCOPE(UMetaSoundSourceBuilder::Audition);
+
 	if (!AudioComponent)
 	{
 		UE_LOG(LogMetaSound, Error, TEXT("Failed to audition MetaSoundBuilder '%s': No AudioComponent supplied"), *GetFullName());
 		return;
 	}
 
-	// kevin.cavanaugh [3/13/2023] - Rebuilding the audition sound currently isn't working. 
-	// Forcing a rebuild each time for now.
-	/*if (AuditionSound.IsValid())
-	{
-		const FMetasoundFrontendDocument& Doc = GetConstBuilder().GetDocument();
-		AuditionSound->SetDocument(Doc);
-		AuditionSound->ConformObjectDataToInterfaces();
+	FMetaSoundBuilderOptions BuilderOptions;
+	BuilderOptions.bAddToRegistry = true;
 
+	if (AuditionSound.IsValid())
+	{
+		BuilderOptions.ExistingMetaSound = AuditionSound.Get();
 		if (USoundBase* Sound = AudioComponent->GetSound())
 		{
 			if (Sound != AuditionSound)
@@ -637,26 +657,18 @@ void UMetaSoundSourceBuilder::Audition(UObject* Parent, UAudioComponent* AudioCo
 				UE_LOG(LogMetaSound, Warning, TEXT("MetaSoundBuilder '%s' supplied AudioComponent with unlinked sound '%s'. Stopping sound and replacing with builder's audition sound."),
 					*GetFullName(),
 					*Sound->GetFullName());
-				AudioComponent->SetSound(AuditionSound.Get());
 			}
 		}
-		else
-		{
-			AudioComponent->SetSound(AuditionSound.Get());
-		}
 	}
-	else*/
+	else
 	{
-		FMetaSoundBuilderOptions BuildOptions;
-		BuildOptions.Name = MakeUniqueObjectName(nullptr, UMetaSoundSource::StaticClass(), FName(GetName() + TEXT("_Audition")));
-		BuildOptions.bAddToRegistry = false;
+		BuilderOptions.Name = MakeUniqueObjectName(nullptr, UMetaSoundSource::StaticClass(), FName(GetName() + TEXT("_Audition")));
+	}
 
-		TScriptInterface<IMetaSoundDocumentInterface> DocumentInterface = Build(Parent, BuildOptions);
-		AuditionSound = Cast<UMetaSoundSource>(DocumentInterface.GetObject());
-		if (AuditionSound.IsValid())
-		{
-			AudioComponent->SetSound(AuditionSound.Get());
-		}
+	AuditionSound = BuildInternal<UMetaSoundSource>(Parent, BuilderOptions);
+	if (AuditionSound.IsValid())
+	{
+		AudioComponent->SetSound(AuditionSound.Get());
 	}
 
 	if (CreateGenerator.IsBound())
@@ -669,9 +681,9 @@ void UMetaSoundSourceBuilder::Audition(UObject* Parent, UAudioComponent* AudioCo
 	AudioComponent->Play();
 }
 
-TScriptInterface<IMetaSoundDocumentInterface> UMetaSoundSourceBuilder::Build(UObject* Parent, const FMetaSoundBuilderOptions& BuildOptions) const
+TScriptInterface<IMetaSoundDocumentInterface> UMetaSoundSourceBuilder::Build(UObject* Parent, const FMetaSoundBuilderOptions& InBuilderOptions) const
 {
-	return BuildInternal<UMetaSoundSource>(Parent, BuildOptions);
+	return BuildInternal<UMetaSoundSource>(Parent, InBuilderOptions);
 }
 
 const Metasound::Engine::FOutputAudioFormatInfoPair* UMetaSoundSourceBuilder::FindOutputAudioFormatInfo() const
diff --git a/Engine/Plugins/Runtime/Metasound/Source/MetasoundEngine/Public/MetasoundBuilderSubsystem.h b/Engine/Plugins/Runtime/Metasound/Source/MetasoundEngine/Public/MetasoundBuilderSubsystem.h
--- a/Engine/Plugins/Runtime/Metasound/Source/MetasoundEngine/Public/MetasoundBuilderSubsystem.h
+++ b/Engine/Plugins/Runtime/Metasound/Source/MetasoundEngine/Public/MetasoundBuilderSubsystem.h
@@ -70,15 +70,19 @@ struct METASOUNDENGINE_API FMetaSoundBuilderOptions
 {
 	GENERATED_BODY()
 
+	// Name of generated object. If object already exists, used as the base name to ensure new object is unique.
 	UPROPERTY(EditAnywhere, BlueprintReadWrite, Category = "MetaSound|Builder")
 	FName Name;
 
 	// If true, adds MetaSound to node registry, making it available
 	// for reference by other dynamically created MetaSounds.
-	UPROPERTY(EditAnywhere, BlueprintReadWrite, Category = "MetaSound|Builder")
-	bool bAddToRegistry = false;
+	UPROPERTY()
+	bool bAddToRegistry = true;
 
-	TScriptInterface<IMetaSoundDocumentInterface> ExistingAsset;
+	// If set, builder overwrites the given MetaSound's document with the builder's copy
+	// (ignores the Name field above).
+	UPROPERTY()
+	TScriptInterface<IMetaSoundDocumentInterface> ExistingMetaSound;
 };
 
 UENUM(BlueprintType)
@@ -171,63 +175,63 @@ public:
 
 	// Returns graph input node by the given name if it exists, or an invalid handle if not found.
 	UFUNCTION(BlueprintCallable, Category = "Audio|MetaSound|Builder", meta = (ExpandEnumAsExecs = "OutResult"))
-	UPARAM(DisplayName = "Node Handle") FMetaSoundNodeHandle FindGraphInputNode(FName InputName, EMetaSoundBuilderResult& OutResult) const;
+	UPARAM(DisplayName = "Node Handle") FMetaSoundNodeHandle FindGraphInputNode(FName InputName, EMetaSoundBuilderResult& OutResult);
 
 	// Returns graph output node by the given name if it exists, or an invalid handle if not found.
 	UFUNCTION(BlueprintCallable, Category = "Audio|MetaSound|Builder", meta = (ExpandEnumAsExecs = "OutResult"))
-	UPARAM(DisplayName = "Node Handle") FMetaSoundNodeHandle FindGraphOutputNode(FName OutputName, EMetaSoundBuilderResult& OutResult) const;
+	UPARAM(DisplayName = "Node Handle") FMetaSoundNodeHandle FindGraphOutputNode(FName OutputName, EMetaSoundBuilderResult& OutResult);
 
 	// Returns node input by the given name if it exists, or an invalid handle if not found.
 	UFUNCTION(BlueprintCallable, Category = "Audio|MetaSound|Builder", meta = (ExpandEnumAsExecs = "OutResult"))
-	UPARAM(DisplayName = "Input Handle") FMetaSoundBuilderNodeInputHandle FindNodeInputByName(const FMetaSoundNodeHandle& NodeHandle, FName InputName, EMetaSoundBuilderResult& OutResult) const;
+	UPARAM(DisplayName = "Input Handle") FMetaSoundBuilderNodeInputHandle FindNodeInputByName(const FMetaSoundNodeHandle& NodeHandle, FName InputName, EMetaSoundBuilderResult& OutResult);
 
 	// Returns node output by the given name.
 	UFUNCTION(BlueprintCallable, Category = "Audio|MetaSound|Builder", meta = (ExpandEnumAsExecs = "OutResult"))
-	UPARAM(DisplayName = "Input  Handles") TArray<FMetaSoundBuilderNodeInputHandle> FindNodeInputs(const FMetaSoundNodeHandle& NodeHandle, EMetaSoundBuilderResult& OutResult) const;
+	UPARAM(DisplayName = "Input  Handles") TArray<FMetaSoundBuilderNodeInputHandle> FindNodeInputs(const FMetaSoundNodeHandle& NodeHandle, EMetaSoundBuilderResult& OutResult);
 
 	// Returns node inputs by the given DataType (ex. "Audio", "Trigger", "String", "Bool", "Float", "Int32", etc.).
 	UFUNCTION(BlueprintCallable, Category = "Audio|MetaSound|Builder", meta = (ExpandEnumAsExecs = "OutResult"))
-	UPARAM(DisplayName = "Input Handles") TArray<FMetaSoundBuilderNodeInputHandle> FindNodeInputsByDataType(const FMetaSoundNodeHandle& NodeHandle, EMetaSoundBuilderResult& OutResult, FName DataType) const;
+	UPARAM(DisplayName = "Input Handles") TArray<FMetaSoundBuilderNodeInputHandle> FindNodeInputsByDataType(const FMetaSoundNodeHandle& NodeHandle, EMetaSoundBuilderResult& OutResult, FName DataType);
 
 	// Returns node output by the given name.
 	UFUNCTION(BlueprintCallable, Category = "Audio|MetaSound|Builder", meta = (ExpandEnumAsExecs = "OutResult"))
-	UPARAM(DisplayName = "Output Handle") FMetaSoundBuilderNodeOutputHandle FindNodeOutputByName(const FMetaSoundNodeHandle& NodeHandle, FName OutputName, EMetaSoundBuilderResult& OutResult) const;
+	UPARAM(DisplayName = "Output Handle") FMetaSoundBuilderNodeOutputHandle FindNodeOutputByName(const FMetaSoundNodeHandle& NodeHandle, FName OutputName, EMetaSoundBuilderResult& OutResult);
 
 	// Returns all node outputs.
 	UFUNCTION(BlueprintCallable, Category = "Audio|MetaSound|Builder", meta = (ExpandEnumAsExecs = "OutResult"))
-	UPARAM(DisplayName = "Output Handles") TArray<FMetaSoundBuilderNodeOutputHandle> FindNodeOutputs(const FMetaSoundNodeHandle& NodeHandle, EMetaSoundBuilderResult& OutResult) const;
+	UPARAM(DisplayName = "Output Handles") TArray<FMetaSoundBuilderNodeOutputHandle> FindNodeOutputs(const FMetaSoundNodeHandle& NodeHandle, EMetaSoundBuilderResult& OutResult);
 
 	// Returns node outputs by the given DataType (ex. "Audio", "Trigger", "String", "Bool", "Float", "Int32", etc.).
 	UFUNCTION(BlueprintCallable, Category = "Audio|MetaSound|Builder", meta = (ExpandEnumAsExecs = "OutResult"))
-	UPARAM(DisplayName = "Output Handles") TArray<FMetaSoundBuilderNodeOutputHandle> FindNodeOutputsByDataType(const FMetaSoundNodeHandle& NodeHandle, EMetaSoundBuilderResult& OutResult, FName DataType) const;
+	UPARAM(DisplayName = "Output Handles") TArray<FMetaSoundBuilderNodeOutputHandle> FindNodeOutputsByDataType(const FMetaSoundNodeHandle& NodeHandle, EMetaSoundBuilderResult& OutResult, FName DataType);
 
 	// Returns input nodes associated with a given interface.
 	UFUNCTION(BlueprintCallable, Category = "Audio|MetaSound|Builder", meta = (ExpandEnumAsExecs = "OutResult"))
-	UPARAM(DisplayName = "Input Node Handles") TArray<FMetaSoundNodeHandle> FindInterfaceInputNodes(FName InterfaceName, EMetaSoundBuilderResult& OutResult) const;
+	UPARAM(DisplayName = "Input Node Handles") TArray<FMetaSoundNodeHandle> FindInterfaceInputNodes(FName InterfaceName, EMetaSoundBuilderResult& OutResult);
 
 	// Returns output nodes associated with a given interface.
 	UFUNCTION(BlueprintCallable, Category = "Audio|MetaSound|Builder", meta = (ExpandEnumAsExecs = "OutResult"))
-	UPARAM(DisplayName = "Output Node Handles") TArray<FMetaSoundNodeHandle> FindInterfaceOutputNodes(FName InterfaceName, EMetaSoundBuilderResult& OutResult) const;
+	UPARAM(DisplayName = "Output Node Handles") TArray<FMetaSoundNodeHandle> FindInterfaceOutputNodes(FName InterfaceName, EMetaSoundBuilderResult& OutResult);
 
 	// Returns input's parent node if the input is valid, otherwise returns invalid node handle.
 	UFUNCTION(BlueprintCallable, Category = "Audio|MetaSound|Builder", meta = (ExpandEnumAsExecs = "OutResult"))
-	UPARAM(DisplayName = "Node Handle") FMetaSoundNodeHandle FindNodeInputParent(const FMetaSoundBuilderNodeInputHandle& InputHandle, EMetaSoundBuilderResult& OutResult) const;
+	UPARAM(DisplayName = "Node Handle") FMetaSoundNodeHandle FindNodeInputParent(const FMetaSoundBuilderNodeInputHandle& InputHandle, EMetaSoundBuilderResult& OutResult);
 
 	// Returns output's parent node if the input is valid, otherwise returns invalid node handle.
 	UFUNCTION(BlueprintCallable, Category = "Audio|MetaSound|Builder", meta = (ExpandEnumAsExecs = "OutResult"))
-	UPARAM(DisplayName = "Node Handle") FMetaSoundNodeHandle FindNodeOutputParent(const FMetaSoundBuilderNodeOutputHandle& OutputHandle, EMetaSoundBuilderResult& OutResult) const;
+	UPARAM(DisplayName = "Node Handle") FMetaSoundNodeHandle FindNodeOutputParent(const FMetaSoundBuilderNodeOutputHandle& OutputHandle, EMetaSoundBuilderResult& OutResult);
 
 	// Returns output's parent node if the input is valid, otherwise returns invalid node handle.
 	UFUNCTION(BlueprintCallable, Category = "Audio|MetaSound|Builder", meta = (ExpandEnumAsExecs = "OutResult"))
-	UPARAM(DisplayName = "Node ClassVersion") FMetasoundFrontendVersion FindNodeClassVersion(const FMetaSoundNodeHandle& NodeHandle, EMetaSoundBuilderResult& OutResult) const;
+	UPARAM(DisplayName = "Node ClassVersion") FMetasoundFrontendVersion FindNodeClassVersion(const FMetaSoundNodeHandle& NodeHandle, EMetaSoundBuilderResult& OutResult);
 
 	// Returns node input's data if valid (including things like name and datatype).
 	UFUNCTION(BlueprintCallable, Category = "Audio|MetaSound|Builder", meta = (ExpandEnumAsExecs = "OutResult"))
-	void GetNodeInputData(const FMetaSoundBuilderNodeInputHandle& InputHandle, FName& Name, FName& DataType, EMetaSoundBuilderResult& OutResult) const;
+	void GetNodeInputData(const FMetaSoundBuilderNodeInputHandle& InputHandle, FName& Name, FName& DataType, EMetaSoundBuilderResult& OutResult);
 
 	// Returns node output's data if valid (including things like name and datatype).
 	UFUNCTION(BlueprintCallable, Category = "Audio|MetaSound|Builder", meta = (ExpandEnumAsExecs = "OutResult"))
-	void GetNodeOutputData(const FMetaSoundBuilderNodeOutputHandle& OutputHandle, FName& Name, FName& DataType, EMetaSoundBuilderResult& OutResult) const;
+	void GetNodeOutputData(const FMetaSoundBuilderNodeOutputHandle& OutputHandle, FName& Name, FName& DataType, EMetaSoundBuilderResult& OutResult);
 
 	// Returns if a given interface is declared.
 	UFUNCTION(BlueprintCallable, Category = "Audio|MetaSound|Builder")
@@ -280,41 +284,67 @@ protected:
 		return Builder;
 	}
 
+	// Runs build, conforming the document and corresponding object data on a MetaSound UObject to that managed by this builder.
 	template <typename UClassType>
-	TScriptInterface<IMetaSoundDocumentInterface> BuildInternal(UObject* Parent, const FMetaSoundBuilderOptions& BuilderOptions) const
+	UClassType* BuildInternal(UObject* Parent, const FMetaSoundBuilderOptions& BuilderOptions) const
 	{
-		FName ObjectName = BuilderOptions.Name;
-		if (!ObjectName.IsNone())
-		{
-			ObjectName = MakeUniqueObjectName(Parent, UClassType::StaticClass(), BuilderOptions.Name);
-		}
+		using namespace Metasound::Frontend;
 
-		UClassType* NewMetaSound = BuilderOptions.ExistingAsset
-			? CastChecked<UClassType>(BuilderOptions.ExistingAsset.GetObject())
-			: NewObject<UClassType>(Parent, ObjectName, RF_Public | RF_Transient);
-		if (!NewMetaSound)
+		UClassType* MetaSound = nullptr;
+		if (BuilderOptions.ExistingMetaSound)
 		{
-			return nullptr;
+			MetaSound = CastChecked<UClassType>(BuilderOptions.ExistingMetaSound.GetObject());
+
+			// Always unregister if mutating existing object. If bAddToRegistry is set to false,
+			// leaving registered would result in any references to this MetaSound executing on
+			// out-of-date data. If bAddToRegistry is set, then it needs to be unregistered before
+			// being registered as it is below.
+			if (MetaSound)
+			{
+				const FNodeRegistryKey& RegistryKey = MetaSound->GetRegistryKey();
+				if (NodeRegistryKey::IsValid(RegistryKey))
+				{
+					if (FMetasoundFrontendRegistryContainer::Get()->IsNodeRegistered(RegistryKey))
+					{
+						MetaSound->UnregisterGraphWithFrontend();
+					}
+				}
+			}
 		}
-
-		FMetasoundFrontendDocument NewDocument = GetConstBuilder().GetDocument();
+		else
 		{
-			// This is required to ensure the newly build document has a unique class
-			// identifier to avoid collisions if added to the Frontend class registry
-			// (either below or at a later point in time).
-			constexpr bool bResetVersion = false;
-			FMetaSoundFrontendDocumentBuilder::InitGraphClassMetadata(NewDocument.RootGraph.Metadata, bResetVersion);
+			FName ObjectName = BuilderOptions.Name;
+			if (!ObjectName.IsNone())
+			{
+				ObjectName = MakeUniqueObjectName(Parent, UClassType::StaticClass(), BuilderOptions.Name);
+			}
+
+			MetaSound = NewObject<UClassType>(Parent, ObjectName, RF_Public | RF_Transient);
 		}
-		NewMetaSound->SetDocument(MoveTemp(NewDocument));
-		NewMetaSound->ConformObjectDataToInterfaces();
 
-		if (BuilderOptions.bAddToRegistry)
+		if (MetaSound)
 		{
-			NewMetaSound->RegisterGraphWithFrontend();
+			FMetasoundFrontendDocument NewDocument = GetConstBuilder().GetDocument();
+			{
+				// This is required to ensure the newly build document has a unique class
+				// identifier to avoid collisions if added to the Frontend class registry
+				// (either below or at a later point in time).
+				constexpr bool bResetVersion = false;
+				FMetaSoundFrontendDocumentBuilder::InitGraphClassMetadata(NewDocument.RootGraph.Metadata, bResetVersion);
+			}
+
+			MetaSound->SetDocument(MoveTemp(NewDocument));
+			MetaSound->ConformObjectDataToInterfaces();
+
+			if (BuilderOptions.bAddToRegistry)
+			{
+				MetaSound->RegisterGraphWithFrontend();
+			}
+
+			UE_LOG(LogMetaSound, VeryVerbose, TEXT("MetaSound '%s' built from '%s'"), *BuilderOptions.Name.ToString(), *GetFullName());
 		}
 
-		UE_LOG(LogMetaSound, VeryVerbose, TEXT("New MetaSound '%s' built from '%s'"), *BuilderOptions.Name.ToString(), *GetFullName());
-		return NewMetaSound;
+		return MetaSound;
 	}
 
 	UPROPERTY()
diff --git a/Engine/Plugins/Runtime/Metasound/Source/MetasoundFrontend/Private/MetasoundAssetBase.cpp b/Engine/Plugins/Runtime/Metasound/Source/MetasoundFrontend/Private/MetasoundAssetBase.cpp
--- a/Engine/Plugins/Runtime/Metasound/Source/MetasoundFrontend/Private/MetasoundAssetBase.cpp
+++ b/Engine/Plugins/Runtime/Metasound/Source/MetasoundFrontend/Private/MetasoundAssetBase.cpp
@@ -766,6 +766,11 @@ const FMetasoundFrontendDocument& FMetasoundAssetBase::GetDocumentChecked() cons
 	return *Document;
 }
 
+const Metasound::Frontend::FNodeRegistryKey& FMetasoundAssetBase::GetRegistryKey() const
+{
+	return RegistryKey;
+}
+
 FString FMetasoundAssetBase::GetOwningAssetName() const
 {
 	if (const UObject* OwningAsset = GetOwningAsset())
diff --git a/Engine/Plugins/Runtime/Metasound/Source/MetasoundFrontend/Public/MetasoundAssetBase.h b/Engine/Plugins/Runtime/Metasound/Source/MetasoundFrontend/Public/MetasoundAssetBase.h
--- a/Engine/Plugins/Runtime/Metasound/Source/MetasoundFrontend/Public/MetasoundAssetBase.h
+++ b/Engine/Plugins/Runtime/Metasound/Source/MetasoundFrontend/Public/MetasoundAssetBase.h
@@ -138,6 +138,8 @@ public:
 	FMetasoundFrontendDocument& GetDocumentChecked();
 	const FMetasoundFrontendDocument& GetDocumentChecked() const;
 
+	const Metasound::Frontend::FNodeRegistryKey& GetRegistryKey() const;
+
 	UE_DEPRECATED(5.3, "AddDefaultInterfaces is included in now applied via FMetaSoundFrontendDocumentBuilder::InitDocument and no longer directly supported via this function.")
 	void AddDefaultInterfaces();
 
