Metasound bug fixes discovered while using the output subsystem

diff --git a/Engine/Plugins/Runtime/Metasound/Source/MetasoundEngine/Private/MetasoundGeneratorHandle.cpp b/Engine/Plugins/Runtime/Metasound/Source/MetasoundEngine/Private/MetasoundGeneratorHandle.cpp
--- a/Engine/Plugins/Runtime/Metasound/Source/MetasoundEngine/Private/MetasoundGeneratorHandle.cpp
+++ b/Engine/Plugins/Runtime/Metasound/Source/MetasoundEngine/Private/MetasoundGeneratorHandle.cpp
@@ -65,12 +65,12 @@ void UMetasoundGeneratorHandle::CacheMetasoundSource()
 
 void UMetasoundGeneratorHandle::AttachGeneratorDelegates()
 {
-	GeneratorCreatedDelegateHandle = CachedMetasoundSource->OnGeneratorInstanceCreated.AddLambda(
+	GeneratorCreatedDelegateHandle = CachedMetasoundSource->OnGeneratorInstanceCreated.AddWeakLambda(this,
 		[this](uint64 InAudioComponentId, TSharedPtr<Metasound::FMetasoundGenerator> InGenerator)
 		{
 			OnSourceCreatedAGenerator(InAudioComponentId, InGenerator);
 		});
-	GeneratorDestroyedDelegateHandle = CachedMetasoundSource->OnGeneratorInstanceDestroyed.AddLambda(
+	GeneratorDestroyedDelegateHandle = CachedMetasoundSource->OnGeneratorInstanceDestroyed.AddWeakLambda(this,
 		[this](uint64 InAudioComponentId, TSharedPtr<Metasound::FMetasoundGenerator> InGenerator)
 		{
 			OnSourceDestroyedAGenerator(InAudioComponentId, InGenerator);
diff --git a/Engine/Plugins/Runtime/Metasound/Source/MetasoundEngine/Private/MetasoundOutputSubsystem.cpp b/Engine/Plugins/Runtime/Metasound/Source/MetasoundEngine/Private/MetasoundOutputSubsystem.cpp
--- a/Engine/Plugins/Runtime/Metasound/Source/MetasoundEngine/Private/MetasoundOutputSubsystem.cpp
+++ b/Engine/Plugins/Runtime/Metasound/Source/MetasoundEngine/Private/MetasoundOutputSubsystem.cpp
@@ -154,7 +154,12 @@ UMetaSoundOutputSubsystem::FGeneratorInfo* UMetaSoundOutputSubsystem::FindOrAddG
 		&& GeneratorInfo.AudioComponent->GetAudioComponentID() == AudioComponentId;
 	}))
 	{
-		return Info;
+		if (Info->Source.IsValid())
+		{
+			return Info;
+		}
+		// Source is no longer valid, remove info
+		OnGeneratorDestroyed(AudioComponentId, Info->Generator.Pin());
 	}
 
 	UMetaSoundSource* Source = Cast<UMetaSoundSource>(AudioComponent->GetSound());
diff --git a/Engine/Plugins/Runtime/Metasound/Source/MetasoundEngine/Public/MetasoundOutputSubsystem.h b/Engine/Plugins/Runtime/Metasound/Source/MetasoundEngine/Public/MetasoundOutputSubsystem.h
--- a/Engine/Plugins/Runtime/Metasound/Source/MetasoundEngine/Public/MetasoundOutputSubsystem.h
+++ b/Engine/Plugins/Runtime/Metasound/Source/MetasoundEngine/Public/MetasoundOutputSubsystem.h
@@ -29,7 +29,7 @@ DECLARE_DYNAMIC_DELEGATE_OneParam(FOnMetasoundOutputValueChanged, const FMetaSou
  * Provides access to a playing Metasound generator's outputs
  */
 UCLASS()
-class UMetaSoundOutputSubsystem : public UTickableWorldSubsystem
+class METASOUNDENGINE_API UMetaSoundOutputSubsystem : public UTickableWorldSubsystem
 {
 	GENERATED_BODY()
 
