- Fix crash when attempting to add invalid MetaSound node classes (due to them being transient assets) to a given MetaSound in its respective asset Editor.
- Remove temporary gating cvar for MetaSound builder subsystem
#rnx

diff --git a/Engine/Plugins/Runtime/Metasound/Source/MetasoundEditor/Private/MetasoundEditorGraphSchema.cpp b/Engine/Plugins/Runtime/Metasound/Source/MetasoundEditor/Private/MetasoundEditorGraphSchema.cpp
--- a/Engine/Plugins/Runtime/Metasound/Source/MetasoundEditor/Private/MetasoundEditorGraphSchema.cpp
+++ b/Engine/Plugins/Runtime/Metasound/Source/MetasoundEditor/Private/MetasoundEditorGraphSchema.cpp
@@ -1985,17 +1985,25 @@ void UMetasoundEditorGraphSchema::GetFunctionActions(FGraphActionMenuBuilder& Ac
 			continue;
 		}
 
-		const FSoftObjectPath* Path = nullptr;
+		bool bIsAsset = false;
 		if (ParentAsset)
 		{
 			const FNodeRegistryKey RegistryKey = NodeRegistryKey::CreateKey(FrontendClass.Metadata);
-			Path = AssetManager.FindObjectPathFromKey(RegistryKey);
-			if (Path)
+			if (const FSoftObjectPath* Path = AssetManager.FindObjectPathFromKey(RegistryKey))
 			{
+				// The path can be not an asset if it was transiently created (eg. in PIE if generated by a runtime MetaSoundDocumentBuilder)
+				// Ignore these assets and do not include as this would create a dependency on a transient MetaSound from a MetaSound asset.
+				if (!Path->IsAsset())
+				{
+					continue;
+				}
+
 				if (ParentAsset->AddingReferenceCausesLoop(*Path))
 				{
 					continue;
 				}
+
+				bIsAsset = true;
 			}
 		}
 
@@ -2006,7 +2014,7 @@ void UMetasoundEditorGraphSchema::GetFunctionActions(FGraphActionMenuBuilder& Ac
 
 		if (Metadata.GetCategoryHierarchy().IsEmpty() || Metadata.GetCategoryHierarchy()[0].CompareTo(Metasound::NodeCategories::Conversions))
 		{
-			const EPrimaryContextGroup ContextGroup = Path ? EPrimaryContextGroup::Graphs : EPrimaryContextGroup::Functions;
+			const EPrimaryContextGroup ContextGroup = bIsAsset ? EPrimaryContextGroup::Graphs : EPrimaryContextGroup::Functions;
 			TArray<FText> CategoryHierarchy { GetContextGroupDisplayName(ContextGroup) };
 			CategoryHierarchy.Append(Metadata.GetCategoryHierarchy());
 			const FText KeywordsText = FText::Join(KeywordDelim, Metadata.GetKeywords());
diff --git a/Engine/Plugins/Runtime/Metasound/Source/MetasoundEngine/Private/MetasoundBuilderSubsystem.cpp b/Engine/Plugins/Runtime/Metasound/Source/MetasoundEngine/Private/MetasoundBuilderSubsystem.cpp
--- a/Engine/Plugins/Runtime/Metasound/Source/MetasoundEngine/Private/MetasoundBuilderSubsystem.cpp
+++ b/Engine/Plugins/Runtime/Metasound/Source/MetasoundEngine/Private/MetasoundBuilderSubsystem.cpp
@@ -26,14 +26,6 @@ namespace Metasound::Engine
 {
 	namespace BuilderSubsystemPrivate
 	{
-		static int32 DocumentBuilderEnabledCVar = 0;
-		FAutoConsoleVariableRef CVarEnableDocumentBuilder(
-			TEXT("au.MetaSound.BuilderSubsystem.Enabled"),
-			DocumentBuilderEnabledCVar,
-			TEXT("Enables the the Runtime MetaSoundBuilderSubsystem, which (as of 5.3) is in beta (Default: Disabled).\n")
-			TEXT("0: Disabled, 1: Enabled"),
-			ECVF_Default);
-
 		template <typename TLiteralType>
 		FMetasoundFrontendLiteral CreatePODMetaSoundLiteral(const TLiteralType& Value, FName& OutDataType)
 		{
@@ -740,14 +732,7 @@ void UMetaSoundSourceBuilder::SetFormat(EMetaSoundOutputAudioFormat OutputFormat
 
 UMetaSoundPatchBuilder* UMetaSoundBuilderSubsystem::CreatePatchBuilder(FName BuilderName, EMetaSoundBuilderResult& OutResult)
 {
-	using namespace Metasound::Engine::BuilderSubsystemPrivate;
-
-	if (!DocumentBuilderEnabledCVar)
-	{
-		UE_LOG(LogMetaSound, Error, TEXT("CreatePatchBuilder Failed to create MetaSound Builder '%s': See CVar au.MetaSound.BuilderSubsystem.Enabled (%s')"), *BuilderName.ToString(), CVarEnableDocumentBuilder->GetHelp());
-		return nullptr;
-	}
-	return &CreateTransientBuilder<UMetaSoundPatchBuilder>();
+	return &Metasound::Engine::BuilderSubsystemPrivate::CreateTransientBuilder<UMetaSoundPatchBuilder>();
 }
 
 UMetaSoundSourceBuilder* UMetaSoundBuilderSubsystem::CreateSourceBuilder(
@@ -768,12 +753,6 @@ UMetaSoundSourceBuilder* UMetaSoundBuilderSubsystem::CreateSourceBuilder(
 
 	OutResult = EMetaSoundBuilderResult::Failed;
 
-	if (!DocumentBuilderEnabledCVar)
-	{
-		UE_LOG(LogMetaSound, Error, TEXT("CreateSourceBuilder Failed to create MetaSound Builder '%s': See CVar au.MetaSound.BuilderSubsystem.Enabled (%s')"), *BuilderName.ToString(), CVarEnableDocumentBuilder->GetHelp());
-		return nullptr;
-	}
-
 	UMetaSoundSourceBuilder& NewBuilder = CreateTransientBuilder<UMetaSoundSourceBuilder>();
 
 	bool bFormatSet = OutputFormat == EMetaSoundOutputAudioFormat::Mono;
