[Metasound] TEnum: Add support for integral underlying types which can fit in int32

diff --git a/Engine/Plugins/Runtime/Metasound/Source/MetasoundFrontend/Public/MetasoundEnum.h b/Engine/Plugins/Runtime/Metasound/Source/MetasoundFrontend/Public/MetasoundEnum.h
--- a/Engine/Plugins/Runtime/Metasound/Source/MetasoundFrontend/Public/MetasoundEnum.h
+++ b/Engine/Plugins/Runtime/Metasound/Source/MetasoundFrontend/Public/MetasoundEnum.h
@@ -109,13 +109,50 @@ namespace Metasound
 	{
 	public:
 		using InnerType = EnumType;
-	
-		static_assert(TIsEnum<EnumType>::Value, "Expecting an Enum type");
-		
-		// Allow Non-enum class Enum or Enum class that's are derived from int32
-		static_assert(std::is_same_v<typename std::underlying_type<EnumType>::type, int32> ||
-			!TIsEnumClass<EnumType>::Value, "We serialize to int32, so limit to that for now");
+		using SerializedType = int32;
+		using UnderlyingType = std::underlying_type_t<EnumType>;
+
+		static constexpr bool EnumTypeIsSupported()
+		{
+			// make sure this is an enum
+			if constexpr (!TIsEnum<EnumType>::Value)
+			{
+				return false;
+			}
+
+			// make sure the underlying type is integral
+			if constexpr (!TIsIntegral<UnderlyingType>::Value)
+			{
+				return false;
+			}
+
+			// if this is a scoped enum make sure the underlying type will fit in the serialized type
+			if constexpr (TIsEnumClass<EnumType>::Value)
+			{
+				constexpr bool SerializedIsSigned = TIsSigned<SerializedType>::Value;
+				constexpr bool UnderlyingIsSigned = TIsSigned<UnderlyingType>::Value;
 
+				// if the serialized type is signed, the underlying type must be
+				// smaller than the serialized type if it is unsigned,
+				// or the same size or smaller than the serialized type if it's signed
+				if constexpr (SerializedIsSigned)
+				{
+					return UnderlyingIsSigned
+						? sizeof(UnderlyingType) <= sizeof(SerializedType)
+						: sizeof(UnderlyingType) < sizeof(SerializedType);
+				}
+
+				// otherwise the underlying type must be unsigned and at most the size of the serialized type
+				// This is just here in case someone changes the serialized type
+				return !UnderlyingIsSigned && sizeof(UnderlyingType) <= sizeof(SerializedType);
+			}
+
+			// make an exception for C-style enums so we don't break backward compatibility
+			return true;
+		}
+		
+		static_assert(EnumTypeIsSupported(), "EnumType is not supported for TEnum");
+		
 		// Default.
 		explicit TEnum(EnumType InValue = DefaultValue)
 		{
