[Metasound] Avoid dangling pointer in output watcher handler

#rnx

diff --git a/Engine/Plugins/Runtime/Metasound/Source/MetasoundEngine/Private/MetasoundOutputSubsystem.cpp b/Engine/Plugins/Runtime/Metasound/Source/MetasoundEngine/Private/MetasoundOutputSubsystem.cpp
--- a/Engine/Plugins/Runtime/Metasound/Source/MetasoundEngine/Private/MetasoundOutputSubsystem.cpp
+++ b/Engine/Plugins/Runtime/Metasound/Source/MetasoundEngine/Private/MetasoundOutputSubsystem.cpp
@@ -81,18 +81,6 @@ bool UMetaSoundOutputSubsystem::WatchOutput(
 
 	// either way, add the delegate to the map
 	GeneratorInfo->OutputChangedDelegates.FindOrAdd(OutputName).FindOrAdd(AnalyzerOutputName).AddUnique(OnOutputValueChanged);
-
-	// set up the output handler
-	GeneratorInfo->HandleOutputChanged = [GeneratorInfo](const FName GeneratorOutputName, const FMetaSoundOutput& AnalyzerOutput)
-	{
-		if (TMap<FName, FOnOutputValueChangedMulticast>* Map = GeneratorInfo->OutputChangedDelegates.Find(GeneratorOutputName))
-		{
-			if (const FOnOutputValueChangedMulticast* Delegate = Map->Find(AnalyzerOutput.Name))
-			{
-				Delegate->Broadcast(AnalyzerOutput);
-			}
-		}
-	};
 	
 	return true;
 }
@@ -116,7 +104,10 @@ void UMetaSoundOutputSubsystem::Tick(float DeltaTime)
 		{
 			for (Metasound::Private::FMetasoundOutputWatcher& Watcher : GeneratorIt->OutputWatchers)
 			{
-				Watcher.Update(GeneratorIt->HandleOutputChanged);
+				Watcher.Update([&GeneratorIt](FName OutputName, const FMetaSoundOutput& Output)
+				{
+					GeneratorIt->HandleOutputChanged(OutputName, Output);
+				});
 			}
 		}
 	}
@@ -133,6 +124,17 @@ void UMetaSoundOutputSubsystem::RegisterPassthroughAnalyzerForType(const FName T
 	PassthroughAnalyzers.Add(TypeName, AnalyzerName);
 }
 
+void UMetaSoundOutputSubsystem::FGeneratorInfo::HandleOutputChanged(FName OutputName, const FMetaSoundOutput& Output)
+{
+	if (TMap<FName, FOnOutputValueChangedMulticast>* Map = OutputChangedDelegates.Find(OutputName))
+	{
+		if (const FOnOutputValueChangedMulticast* Delegate = Map->Find(Output.Name))
+		{
+			Delegate->Broadcast(Output);
+		}
+	}
+}
+
 bool UMetaSoundOutputSubsystem::FGeneratorInfo::IsValid() const
 {
 	return AudioComponent.IsValid() && Source.IsValid();
diff --git a/Engine/Plugins/Runtime/Metasound/Source/MetasoundEngine/Public/MetasoundOutputSubsystem.h b/Engine/Plugins/Runtime/Metasound/Source/MetasoundEngine/Public/MetasoundOutputSubsystem.h
--- a/Engine/Plugins/Runtime/Metasound/Source/MetasoundEngine/Public/MetasoundOutputSubsystem.h
+++ b/Engine/Plugins/Runtime/Metasound/Source/MetasoundEngine/Public/MetasoundOutputSubsystem.h
@@ -83,7 +83,7 @@ private:
 		TQueue<Metasound::Frontend::FAnalyzerAddress> AnalyzersToCreate;
 		TMap<FName, TMap<FName, FOnOutputValueChangedMulticast>> OutputChangedDelegates;
 		TArray<Metasound::Private::FMetasoundOutputWatcher> OutputWatchers;
-		TFunction<void(FName, const FMetaSoundOutput&)> HandleOutputChanged;
+		void HandleOutputChanged(FName OutputName, const FMetaSoundOutput& Output);
 		bool IsValid() const;
 	};
 
