Add delegates for when the MusicClockComponent connects and disconnects from an underling MidiClock in the associated metasound.

diff --git a/Engine/Plugins/Runtime/Metasound/Source/MetasoundEngine/Private/MetasoundGeneratorHandle.cpp b/Engine/Plugins/Runtime/Metasound/Source/MetasoundEngine/Private/MetasoundGeneratorHandle.cpp
--- a/Engine/Plugins/Runtime/Metasound/Source/MetasoundEngine/Private/MetasoundGeneratorHandle.cpp
+++ b/Engine/Plugins/Runtime/Metasound/Source/MetasoundEngine/Private/MetasoundGeneratorHandle.cpp
@@ -65,25 +65,72 @@ void UMetasoundGeneratorHandle::CacheMetasoundSource()
 
 void UMetasoundGeneratorHandle::AttachGeneratorDelegates()
 {
-	GeneratorCreatedDelegateHandle = CachedMetasoundSource->OnGeneratorInstanceCreated.AddWeakLambda(this,
-		[this](uint64 InAudioComponentId, TSharedPtr<Metasound::FMetasoundGenerator> InGenerator)
+	// These delegates can be called on a separate thread, so we need to take steps
+	// to assure this UObject hasn't been garbage collected before trying to dereference. 
+	// That is why we capture a TWeakObjectPtr and try to dereference that later!
+	GeneratorCreatedDelegateHandle = CachedMetasoundSource->OnGeneratorInstanceCreated.AddLambda(
+		[WeakGeneratorHandlePtr = TWeakObjectPtr<UMetasoundGeneratorHandle>(this),
+		StatId = this->GetStatID(true)]
+		(uint64 InAudioComponentId, TSharedPtr<Metasound::FMetasoundGenerator> InGenerator)
 		{
-			OnSourceCreatedAGenerator(InAudioComponentId, InGenerator);
+			// We are in the audio render (or control) thread here, so create a "dispatch task" to be
+			// executed later on the game thread...
+			FFunctionGraphTask::CreateAndDispatchWhenReady(
+				[WeakGeneratorHandlePtr, InAudioComponentId, InGenerator]()
+				{
+					check(IsInGameThread());
+					// Now, since we are in the game thread, try to dereference the pointer to
+					// to the UMetasoundGeneratorHandle. This should only succeed if the UObject
+					// hasn't been garbage collected.
+					if (UMetasoundGeneratorHandle* TheHandle = WeakGeneratorHandlePtr.Get())
+					{
+						TheHandle->OnSourceCreatedAGenerator(InAudioComponentId, InGenerator);
+					}
+				}, 
+				StatId, nullptr, ENamedThreads::GameThread);
 		});
 	GeneratorDestroyedDelegateHandle = CachedMetasoundSource->OnGeneratorInstanceDestroyed.AddWeakLambda(this,
 		[this](uint64 InAudioComponentId, TSharedPtr<Metasound::FMetasoundGenerator> InGenerator)
 		{
-			OnSourceDestroyedAGenerator(InAudioComponentId, InGenerator);
+				// We are in the audio render thread here, so create a "dispatch task" to be
+				// executed later on the game thread...
+				FFunctionGraphTask::CreateAndDispatchWhenReady(
+					[WeakGeneratorHandlePtr]()
+					{
+						check(IsInGameThread());
+						// Now, since we are in the game thread, try to dereference the pointer to
+						// to the UMetasoundGeneratorHandle. This should only succeed if the UObject
+						// hasn't been garbage collected.
+						if (UMetasoundGeneratorHandle* TheHandle = WeakGeneratorHandlePtr.Get())
+						{
+							TheHandle->OnGeneratorsGraphChanged.Broadcast();
+						}
+					},
+					StatId, nullptr, ENamedThreads::GameThread);
 		});
 }
+}
 
 void UMetasoundGeneratorHandle::DetachGeneratorDelegates()
 {
+	// First detach any callbacks that tell us when a generator is created or destroyed
+	// for the UMetasoundSource of interest...
+	if (CachedMetasoundSource)
+	{
 	CachedMetasoundSource->OnGeneratorInstanceCreated.Remove(GeneratorCreatedDelegateHandle);
 	GeneratorCreatedDelegateHandle.Reset();
 	CachedMetasoundSource->OnGeneratorInstanceDestroyed.Remove(GeneratorDestroyedDelegateHandle);
 	GeneratorDestroyedDelegateHandle.Reset();
 }
+	// Now detach any callback we may have registered to get callbacks when the generator's
+	// graph has been changed...
+	TSharedPtr<Metasound::FMetasoundGenerator> PinnedGenerator = CachedGeneratorPtr.Pin();
+	if (PinnedGenerator)
+	{
+		PinnedGenerator->OnSetGraph.Remove(GeneratorGraphChangedDelegateHandle);
+		GeneratorGraphChangedDelegateHandle.Reset();
+	}
+}
 
 TSharedPtr<Metasound::FMetasoundGenerator> UMetasoundGeneratorHandle::PinGenerator()
 {
@@ -144,8 +191,29 @@ TSharedPtr<Metasound::FMetasoundGenerator> UMetasoundGeneratorHandle::GetGenerat
 	return PinGenerator();
 }
 
+FDelegateHandle UMetasoundGeneratorHandle::AddGraphSetCallback(const FOnSetGraph& Delegate)
+{
+	
+	FDelegateHandle Handle = OnGeneratorsGraphChanged.Add(Delegate);
+	CacheMetasoundSource();
+	AttachGraphChangedDelegate();
+	return Handle;
+}
+
+bool UMetasoundGeneratorHandle::RemoveGraphSetCallback(const FDelegateHandle& Handle)
+{
+	return OnGeneratorsGraphChanged.Remove(Handle);
+}
+
+void UMetasoundGeneratorHandle::BeginDestroy()
+{
+	Super::BeginDestroy();
+	DetachGeneratorDelegates();
+}
+
 void UMetasoundGeneratorHandle::OnSourceCreatedAGenerator(uint64 InAudioComponentId, TSharedPtr<Metasound::FMetasoundGenerator> InGenerator)
 {
+	check(IsInGameThread());
 	if (InAudioComponentId == AudioComponentId)
 	{
 		CachedGeneratorPtr = InGenerator;
@@ -156,12 +224,16 @@ void UMetasoundGeneratorHandle::OnSourceCreatedAGenerator(uint64 InAudioComponen
 				InGenerator->QueueParameterPack(CachedParameterPack);
 			}
 			OnGeneratorHandleAttached.Broadcast();
+			// If anyone has told us they are interested in being notified when a generator's 
+			// graph has changed go ahead and set that up now...
+			AttachGraphChangedDelegate();
 		}
 	}
 }
 
 void UMetasoundGeneratorHandle::OnSourceDestroyedAGenerator(uint64 InAudioComponentId, TSharedPtr<Metasound::FMetasoundGenerator> InGenerator)
 {
+	check(IsInGameThread());
 	if (InAudioComponentId == AudioComponentId)
 	{
 		if (CachedGeneratorPtr.IsValid())
@@ -172,3 +244,4 @@ void UMetasoundGeneratorHandle::OnSourceDestroyedAGenerator(uint64 InAudioCompon
 	}
 }
 
+
diff --git a/Engine/Plugins/Runtime/Metasound/Source/MetasoundEngine/Public/MetasoundGeneratorHandle.h b/Engine/Plugins/Runtime/Metasound/Source/MetasoundEngine/Public/MetasoundGeneratorHandle.h
--- a/Engine/Plugins/Runtime/Metasound/Source/MetasoundEngine/Public/MetasoundGeneratorHandle.h
+++ b/Engine/Plugins/Runtime/Metasound/Source/MetasoundEngine/Public/MetasoundGeneratorHandle.h
@@ -14,8 +14,16 @@ namespace Metasound
 	class FMetasoundGenerator;
 }
 
-DECLARE_TS_MULTICAST_DELEGATE(FOnGeneratorHandleAttached);
-DECLARE_TS_MULTICAST_DELEGATE(FOnGeneratorHandleDetached);
+// UMetasoundGeneratorHandle shields its "clients" from "cross thread" issues
+// related to callbacks coming in the audio control or rendering threads that 
+// game thread clients (e.g. blueprints) want to know about. That is why these 
+// next delegate definitions do not need to be the "TS" variants. Assignments 
+// to members of this type, and the broadcasts there to will all happen on the
+// game thread. EVEN IF the instigator of those callbacks is on the audio
+// render thread. 
+DECLARE_MULTICAST_DELEGATE(FOnGeneratorHandleAttached);
+DECLARE_MULTICAST_DELEGATE(FOnGeneratorHandleDetached);
+DECLARE_DELEGATE(FOnSetGraph);
 
 UCLASS(BlueprintType,Category="MetaSound")
 class METASOUNDENGINE_API UMetasoundGeneratorHandle : public UObject
@@ -27,6 +35,8 @@ public:
 	UFUNCTION(BlueprintCallable, Category="MetaSound")
 	static UMetasoundGeneratorHandle* CreateMetaSoundGeneratorHandle(UAudioComponent* OnComponent);
 
+	void BeginDestroy() override;
+
 	/**
 	 * Makes a copy of the supplied parameter pack and passes it to the MetaSoundGenerator
 	 * for asynchronous processing. IT ALSO caches this copy so that if the AudioComponent
@@ -38,11 +48,13 @@ public:
 
 	TSharedPtr<Metasound::FMetasoundGenerator> GetGenerator();
 
-	/**
-	 * NOTE: These delegates are called on the audio render thread. Careful!
-	 */
 	FOnGeneratorHandleAttached OnGeneratorHandleAttached;
 	FOnGeneratorHandleDetached OnGeneratorHandleDetached;
+	// Note: We don't allow direct assignment to the OnGeneratorsGraphChanged delegate
+	// because we need to know that someone actually wants this message so we can
+	// start actively listening for the corresponding audio render thread callback...
+	FDelegateHandle AddGraphSetCallback(const FOnSetGraph& Delegate);
+	bool RemoveGraphSetCallback(const FDelegateHandle& Handle);
 
 private:
 
@@ -61,6 +73,7 @@ private:
 	 * Functions for adding and removing our MetaSoundGenerator lifecycle delegates
 	 */
 	void AttachGeneratorDelegates();
+	void AttachGraphChangedDelegate();
 	void DetachGeneratorDelegates();
 
 	/**
@@ -69,6 +82,7 @@ private:
 	void OnSourceCreatedAGenerator(uint64 InAudioComponentId, TSharedPtr<Metasound::FMetasoundGenerator> InGenerator);
 	void OnSourceDestroyedAGenerator(uint64 InAudioComponentId, TSharedPtr<Metasound::FMetasoundGenerator> InGenerator);
 
+
 	UPROPERTY(Transient)
 	TObjectPtr<UAudioComponent> AudioComponent;
 	uint64 AudioComponentId;
@@ -77,6 +91,15 @@ private:
 	TWeakPtr<Metasound::FMetasoundGenerator> CachedGeneratorPtr;
 	FSharedMetasoundParameterStoragePtr CachedParameterPack;
 
+	// Note: We don't allow direct assignment to the OnGeneratorsGraphChanged delegate
+	// because we need to know that someone actually wants this message so we can
+	// start actively listening for the corresponding audio render thread callback. 
+	// So these next members are private and a "client" that wants to be notified of
+	// the graph change have to call public functions declared above to add themselves.
+	DECLARE_MULTICAST_DELEGATE(FOnSetGraphMulticast);
+	FOnSetGraphMulticast OnGeneratorsGraphChanged;
+
 	FDelegateHandle GeneratorCreatedDelegateHandle;
 	FDelegateHandle GeneratorDestroyedDelegateHandle;
+	FDelegateHandle GeneratorGraphChangedDelegateHandle;
 };
\ No newline at end of file
diff --git a/Engine/Plugins/Runtime/Metasound/Source/MetasoundEngine/Public/MetasoundSource.h b/Engine/Plugins/Runtime/Metasound/Source/MetasoundEngine/Public/MetasoundSource.h
--- a/Engine/Plugins/Runtime/Metasound/Source/MetasoundEngine/Public/MetasoundSource.h
+++ b/Engine/Plugins/Runtime/Metasound/Source/MetasoundEngine/Public/MetasoundSource.h
@@ -27,8 +27,8 @@ namespace Metasound
 
 } // namespace Metasound
 
-DECLARE_MULTICAST_DELEGATE_TwoParams(FOnGeneratorInstanceCreated, uint64, TSharedPtr<Metasound::FMetasoundGenerator>);
-DECLARE_MULTICAST_DELEGATE_TwoParams(FOnGeneratorInstanceDestroyed, uint64, TSharedPtr<Metasound::FMetasoundGenerator>);
+DECLARE_TS_MULTICAST_DELEGATE_TwoParams(FOnGeneratorInstanceCreated, uint64, TSharedPtr<Metasound::FMetasoundGenerator>);
+DECLARE_TS_MULTICAST_DELEGATE_TwoParams(FOnGeneratorInstanceDestroyed, uint64, TSharedPtr<Metasound::FMetasoundGenerator>);
 
 /**
  * This Metasound type can be played as an audio source.
