Namespace the names of the multicast delegate definitions (move them into the class definition) to avoid polluting the global namespace.

diff --git a/Engine/Plugins/Runtime/Metasound/Source/MetasoundEngine/Private/MetasoundGeneratorHandle.cpp b/Engine/Plugins/Runtime/Metasound/Source/MetasoundEngine/Private/MetasoundGeneratorHandle.cpp
--- a/Engine/Plugins/Runtime/Metasound/Source/MetasoundEngine/Private/MetasoundGeneratorHandle.cpp
+++ b/Engine/Plugins/Runtime/Metasound/Source/MetasoundEngine/Private/MetasoundGeneratorHandle.cpp
@@ -191,7 +191,7 @@ TSharedPtr<Metasound::FMetasoundGenerator> UMetasoundGeneratorHandle::GetGenerat
 	return PinGenerator();
 }
 
-FDelegateHandle UMetasoundGeneratorHandle::AddGraphSetCallback(const FOnSetGraph& Delegate)
+FDelegateHandle UMetasoundGeneratorHandle::AddGraphSetCallback(const UMetasoundGeneratorHandle::FOnSetGraph& Delegate)
 {
 	
 	FDelegateHandle Handle = OnGeneratorsGraphChanged.Add(Delegate);
diff --git a/Engine/Plugins/Runtime/Metasound/Source/MetasoundEngine/Public/MetasoundGeneratorHandle.h b/Engine/Plugins/Runtime/Metasound/Source/MetasoundEngine/Public/MetasoundGeneratorHandle.h
--- a/Engine/Plugins/Runtime/Metasound/Source/MetasoundEngine/Public/MetasoundGeneratorHandle.h
+++ b/Engine/Plugins/Runtime/Metasound/Source/MetasoundEngine/Public/MetasoundGeneratorHandle.h
@@ -14,17 +14,6 @@ namespace Metasound
 	class FMetasoundGenerator;
 }
 
-// UMetasoundGeneratorHandle shields its "clients" from "cross thread" issues
-// related to callbacks coming in the audio control or rendering threads that 
-// game thread clients (e.g. blueprints) want to know about. That is why these 
-// next delegate definitions do not need to be the "TS" variants. Assignments 
-// to members of this type, and the broadcasts there to will all happen on the
-// game thread. EVEN IF the instigator of those callbacks is on the audio
-// render thread. 
-DECLARE_MULTICAST_DELEGATE(FOnGeneratorHandleAttached);
-DECLARE_MULTICAST_DELEGATE(FOnGeneratorHandleDetached);
-DECLARE_DELEGATE(FOnSetGraph);
-
 UCLASS(BlueprintType,Category="MetaSound")
 class METASOUNDENGINE_API UMetasoundGeneratorHandle : public UObject
 {
@@ -37,6 +26,17 @@ public:
 
 	void BeginDestroy() override;
 
+	// UMetasoundGeneratorHandle shields its "clients" from "cross thread" issues
+	// related to callbacks coming in the audio control or rendering threads that 
+	// game thread clients (e.g. blueprints) want to know about. That is why these 
+	// next delegate definitions do not need to be the "TS" variants. Assignments 
+	// to members of this type, and the broadcasts there to will all happen on the
+	// game thread. EVEN IF the instigator of those callbacks is on the audio
+	// render thread. 
+	DECLARE_MULTICAST_DELEGATE(FOnAttached);
+	DECLARE_MULTICAST_DELEGATE(FOnDetached);
+	DECLARE_DELEGATE(FOnSetGraph);
+
 	/**
 	 * Makes a copy of the supplied parameter pack and passes it to the MetaSoundGenerator
 	 * for asynchronous processing. IT ALSO caches this copy so that if the AudioComponent
@@ -48,12 +48,12 @@ public:
 
 	TSharedPtr<Metasound::FMetasoundGenerator> GetGenerator();
 
-	FOnGeneratorHandleAttached OnGeneratorHandleAttached;
-	FOnGeneratorHandleDetached OnGeneratorHandleDetached;
+	UMetasoundGeneratorHandle::FOnAttached OnGeneratorHandleAttached;
+	UMetasoundGeneratorHandle::FOnDetached OnGeneratorHandleDetached;
 	// Note: We don't allow direct assignment to the OnGeneratorsGraphChanged delegate
 	// because we need to know that someone actually wants this message so we can
 	// start actively listening for the corresponding audio render thread callback...
-	FDelegateHandle AddGraphSetCallback(const FOnSetGraph& Delegate);
+	FDelegateHandle AddGraphSetCallback(const UMetasoundGeneratorHandle::FOnSetGraph& Delegate);
 	bool RemoveGraphSetCallback(const FDelegateHandle& Handle);
 
 private:
