Fix bad merge
#rnx

diff --git a/Engine/Plugins/Runtime/Metasound/Source/MetasoundEngine/Private/MetasoundGeneratorHandle.cpp b/Engine/Plugins/Runtime/Metasound/Source/MetasoundEngine/Private/MetasoundGeneratorHandle.cpp
--- a/Engine/Plugins/Runtime/Metasound/Source/MetasoundEngine/Private/MetasoundGeneratorHandle.cpp
+++ b/Engine/Plugins/Runtime/Metasound/Source/MetasoundEngine/Private/MetasoundGeneratorHandle.cpp
@@ -89,9 +89,46 @@ void UMetasoundGeneratorHandle::AttachGeneratorDelegates()
 				}, 
 				StatId, nullptr, ENamedThreads::GameThread);
 		});
-	GeneratorDestroyedDelegateHandle = CachedMetasoundSource->OnGeneratorInstanceDestroyed.AddWeakLambda(this,
-		[this](uint64 InAudioComponentId, TSharedPtr<Metasound::FMetasoundGenerator> InGenerator)
+	GeneratorDestroyedDelegateHandle = CachedMetasoundSource->OnGeneratorInstanceDestroyed.AddLambda(
+		[WeakGeneratorHandlePtr = TWeakObjectPtr<UMetasoundGeneratorHandle>(this),
+		StatId = this->GetStatID(true)]
+		(uint64 InAudioComponentId, TSharedPtr<Metasound::FMetasoundGenerator> InGenerator)
 		{
+			// We are in the audio render (or control) thread here, so create a "dispatch task" to be
+			// executed later on the game thread...
+			FFunctionGraphTask::CreateAndDispatchWhenReady(
+				[WeakGeneratorHandlePtr, InAudioComponentId, InGenerator]()
+				{
+					check(IsInGameThread());
+					// Now, since we are in the game thread, try to dereference the pointer to
+					// to the UMetasoundGeneratorHandle. This should only succeed if the UObject
+					// hasn't been garbage collected.
+					if (UMetasoundGeneratorHandle* TheHandle = WeakGeneratorHandlePtr.Get())
+					{
+						TheHandle->OnSourceDestroyedAGenerator(InAudioComponentId, InGenerator);
+					}
+				},
+				StatId, nullptr, ENamedThreads::GameThread);
+		});
+}
+
+void UMetasoundGeneratorHandle::AttachGraphChangedDelegate()
+{
+	if (!CachedGeneratorPtr.IsValid() || !OnGeneratorsGraphChanged.IsBound())
+		return;
+
+	TSharedPtr<Metasound::FMetasoundGenerator> Generator = CachedGeneratorPtr.Pin();
+	if (Generator)
+	{
+		// We're about to add a delegate to the generator. This delegate will be called on 
+		// the audio render thread so we need to take steps to assure this UMetasoundGeneratorHandle
+		// hasn't been garbage collected before trying to dereference it later when the delegate "fires".
+		// That is why we capture a TWeakObjectPtr and try to dereference that later!
+		GeneratorGraphChangedDelegateHandle = Generator->OnSetGraph.AddLambda(
+			[WeakGeneratorHandlePtr = TWeakObjectPtr<UMetasoundGeneratorHandle>(this),
+			StatId = this->GetStatID(true)]
+			()
+			{
 				// We are in the audio render thread here, so create a "dispatch task" to be
 				// executed later on the game thread...
 				FFunctionGraphTask::CreateAndDispatchWhenReady(
