Add TRACE macros for profiling

diff --git a/Engine/Plugins/Runtime/Metasound/Source/MetasoundEngine/Private/MetasoundGeneratorHandle.cpp b/Engine/Plugins/Runtime/Metasound/Source/MetasoundEngine/Private/MetasoundGeneratorHandle.cpp
--- a/Engine/Plugins/Runtime/Metasound/Source/MetasoundEngine/Private/MetasoundGeneratorHandle.cpp
+++ b/Engine/Plugins/Runtime/Metasound/Source/MetasoundEngine/Private/MetasoundGeneratorHandle.cpp
@@ -6,6 +6,7 @@
 #include "AudioDeviceManager.h"
 #include "MetasoundGenerator.h"
 #include "MetasoundSource.h"
+#include "MetasoundTrace.h"
 #include "MetasoundParameterPack.h"
 
 UMetasoundGeneratorHandle* UMetasoundGeneratorHandle::CreateMetaSoundGeneratorHandle(UAudioComponent* OnComponent)
@@ -78,6 +79,7 @@ void UMetasoundGeneratorHandle::AttachGeneratorDelegates()
 			FFunctionGraphTask::CreateAndDispatchWhenReady(
 				[WeakGeneratorHandlePtr, InAudioComponentId, InGenerator]()
 				{
+					METASOUND_TRACE_CPUPROFILER_EVENT_SCOPE(UMetasoundGeneratorHandle::CallingGeneratorAttachedDelegates);
 					check(IsInGameThread());
 					// Now, since we are in the game thread, try to dereference the pointer to
 					// to the UMetasoundGeneratorHandle. This should only succeed if the UObject
@@ -99,6 +101,7 @@ void UMetasoundGeneratorHandle::AttachGeneratorDelegates()
 			FFunctionGraphTask::CreateAndDispatchWhenReady(
 				[WeakGeneratorHandlePtr, InAudioComponentId, InGenerator]()
 				{
+					METASOUND_TRACE_CPUPROFILER_EVENT_SCOPE(UMetasoundGeneratorHandle::CallingGeneratorDetachedDelegates);
 					check(IsInGameThread());
 					// Now, since we are in the game thread, try to dereference the pointer to
 					// to the UMetasoundGeneratorHandle. This should only succeed if the UObject
@@ -134,6 +137,7 @@ void UMetasoundGeneratorHandle::AttachGraphChangedDelegate()
 				FFunctionGraphTask::CreateAndDispatchWhenReady(
 					[WeakGeneratorHandlePtr]()
 					{
+						METASOUND_TRACE_CPUPROFILER_EVENT_SCOPE(UMetasoundGeneratorHandle::CallingGraphChangedDelegates);
 						check(IsInGameThread());
 						// Now, since we are in the game thread, try to dereference the pointer to
 						// to the UMetasoundGeneratorHandle. This should only succeed if the UObject
