More MetaSound Interface registry clean-up
- Move bIsDefault/bCanEdit(renamed bIsModifiable) from InterfaceRegistryEntry to Interface data definition
- Add checks in builder to disallow mutation of interfaces that shouldn't be via editor or document builder (ex. SourceInterface which is default and then never changed)
- Optimize SearchEngine query for finding default interfaces
- Tweak MetaSoundUObjectRegistry to support non MetasoundAssetBase classes and add UMetaSoundBuilderDocument. Update iterator w/optional param to only return asset types (true by default for back compat)
- Move FName class name comparitors to use new TopLevelAssetPath
#rnx
#p4v-preflight-copy 24658328

diff --git a/Engine/Plugins/Runtime/Metasound/Source/MetasoundEditor/Private/MetasoundDetailCustomization.cpp b/Engine/Plugins/Runtime/Metasound/Source/MetasoundEditor/Private/MetasoundDetailCustomization.cpp
--- a/Engine/Plugins/Runtime/Metasound/Source/MetasoundEditor/Private/MetasoundDetailCustomization.cpp
+++ b/Engine/Plugins/Runtime/Metasound/Source/MetasoundEditor/Private/MetasoundDetailCustomization.cpp
@@ -2,7 +2,6 @@
 #include "MetasoundDetailCustomization.h"
 
 #include "Containers/Set.h"
-#include "CoreMinimal.h"
 #include "Delegates/Delegate.h"
 #include "DetailCategoryBuilder.h"
 #include "DetailLayoutBuilder.h"
@@ -529,8 +528,6 @@ namespace Metasound
 
 		void FMetasoundInterfacesDetailCustomization::UpdateInterfaceNames()
 		{
-			using namespace Frontend;
-
 			AddableInterfaceNames.Reset();
 			ImplementedInterfaceNames.Reset();
 
@@ -539,14 +536,20 @@ namespace Metasound
 			{
 				auto GetVersionName = [](const FMetasoundFrontendVersion& Version) { return Version.Name; };
 				const UClass* MetaSoundClass = MetaSoundObject->GetClass();
-				auto CanAddOrRemoveInterface = [ClassName = MetaSoundClass->GetFName()](const FMetasoundFrontendVersion& Version)
+				auto CanAddOrRemoveInterface = [ClassName = MetaSoundClass->GetClassPathName()](const FMetasoundFrontendVersion& Version)
 				{
-					using namespace Metasound::Frontend;
+					using namespace Frontend;
 
 					const FInterfaceRegistryKey Key = GetInterfaceRegistryKey(Version);
 					if (const IInterfaceRegistryEntry* Entry = IInterfaceRegistry::Get().FindInterfaceRegistryEntry(Key))
 					{
-						return Entry->EditorCanAddOrRemove(ClassName);
+						if (const FMetasoundFrontendInterfaceUClassOptions* Options = Entry->GetInterface().FindClassOptions(ClassName))
+						{
+							return Options->bIsModifiable;
+						}
+
+						// If no options are found for the given class, interface is modifiable by default.
+						return true;
 					}
 
 					return false;
@@ -555,7 +558,7 @@ namespace Metasound
 				const TSet<FMetasoundFrontendVersion>& ImplementedInterfaces = MetaSoundAsset->GetDocumentChecked().Interfaces;
 				Algo::TransformIf(ImplementedInterfaces, ImplementedInterfaceNames, CanAddOrRemoveInterface, GetVersionName);
 
-				TArray<FMetasoundFrontendInterface> Interfaces = ISearchEngine::Get().FindAllInterfaces();
+				TArray<FMetasoundFrontendInterface> Interfaces = Frontend::ISearchEngine::Get().FindAllInterfaces();
 				for (const FMetasoundFrontendInterface& Interface : Interfaces)
 				{
 					if (!ImplementedInterfaceNames.Contains(Interface.Version.Name))
diff --git a/Engine/Plugins/Runtime/Metasound/Source/MetasoundEditor/Private/MetasoundEditorGraphBuilder.cpp b/Engine/Plugins/Runtime/Metasound/Source/MetasoundEditor/Private/MetasoundEditorGraphBuilder.cpp
--- a/Engine/Plugins/Runtime/Metasound/Source/MetasoundEditor/Private/MetasoundEditorGraphBuilder.cpp
+++ b/Engine/Plugins/Runtime/Metasound/Source/MetasoundEditor/Private/MetasoundEditorGraphBuilder.cpp
@@ -1566,8 +1566,8 @@ namespace Metasound
 				UClass* AssetClass = InMetaSound.GetClass();
 				check(AssetClass);
 
-				TArray<FMetasoundFrontendInterface> InitInterfaces = ISearchEngine::Get().FindUClassDefaultInterfaces(AssetClass->GetFName());
-				FModifyRootGraphInterfaces({ }, InitInterfaces).Transform(Doc);
+				TArray<FMetasoundFrontendVersion> InitVersions = ISearchEngine::Get().FindUClassDefaultInterfaceVersions(AssetClass->GetClassPathName());
+				FModifyRootGraphInterfaces({ }, InitVersions).Transform(Doc);
 			}
 
 			FGraphHandle GraphHandle = MetaSoundAsset->GetRootGraphHandle();
diff --git a/Engine/Plugins/Runtime/Metasound/Source/MetasoundEditor/Private/MetasoundEditorGraphSchema.cpp b/Engine/Plugins/Runtime/Metasound/Source/MetasoundEditor/Private/MetasoundEditorGraphSchema.cpp
--- a/Engine/Plugins/Runtime/Metasound/Source/MetasoundEditor/Private/MetasoundEditorGraphSchema.cpp
+++ b/Engine/Plugins/Runtime/Metasound/Source/MetasoundEditor/Private/MetasoundEditorGraphSchema.cpp
@@ -1652,10 +1652,12 @@ void UMetasoundEditorGraphSchema::GetAssetsGraphHoverMessage(const TArray<FAsset
 
 	OutTooltipText = TEXT("Add MetaSound reference to Graph.");
 
-	IMetasoundEditorModule& EditorModule = FModuleManager::GetModuleChecked<IMetasoundEditorModule>("MetaSoundEditor");
 	for (const FAssetData& Data : Assets)
 	{
-		if (!EditorModule.IsMetaSoundAssetClass(Data.GetClass()->GetClassPathName()))
+		UClass* AssetClass = Data.GetClass();
+		check(AssetClass);
+		const bool bIsRegisteredClass = IMetasoundUObjectRegistry::Get().IsRegisteredClass(*AssetClass);
+		if (!bIsRegisteredClass)
 		{
 			OutOkIcon = false;
 			OutTooltipText = TEXT("Asset(s) must all be MetaSounds.");
diff --git a/Engine/Plugins/Runtime/Metasound/Source/MetasoundEditor/Private/MetasoundEditorModule.cpp b/Engine/Plugins/Runtime/Metasound/Source/MetasoundEditor/Private/MetasoundEditorModule.cpp
--- a/Engine/Plugins/Runtime/Metasound/Source/MetasoundEditor/Private/MetasoundEditorModule.cpp
+++ b/Engine/Plugins/Runtime/Metasound/Source/MetasoundEditor/Private/MetasoundEditorModule.cpp
@@ -287,25 +287,30 @@ namespace Metasound
 			{
 				using namespace Frontend;
 
-				if (!IsMetaSoundAssetClass(InAssetData.AssetClassPath))
+				// If an object's class could not be found, ignore this asset.  This can hit for non-MetaSound assets
+				// and it is up to the system in charge of interacting with that asset or the loading behavior to
+				// report the failed load of the class.
+				if (const UClass* AssetClass = InAssetData.GetClass())
 				{
-					return;
-				}
-
-				check(GEngine);
-				UMetaSoundAssetSubsystem* AssetSubsystem = GEngine->GetEngineSubsystem<UMetaSoundAssetSubsystem>();
-				check(AssetSubsystem);
+					const bool bIsRegisteredClass = IMetasoundUObjectRegistry::Get().IsRegisteredClass(*AssetClass);
+					if (bIsRegisteredClass)
+					{
+						check(GEngine);
+						UMetaSoundAssetSubsystem* AssetSubsystem = GEngine->GetEngineSubsystem<UMetaSoundAssetSubsystem>();
+						check(AssetSubsystem);
 
-				const FNodeRegistryKey RegistryKey = AssetSubsystem->AddOrUpdateAsset(InAssetData);
+						const FNodeRegistryKey RegistryKey = AssetSubsystem->AddOrUpdateAsset(InAssetData);
 
-				// Can be invalid if being called for the first time on an asset before FRenameRootGraphClass is called
-				if (NodeRegistryKey::IsValid(RegistryKey))
-				{
-					const bool bPrimeRequested = AssetPrimeStatus > EAssetPrimeStatus::NotRequested;
-					const bool bIsRegistered = FMetasoundFrontendRegistryContainer::Get()->IsNodeRegistered(RegistryKey);
-					if (bPrimeRequested && !bIsRegistered)
-					{
-						LoadAndRegisterAsset(InAssetData);
+						// Can be invalid if being called for the first time on an asset before FRenameRootGraphClass is called
+						if (NodeRegistryKey::IsValid(RegistryKey))
+						{
+							const bool bPrimeRequested = AssetPrimeStatus > EAssetPrimeStatus::NotRequested;
+							const bool bIsRegistered = FMetasoundFrontendRegistryContainer::Get()->IsNodeRegistered(RegistryKey);
+							if (bPrimeRequested && !bIsRegistered)
+							{
+								LoadAndRegisterAsset(InAssetData);
+							}
+						}
 					}
 				}
 			}
@@ -314,23 +319,28 @@ namespace Metasound
 			{
 				using namespace Frontend;
 
-				if (!IsMetaSoundAssetClass(InAssetData.AssetClassPath))
+				// If an object's class could not be found, ignore this asset.  This can hit for non-MetaSound assets
+				// and it is up to the system in charge of interacting with that asset or the loading behavior to
+				// report the failed load of the class.
+				if (const UClass* AssetClass = InAssetData.GetClass())
 				{
-					return;
-				}
-
-				check(GEngine);
-				UMetaSoundAssetSubsystem* AssetSubsystem = GEngine->GetEngineSubsystem<UMetaSoundAssetSubsystem>();
-				check(AssetSubsystem);
+					const bool bIsRegisteredClass = IMetasoundUObjectRegistry::Get().IsRegisteredClass(*AssetClass);
+					if (bIsRegisteredClass)
+					{
+						check(GEngine);
+						UMetaSoundAssetSubsystem* AssetSubsystem = GEngine->GetEngineSubsystem<UMetaSoundAssetSubsystem>();
+						check(AssetSubsystem);
 
-				const FNodeRegistryKey RegistryKey = AssetSubsystem->AddOrUpdateAsset(InAssetData);
-				const bool bPrimeRequested = AssetPrimeStatus > EAssetPrimeStatus::NotRequested;
-				const bool bIsRegistered = FMetasoundFrontendRegistryContainer::Get()->IsNodeRegistered(RegistryKey);
+						const FNodeRegistryKey RegistryKey = AssetSubsystem->AddOrUpdateAsset(InAssetData);
+						const bool bPrimeRequested = AssetPrimeStatus > EAssetPrimeStatus::NotRequested;
+						const bool bIsRegistered = FMetasoundFrontendRegistryContainer::Get()->IsNodeRegistered(RegistryKey);
 
-				// Have to re-register even if prime was not requested to avoid registry desync.
-				if (bPrimeRequested || bIsRegistered)
-				{
-					LoadAndRegisterAsset(InAssetData);
+						// Have to re-register even if prime was not requested to avoid registry desync.
+						if (bPrimeRequested || bIsRegistered)
+						{
+							LoadAndRegisterAsset(InAssetData);
+						}
+					}	
 				}
 			}
 
@@ -350,11 +360,22 @@ namespace Metasound
 					return;
 				}
 
+				auto IsAssetMetaSound = [](const UObject* Obj)
+				{
+					check(Obj);
+					if (const UClass* AssetClass = Obj->GetClass())
+					{
+						return IMetasoundUObjectRegistry::Get().IsRegisteredClass(*AssetClass);
+					}
+					
+					return false;	
+				};
+
 				for (const TPair<UObject*, UObject*>& Pair : InPackageReloadedEvent->GetRepointedObjects())
 				{
 					if (UObject* Obj = Pair.Key)
 					{
-						if (IsMetaSoundAssetClass(Obj->GetClass()->GetClassPathName()))
+						if (IsAssetMetaSound(Obj))
 						{
 							check(GEngine);
 							UMetaSoundAssetSubsystem* AssetSubsystem = GEngine->GetEngineSubsystem<UMetaSoundAssetSubsystem>();
@@ -368,7 +389,7 @@ namespace Metasound
 
 					if (UObject* Obj = Pair.Value)
 					{
-						if (IsMetaSoundAssetClass(Obj->GetClass()->GetClassPathName()))
+						if (IsAssetMetaSound(Obj))
 						{
 							check(GEngine);
 							UMetaSoundAssetSubsystem* AssetSubsystem = GEngine->GetEngineSubsystem<UMetaSoundAssetSubsystem>();
@@ -403,37 +424,45 @@ namespace Metasound
 
 			void RemoveAssetFromClassRegistry(const FAssetData& InAssetData)
 			{
-				if (IsMetaSoundAssetClass(InAssetData.AssetClassPath))
+				if (const UClass* AssetClass = InAssetData.GetClass())
 				{
-					check(GEngine);
-					UMetaSoundAssetSubsystem* AssetSubsystem = GEngine->GetEngineSubsystem<UMetaSoundAssetSubsystem>();
-					check(AssetSubsystem);
-
-					// Use the editor version of UnregisterWithFrontend so it refreshes any open MetaSound editors
-					AssetSubsystem->RemoveAsset(InAssetData);
-					if (UObject* AssetObject = InAssetData.GetAsset())
+					const bool bIsRegisteredClass = IMetasoundUObjectRegistry::Get().IsRegisteredClass(*AssetClass);
+					if (bIsRegisteredClass)
 					{
-						FGraphBuilder::UnregisterGraphWithFrontend(*AssetObject);
+						check(GEngine);
+						UMetaSoundAssetSubsystem* AssetSubsystem = GEngine->GetEngineSubsystem<UMetaSoundAssetSubsystem>();
+						check(AssetSubsystem);
+
+						// Use the editor version of UnregisterWithFrontend so it refreshes any open MetaSound editors
+						AssetSubsystem->RemoveAsset(InAssetData);
+						if (UObject* AssetObject = InAssetData.GetAsset())
+						{
+							FGraphBuilder::UnregisterGraphWithFrontend(*AssetObject);
+						}
 					}
 				}
 			}
 
 			void RenameAssetInClassRegistry(const FAssetData& InAssetData, const FString& InOldObjectPath)
 			{
-				if (IsMetaSoundAssetClass(InAssetData.AssetClassPath))
+				if (const UClass* AssetClass = InAssetData.GetClass())
 				{
-					check(GEngine);
-					UMetaSoundAssetSubsystem* AssetSubsystem = GEngine->GetEngineSubsystem<UMetaSoundAssetSubsystem>();
-					check(AssetSubsystem);
-
-					// Use the FGraphBuilder Register call instead of registering via the
-					// MetaSoundAssetSubsystem so as to properly refresh respective open editors.
-					constexpr bool bReregisterWithFrontend = false;
-					AssetSubsystem->RenameAsset(InAssetData, bReregisterWithFrontend);
-
-					constexpr bool bForceViewSynchronization = true;
-					UObject* AssetObject = InAssetData.GetAsset();
-					FGraphBuilder::RegisterGraphWithFrontend(*AssetObject, bForceViewSynchronization);
+					const bool bIsRegisteredClass = IMetasoundUObjectRegistry::Get().IsRegisteredClass(*AssetClass);
+					if (bIsRegisteredClass)
+					{
+						check(GEngine);
+						UMetaSoundAssetSubsystem* AssetSubsystem = GEngine->GetEngineSubsystem<UMetaSoundAssetSubsystem>();
+						check(AssetSubsystem);
+
+						// Use the FGraphBuilder Register call instead of registering via the
+						// MetaSoundAssetSubsystem so as to properly refresh respective open editors.
+						constexpr bool bReregisterWithFrontend = false;
+						AssetSubsystem->RenameAsset(InAssetData, bReregisterWithFrontend);
+
+						constexpr bool bForceViewSynchronization = true;
+						UObject* AssetObject = InAssetData.GetAsset();
+						FGraphBuilder::RegisterGraphWithFrontend(*AssetObject, bForceViewSynchronization);
+					}
 				}
 			}
 
@@ -609,8 +638,14 @@ namespace Metasound
 				{
 					AssetPrimeStatus = EAssetPrimeStatus::InProgress;
 
+					TArray<FTopLevelAssetPath> ClassNames;
+					IMetasoundUObjectRegistry::Get().IterateRegisteredUClasses([&ClassNames](UClass& InClass)
+					{
+						ClassNames.Add(InClass.GetClassPathName());
+					});
+
 					FARFilter Filter;
-					Filter.ClassPaths = MetaSoundClassNames;
+					Filter.ClassPaths = ClassNames;
 
 					FAssetRegistryModule& AssetRegistryModule = FModuleManager::LoadModuleChecked<FAssetRegistryModule>("AssetRegistry");
 					AssetRegistryModule.Get().EnumerateAssets(Filter, [this](const FAssetData& AssetData)
@@ -680,8 +715,12 @@ namespace Metasound
 
 			virtual bool IsMetaSoundAssetClass(const FTopLevelAssetPath& InClassName) const override
 			{
-				// TODO: Move to IMetasoundUObjectRegistry (overload IsRegisteredClass to take in class name?)
-				return MetaSoundClassNames.Contains(InClassName);
+				if (const UClass* ClassObject = FindObject<const UClass>(InClassName))
+				{
+					return IMetasoundUObjectRegistry::Get().IsRegisteredClass(*ClassObject);
+				}
+
+				return false;
 			}
 
 			virtual void StartupModule() override
@@ -757,12 +796,6 @@ namespace Metasound
 					GetMutableDefault<UMetasoundEditorSettings>()
 				);
 
-				MetaSoundClassNames.Add(UMetaSoundPatch::StaticClass()->GetClassPathName());
-				MetaSoundClassNames.Add(UMetaSoundSource::StaticClass()->GetClassPathName());
-
-				// Required to query MetaSound assets (that have been redirected to MetaSoundPatch assets) created before UE Release 5.1
-				MetaSoundClassNames.Add(FTopLevelAssetPath(TEXT("/Script/MetasoundEngine.MetaSound")));
-
 				FAssetTypeActions_MetaSoundPatch::RegisterMenuActions();
 				FAssetTypeActions_MetaSoundSource::RegisterMenuActions();
 
@@ -832,11 +865,8 @@ namespace Metasound
 
 				AssetActions.Reset();
 				PinTypes.Reset();
-				MetaSoundClassNames.Reset();
 			}
 
-			TArray<FTopLevelAssetPath> MetaSoundClassNames;
-
 			TArray<TSharedPtr<FAssetTypeActions_Base>> AssetActions;
 			TMap<EMetasoundFrontendLiteralType, const TSubclassOf<UMetasoundEditorGraphMemberDefaultLiteral>> InputDefaultLiteralClassRegistry;
 			TMap<FName, FEdGraphPinType> PinTypes;
diff --git a/Engine/Plugins/Runtime/Metasound/Source/MetasoundEditor/Private/MetasoundEditorSubsystem.cpp b/Engine/Plugins/Runtime/Metasound/Source/MetasoundEditor/Private/MetasoundEditorSubsystem.cpp
--- a/Engine/Plugins/Runtime/Metasound/Source/MetasoundEditor/Private/MetasoundEditorSubsystem.cpp
+++ b/Engine/Plugins/Runtime/Metasound/Source/MetasoundEditor/Private/MetasoundEditorSubsystem.cpp
@@ -21,7 +21,10 @@ TScriptInterface<IMetaSoundDocumentInterface> UMetaSoundEditorSubsystem::BuildTo
 	{
 		constexpr UObject* Parent = nullptr;
 		constexpr UFactory* Factory = nullptr;
-		if (UObject* NewMetaSound = IAssetTools::Get().CreateAsset(AssetName, PackagePath, &InBuilder->GetBuilderUClass(), Factory))
+
+		// Not about to follow this lack of const correctness down a multidecade in the works rabbit hole.
+		UClass& BuilderUClass = const_cast<UClass&>(InBuilder->GetBaseMetaSoundUClass());
+		if (UObject* NewMetaSound = IAssetTools::Get().CreateAsset(AssetName, PackagePath, &BuilderUClass, Factory))
 		{
 			InBuilder->InitNodeLocations();
 			InBuilder->SetAuthor(Author);
diff --git a/Engine/Plugins/Runtime/Metasound/Source/MetasoundEditor/Public/MetasoundEditorModule.h b/Engine/Plugins/Runtime/Metasound/Source/MetasoundEditor/Public/MetasoundEditorModule.h
--- a/Engine/Plugins/Runtime/Metasound/Source/MetasoundEditor/Public/MetasoundEditorModule.h
+++ b/Engine/Plugins/Runtime/Metasound/Source/MetasoundEditor/Public/MetasoundEditorModule.h
@@ -88,6 +88,7 @@ namespace Metasound
 			// By default, proxy classes support child classes & inheritance.
 			virtual void RegisterExplicitProxyClass(const UClass& InClass) = 0;
 
+			UE_DEPRECATED(5.3, "IsMetaSoundAssetClass is deprecated, use IMetasoundUObjectRegistry::IsRegisteredClass")
 			virtual bool IsMetaSoundAssetClass(const FTopLevelAssetPath& InClassName) const = 0;
 
 			// Primes MetaSound assets, effectively loading the asset asynchronously (if not already
diff --git a/Engine/Plugins/Runtime/Metasound/Source/MetasoundEngine/Private/Interfaces/MetasoundDeprecatedInterfaces.cpp b/Engine/Plugins/Runtime/Metasound/Source/MetasoundEngine/Private/Interfaces/MetasoundDeprecatedInterfaces.cpp
--- a/Engine/Plugins/Runtime/Metasound/Source/MetasoundEngine/Private/Interfaces/MetasoundDeprecatedInterfaces.cpp
+++ b/Engine/Plugins/Runtime/Metasound/Source/MetasoundEngine/Private/Interfaces/MetasoundDeprecatedInterfaces.cpp
@@ -1,6 +1,7 @@
 // Copyright Epic Games, Inc. All Rights Reserved.
 #include "Interfaces/MetasoundDeprecatedInterfaces.h"
 
+#include "AudioParameter.h"
 #include "IAudioParameterInterfaceRegistry.h"
 #include "Interfaces/MetasoundFrontendInterfaceRegistry.h"
 #include "Interfaces/MetasoundInterface.h"
@@ -21,6 +22,20 @@ namespace Metasound::Engine
 {
 	namespace DeprecatedInterfacesPrivate
 	{
+		const TArray<FMetasoundFrontendInterfaceUClassOptions>& GetDeprecatedClassOptions()
+		{
+			auto GetOptions = []()
+			{
+				constexpr bool bIsModifiable = false;
+				TArray<FMetasoundFrontendInterfaceUClassOptions> Options;
+				Options.Add(FMetasoundFrontendInterfaceUClassOptions(UMetaSoundPatch::StaticClass()->GetClassPathName(), bIsModifiable));
+				Options.Add(FMetasoundFrontendInterfaceUClassOptions(UMetaSoundSource::StaticClass()->GetClassPathName(), bIsModifiable));
+				return Options;
+			};
+			static const TArray<FMetasoundFrontendInterfaceUClassOptions> Options = GetOptions();
+			return Options;
+		}
+
 		auto MatchMemberNamesIgnoreSpaces = [] (FName InNameA, FName InNameB)
 		{
 			FName ParamA;
@@ -71,6 +86,7 @@ namespace Metasound::Engine
 		{
 			FMetasoundFrontendInterface Interface;
 			Interface.Version = GetVersion();
+			Interface.UClassOptions = DeprecatedInterfacesPrivate::GetDeprecatedClassOptions();
 			return Interface;
 		}
 	}
@@ -93,7 +109,8 @@ namespace Metasound::Engine
 		{
 			FMetasoundFrontendInterface Interface;
 			Interface.Version = GetVersion();
-				
+			Interface.UClassOptions = DeprecatedInterfacesPrivate::GetDeprecatedClassOptions();
+
 			// Inputs
 			FMetasoundFrontendClassVertex OnPlayTrigger;
 			OnPlayTrigger.Name = DeprecatedInterfacesPrivate::GetOnPlayInputName();
@@ -160,7 +177,8 @@ namespace Metasound::Engine
 		{
 			FMetasoundFrontendInterface Interface;
 			Interface.Version = GetVersion();
-				
+			Interface.UClassOptions = DeprecatedInterfacesPrivate::GetDeprecatedClassOptions();
+
 			// Inputs
 			FMetasoundFrontendClassVertex OnPlayTrigger;
 			OnPlayTrigger.Name = DeprecatedInterfacesPrivate::GetOnPlayInputName();
@@ -242,6 +260,7 @@ namespace Metasound::Engine
 		{
 			FMetasoundFrontendInterface Interface;
 			Interface.Version = GetVersion();
+			Interface.UClassOptions = DeprecatedInterfacesPrivate::GetDeprecatedClassOptions();
 
 			// Inputs
 			FMetasoundFrontendClassVertex OnPlayTrigger;
@@ -405,7 +424,8 @@ namespace Metasound::Engine
 		{
 			FMetasoundFrontendInterface Interface;
 			Interface.Version = GetVersion();
-				
+			Interface.UClassOptions = DeprecatedInterfacesPrivate::GetDeprecatedClassOptions();
+
 			// Inputs
 			FMetasoundFrontendClassVertex OnPlayTrigger;
 			OnPlayTrigger.Name = DeprecatedInterfacesPrivate::GetOnPlayInputName();
@@ -542,6 +562,7 @@ namespace Metasound::Engine
 		{
 			FMetasoundFrontendInterface Interface;
 			Interface.Version = GetVersion();
+			Interface.UClassOptions = DeprecatedInterfacesPrivate::GetDeprecatedClassOptions();
 			return Interface;
 		}
 
@@ -579,6 +600,7 @@ namespace Metasound::Engine
 		{
 			FMetasoundFrontendInterface Interface;
 			Interface.Version = GetVersion();
+			Interface.UClassOptions = DeprecatedInterfacesPrivate::GetDeprecatedClassOptions();
 			return Interface;
 		}
 
@@ -608,24 +630,25 @@ namespace Metasound::Engine
 	{
 		using namespace Frontend;
 
-		const UClass& PatchClass = *UMetaSoundPatch::StaticClass();
-		const UClass& SourceClass = *UMetaSoundSource::StaticClass();
-
-		// Default Patch Interfaces (legacy)
-		RegisterInterfaceForSingleClass(PatchClass, MetasoundV1_0::GetInterface(), nullptr, true, false);
+		IInterfaceRegistry& Reg = IInterfaceRegistry::Get();
 
-		RegisterInterfaceForSingleClass(SourceClass, SourceInterfaceV1_0::CreateInterface(), nullptr);
+		Reg.RegisterInterface(MakeUnique<FInterfaceRegistryEntry>(MetasoundV1_0::GetInterface()));
+		Reg.RegisterInterface(MakeUnique<FInterfaceRegistryEntry>(SourceInterfaceV1_0::CreateInterface(*UMetaSoundSource::StaticClass())));
 
 		// Set default interface with unset version to use base UMetaSoundPatch class implementation (legacy requirement for 5.0 alpha).
-		RegisterInterfaceForSingleClass(PatchClass, FMetasoundFrontendInterface(), nullptr);
+		{
+			FMetasoundFrontendInterface DeprecatedInterface;
+			DeprecatedInterface.UClassOptions = DeprecatedInterfacesPrivate::GetDeprecatedClassOptions();
+			Reg.RegisterInterface(MakeUnique<FInterfaceRegistryEntry>(MoveTemp(DeprecatedInterface)));
+		}
 
-		RegisterInterfaceForSingleClass(SourceClass, MetasoundOutputFormatStereoV1_0::GetInterface(), nullptr);
-		RegisterInterfaceForSingleClass(SourceClass, MetasoundOutputFormatStereoV1_1::GetInterface(), MakeUnique<MetasoundOutputFormatStereoV1_1::FUpdateInterface>());
-		RegisterInterfaceForSingleClass(SourceClass, MetasoundOutputFormatStereoV1_2::GetInterface(), MakeUnique<MetasoundOutputFormatStereoV1_2::FUpdateInterface>());
+		Reg.RegisterInterface(MakeUnique<FInterfaceRegistryEntry>(MetasoundOutputFormatStereoV1_0::GetInterface()));
+		Reg.RegisterInterface(MakeUnique<FInterfaceRegistryEntry>(MetasoundOutputFormatStereoV1_1::GetInterface(), MakeUnique<MetasoundOutputFormatStereoV1_1::FUpdateInterface>()));
+		Reg.RegisterInterface(MakeUnique<FInterfaceRegistryEntry>(MetasoundOutputFormatStereoV1_2::GetInterface(), MakeUnique<MetasoundOutputFormatStereoV1_2::FUpdateInterface>()));
 
-		RegisterInterfaceForSingleClass(SourceClass, MetasoundOutputFormatMonoV1_0::GetInterface(), nullptr);
-		RegisterInterfaceForSingleClass(SourceClass, MetasoundOutputFormatMonoV1_1::GetInterface(), MakeUnique<MetasoundOutputFormatMonoV1_1::FUpdateInterface>());
-		RegisterInterfaceForSingleClass(SourceClass, MetasoundOutputFormatMonoV1_2::GetInterface(), MakeUnique<MetasoundOutputFormatMonoV1_2::FUpdateInterface>());
+		Reg.RegisterInterface(MakeUnique<FInterfaceRegistryEntry>(MetasoundOutputFormatMonoV1_0::GetInterface()));
+		Reg.RegisterInterface(MakeUnique<FInterfaceRegistryEntry>(MetasoundOutputFormatMonoV1_1::GetInterface(), MakeUnique<MetasoundOutputFormatMonoV1_1::FUpdateInterface>()));
+		Reg.RegisterInterface(MakeUnique<FInterfaceRegistryEntry>(MetasoundOutputFormatMonoV1_2::GetInterface(), MakeUnique<MetasoundOutputFormatMonoV1_2::FUpdateInterface>()));
 	}
 } // namespace Metasound::Engine
 #undef LOCTEXT_NAMESPACE // MetasoundEngineDeprecatedInterfaces
diff --git a/Engine/Plugins/Runtime/Metasound/Source/MetasoundEngine/Private/Interfaces/MetasoundInputFormatInterfaces.cpp b/Engine/Plugins/Runtime/Metasound/Source/MetasoundEngine/Private/Interfaces/MetasoundInputFormatInterfaces.cpp
--- a/Engine/Plugins/Runtime/Metasound/Source/MetasoundEngine/Private/Interfaces/MetasoundInputFormatInterfaces.cpp
+++ b/Engine/Plugins/Runtime/Metasound/Source/MetasoundEngine/Private/Interfaces/MetasoundInputFormatInterfaces.cpp
@@ -3,7 +3,9 @@
 
 #include "IAudioParameterInterfaceRegistry.h"
 #include "Interfaces/MetasoundOutputFormatInterfaces.h"
+#include "Metasound.h"
 #include "MetasoundAudioBuffer.h"
+#include "MetasoundSource.h"
 #include "MetasoundTrigger.h"
 #include "Templates/SharedPointer.h"
 #include "UObject/Class.h"
@@ -12,6 +14,18 @@
 #define LOCTEXT_NAMESPACE "MetasoundEngine"
 namespace Metasound::Engine
 {
+	namespace InputFormatPrivate
+	{
+		TArray<Audio::FParameterInterface::FClassOptions> GetUClassOptions()
+		{
+			return
+			{
+				{ UMetaSoundPatch::StaticClass()->GetClassPathName(), true /* bIsModifiable */, false /* bIsDefault */ },
+				{ UMetaSoundSource::StaticClass()->GetClassPathName(), false /* bIsModifiable */ , false /* bIsDefault */ }
+			};
+		};
+	} // namespace InputFormatPrivate
+
 #define AUDIO_PARAMETER_INTERFACE_NAMESPACE "UE.InputFormat.Mono"
 	namespace InputFormatMonoInterface
 	{
@@ -33,6 +47,7 @@ namespace Metasound::Engine
 				FInterface()
 					: FParameterInterface(InputFormatMonoInterface::GetVersion().Name, InputFormatMonoInterface::GetVersion().Number.ToInterfaceVersion())
 				{
+					UClassOptions = InputFormatPrivate::GetUClassOptions();
 					Inputs =
 					{
 						{
@@ -133,6 +148,7 @@ namespace Metasound::Engine
 				FInterface()
 					: FParameterInterface(InputFormatStereoInterface::GetVersion().Name, InputFormatStereoInterface::GetVersion().Number.ToInterfaceVersion())
 				{
+					UClassOptions = InputFormatPrivate::GetUClassOptions();
 					Inputs =
 					{
 						{
diff --git a/Engine/Plugins/Runtime/Metasound/Source/MetasoundEngine/Private/Interfaces/MetasoundInterface.cpp b/Engine/Plugins/Runtime/Metasound/Source/MetasoundEngine/Private/Interfaces/MetasoundInterface.cpp
--- a/Engine/Plugins/Runtime/Metasound/Source/MetasoundEngine/Private/Interfaces/MetasoundInterface.cpp
+++ b/Engine/Plugins/Runtime/Metasound/Source/MetasoundEngine/Private/Interfaces/MetasoundInterface.cpp
@@ -26,318 +26,92 @@
 
 namespace Metasound::Engine
 {
-	// Entry for registered interface.
-	class FInterfaceRegistryEntry : public Frontend::IInterfaceRegistryEntry
+	FInterfaceRegistryEntry::FInterfaceRegistryEntry(FMetasoundFrontendInterface&& InInterface, FName InRouterName)
+		: Interface(MoveTemp(InInterface))
+		, RouterName(InRouterName)
 	{
-		const FInterfaceRegistryUClassOptions* FindOptionsByClassName(FName ClassName) const
-		{
-			auto NameIsEqual = [&ClassName](const FInterfaceRegistryUClassOptions& InOptions)
-			{
-				return InOptions.ClassName == ClassName;
-			};
-
-			return Options.UClassOptions.FindByPredicate(NameIsEqual);
-		}
-
-	public:
-		FInterfaceRegistryEntry(
-			const FMetasoundFrontendInterface& InInterface,
-			TUniquePtr<Frontend::IDocumentTransform>&& InUpdateTransform,
-			FInterfaceRegistryOptions&& InOptions
-		)
-			: Interface(InInterface)
-			, UpdateTransform(MoveTemp(InUpdateTransform))
-			, Options(MoveTemp(InOptions))
-		{
-		}
-
-		virtual bool EditorCanAddOrRemove(FName InUClassName) const override
-		{
-			if (const FInterfaceRegistryUClassOptions* UClassOptions = FindOptionsByClassName(InUClassName))
-			{
-				return UClassOptions->bEditorCanAddOrRemove;
-			}
-
-			return false;
-		}
-
-		virtual bool UClassIsSupported(FName InUClassName) const override
-		{
-			if (const FInterfaceRegistryUClassOptions* UClassOptions = FindOptionsByClassName(InUClassName))
-			{
-				return true;
-			}
-
-			// TODO: Support child asset class types.
-			return false;
-		}
-
-		virtual bool IsDefault(FName InUClassName) const override
-		{
-			if (const FInterfaceRegistryUClassOptions* UClassOptions = FindOptionsByClassName(InUClassName))
-			{
-				return UClassOptions->bIsDefault;
-			}
-
-			return false;
-		}
-
-		virtual FName GetRouterName() const override
-		{
-			return Options.InputSystemName;
-		}
-
-		virtual const FMetasoundFrontendInterface& GetInterface() const override
-		{
-			return Interface;
-		}
-
-		virtual bool UpdateRootGraphInterface(Frontend::FDocumentHandle InDocument) const override
-		{
-			if (UpdateTransform.IsValid())
-			{
-				return UpdateTransform->Transform(InDocument);
-			}
-			return false;
-		}
-
-	private:
-		FMetasoundFrontendInterface Interface;
-		TArray<FMetasoundFrontendInterfaceBinding> OutputBindings;
-		TUniquePtr<Frontend::IDocumentTransform> UpdateTransform;
-		FInterfaceRegistryOptions Options;
-	};
+	}
 
-	FMetasoundFrontendInterface ConvertParameterToFrontendInterface(const Audio::FParameterInterface& InInterface)
+	FInterfaceRegistryEntry::FInterfaceRegistryEntry(const FMetasoundFrontendInterface& InInterface, FName InRouterName)
+		: Interface(InInterface)
+		, RouterName(InRouterName)
 	{
-		using namespace Frontend;
-
-		auto ResolveMemberDataType = [](FName DataType, EAudioParameterType ParamType)
-		{
-			if (!DataType.IsNone())
-			{
-				const bool bIsRegisteredType = Frontend::IDataTypeRegistry::Get().IsRegistered(DataType);
-				if (ensureAlwaysMsgf(bIsRegisteredType, TEXT("Attempting to register Interface member with unregistered DataType '%s'."), *DataType.ToString()))
-				{
-					return DataType;
-				}
-			}
-
-			return ConvertParameterToDataType(ParamType);
-		};
-
-		FMetasoundFrontendInterface FrontendInterface;
-		FrontendInterface.Version = { InInterface.GetName(), FMetasoundFrontendVersionNumber { InInterface.GetVersion().Major, InInterface.GetVersion().Minor } };
-
-		// Transfer all input data from AudioExtension interface struct to FrontendInterface
-		{
-			Algo::Transform(InInterface.GetInputs(), FrontendInterface.Inputs, [&](const Audio::FParameterInterface::FInput& Input)
-			{
-				FMetasoundFrontendClassInput ClassInput;
-				ClassInput.Name = Input.InitValue.ParamName;
-				ClassInput.DefaultLiteral = FMetasoundFrontendLiteral(Input.InitValue);
-				ClassInput.TypeName = ResolveMemberDataType(Input.DataType, Input.InitValue.ParamType);
-
-#if WITH_EDITOR
-				// Interfaces should never serialize text to avoid desync between
-				// copied versions serialized in assets and those defined in code.
-				ClassInput.Metadata.SetSerializeText(false);
-				ClassInput.Metadata.SetDisplayName(Input.DisplayName);
-				ClassInput.Metadata.SetDescription(Input.Description);
-				ClassInput.Metadata.SortOrderIndex = Input.SortOrderIndex;
-					
-				FrontendInterface.AddSortOrderToInputStyle(Input.SortOrderIndex);
+	}
 
-				// Setup required inputs by telling the style that the input is required
-				// This will later be validated against.
-				if (!Input.RequiredText.IsEmpty())
-				{
-					FrontendInterface.AddRequiredInputToStyle(Input.InitValue.ParamName, Input.RequiredText);
-				}
-#endif // WITH_EDITOR
+	FInterfaceRegistryEntry::FInterfaceRegistryEntry(const FMetasoundFrontendInterface& InInterface, TUniquePtr<Frontend::IDocumentTransform>&& InUpdateTransform, FName InRouterName)
+		: Interface(InInterface)
+		, UpdateTransform(MoveTemp(InUpdateTransform))
+		, RouterName(InRouterName)
+	{
+	}
 
-				ClassInput.VertexID = FGuid::NewGuid();
+	FName FInterfaceRegistryEntry::GetRouterName() const
+	{
+		return RouterName;
+	}
 
-				return ClassInput;
-			});
-		}
+	const FMetasoundFrontendInterface& FInterfaceRegistryEntry::GetInterface() const
+	{
+		return Interface;
+	}
 
-		// Transfer all output data from AudioExtension interface struct to FrontendInterface
+	bool FInterfaceRegistryEntry::UpdateRootGraphInterface(Frontend::FDocumentHandle InDocument) const
+	{
+		if (UpdateTransform.IsValid())
 		{
-			Algo::Transform(InInterface.GetOutputs(), FrontendInterface.Outputs, [&](const Audio::FParameterInterface::FOutput& Output)
-			{
-				FMetasoundFrontendClassOutput ClassOutput;
-				ClassOutput.Name = Output.ParamName;
-				ClassOutput.TypeName = ResolveMemberDataType(Output.DataType, Output.ParamType);
-
-#if WITH_EDITOR
-				// Interfaces should never serialize text to avoid desync between
-				// copied versions serialized in assets and those defined in code.
-				ClassOutput.Metadata.SetSerializeText(false);
-				ClassOutput.Metadata.SetDisplayName(Output.DisplayName);
-				ClassOutput.Metadata.SetDescription(Output.Description);
-				ClassOutput.Metadata.SortOrderIndex = Output.SortOrderIndex;
-					
-				FrontendInterface.AddSortOrderToOutputStyle(Output.SortOrderIndex);
-
-				// Setup required outputs by telling the style that the output is required
-				// This will later be validated against.
-				if (!Output.RequiredText.IsEmpty())
-				{
-					FrontendInterface.AddRequiredOutputToStyle(Output.ParamName, Output.RequiredText);
-				}
-#endif // WITH_EDITOR
-
-				ClassOutput.VertexID = FGuid::NewGuid();
-
-				return ClassOutput;
-			});
+			return UpdateTransform->Transform(InDocument);
 		}
-
-		Algo::Transform(InInterface.GetEnvironment(), FrontendInterface.Environment, [&](const Audio::FParameterInterface::FEnvironmentVariable& Environment)
-		{
-			FMetasoundFrontendClassEnvironmentVariable EnvironmentVariable;
-			EnvironmentVariable.Name = Environment.ParamName;
-
-			// Disabled as it isn't used to infer type when getting/setting at a lower level.
-			// TODO: Either remove type info for environment variables all together or enforce type.
-			// EnvironmentVariable.TypeName = ResolveMemberDataType(Environment.DataType, Environment.ParamType);
-
-			return EnvironmentVariable;
-		});
-
-		return FrontendInterface;
+		return false;
 	}
 
 	void RegisterAudioFormatInterfaces()
 	{
 		using namespace Frontend;
 
-		const FName SourceClassName = UMetaSoundSource::StaticClass()->GetFName();
-		const FName PatchClassName = UMetaSoundPatch::StaticClass()->GetFName();
+		IInterfaceRegistry& Reg = IInterfaceRegistry::Get();
 
 		// Input Formats
 		{
-			RegisterInterface(InputFormatMonoInterface::CreateInterface(), nullptr, FInterfaceRegistryOptions
-				{
-					IDataReference::RouterName,
-					{
-						{ PatchClassName, false /* bIsDefault */, true /* bEditorCanAddOrRemove */ },
-						{ SourceClassName },
-					}
-				});
-
-			RegisterInterface(InputFormatStereoInterface::CreateInterface(), nullptr, FInterfaceRegistryOptions
-				{
-					IDataReference::RouterName,
-					{
-						{ PatchClassName, false /* bIsDefault */, true /* bEditorCanAddOrRemove */ },
-						{ SourceClassName },
-					}
-				});
+			Reg.RegisterInterface(MakeUnique<FInterfaceRegistryEntry>(FMetasoundFrontendInterface(InputFormatMonoInterface::CreateInterface()), IDataReference::RouterName));
+			Reg.RegisterInterface(MakeUnique<FInterfaceRegistryEntry>(InputFormatStereoInterface::CreateInterface(), IDataReference::RouterName));
 		}
 
 		// Output Formats
 		{
-			RegisterInterface(OutputFormatMonoInterface::CreateInterface(), nullptr, FInterfaceRegistryOptions
-				{
-					IDataReference::RouterName,
-					{
-						{ PatchClassName, false /* bIsDefault */, true /* bEditorCanAddOrRemove */ },
-						{ SourceClassName, true /* bIsDefault */, false /* bEditorCanAddOrRemove */ },
-					}
-				});
-
-			RegisterInterface(OutputFormatStereoInterface::CreateInterface(), nullptr, FInterfaceRegistryOptions
-				{
-					IDataReference::RouterName,
-					{
-						{ PatchClassName, false /* bIsDefault */, true /* bEditorCanAddOrRemove */ },
-						{ SourceClassName },
-					}
-				});
-
-			RegisterInterface(OutputFormatQuadInterface::CreateInterface(), nullptr, FInterfaceRegistryOptions
-				{
-					IDataReference::RouterName,
-					{
-						{ PatchClassName, false /* bIsDefault */, true /* bEditorCanAddOrRemove */ },
-						{ SourceClassName },
-					}
-				});
-
-			RegisterInterface(OutputFormatFiveDotOneInterface::CreateInterface(), nullptr, FInterfaceRegistryOptions
-				{
-					IDataReference::RouterName,
-					{
-						{ PatchClassName, false /* bIsDefault */, true /* bEditorCanAddOrRemove */ },
-						{ SourceClassName },
-					}
-				});
-
-			RegisterInterface(OutputFormatSevenDotOneInterface::CreateInterface(), nullptr, FInterfaceRegistryOptions
-				{
-					IDataReference::RouterName,
-					{
-						{ PatchClassName, false /* bIsDefault */, true /* bEditorCanAddOrRemove */ },
-						{ SourceClassName },
-					}
-				});
+			Reg.RegisterInterface(MakeUnique<FInterfaceRegistryEntry>(OutputFormatMonoInterface::CreateInterface()));
+			Reg.RegisterInterface(MakeUnique<FInterfaceRegistryEntry>(OutputFormatStereoInterface::CreateInterface()));
+			Reg.RegisterInterface(MakeUnique<FInterfaceRegistryEntry>(OutputFormatQuadInterface::CreateInterface()));
+			Reg.RegisterInterface(MakeUnique<FInterfaceRegistryEntry>(OutputFormatFiveDotOneInterface::CreateInterface()));
+			Reg.RegisterInterface(MakeUnique<FInterfaceRegistryEntry>(OutputFormatSevenDotOneInterface::CreateInterface()));
 		}
 	}
 
 	void RegisterExternalInterfaces()
 	{
-		using namespace Frontend;
-
 		// Register External Interfaces (Interfaces defined externally & can be managed directly by end-user).
 		auto RegisterExternalInterface = [](Audio::FParameterInterfacePtr Interface)
 		{
-			FInterfaceRegistryOptions Options{ Audio::IParameterTransmitter::RouterName };
+			using namespace Frontend;
 
-			// Currently, no externally defined interfaces can be added as default for protection against undesired
-			// interfaces automatically being added when creating a new MetaSound asset through the editor. Also,
-			// all parameter interfaces are enabled in the editor for addition/removal.
-			constexpr bool bIsDefault = false;
-			constexpr bool bEditorCanAddOrRemove = true;
-
-			bool bInterfaceSupportsPatch = true;
-			bool bInterfaceSupportsSource = true;
-
-			const TArray<const UClass*> SupportedUClasses = Interface->FindSupportedUClasses();
-			if (!SupportedUClasses.IsEmpty())
+			bool bSupportedInterface = false;
+			IMetasoundUObjectRegistry::Get().IterateRegisteredUClasses([&bSupportedInterface, &Interface](UClass& InRegisteredClass)
 			{
-				bInterfaceSupportsPatch = false;
-				bInterfaceSupportsSource = false;
+				const TArray<const UClass*> SupportedUClasses = Interface->FindSupportedUClasses();
+				bSupportedInterface |= SupportedUClasses.IsEmpty();
 				for (const UClass* SupportedUClass : SupportedUClasses)
 				{
-					const UClass* PatchClass = UMetaSoundPatch::StaticClass();
-					check(PatchClass);
-					bInterfaceSupportsPatch |= PatchClass->IsChildOf(SupportedUClass);
-
-					const UClass* SourceClass = UMetaSoundSource::StaticClass();
-					check(SourceClass);
-					bInterfaceSupportsSource |= SourceClass->IsChildOf(SupportedUClass);
+					check(SupportedUClass);
+					bSupportedInterface |= SupportedUClass->IsChildOf(&InRegisteredClass);
 				}
-			}
+			});
 
-			const bool bSupported = bInterfaceSupportsSource || bInterfaceSupportsPatch;
-			if (ensureAlwaysMsgf(bSupported, TEXT("Parameter interface '%s' type not supported by MetaSounds"), *Interface->GetName().ToString()))
+			if (bSupportedInterface)
 			{
-				const FMetasoundFrontendInterface FrontendInterface = ConvertParameterToFrontendInterface(*Interface);
-				if (bInterfaceSupportsPatch)
-				{
-					const FName ClassName = UMetaSoundPatch::StaticClass()->GetFName();
-					Options.UClassOptions.Add(FInterfaceRegistryUClassOptions{ ClassName, bIsDefault, bEditorCanAddOrRemove });
-				}
-
-				if (bInterfaceSupportsSource)
-				{
-					const FName ClassName = UMetaSoundSource::StaticClass()->GetFName();
-					Options.UClassOptions.Add(FInterfaceRegistryUClassOptions{ ClassName, bIsDefault, bEditorCanAddOrRemove });
-				}
-
-				IInterfaceRegistry::Get().RegisterInterface(MakeUnique<FInterfaceRegistryEntry>(FrontendInterface, nullptr, MoveTemp(Options)));
+				IInterfaceRegistry::Get().RegisterInterface(MakeUnique<FInterfaceRegistryEntry>(FMetasoundFrontendInterface(Interface), Audio::IParameterTransmitter::RouterName));
+			}
+			else
+			{
+				UE_LOG(LogMetaSound, Warning, TEXT("Parameter interface '%s' not supported by MetaSounds"), *Interface->GetName().ToString());
 			}
 		};
 
@@ -345,63 +119,16 @@ namespace Metasound::Engine
 		Audio::IAudioParameterInterfaceRegistry::Get().OnRegistration(MoveTemp(RegisterExternalInterface));
 	}
 
-	void RegisterInterface(const FMetasoundFrontendInterface& InInterface, TUniquePtr<Frontend::IDocumentTransform>&& InUpdateTransform, FInterfaceRegistryOptions&& InOptions)
-	{
-		using namespace Frontend;
-		IInterfaceRegistry::Get().RegisterInterface(MakeUnique<FInterfaceRegistryEntry>(InInterface, MoveTemp(InUpdateTransform), MoveTemp(InOptions)));
-	}
-
-	void RegisterInterface(Audio::FParameterInterfacePtr Interface, TUniquePtr<Frontend::IDocumentTransform>&& InUpdateTransform, FInterfaceRegistryOptions&& InOptions)
-	{
-		using namespace Frontend;
-
-		const FMetasoundFrontendInterface FrontendInterface = ConvertParameterToFrontendInterface(*Interface);
-		IInterfaceRegistry::Get().RegisterInterface(MakeUnique<FInterfaceRegistryEntry>(FrontendInterface, MoveTemp(InUpdateTransform), MoveTemp(InOptions)));
-	}
-
-	void RegisterInterfaceForSingleClass(const UClass& InClass, Audio::FParameterInterfacePtr Interface, TUniquePtr<Frontend::IDocumentTransform>&& InUpdateTransform, bool bInIsDefault, bool bInEditorCanAddOrRemove, FName InRouterName)
-	{
-		using namespace Frontend;
-
-		const FMetasoundFrontendInterface FrontendInterface = ConvertParameterToFrontendInterface(*Interface);
-		RegisterInterface(FrontendInterface, MoveTemp(InUpdateTransform), FInterfaceRegistryOptions
-			{
-				InRouterName,
-				{
-					{
-						InClass.GetFName(),
-						bInIsDefault,
-						bInEditorCanAddOrRemove
-					}
-				}
-			});
-	}
-
-	void RegisterInterfaceForSingleClass(const UClass& InClass, const FMetasoundFrontendInterface& InInterface, TUniquePtr<Frontend::IDocumentTransform>&& InUpdateTransform, bool bInIsDefault, bool bInEditorCanAddOrRemove, FName InRouterName)
-	{
-		using namespace Frontend;
-
-		RegisterInterface(InInterface, MoveTemp(InUpdateTransform), FInterfaceRegistryOptions
-			{
-				InRouterName,
-				{
-					{
-						InClass.GetFName(),
-						bInIsDefault,
-						bInEditorCanAddOrRemove
-					}
-				}
-			});
-	}
-
 	void RegisterInterfaces()
 	{
 		using namespace Frontend;
 
+		IInterfaceRegistry& Reg = IInterfaceRegistry::Get();
+
 		// Default Source Interfaces
 		{
-			RegisterInterfaceForSingleClass(*UMetaSoundSource::StaticClass(), SourceInterface::CreateInterface(), MakeUnique<SourceInterface::FUpdateInterface>(), true, false);
-			RegisterInterfaceForSingleClass(*UMetaSoundSource::StaticClass(), SourceOneShotInterface::CreateInterface(), nullptr, true, true);
+			Reg.RegisterInterface(MakeUnique<FInterfaceRegistryEntry>(SourceInterface::CreateInterface(*UMetaSoundSource::StaticClass()), MakeUnique<SourceInterface::FUpdateInterface>()));
+			Reg.RegisterInterface(MakeUnique<FInterfaceRegistryEntry>(SourceOneShotInterface::CreateInterface(*UMetaSoundSource::StaticClass())));
 		}
 
 		RegisterAudioFormatInterfaces();
diff --git a/Engine/Plugins/Runtime/Metasound/Source/MetasoundEngine/Private/Interfaces/MetasoundInterface.h b/Engine/Plugins/Runtime/Metasound/Source/MetasoundEngine/Private/Interfaces/MetasoundInterface.h
--- a/Engine/Plugins/Runtime/Metasound/Source/MetasoundEngine/Private/Interfaces/MetasoundInterface.h
+++ b/Engine/Plugins/Runtime/Metasound/Source/MetasoundEngine/Private/Interfaces/MetasoundInterface.h
@@ -2,30 +2,30 @@
 #pragma once
 
 #include "IAudioParameterInterfaceRegistry.h"
+#include "Interfaces/MetasoundFrontendInterfaceRegistry.h"
 #include "MetasoundFrontendDocument.h"
 #include "MetasoundFrontendTransform.h"
 
 
 namespace Metasound::Engine
 {
-	struct FInterfaceRegistryUClassOptions
+	// Entry for registered interface.
+	class FInterfaceRegistryEntry : public Frontend::IInterfaceRegistryEntry
 	{
-		FName ClassName;
-		bool bIsDefault = false;
-		bool bEditorCanAddOrRemove = false;
-	};
+	public:
+		FInterfaceRegistryEntry(FMetasoundFrontendInterface&& InInterface, FName InRouterName = IDataReference::RouterName);
+		FInterfaceRegistryEntry(const FMetasoundFrontendInterface& InInterface, FName InRouterName = IDataReference::RouterName);
+		FInterfaceRegistryEntry(const FMetasoundFrontendInterface& InInterface, TUniquePtr<Frontend::IDocumentTransform>&& InUpdateTransform, FName InRouterName = IDataReference::RouterName);
 
-	struct FInterfaceRegistryOptions
-	{
-		FName InputSystemName;
-		TArray<FInterfaceRegistryUClassOptions> UClassOptions;
-	};
+		virtual FName GetRouterName() const override;
+		virtual const FMetasoundFrontendInterface& GetInterface() const override;
+		virtual bool UpdateRootGraphInterface(Frontend::FDocumentHandle InDocument) const override;
 
-	FMetasoundFrontendInterface ConvertParameterToFrontendInterface(const Audio::FParameterInterface& InInterface);
+	private:
+		FMetasoundFrontendInterface Interface;
+		TUniquePtr<Frontend::IDocumentTransform> UpdateTransform;
+		FName RouterName;
+	};
 
 	void RegisterInterfaces();
-	void RegisterInterface(Audio::FParameterInterfacePtr Interface, TUniquePtr<Frontend::IDocumentTransform>&& InUpdateTransform, FInterfaceRegistryOptions&& InOptions);
-	void RegisterInterface(const FMetasoundFrontendInterface& InInterface, TUniquePtr<Frontend::IDocumentTransform>&& InUpdateTransform, FInterfaceRegistryOptions&& InOptions);
-	void RegisterInterfaceForSingleClass(const UClass& InClass, Audio::FParameterInterfacePtr Interface, TUniquePtr<Frontend::IDocumentTransform>&& InUpdateTransform, bool bInIsDefault = false, bool bInEditorCanAddOrRemove = false, FName InRouterName = IDataReference::RouterName);
-	void RegisterInterfaceForSingleClass(const UClass& InClass, const FMetasoundFrontendInterface& InInterface, TUniquePtr<Frontend::IDocumentTransform>&& InUpdateTransform, bool bInIsDefault = false, bool bInEditorCanAddOrRemove = false, FName InRouterName = IDataReference::RouterName);
 } // namespace Metasound::Engine
diff --git a/Engine/Plugins/Runtime/Metasound/Source/MetasoundEngine/Private/Interfaces/MetasoundOutputFormatInterfaces.cpp b/Engine/Plugins/Runtime/Metasound/Source/MetasoundEngine/Private/Interfaces/MetasoundOutputFormatInterfaces.cpp
--- a/Engine/Plugins/Runtime/Metasound/Source/MetasoundEngine/Private/Interfaces/MetasoundOutputFormatInterfaces.cpp
+++ b/Engine/Plugins/Runtime/Metasound/Source/MetasoundEngine/Private/Interfaces/MetasoundOutputFormatInterfaces.cpp
@@ -2,7 +2,9 @@
 #include "Interfaces/MetasoundOutputFormatInterfaces.h"
 
 #include "IAudioParameterInterfaceRegistry.h"
+#include "Metasound.h"
 #include "MetasoundAudioBuffer.h"
+#include "MetasoundSource.h"
 #include "MetasoundTrigger.h"
 #include "Templates/SharedPointer.h"
 #include "UObject/Class.h"
@@ -14,7 +16,7 @@
 
 namespace Metasound::Engine
 {
-	namespace MetasoundOutputFormatInterfacesPrivate
+	namespace OutputFormatPrivate
 	{
 		Audio::FParameterInterface::FOutput GetFrontLeftOutput(const FName& InVertexName)
 		{
@@ -119,7 +121,16 @@ namespace Metasound::Engine
 				107
 			};
 		}
-	} // namespace MetasoundOutputFormatInterfacesPrivate
+
+		TArray<Audio::FParameterInterface::FClassOptions> GetUClassOptions(bool bIsSourceDefault = false)
+		{
+			return
+			{
+				{ UMetaSoundPatch::StaticClass()->GetClassPathName(), true /* bIsModifiable */, false /* bIsDefault */ },
+				{ UMetaSoundSource::StaticClass()->GetClassPathName(),  false /* bIsModifiable */, bIsSourceDefault }
+			};
+		};
+	} // namespace OutputFormatPrivate
 
 #define AUDIO_PARAMETER_INTERFACE_NAMESPACE "UE.OutputFormat.Mono"
 	namespace OutputFormatMonoInterface
@@ -147,6 +158,9 @@ namespace Metasound::Engine
 				FInterface()
 					: FParameterInterface(OutputFormatMonoInterface::GetVersion().Name, OutputFormatMonoInterface::GetVersion().Number.ToInterfaceVersion())
 				{
+					constexpr bool bIsSourceDefault = true;
+					UClassOptions = OutputFormatPrivate::GetUClassOptions(bIsSourceDefault);
+
 					Outputs =
 					{
 						{
@@ -194,6 +208,7 @@ namespace Metasound::Engine
 				FInterface()
 					: FParameterInterface(OutputFormatStereoInterface::GetVersion().Name, OutputFormatStereoInterface::GetVersion().Number.ToInterfaceVersion())
 				{
+					UClassOptions = OutputFormatPrivate::GetUClassOptions();
 					Outputs =
 					{
 						{
@@ -252,8 +267,9 @@ namespace Metasound::Engine
 				FInterface()
 					: FParameterInterface(OutputFormatQuadInterface::GetVersion().Name, OutputFormatQuadInterface::GetVersion().Number.ToInterfaceVersion())
 				{
-					using namespace MetasoundOutputFormatInterfacesPrivate;
+					using namespace OutputFormatPrivate;
 
+					UClassOptions = GetUClassOptions();
 					Outputs =
 					{
 						GetFrontLeftOutput(Outputs::FrontLeftOut),
@@ -300,7 +316,9 @@ namespace Metasound::Engine
 				FInterface()
 					: FParameterInterface(OutputFormatFiveDotOneInterface::GetVersion().Name, OutputFormatFiveDotOneInterface::GetVersion().Number.ToInterfaceVersion())
 				{
-					using namespace MetasoundOutputFormatInterfacesPrivate;
+					using namespace OutputFormatPrivate;
+
+					UClassOptions = GetUClassOptions();
 					Outputs =
 					{
 						GetFrontLeftOutput(Outputs::FrontLeftOut),
@@ -351,7 +369,9 @@ namespace Metasound::Engine
 				FInterface()
 					: FParameterInterface(OutputFormatSevenDotOneInterface::GetVersion().Name, OutputFormatSevenDotOneInterface::GetVersion().Number.ToInterfaceVersion())
 				{
-					using namespace MetasoundOutputFormatInterfacesPrivate;
+					using namespace OutputFormatPrivate;
+
+					UClassOptions = GetUClassOptions();
 					Outputs =
 					{
 						GetFrontLeftOutput(Outputs::FrontLeftOut),
diff --git a/Engine/Plugins/Runtime/Metasound/Source/MetasoundEngine/Private/MetasoundBuilderSubsystem.cpp b/Engine/Plugins/Runtime/Metasound/Source/MetasoundEngine/Private/MetasoundBuilderSubsystem.cpp
--- a/Engine/Plugins/Runtime/Metasound/Source/MetasoundEngine/Private/MetasoundBuilderSubsystem.cpp
+++ b/Engine/Plugins/Runtime/Metasound/Source/MetasoundEngine/Private/MetasoundBuilderSubsystem.cpp
@@ -479,8 +479,12 @@ void UMetaSoundBuilderBase::CreateTransientDocument()
 	TObjectPtr<UMetaSoundBuilderDocument> NewBuilder = NewObject<UMetaSoundBuilderDocument>();
 	TScriptInterface<IMetaSoundDocumentInterface> MetaSoundSourceDocInterface = NewBuilder;
 	Builder = FMetaSoundFrontendDocumentBuilder(MetaSoundSourceDocInterface);
-	const UClass& BuilderClass = GetBuilderUClass();
-	Builder.InitDocument(BuilderClass.GetFName());
+	Builder.InitDocument();
+}
+
+const UClass& UMetaSoundBuilderBase::GetBaseMetaSoundUClass() const
+{
+	return Builder.GetDocumentInterface().GetBaseMetaSoundUClass();
 }
 
 void UMetaSoundBuilderBase::InitNodeLocations()
@@ -613,16 +617,6 @@ TScriptInterface<IMetaSoundDocumentInterface> UMetaSoundPatchBuilder::Build(UObj
 	return BuildInternal<UMetaSoundPatch>(Parent, InBuilderOptions);
 }
 
-UClass& UMetaSoundPatchBuilder::GetBuilderUClass() const
-{
-	return *UMetaSoundPatch::StaticClass();
-}
-
-UClass& UMetaSoundSourceBuilder::GetBuilderUClass() const
-{
-	return *UMetaSoundSource::StaticClass();
-}
-
 void UMetaSoundSourceBuilder::Audition(UObject* Parent, UAudioComponent* AudioComponent, FOnCreateAuditionGeneratorHandleDelegate CreateGenerator)
 {
 	using namespace Metasound;
diff --git a/Engine/Plugins/Runtime/Metasound/Source/MetasoundEngine/Private/MetasoundEngineModule.cpp b/Engine/Plugins/Runtime/Metasound/Source/MetasoundEngine/Private/MetasoundEngineModule.cpp
--- a/Engine/Plugins/Runtime/Metasound/Source/MetasoundEngine/Private/MetasoundEngineModule.cpp
+++ b/Engine/Plugins/Runtime/Metasound/Source/MetasoundEngine/Private/MetasoundEngineModule.cpp
@@ -1,10 +1,17 @@
 // Copyright Epic Games, Inc. All Rights Reserved.
 #include "MetasoundEngineModule.h"
 
+#include "Analysis/MetasoundFrontendAnalyzerRegistry.h"
+#include "Analysis/MetasoundFrontendVertexAnalyzerForwardValue.h"
+#include "Interfaces/MetasoundDeprecatedInterfaces.h"
+#include "Interfaces/MetasoundInterface.h"
+#include "Interfaces/MetasoundInterfaceBindingsPrivate.h"
 #include "Metasound.h"
 #include "MetasoundAudioBus.h"
 #include "MetasoundDataReference.h"
 #include "MetasoundDataTypeRegistrationMacro.h"
+#include "MetasoundFrontendDataTypeTraits.h"
+#include "MetasoundFrontendDocumentBuilder.h"
 #include "MetasoundOutputSubsystem.h"
 #include "MetasoundOutputWatcher.h"
 #include "MetasoundSource.h"
@@ -12,11 +19,6 @@
 #include "MetasoundUObjectRegistry.h"
 #include "MetasoundWave.h"
 #include "MetasoundWaveTable.h"
-#include "Analysis/MetasoundFrontendAnalyzerRegistry.h"
-#include "Analysis/MetasoundFrontendVertexAnalyzerForwardValue.h"
-#include "Interfaces/MetasoundDeprecatedInterfaces.h"
-#include "Interfaces/MetasoundInterface.h"
-#include "Interfaces/MetasoundInterfaceBindingsPrivate.h"
 #include "Modules/ModuleManager.h"
 #include "Sound/AudioSettings.h"
 
@@ -50,6 +52,7 @@ class FMetasoundEngineModule : public IMetasoundEngineModule
 		check(AudioSettings);
 		AudioSettings->RegisterParameterInterfaces();
 
+		IMetasoundUObjectRegistry::Get().RegisterUClass(MakeUnique<TMetasoundUObjectRegistryEntry<UMetaSoundBuilderDocument>>());
 		IMetasoundUObjectRegistry::Get().RegisterUClass(MakeUnique<TMetasoundUObjectRegistryEntry<UMetaSoundPatch>>());
 		IMetasoundUObjectRegistry::Get().RegisterUClass(MakeUnique<TMetasoundUObjectRegistryEntry<UMetaSoundSource>>());
 
diff --git a/Engine/Plugins/Runtime/Metasound/Source/MetasoundEngine/Private/MetasoundUObjectRegistry.cpp b/Engine/Plugins/Runtime/Metasound/Source/MetasoundEngine/Private/MetasoundUObjectRegistry.cpp
--- a/Engine/Plugins/Runtime/Metasound/Source/MetasoundEngine/Private/MetasoundUObjectRegistry.cpp
+++ b/Engine/Plugins/Runtime/Metasound/Source/MetasoundEngine/Private/MetasoundUObjectRegistry.cpp
@@ -37,6 +37,21 @@ namespace Metasound
 				return (nullptr != GetEntryByUObject(InObject));
 			}
 
+			bool IsRegisteredClass(const UClass& InClass) const override
+			{
+				auto IsChildClassOfRegisteredClass = [&](const IMetasoundUObjectRegistryEntry* Entry)
+				{
+					if (Entry)
+					{
+						return Entry->IsChildClass(&InClass);
+					}
+
+					return false;
+				};
+
+				return FindEntryByPredicate(IsChildClassOfRegisteredClass) != nullptr;
+			}
+
 			FMetasoundAssetBase* GetObjectAsAssetBase(UObject* InObject) const override
 			{
 				if (const IMetasoundUObjectRegistryEntry* Entry = GetEntryByUObject(InObject))
@@ -55,22 +70,25 @@ namespace Metasound
 				return nullptr;
 			}
 
-			void IterateRegisteredUClasses(TFunctionRef<void(UClass&)> InFunc) const override
+			void IterateRegisteredUClasses(TFunctionRef<void(UClass&)> InFunc, bool bAssetTypesOnly) const override
 			{
 				for (const IMetasoundUObjectRegistryEntry* Entry : Entries)
 				{
 					if (Entry)
 					{
-						if (UClass* Class = Entry->GetUClass())
+						if (!bAssetTypesOnly || Entry->IsAssetType())
 						{
-							InFunc(*Class);
+							if (UClass* Class = Entry->GetUClass())
+							{
+								InFunc(*Class);
+							}
 						}
 					}
 				}
 			}
 
 		private:
-			const IMetasoundUObjectRegistryEntry* FindEntryByPredicate(TFunction<bool (const IMetasoundUObjectRegistryEntry*)> InPredicate) const
+			const IMetasoundUObjectRegistryEntry* FindEntryByPredicate(TFunctionRef<bool (const IMetasoundUObjectRegistryEntry*)> InPredicate) const
 			{
 				const IMetasoundUObjectRegistryEntry* const* Entry = Entries.FindByPredicate(InPredicate);
 
@@ -82,7 +100,7 @@ namespace Metasound
 				return *Entry;
 			}
 
-			TArray<const IMetasoundUObjectRegistryEntry*> FindEntriesByPredicate(TFunction<bool (const IMetasoundUObjectRegistryEntry*)> InPredicate) const
+			TArray<const IMetasoundUObjectRegistryEntry*> FindEntriesByPredicate(TFunctionRef<bool (const IMetasoundUObjectRegistryEntry*)> InPredicate) const
 			{
 				TArray<const IMetasoundUObjectRegistryEntry*> FoundEntries;
 
diff --git a/Engine/Plugins/Runtime/Metasound/Source/MetasoundEngine/Public/Metasound.h b/Engine/Plugins/Runtime/Metasound/Source/MetasoundEngine/Public/Metasound.h
--- a/Engine/Plugins/Runtime/Metasound/Source/MetasoundEngine/Public/Metasound.h
+++ b/Engine/Plugins/Runtime/Metasound/Source/MetasoundEngine/Public/Metasound.h
@@ -135,6 +135,11 @@ public:
 		return RootMetaSoundDocument;
 	}
 
+	virtual const UClass& GetBaseMetaSoundUClass() const final override
+	{
+		return *UMetaSoundPatch::StaticClass();
+	}
+
 #if WITH_EDITOR
 	virtual void PostDuplicate(EDuplicateMode::Type DuplicateMode) override;
 	virtual void PostEditUndo() override;
diff --git a/Engine/Plugins/Runtime/Metasound/Source/MetasoundEngine/Public/MetasoundBuilderSubsystem.h b/Engine/Plugins/Runtime/Metasound/Source/MetasoundEngine/Public/MetasoundBuilderSubsystem.h
--- a/Engine/Plugins/Runtime/Metasound/Source/MetasoundEngine/Public/MetasoundBuilderSubsystem.h
+++ b/Engine/Plugins/Runtime/Metasound/Source/MetasoundEngine/Public/MetasoundBuilderSubsystem.h
@@ -269,12 +269,12 @@ public:
 
 	virtual TScriptInterface<IMetaSoundDocumentInterface> Build(UObject* Parent, const FMetaSoundBuilderOptions& Options) const PURE_VIRTUAL(UMetaSoundBuilderBase::Build, return { }; );
 
-	// Returns the builder's supported UClass
-	virtual UClass& GetBuilderUClass() const PURE_VIRTUAL(UMetaSoundBuilderBase::GetBuilderUClass, return *UObject::StaticClass(); );
-
 	// Constructs a transient UMetaSoundBuilderDocument & attaches builder to it
 	void CreateTransientDocument();
 
+	// Returns the base MetaSound UClass the builder is operating on (ex. MetaSoundSource, MetaSoundPatch, etc.)
+	const UClass& GetBaseMetaSoundUClass() const;
+
 	// Initializes and ensures all nodes have a position (required prior to exporting to an asset if expected to be viewed in the editor).
 	void InitNodeLocations();
 
@@ -363,8 +363,6 @@ class METASOUNDENGINE_API UMetaSoundPatchBuilder : public UMetaSoundBuilderBase
 public:
 	UFUNCTION(BlueprintCallable, Category = "Audio|MetaSound|Builder", meta = (WorldContext = "Parent"))
 	virtual UPARAM(DisplayName = "MetaSound") TScriptInterface<IMetaSoundDocumentInterface> Build(UObject* Parent, const FMetaSoundBuilderOptions& Options) const override;
-
-	virtual UClass& GetBuilderUClass() const override;
 };
 
 /** Builder in charge of building a MetaSound Source */
@@ -384,8 +382,6 @@ public:
 	UFUNCTION(BlueprintCallable, Category = "Audio|MetaSound|Builder", meta = (ExpandEnumAsExecs = "OutResult"))
 	void SetFormat(EMetaSoundOutputAudioFormat OutputFormat, EMetaSoundBuilderResult& OutResult);
 
-	virtual UClass& GetBuilderUClass() const override;
-
 	const Metasound::Engine::FOutputAudioFormatInfoPair* FindOutputAudioFormatInfo() const;
 
 private:
diff --git a/Engine/Plugins/Runtime/Metasound/Source/MetasoundEngine/Public/MetasoundSource.h b/Engine/Plugins/Runtime/Metasound/Source/MetasoundEngine/Public/MetasoundSource.h
--- a/Engine/Plugins/Runtime/Metasound/Source/MetasoundEngine/Public/MetasoundSource.h
+++ b/Engine/Plugins/Runtime/Metasound/Source/MetasoundEngine/Public/MetasoundSource.h
@@ -169,8 +169,8 @@ public:
 	{
 		return this;
 	}
-
 	virtual void InitParameters(TArray<FAudioParameter>& ParametersToInit, FName InFeatureName) override;
+
 	virtual void InitResources() override;
 
 	virtual bool IsPlayable() const override;
@@ -211,6 +211,11 @@ protected:
 		return RootMetasoundDocument;
 	}
 
+	virtual const UClass& GetBaseMetaSoundUClass() const final override
+	{
+		return *UMetaSoundSource::StaticClass();
+	}
+
 	/** Gets all the default parameters for this Asset.  */
 	virtual bool GetAllDefaultParameters(TArray<FAudioParameter>& OutParameters) const override;
 
diff --git a/Engine/Plugins/Runtime/Metasound/Source/MetasoundEngine/Public/MetasoundUObjectRegistry.h b/Engine/Plugins/Runtime/Metasound/Source/MetasoundEngine/Public/MetasoundUObjectRegistry.h
--- a/Engine/Plugins/Runtime/Metasound/Source/MetasoundEngine/Public/MetasoundUObjectRegistry.h
+++ b/Engine/Plugins/Runtime/Metasound/Source/MetasoundEngine/Public/MetasoundUObjectRegistry.h
@@ -1,10 +1,10 @@
 // Copyright Epic Games, Inc. All Rights Reserved.
 #pragma once
 
-#include "CoreMinimal.h"
 
 #include "Engine/Engine.h"
 #include "MetasoundAssetBase.h"
+#include "MetasoundDocumentInterface.h"
 #include "Subsystems/EngineSubsystem.h"
 #include "UObject/Object.h"
 #include "Templates/Function.h"
@@ -42,6 +42,9 @@ namespace Metasound
 		/** Returns true if the UClass is a child of this UClass associated with this entry. */
 		virtual bool IsChildClass(const UClass* InClass) const = 0;
 
+		/** Returns whether entry's class is a serialized asset or a transient type */
+		virtual bool IsAssetType() const = 0;
+
 		/** Attempts to cast the UObject to an FMetasoundAssetBase */
 		virtual FMetasoundAssetBase* Cast(UObject* InObject) const = 0;
 
@@ -66,13 +69,14 @@ namespace Metasound
 
 	/** An entry into the Metasound-UObject registry.
 	 *
-	 * @Tparam UClassType A class which derives from UObject and FMetasoundAssetBase.
+	 * @Tparam UClassType A class which derives from UObject and IMetaSoundDocumentInterface.
+	 * @Tparam IsAssetType If true, type derives from FMetasoundAssetBase and is considered a serializable, playable asset.
 	 */
 	template<typename UClassType>
 	class TMetasoundUObjectRegistryEntry : public IMetasoundUObjectRegistryEntry
 	{
-		// Ensure that this is a subclass of FMetasoundAssetBase and UObject.
-		static_assert(std::is_base_of<FMetasoundAssetBase, UClassType>::value, "UClass must be derived from FMetasoundAssetBase");
+		// Ensure that this is a subclass of IMetaSoundDocumentInterface and UObject.
+		static_assert(std::is_base_of<IMetaSoundDocumentInterface, UClassType>::value, "UClass must be derived from IMetaSoundDocumentInterface");
 		static_assert(std::is_base_of<UObject, UClassType>::value, "UClass must be derived from UObject");
 
 	public:
@@ -84,7 +88,7 @@ namespace Metasound
 			return UClassType::StaticClass();
 		}
 
-		bool IsChildClass(const UObject* InObject) const override
+		virtual bool IsChildClass(const UObject* InObject) const override
 		{
 			if (nullptr != InObject)
 			{
@@ -93,7 +97,7 @@ namespace Metasound
 			return false;
 		}
 
-		bool IsChildClass(const UClass* InClass) const override
+		virtual bool IsChildClass(const UClass* InClass) const override
 		{
 			if (nullptr != InClass)
 			{
@@ -102,22 +106,43 @@ namespace Metasound
 			return false;
 		}
 
-		FMetasoundAssetBase* Cast(UObject* InObject) const override
+		virtual bool IsAssetType() const override
+		{
+			return std::is_base_of<FMetasoundAssetBase, UClassType>::value;
+		}
+
+		virtual FMetasoundAssetBase* Cast(UObject * InObject) const override
 		{
-			if (nullptr == InObject)
+			if constexpr (!std::is_base_of<FMetasoundAssetBase, UClassType>::value)
 			{
 				return nullptr;
 			}
-			return static_cast<FMetasoundAssetBase*>(CastChecked<UClassType>(InObject));
+			else
+			{
+				if (InObject)
+				{
+					return static_cast<FMetasoundAssetBase*>(CastChecked<UClassType>(InObject));
+				}
+
+				return nullptr;
+			}
 		}
 
-		const FMetasoundAssetBase* Cast(const UObject* InObject) const override
+		virtual const FMetasoundAssetBase* Cast(const UObject* InObject) const override
 		{
-			if (nullptr == InObject)
+			if constexpr (!std::is_base_of<FMetasoundAssetBase, UClassType>::value)
+			{
+				return nullptr;
+			}
+			else
 			{
+				if (InObject)
+				{
+					return static_cast<const FMetasoundAssetBase*>(CastChecked<const UClassType>(InObject));
+				}
+
 				return nullptr;
 			}
-			return static_cast<const FMetasoundAssetBase*>(CastChecked<const UClassType>(InObject));
 		}
 	};
 
@@ -160,11 +185,14 @@ namespace Metasound
 			virtual UObject* NewObject(UClass* InClass, const FMetasoundFrontendDocument& InDocument, const FString& InPath) const { return nullptr; }
 
 			/** Iterate all registered UClasses that serve as MetaSound assets.*/
-			virtual void IterateRegisteredUClasses(TFunctionRef<void(UClass&)> InFunc) const = 0;
+			virtual void IterateRegisteredUClasses(TFunctionRef<void(UClass&)> InFunc, bool bAssetTypesOnly = true) const = 0;
 
 			/** Returns true if the InObject is of a class or child class which is registered. */
 			virtual bool IsRegisteredClass(UObject* InObject) const = 0;
 
+			/** Returns true if the InClass matches a class or child class which is registered. */
+			virtual bool IsRegisteredClass(const UClass& InClass) const = 0;
+
 			/** Returns casts the UObject to a FMetasoundAssetBase if the UObject is of a registered type.
 			 * If the UObject's UClass is not registered, then a nullptr is returned. 
 			 */
diff --git a/Engine/Plugins/Runtime/Metasound/Source/MetasoundFrontend/Private/Interfaces/MetasoundFrontendInterfaceRegistry.cpp b/Engine/Plugins/Runtime/Metasound/Source/MetasoundFrontend/Private/Interfaces/MetasoundFrontendInterfaceRegistry.cpp
--- a/Engine/Plugins/Runtime/Metasound/Source/MetasoundFrontend/Private/Interfaces/MetasoundFrontendInterfaceRegistry.cpp
+++ b/Engine/Plugins/Runtime/Metasound/Source/MetasoundFrontend/Private/Interfaces/MetasoundFrontendInterfaceRegistry.cpp
@@ -15,7 +15,7 @@ namespace Metasound::Frontend
 	{
 	}
 
-	void FInterfaceRegistry::RegisterInterface(TUniquePtr<IInterfaceRegistryEntry>&& InEntry)
+	bool FInterfaceRegistry::RegisterInterface(TUniquePtr<IInterfaceRegistryEntry>&& InEntry)
 	{
 		METASOUND_LLM_SCOPE;
 
@@ -35,10 +35,12 @@ namespace Metasound::Frontend
 					
 				FInterfaceRegistryTransaction Transaction{FInterfaceRegistryTransaction::ETransactionType::InterfaceRegistration, Key, InEntry->GetInterface().Version, TransactionTime};
 				TransactionBuffer->AddTransaction(MoveTemp(Transaction));
-
 				Entries.Add(Key, MoveTemp(InEntry));
+				return true;
 			}
 		}
+
+		return false;
 	}
 
 	const IInterfaceRegistryEntry* FInterfaceRegistry::FindInterfaceRegistryEntry(const FInterfaceRegistryKey& InKey) const
diff --git a/Engine/Plugins/Runtime/Metasound/Source/MetasoundFrontend/Private/Interfaces/MetasoundFrontendInterfaceRegistryPrivate.h b/Engine/Plugins/Runtime/Metasound/Source/MetasoundFrontend/Private/Interfaces/MetasoundFrontendInterfaceRegistryPrivate.h
--- a/Engine/Plugins/Runtime/Metasound/Source/MetasoundFrontend/Private/Interfaces/MetasoundFrontendInterfaceRegistryPrivate.h
+++ b/Engine/Plugins/Runtime/Metasound/Source/MetasoundFrontend/Private/Interfaces/MetasoundFrontendInterfaceRegistryPrivate.h
@@ -6,6 +6,7 @@
 #include "MetasoundFrontendDocument.h"
 #include "MetasoundFrontendRegistryTransaction.h"
 #include "UObject/NoExportTypes.h"
+#include "UObject/TopLevelAssetPath.h"
 
 
 namespace Metasound::Frontend
@@ -21,10 +22,8 @@ namespace Metasound::Frontend
 
 		virtual ~FInterfaceRegistry() = default;
 
-		virtual void RegisterInterface(TUniquePtr<IInterfaceRegistryEntry>&& InEntry) override;
-
+		virtual bool RegisterInterface(TUniquePtr<IInterfaceRegistryEntry>&& InEntry) override;
 		virtual const IInterfaceRegistryEntry* FindInterfaceRegistryEntry(const FInterfaceRegistryKey& InKey) const override;
-
 		virtual bool FindInterface(const FInterfaceRegistryKey& InKey, FMetasoundFrontendInterface& OutInterface) const override;
 
 		TUniquePtr<FInterfaceTransactionStream> CreateTransactionStream();
diff --git a/Engine/Plugins/Runtime/Metasound/Source/MetasoundFrontend/Private/Interfaces/MetasoundFrontendSourceInterface.cpp b/Engine/Plugins/Runtime/Metasound/Source/MetasoundFrontend/Private/Interfaces/MetasoundFrontendSourceInterface.cpp
--- a/Engine/Plugins/Runtime/Metasound/Source/MetasoundFrontend/Private/Interfaces/MetasoundFrontendSourceInterface.cpp
+++ b/Engine/Plugins/Runtime/Metasound/Source/MetasoundFrontend/Private/Interfaces/MetasoundFrontendSourceInterface.cpp
@@ -10,287 +10,289 @@
 
 #define LOCTEXT_NAMESPACE "Metasound"
 
-namespace Metasound
+namespace Metasound::Frontend
 {
-	namespace Frontend
+	namespace SourceInterfacePrivate
 	{
-#define AUDIO_PARAMETER_INTERFACE_NAMESPACE "UE.Source.OneShot"
-
-		namespace SourceOneShotInterface
+		TArray<Audio::FParameterInterface::FClassOptions> GetDefaultSourceClassOptions(const FTopLevelAssetPath& InClassPath)
 		{
-			const FMetasoundFrontendVersion& GetVersion()
+			constexpr bool bIsModifiable = false;
+			constexpr bool bIsDefault = true;
+			return TArray<Audio::FParameterInterface::FClassOptions>
 			{
-				static const FMetasoundFrontendVersion Version = { AUDIO_PARAMETER_INTERFACE_NAMESPACE, { 1, 0 } };
-				return Version;
-			}
+				{ InClassPath, bIsModifiable, bIsDefault }
+			};
+		}
+	}
 
-			namespace Outputs
-			{
-				const FName OnFinished = AUDIO_PARAMETER_INTERFACE_MEMBER_DEFINE("OnFinished");
-			}
+#define AUDIO_PARAMETER_INTERFACE_NAMESPACE "UE.Source.OneShot"
+	namespace SourceOneShotInterface
+	{
+		const FMetasoundFrontendVersion& GetVersion()
+		{
+			static const FMetasoundFrontendVersion Version = { AUDIO_PARAMETER_INTERFACE_NAMESPACE, { 1, 0 } };
+			return Version;
+		}
 
-			Audio::FParameterInterfacePtr CreateInterface(const UClass& InClass)
-			{
-				return CreateInterface();
-			}
+		namespace Outputs
+		{
+			const FName OnFinished = AUDIO_PARAMETER_INTERFACE_MEMBER_DEFINE("OnFinished");
+		}
 
-			Audio::FParameterInterfacePtr CreateInterface()
+		Audio::FParameterInterfacePtr CreateInterface(const UClass& InClass)
+		{
+			struct FInterface : public Audio::FParameterInterface
 			{
-				struct FInterface : public Audio::FParameterInterface
+				FInterface(const FTopLevelAssetPath& InClassPath)
+					: FParameterInterface(SourceOneShotInterface::GetVersion().Name, SourceOneShotInterface::GetVersion().Number.ToInterfaceVersion())
 				{
-					FInterface()
-						: FParameterInterface(SourceOneShotInterface::GetVersion().Name, SourceOneShotInterface::GetVersion().Number.ToInterfaceVersion())
+					UClassOptions = SourceInterfacePrivate::GetDefaultSourceClassOptions(InClassPath);
+
+					Outputs =
 					{
-						Outputs =
 						{
-							{
-								LOCTEXT("OnFinished", "On Finished"),
-								LOCTEXT("OnFinishedDescription", "Trigger executed to initiate stopping the source."),
-								GetMetasoundDataTypeName<FTrigger>(),
-								Outputs::OnFinished,
-								LOCTEXT("OnFinishedWarning", "\"On Finished\" should be connected for OneShot MetaSound sources. For sources with undefined duration (e.g. looping), remove the OneShot interface and use an audio component to avoid leaking the source."),
-							}
-						};
-					}
-				};
+							LOCTEXT("OnFinished", "On Finished"),
+							LOCTEXT("OnFinishedDescription", "Trigger executed to initiate stopping the source."),
+							GetMetasoundDataTypeName<FTrigger>(),
+							Outputs::OnFinished,
+							LOCTEXT("OnFinishedWarning", "\"On Finished\" should be connected for OneShot MetaSound sources. For sources with undefined duration (e.g. looping), remove the OneShot interface and use an audio component to avoid leaking the source."),
+						}
+					};
+				}
+			};
 
-				return MakeShared<FInterface>();
-			}
-		} // namespace SourceOneShotInterface
+			return MakeShared<FInterface>(InClass.GetClassPathName());
+		}
+	} // namespace SourceOneShotInterface
 
 #undef AUDIO_PARAMETER_INTERFACE_NAMESPACE
 
 #define AUDIO_PARAMETER_INTERFACE_NAMESPACE "UE.Source"
+	namespace SourceInterfaceV1_0
+	{
+		const FMetasoundFrontendVersion& GetVersion()
+		{
+			static const FMetasoundFrontendVersion Version = { AUDIO_PARAMETER_INTERFACE_NAMESPACE, { 1, 0 } };
+			return Version;
+		}
 
-		namespace SourceInterfaceV1_0
+		namespace Inputs
 		{
-			const FMetasoundFrontendVersion& GetVersion()
-			{
-				static const FMetasoundFrontendVersion Version = { AUDIO_PARAMETER_INTERFACE_NAMESPACE, { 1, 0 }};
-				return Version;
-			}
+			const FName OnPlay = AUDIO_PARAMETER_INTERFACE_MEMBER_DEFINE("OnPlay");
+		}
 
-			namespace Inputs
-			{
-				const FName OnPlay = AUDIO_PARAMETER_INTERFACE_MEMBER_DEFINE("OnPlay");
-			}
+		namespace Outputs
+		{
+			const FName OnFinished = AUDIO_PARAMETER_INTERFACE_MEMBER_DEFINE("OnFinished");
+		}
 
-			namespace Outputs
+		namespace Environment
+		{
+			const FName DeviceID = AUDIO_PARAMETER_INTERFACE_MEMBER_DEFINE("AudioDeviceID");
+			const FName GraphName = AUDIO_PARAMETER_INTERFACE_MEMBER_DEFINE("GraphName");
+			const FName IsPreview = AUDIO_PARAMETER_INTERFACE_MEMBER_DEFINE("IsPreviewSound");
+			const FName SoundUniqueID = AUDIO_PARAMETER_INTERFACE_MEMBER_DEFINE("SoundUniqueID");
+			const FName TransmitterID = AUDIO_PARAMETER_INTERFACE_MEMBER_DEFINE("TransmitterID");
+		}
+
+		Audio::FParameterInterfacePtr CreateInterface(const UClass& InClass)
+		{
+			struct FInterface : public Audio::FParameterInterface
 			{
-				const FName OnFinished = AUDIO_PARAMETER_INTERFACE_MEMBER_DEFINE("OnFinished");
-			}
+				FInterface(const FTopLevelAssetPath& InClassPath)
+					: FParameterInterface(SourceInterfaceV1_0::GetVersion().Name, SourceInterfaceV1_0::GetVersion().Number.ToInterfaceVersion())
+				{
+					constexpr bool bIsModifiable = false;
+					constexpr bool bIsDefault = false;
+					UClassOptions = TArray<FClassOptions>
+					{
+						{ InClassPath, bIsModifiable, bIsDefault }
+					};
 
-			namespace Environment
-			{
-				const FName DeviceID = AUDIO_PARAMETER_INTERFACE_MEMBER_DEFINE("AudioDeviceID");
-				const FName GraphName = AUDIO_PARAMETER_INTERFACE_MEMBER_DEFINE("GraphName");
-				const FName IsPreview = AUDIO_PARAMETER_INTERFACE_MEMBER_DEFINE("IsPreviewSound");
-				const FName SoundUniqueID = AUDIO_PARAMETER_INTERFACE_MEMBER_DEFINE("SoundUniqueID");
-				const FName TransmitterID = AUDIO_PARAMETER_INTERFACE_MEMBER_DEFINE("TransmitterID");
-			}
+					Inputs =
+					{
+						{
+							LOCTEXT("OnPlay", "On Play"),
+							LOCTEXT("OnPlayDescription", "Trigger executed when the source is played."),
+							GetMetasoundDataTypeName<FTrigger>(),
+							{ Inputs::OnPlay, false }
+						}
+					};
 
-			Audio::FParameterInterfacePtr CreateInterface(const UClass& InClass)
-			{
-				return CreateInterface();
-			}
+					Outputs =
+					{
+						{
+							LOCTEXT("OnFinished", "On Finished"),
+							LOCTEXT("OnFinishedDescription", "Trigger executed to initiate stopping the source."),
+							GetMetasoundDataTypeName<FTrigger>(),
+							Outputs::OnFinished
+						}
+					};
 
-			Audio::FParameterInterfacePtr CreateInterface()
-			{
-				struct FInterface : public Audio::FParameterInterface
-				{
-					FInterface()
-						: FParameterInterface(SourceInterfaceV1_0::GetVersion().Name, SourceInterfaceV1_0::GetVersion().Number.ToInterfaceVersion())
+					Environment =
 					{
-						Inputs =
 						{
-							{
-								LOCTEXT("OnPlay", "On Play"),
-								LOCTEXT("OnPlayDescription", "Trigger executed when the source is played."),
-								GetMetasoundDataTypeName<FTrigger>(),
-								{ Inputs::OnPlay, false }
-							}
-						};
-
-						Outputs =
+							LOCTEXT("AudioDeviceIDDisplayName", "Audio Device ID"),
+							LOCTEXT("AudioDeviceIDDescription", "ID of AudioDevice source is played from."),
+							FName(), // TODO: Align environment data types with environment (ex. this is actually set/get as a uint32)
+							Environment::DeviceID
+						},
 						{
-							{
-								LOCTEXT("OnFinished", "On Finished"),
-								LOCTEXT("OnFinishedDescription", "Trigger executed to initiate stopping the source."),
-								GetMetasoundDataTypeName<FTrigger>(),
-								Outputs::OnFinished
-							}
-						};
-
-						Environment =
+							LOCTEXT("GraphNameDisplayName", "Graph Name"),
+							LOCTEXT("GraphNameDescription", "Name of source graph (for debugging/logging)."),
+							GetMetasoundDataTypeName<FString>(),
+							Environment::GraphName
+						},
 						{
-							{
-								LOCTEXT("AudioDeviceIDDisplayName", "Audio Device ID"),
-								LOCTEXT("AudioDeviceIDDescription", "ID of AudioDevice source is played from."),
-								FName(), // TODO: Align environment data types with environment (ex. this is actually set/get as a uint32)
-								Environment::DeviceID
-							},
-							{
-								LOCTEXT("GraphNameDisplayName", "Graph Name"),
-								LOCTEXT("GraphNameDescription", "Name of source graph (for debugging/logging)."),
-								GetMetasoundDataTypeName<FString>(),
-								Environment::GraphName
-							},
-							{
-								LOCTEXT("IsPreviewSoundDisplayName", "Is Preview Sound"),
-								LOCTEXT("IsPreviewSoundDescription", "Whether source is being played as a previewed sound."),
-								GetMetasoundDataTypeName<bool>(),
-								Environment::IsPreview
-							},
-							{
-								LOCTEXT("TransmitterIDDisplayName", "Transmitter ID"),
-								LOCTEXT("TransmitterIDDescription", "ID used by Transmission System to generate a unique send address for each source instance."),
-								FName(), // TODO: Align environment data types with environment (ex. this is actually set/get as a uint64)
-								Environment::TransmitterID
-							},
-							{
-								LOCTEXT("SoundUniqueIdDisplayName", "Sound Unique ID"),
-								LOCTEXT("SoundUniqueIdDescription", "ID of unique source instance."),
-								FName(), // TODO: Align environment data types with environment (ex. this is actually set/get as a uint32)
-								Environment::SoundUniqueID
-							}
-						};
-					}
-				};
+							LOCTEXT("IsPreviewSoundDisplayName", "Is Preview Sound"),
+							LOCTEXT("IsPreviewSoundDescription", "Whether source is being played as a previewed sound."),
+							GetMetasoundDataTypeName<bool>(),
+							Environment::IsPreview
+						},
+						{
+							LOCTEXT("TransmitterIDDisplayName", "Transmitter ID"),
+							LOCTEXT("TransmitterIDDescription", "ID used by Transmission System to generate a unique send address for each source instance."),
+							FName(), // TODO: Align environment data types with environment (ex. this is actually set/get as a uint64)
+							Environment::TransmitterID
+						},
+						{
+							LOCTEXT("SoundUniqueIdDisplayName", "Sound Unique ID"),
+							LOCTEXT("SoundUniqueIdDescription", "ID of unique source instance."),
+							FName(), // TODO: Align environment data types with environment (ex. this is actually set/get as a uint32)
+							Environment::SoundUniqueID
+						}
+					};
+				}
+			};
 
-				return MakeShared<FInterface>();
-			}
-		} // namespace SourceInterfaceV1_0
+			return MakeShared<FInterface>(InClass.GetClassPathName());
+		}
+	} // namespace SourceInterfaceV1_0
 
-		namespace SourceInterface
+	namespace SourceInterface
+	{
+		const FMetasoundFrontendVersion& GetVersion()
 		{
-			const FMetasoundFrontendVersion& GetVersion()
-			{
-				static const FMetasoundFrontendVersion Version = { AUDIO_PARAMETER_INTERFACE_NAMESPACE, { 1, 1 } };
-				return Version;
-			}
+			static const FMetasoundFrontendVersion Version = { AUDIO_PARAMETER_INTERFACE_NAMESPACE, { 1, 1 } };
+			return Version;
+		}
 
-			namespace Inputs
-			{
-				const FName OnPlay = AUDIO_PARAMETER_INTERFACE_MEMBER_DEFINE("OnPlay");
-			}
+		namespace Inputs
+		{
+			const FName OnPlay = AUDIO_PARAMETER_INTERFACE_MEMBER_DEFINE("OnPlay");
+		}
 
-			namespace Environment
+		namespace Environment
+		{
+			const FName DeviceID = AUDIO_PARAMETER_INTERFACE_MEMBER_DEFINE("AudioDeviceID");
+			const FName GraphName = AUDIO_PARAMETER_INTERFACE_MEMBER_DEFINE("GraphName");
+			const FName IsPreview = AUDIO_PARAMETER_INTERFACE_MEMBER_DEFINE("IsPreviewSound");
+			const FName SoundUniqueID = AUDIO_PARAMETER_INTERFACE_MEMBER_DEFINE("SoundUniqueID");
+			const FName TransmitterID = AUDIO_PARAMETER_INTERFACE_MEMBER_DEFINE("TransmitterID");
+			const FName AudioMixerNumOutputFrames = AUDIO_PARAMETER_INTERFACE_MEMBER_DEFINE("AudioMixerNumOutputFrames");
+		}
+
+		Audio::FParameterInterfacePtr CreateInterface(const UClass& InClass)
+		{
+			struct FInterface : public Audio::FParameterInterface
 			{
-				const FName DeviceID = AUDIO_PARAMETER_INTERFACE_MEMBER_DEFINE("AudioDeviceID");
-				const FName GraphName = AUDIO_PARAMETER_INTERFACE_MEMBER_DEFINE("GraphName");
-				const FName IsPreview = AUDIO_PARAMETER_INTERFACE_MEMBER_DEFINE("IsPreviewSound");
-				const FName SoundUniqueID = AUDIO_PARAMETER_INTERFACE_MEMBER_DEFINE("SoundUniqueID");
-				const FName TransmitterID = AUDIO_PARAMETER_INTERFACE_MEMBER_DEFINE("TransmitterID");
-				const FName AudioMixerNumOutputFrames = AUDIO_PARAMETER_INTERFACE_MEMBER_DEFINE("AudioMixerNumOutputFrames");
-			}
+				FInterface(const FTopLevelAssetPath& InClassPath)
+					: FParameterInterface(SourceInterface::GetVersion().Name, SourceInterface::GetVersion().Number.ToInterfaceVersion())
+				{
+					UClassOptions = SourceInterfacePrivate::GetDefaultSourceClassOptions(InClassPath);
 
-			Audio::FParameterInterfacePtr CreateInterface(const UClass& InClass)
-			{
-				return CreateInterface();
-			}
+					Inputs =
+					{
+						{
+							LOCTEXT("OnPlay", "On Play"),
+							LOCTEXT("OnPlayDescription", "Trigger executed when the source is played."),
+							GetMetasoundDataTypeName<FTrigger>(),
+							{ Inputs::OnPlay, false }
+						}
+					};
 
-			Audio::FParameterInterfacePtr CreateInterface()
-			{
-				struct FInterface : public Audio::FParameterInterface
-				{
-					FInterface()
-						: FParameterInterface(SourceInterface::GetVersion().Name, SourceInterface::GetVersion().Number.ToInterfaceVersion())
+					Environment =
 					{
-						Inputs =
 						{
-							{
-								LOCTEXT("OnPlay", "On Play"),
-								LOCTEXT("OnPlayDescription", "Trigger executed when the source is played."),
-								GetMetasoundDataTypeName<FTrigger>(),
-								{ Inputs::OnPlay, false }
-							}
-						};
-
-						Environment =
+							LOCTEXT("AudioDeviceIDDisplayName1", "Audio Device ID"),
+							LOCTEXT("AudioDeviceIDDescription2", "ID of AudioDevice source is played from."),
+							FName(), // TODO: Align environment data types with environment (ex. this is actually set/get as a uint32)
+							Environment::DeviceID
+						},
 						{
-							{
-								LOCTEXT("AudioDeviceIDDisplayName1", "Audio Device ID"),
-								LOCTEXT("AudioDeviceIDDescription2", "ID of AudioDevice source is played from."),
-								FName(), // TODO: Align environment data types with environment (ex. this is actually set/get as a uint32)
-								Environment::DeviceID
-							},
-							{
-								LOCTEXT("GraphNameDisplayName", "Graph Name"),
-								LOCTEXT("AudioDeviceIDDescription3", "Name of source graph (for debugging/logging)."),
-								GetMetasoundDataTypeName<FString>(),
-								Environment::GraphName
-							},
-							{
-								LOCTEXT("IsPreviewSoundDisplayName", "Is Preview Sound"),
-								LOCTEXT("IsPreviewSoundDescription4", "Whether source is being played as a previewed sound."),
-								GetMetasoundDataTypeName<bool>(),
-								Environment::IsPreview
-							},
-							{
-								LOCTEXT("TransmitterIDDisplayName", "Transmitter ID"),
-								LOCTEXT("TransmitterIDDescription", "ID used by Transmission System to generate a unique send address for each source instance."),
-								FName(), // TODO: Align environment data types with environment (ex. this is actually set/get as a uint64)
-								Environment::TransmitterID
-							},
-							{
-								LOCTEXT("SoundUniqueDisplayName", "Sound Unique ID"),
-								LOCTEXT("SoundUniqueDescription", "ID of unique source instance."),
-								FName(), // TODO: Align environment data types with environment (ex. this is actually set/get as a uint32)
-								Environment::SoundUniqueID
-							},
-							{
-								LOCTEXT("AudioMixerOutputFramesDisplayName", "Audio Mixer Output Frames"),
-								LOCTEXT("AudioMixerOutputFramesDescription", "The number of output frames rendered per buffer in the audio mixer."),
-								FName(), 
-								Environment::AudioMixerNumOutputFrames
-							}
-						};
-					}
-				};
+							LOCTEXT("GraphNameDisplayName", "Graph Name"),
+							LOCTEXT("AudioDeviceIDDescription3", "Name of source graph (for debugging/logging)."),
+							GetMetasoundDataTypeName<FString>(),
+							Environment::GraphName
+						},
+						{
+							LOCTEXT("IsPreviewSoundDisplayName", "Is Preview Sound"),
+							LOCTEXT("IsPreviewSoundDescription4", "Whether source is being played as a previewed sound."),
+							GetMetasoundDataTypeName<bool>(),
+							Environment::IsPreview
+						},
+						{
+							LOCTEXT("TransmitterIDDisplayName", "Transmitter ID"),
+							LOCTEXT("TransmitterIDDescription", "ID used by Transmission System to generate a unique send address for each source instance."),
+							FName(), // TODO: Align environment data types with environment (ex. this is actually set/get as a uint64)
+							Environment::TransmitterID
+						},
+						{
+							LOCTEXT("SoundUniqueDisplayName", "Sound Unique ID"),
+							LOCTEXT("SoundUniqueDescription", "ID of unique source instance."),
+							FName(), // TODO: Align environment data types with environment (ex. this is actually set/get as a uint32)
+							Environment::SoundUniqueID
+						},
+						{
+							LOCTEXT("AudioMixerOutputFramesDisplayName", "Audio Mixer Output Frames"),
+							LOCTEXT("AudioMixerOutputFramesDescription", "The number of output frames rendered per buffer in the audio mixer."),
+							FName(), 
+							Environment::AudioMixerNumOutputFrames
+						}
+					};
+				}
+			};
 
-				return MakeShared<FInterface>();
-			}
+			return MakeShared<FInterface>(InClass.GetClassPathName());
+		}
 
-			bool FUpdateInterface::Transform(Frontend::FDocumentHandle InDocument) const
-			{
-				using namespace Frontend;
+		bool FUpdateInterface::Transform(Frontend::FDocumentHandle InDocument) const
+		{
+			using namespace Frontend;
 
-				// When upgrading, we only want to add the one-shot interface if the MetaSound actually has the OnFinished trigger connected.
-				bool bIsOnFinishedConnected = false;
-				InDocument->GetRootGraph()->IterateConstNodes([&](Frontend::FConstNodeHandle NodeHandle)
+			// When upgrading, we only want to add the one-shot interface if the MetaSound actually has the OnFinished trigger connected.
+			bool bIsOnFinishedConnected = false;
+			InDocument->GetRootGraph()->IterateConstNodes([&](Frontend::FConstNodeHandle NodeHandle)
+			{
+				NodeHandle->IterateConstInputs([&](Frontend::FConstInputHandle InputHandle)
 				{
-					NodeHandle->IterateConstInputs([&](Frontend::FConstInputHandle InputHandle)
+					if (InputHandle->GetName() == SourceInterfaceV1_0::Outputs::OnFinished)
 					{
-						if (InputHandle->GetName() == SourceInterfaceV1_0::Outputs::OnFinished)
-						{
-							bIsOnFinishedConnected = InputHandle->IsConnected();
-						}
-					});
-				}, EMetasoundFrontendClassType::Output);
-
-				const TArray<FMetasoundFrontendVersion> InterfacesToRemove
-				{
-					SourceInterfaceV1_0::GetVersion()
-				};
+						bIsOnFinishedConnected = InputHandle->IsConnected();
+					}
+				});
+			}, EMetasoundFrontendClassType::Output);
 
-				TArray<FMetasoundFrontendVersion> InterfacesToAdd
-				{
-					SourceInterface::GetVersion()
-				};
+			const TArray<FMetasoundFrontendVersion> InterfacesToRemove
+			{
+				SourceInterfaceV1_0::GetVersion()
+			};
 
-				if (bIsOnFinishedConnected)
-				{
-					InterfacesToAdd.Add(SourceOneShotInterface::GetVersion());
-				}
+			TArray<FMetasoundFrontendVersion> InterfacesToAdd
+			{
+				SourceInterface::GetVersion()
+			};
 
-				Frontend::FModifyRootGraphInterfaces InterfaceTransform(InterfacesToRemove, InterfacesToAdd);
-				return InterfaceTransform.Transform(InDocument);
+			if (bIsOnFinishedConnected)
+			{
+				InterfacesToAdd.Add(SourceOneShotInterface::GetVersion());
 			}
 
-} // namespace SourceInterface
+			Frontend::FModifyRootGraphInterfaces InterfaceTransform(InterfacesToRemove, InterfacesToAdd);
+			return InterfaceTransform.Transform(InDocument);
+		}
 
+	} // namespace SourceInterface
 #undef AUDIO_PARAMETER_INTERFACE_NAMESPACE
-	} // namespace Frontend
-} // namespace Metasound
-
+} // namespace Metasound::Frontend
 #undef LOCTEXT_NAMESPACE
diff --git a/Engine/Plugins/Runtime/Metasound/Source/MetasoundFrontend/Private/MetasoundAssetBase.cpp b/Engine/Plugins/Runtime/Metasound/Source/MetasoundFrontend/Private/MetasoundAssetBase.cpp
--- a/Engine/Plugins/Runtime/Metasound/Source/MetasoundFrontend/Private/MetasoundAssetBase.cpp
+++ b/Engine/Plugins/Runtime/Metasound/Source/MetasoundFrontend/Private/MetasoundAssetBase.cpp
@@ -320,11 +320,10 @@ void FMetasoundAssetBase::AddDefaultInterfaces()
 
 	UClass* AssetClass = OwningAsset->GetClass();
 	check(AssetClass);
+	const FTopLevelAssetPath& ClassPath = AssetClass->GetClassPathName();
 
-	FDocumentHandle DocumentHandle = GetDocumentHandle();
-
-	TArray<FMetasoundFrontendInterface> InitInterfaces = ISearchEngine::Get().FindUClassDefaultInterfaces(AssetClass->GetFName());
-	FModifyRootGraphInterfaces({ }, InitInterfaces).Transform(DocumentHandle);
+	TArray<FMetasoundFrontendVersion> InitVersions = ISearchEngine::Get().FindUClassDefaultInterfaceVersions(ClassPath);
+	FModifyRootGraphInterfaces({ }, InitVersions).Transform(GetDocumentChecked());
 }
 
 void FMetasoundAssetBase::SetDocument(FMetasoundFrontendDocument&& InDocument)
diff --git a/Engine/Plugins/Runtime/Metasound/Source/MetasoundFrontend/Private/MetasoundFrontendDocument.cpp b/Engine/Plugins/Runtime/Metasound/Source/MetasoundFrontend/Private/MetasoundFrontendDocument.cpp
--- a/Engine/Plugins/Runtime/Metasound/Source/MetasoundFrontend/Private/MetasoundFrontendDocument.cpp
+++ b/Engine/Plugins/Runtime/Metasound/Source/MetasoundFrontend/Private/MetasoundFrontendDocument.cpp
@@ -11,6 +11,7 @@
 #include "MetasoundFrontend.h"
 #include "MetasoundFrontendRegistries.h"
 #include "MetasoundLog.h"
+#include "MetasoundParameterTransmitter.h"
 #include "MetasoundVertex.h"
 
 #include UE_INLINE_GENERATED_CPP_BY_NAME(MetasoundFrontendDocument)
@@ -103,6 +104,20 @@ namespace Metasound
 					return EMetasoundFrontendVertexAccessType::Reference;
 			}
 		}
+
+		FName ResolveMemberDataType(FName DataType, EAudioParameterType ParamType)
+		{
+			if (!DataType.IsNone())
+			{
+				const bool bIsRegisteredType = Metasound::Frontend::IDataTypeRegistry::Get().IsRegistered(DataType);
+				if (ensureAlwaysMsgf(bIsRegisteredType, TEXT("Attempting to register Interface member with unregistered DataType '%s'."), *DataType.ToString()))
+				{
+					return DataType;
+				}
+			}
+
+			return Frontend::ConvertParameterToDataType(ParamType);
+		};
 	} // namespace DocumentPrivate
 } // namespace Metasound
 
@@ -274,6 +289,79 @@ bool FMetasoundFrontendClassVertex::CanConnectVertexAccessTypes(EMetasoundFronte
 	return true;
 }
 
+FMetasoundFrontendInterfaceUClassOptions::FMetasoundFrontendInterfaceUClassOptions(const Audio::FParameterInterface::FClassOptions& InOptions)
+	: ClassPath(InOptions.ClassPath)
+	, bIsModifiable(InOptions.bIsModifiable)
+	, bIsDefault(InOptions.bIsDefault)
+{
+}
+
+FMetasoundFrontendInterfaceUClassOptions::FMetasoundFrontendInterfaceUClassOptions(const FTopLevelAssetPath& InClassPath, bool bInIsModifiable, bool bInIsDefault)
+	: ClassPath(InClassPath)
+	, bIsModifiable(bInIsModifiable)
+	, bIsDefault(bInIsDefault)
+{
+}
+
+FMetasoundFrontendInterface::FMetasoundFrontendInterface(Audio::FParameterInterfacePtr InInterface)
+{
+	using namespace Metasound::DocumentPrivate;
+	using namespace Metasound::Frontend;
+
+	Version = { InInterface->GetName(), FMetasoundFrontendVersionNumber { InInterface->GetVersion().Major, InInterface->GetVersion().Minor } };
+
+	// Transfer all input data from AudioExtension interface struct to FrontendInterface
+	Algo::Transform(InInterface->GetInputs(), Inputs, [this](const Audio::FParameterInterface::FInput& Input)
+	{
+#if WITH_EDITOR
+		AddSortOrderToInputStyle(Input.SortOrderIndex);
+
+		// Setup required inputs by telling the style that the input is required
+		// This will later be validated against.
+		if (!Input.RequiredText.IsEmpty())
+		{
+			AddRequiredInputToStyle(Input.InitValue.ParamName, Input.RequiredText);
+		}
+#endif // WITH_EDITOR
+
+		return FMetasoundFrontendClassInput(Input);
+	});
+
+	// Transfer all output data from AudioExtension interface struct to FrontendInterface
+	Algo::Transform(InInterface->GetOutputs(), Outputs, [this](const Audio::FParameterInterface::FOutput& Output)
+	{
+#if WITH_EDITOR
+		AddSortOrderToOutputStyle(Output.SortOrderIndex);
+
+		// Setup required outputs by telling the style that the output is required.  This will later be validated against.
+		if (!Output.RequiredText.IsEmpty())
+		{
+			AddRequiredOutputToStyle(Output.ParamName, Output.RequiredText);
+		}
+#endif // WITH_EDITOR
+
+		return FMetasoundFrontendClassOutput(Output);
+	});
+
+	// Transfer all environment variables from AudioExtension interface struct to FrontendInterface
+	Algo::Transform(InInterface->GetEnvironment(), Environment, [](const Audio::FParameterInterface::FEnvironmentVariable& Variable)
+	{
+		return FMetasoundFrontendClassEnvironmentVariable(Variable);
+	});
+
+	// Transfer all class options from AudioExtension interface struct to FrontendInterface
+	Algo::Transform(InInterface->GetUClassOptions(), UClassOptions, [](const Audio::FParameterInterface::FClassOptions& Options)
+	{
+		return FMetasoundFrontendInterfaceUClassOptions(Options);
+	});
+}
+
+const FMetasoundFrontendInterfaceUClassOptions* FMetasoundFrontendInterface::FindClassOptions(const FTopLevelAssetPath& InClassPath) const
+{
+	auto FindClassOptionsPredicate = [&InClassPath](const FMetasoundFrontendInterfaceUClassOptions& Options) { return Options.ClassPath == InClassPath; };
+	return UClassOptions.FindByPredicate(FindClassOptionsPredicate);
+}
+
 FMetasoundFrontendClassName::FMetasoundFrontendClassName(const FName& InNamespace, const FName& InName, const FName& InVariant)
 : Namespace(InNamespace)
 , Name(InName)
@@ -655,6 +743,23 @@ FMetasoundFrontendClassInput::FMetasoundFrontendClassInput(const FMetasoundFront
 	DefaultLiteral.SetType(DefaultType);
 }
 
+FMetasoundFrontendClassInput::FMetasoundFrontendClassInput(const Audio::FParameterInterface::FInput& InInput)
+{
+	Name = InInput.InitValue.ParamName;
+	DefaultLiteral = FMetasoundFrontendLiteral(InInput.InitValue);
+	TypeName = Metasound::DocumentPrivate::ResolveMemberDataType(InInput.DataType, InInput.InitValue.ParamType);
+	VertexID = FGuid::NewGuid();
+
+#if WITH_EDITOR
+	// Interfaces should never serialize text to avoid desync between
+	// copied versions serialized in assets and those defined in code.
+	Metadata.SetSerializeText(false);
+	Metadata.SetDisplayName(InInput.DisplayName);
+	Metadata.SetDescription(InInput.Description);
+	Metadata.SortOrderIndex = InInput.SortOrderIndex;
+#endif // WITH_EDITOR
+}
+
 FMetasoundFrontendClassVariable::FMetasoundFrontendClassVariable(const FMetasoundFrontendClassVertex& InOther)
 	: FMetasoundFrontendClassVertex(InOther)
 {
@@ -665,6 +770,35 @@ FMetasoundFrontendClassVariable::FMetasoundFrontendClassVariable(const FMetasoun
 	DefaultLiteral.SetType(DefaultType);
 }
 
+FMetasoundFrontendClassOutput::FMetasoundFrontendClassOutput(const Audio::FParameterInterface::FOutput& Output)
+{
+	Name = Output.ParamName;
+	TypeName = Metasound::DocumentPrivate::ResolveMemberDataType(Output.DataType, Output.ParamType);
+	VertexID = FGuid::NewGuid();
+
+#if WITH_EDITOR
+	// Interfaces should never serialize text to avoid desync between
+	// copied versions serialized in assets and those defined in code.
+	Metadata.SetSerializeText(false);
+	Metadata.SetDisplayName(Output.DisplayName);
+	Metadata.SetDescription(Output.Description);
+	Metadata.SortOrderIndex = Output.SortOrderIndex;
+#endif // WITH_EDITOR
+}
+
+FMetasoundFrontendClassOutput::FMetasoundFrontendClassOutput(const FMetasoundFrontendClassVertex& InOther)
+	: FMetasoundFrontendClassVertex(InOther)
+{
+}
+
+FMetasoundFrontendClassEnvironmentVariable::FMetasoundFrontendClassEnvironmentVariable(const Audio::FParameterInterface::FEnvironmentVariable& InVariable)
+	: Name(InVariable.ParamName)
+	// Disabled as it isn't used to infer type when getting/setting at a lower level.
+	// TODO: Either remove type info for environment variables all together or enforce type.
+	// , TypeName(Metasound::Frontend::DocumentPrivate::ResolveMemberDataType(Environment.DataType, Environment.ParamType))
+{
+}
+
 FMetasoundFrontendGraphClass::FMetasoundFrontendGraphClass()
 {
 	Metadata.SetType(EMetasoundFrontendClassType::Graph);
diff --git a/Engine/Plugins/Runtime/Metasound/Source/MetasoundFrontend/Private/MetasoundFrontendDocumentBuilder.cpp b/Engine/Plugins/Runtime/Metasound/Source/MetasoundFrontend/Private/MetasoundFrontendDocumentBuilder.cpp
--- a/Engine/Plugins/Runtime/Metasound/Source/MetasoundFrontend/Private/MetasoundFrontendDocumentBuilder.cpp
+++ b/Engine/Plugins/Runtime/Metasound/Source/MetasoundFrontend/Private/MetasoundFrontendDocumentBuilder.cpp
@@ -518,9 +518,22 @@ bool FMetaSoundFrontendDocumentBuilder::AddInterface(FName InterfaceName)
 	{
 		if (GetDocument().Interfaces.Contains(Interface.Version))
 		{
+			UE_LOG(LogMetaSound, VeryVerbose, TEXT("MetaSound interface '%s' already found on document. MetaSoundBuilder skipping add request."), *InterfaceName.ToString());
 			return true;
 		}
 
+		const FInterfaceRegistryKey Key = GetInterfaceRegistryKey(Interface.Version);
+		if (const IInterfaceRegistryEntry* Entry = IInterfaceRegistry::Get().FindInterfaceRegistryEntry(Key))
+		{
+			const FTopLevelAssetPath ClassPath = DocumentInterface->GetBaseMetaSoundUClass().GetClassPathName();
+			const FMetasoundFrontendInterfaceUClassOptions* ClassOptions = Entry->GetInterface().FindClassOptions(ClassPath);
+			if (ClassOptions && !ClassOptions->bIsModifiable)
+			{
+				UE_LOG(LogMetaSound, Error, TEXT("DocumentBuilder failed to add MetaSound Interface '%s' to document: is not set to be modifiable for given UClass '%s'"), *InterfaceName.ToString(), *ClassPath.ToString());
+				return false;
+			}
+		}
+
 		const FModifyRootGraphInterfaces ModifyInterfaces({ }, { Interface.Version });
 		if (ModifyInterfaces.Transform(GetDocument()))
 		{
@@ -899,6 +912,13 @@ const FMetasoundFrontendDocument& FMetaSoundFrontendDocumentBuilder::GetDocument
 	return DocumentInterface->GetDocument();
 }
 
+const IMetaSoundDocumentInterface& FMetaSoundFrontendDocumentBuilder::GetDocumentInterface() const
+{
+	IMetaSoundDocumentInterface* Interface = DocumentInterface.GetInterface();
+	check(Interface);
+	return *Interface;
+}
+
 void FMetaSoundFrontendDocumentBuilder::InitGraphClassMetadata(FMetasoundFrontendClassMetadata& InOutMetadata, bool bResetVersion)
 {
 	InOutMetadata.SetClassName(FMetasoundFrontendClassName(FName(), *FGuid::NewGuid().ToString(), FName()));
@@ -911,30 +931,37 @@ void FMetaSoundFrontendDocumentBuilder::InitGraphClassMetadata(FMetasoundFronten
 	InOutMetadata.SetType(EMetasoundFrontendClassType::Graph);
 }
 
-void FMetaSoundFrontendDocumentBuilder::InitDocument(FName UClassName)
+void FMetaSoundFrontendDocumentBuilder::InitDocument()
 {
+	using namespace Metasound;
+	using namespace Metasound::Frontend;
+
 	METASOUND_TRACE_CPUPROFILER_EVENT_SCOPE(FMetaSoundFrontendDocumentBuilder::InitDocument);
 
-	using namespace Metasound::Frontend;
+	IMetaSoundDocumentInterface* DocInterface = DocumentInterface.GetInterface();
+	check(DocInterface);
+	FMetasoundFrontendDocument& Document = DocInterface->GetDocument();
 
 	// 1. Set default class Metadata
 	{
 		constexpr bool bResetVersion = true;
-		FMetasoundFrontendClassMetadata& ClassMetadata = GetDocument().RootGraph.Metadata;
+		FMetasoundFrontendClassMetadata& ClassMetadata = Document.RootGraph.Metadata;
 		InitGraphClassMetadata(ClassMetadata, bResetVersion);
 	}
 
 	// 2. Set default doc version Metadata
 	{
-		FMetasoundFrontendDocumentMetadata& DocMetadata = GetDocument().Metadata;
+		FMetasoundFrontendDocumentMetadata& DocMetadata = Document.Metadata;
 		DocMetadata.Version.Number = FMetasoundFrontendDocument::GetMaxVersion();
 	}
 
 	// 3. Add default interfaces for given UClass
 	{
-		TArray<FMetasoundFrontendInterface> InitInterfaces = ISearchEngine::Get().FindUClassDefaultInterfaces(UClassName);
-		FModifyRootGraphInterfaces ModifyRootGraphInterfaces({ }, InitInterfaces);
-		ModifyRootGraphInterfaces.Transform(GetDocument());
+		const UClass& BaseMetaSoundClass = DocInterface->GetBaseMetaSoundUClass();
+		const FTopLevelAssetPath BaseMetaSoundClassPath = BaseMetaSoundClass.GetClassPathName();
+		TArray<FMetasoundFrontendVersion> InitVersions = ISearchEngine::Get().FindUClassDefaultInterfaceVersions(BaseMetaSoundClassPath);
+		FModifyRootGraphInterfaces ModifyRootGraphInterfaces({ }, InitVersions);
+		ModifyRootGraphInterfaces.Transform(Document);
 	}
 
 	// 5. Reload the whole cache as transforms above mutate cached collections. Can be safely removed once transforms above
@@ -1230,9 +1257,22 @@ bool FMetaSoundFrontendDocumentBuilder::RemoveInterface(FName InterfaceName)
 	{
 		if (!GetDocument().Interfaces.Contains(Interface.Version))
 		{
+			UE_LOG(LogMetaSound, VeryVerbose, TEXT("MetaSound interface '%s' not found on document. MetaSoundBuilder skipping remove request."), *InterfaceName.ToString());
 			return true;
 		}
 
+		const FInterfaceRegistryKey Key = GetInterfaceRegistryKey(Interface.Version);
+		if (const IInterfaceRegistryEntry* Entry = IInterfaceRegistry::Get().FindInterfaceRegistryEntry(Key))
+		{
+			const FTopLevelAssetPath ClassPath = DocumentInterface->GetBaseMetaSoundUClass().GetClassPathName();
+			const FMetasoundFrontendInterfaceUClassOptions* ClassOptions = Entry->GetInterface().FindClassOptions(ClassPath);
+			if (ClassOptions && !ClassOptions->bIsModifiable)
+			{
+				UE_LOG(LogMetaSound, Error, TEXT("DocumentBuilder failed to remove MetaSound Interface '%s' to document: is not set to be modifiable for given UClass '%s'"), *InterfaceName.ToString(), *ClassPath.ToString());
+				return false;
+			}
+		}
+
 		const FModifyRootGraphInterfaces ModifyInterfaces({ { Interface.Version }, { } });
 		if (ModifyInterfaces.Transform(GetDocument()))
 		{
diff --git a/Engine/Plugins/Runtime/Metasound/Source/MetasoundFrontend/Private/MetasoundFrontendQuery.cpp b/Engine/Plugins/Runtime/Metasound/Source/MetasoundFrontend/Private/MetasoundFrontendQuery.cpp
--- a/Engine/Plugins/Runtime/Metasound/Source/MetasoundFrontend/Private/MetasoundFrontendQuery.cpp
+++ b/Engine/Plugins/Runtime/Metasound/Source/MetasoundFrontend/Private/MetasoundFrontendQuery.cpp
@@ -118,6 +118,24 @@ namespace Metasound
 			FMapFunction Func;
 		};
 
+		// Wrapper for step defined by function
+		struct FMultiMapFunctionFrontendQueryStep : IFrontendQueryMultiMapStep
+		{
+			using FMultiMapFunction = FFrontendQueryStep::FMultiMapFunction;
+
+			FMultiMapFunctionFrontendQueryStep(FMultiMapFunction InFunc)
+				: Func(InFunc)
+			{
+			}
+
+			TArray<FFrontendQueryKey> Map(const FFrontendQueryEntry& InEntry) const override
+			{
+				return Func(InEntry);
+			}
+
+			FMultiMapFunction Func;
+		};
+
 		// Wrapper for step defined by function
 		struct FReduceFunctionFrontendQueryStep: IFrontendQueryReduceStep
 		{
@@ -348,13 +366,17 @@ namespace Metasound
 			}
 		};
 
-		struct FMapStepExecuter : TStepExecuterBase<IFrontendQueryMapStep>
+		template <typename MapQueryStep>
+		struct TMapStepExecuterBase : public TStepExecuterBase<MapQueryStep>
 		{
-			using TStepExecuterBase<IFrontendQueryMapStep>::TStepExecuterBase;
+			using TStepExecuterBase<MapQueryStep>::TStepExecuterBase;
+			using TStepExecuterBase<MapQueryStep>::Step;
+
+			virtual void ExecuteStep(FFrontendQueryEntry&& Entry, FFrontendQuerySelection& OutResult) const = 0;
 
 			virtual void Execute(TSet<FFrontendQueryKey>& InOutUpdatedKeys, FFrontendQuerySelection& InOutResult) const override
 			{
-				METASOUND_TRACE_CPUPROFILER_EVENT_SCOPE(MetaSound::MapQueryStep::Execute);
+				METASOUND_TRACE_CPUPROFILER_EVENT_SCOPE(MetaSound::MapStepQuery::Execute);
 				if (Step.IsValid())
 				{
 					if (InOutUpdatedKeys.Num() > 0)
@@ -370,8 +392,7 @@ namespace Metasound
 								{
 									for (FFrontendQueryEntry& Entry : *Partition)
 									{
-										FFrontendQueryKey NewKey = Step->Map(Entry);
-										Result.FindOrAdd(NewKey).Add(Entry);
+										ExecuteStep(MoveTemp(Entry), Result);
 									}
 								}
 
@@ -388,12 +409,37 @@ namespace Metasound
 						}
 
 						// Append the new values to the final result.
-						Append(InOutUpdatedKeys, Result, InOutResult);
+						TStepExecuterBase<MapQueryStep>::Append(InOutUpdatedKeys, Result, InOutResult);
 					}
 				}
 			}
 		};
 
+		struct FMapStepExecuter : TMapStepExecuterBase<IFrontendQueryMapStep>
+		{
+			using TMapStepExecuterBase<IFrontendQueryMapStep>::TMapStepExecuterBase;
+
+			virtual void ExecuteStep(FFrontendQueryEntry&& Entry, FFrontendQuerySelection& OutResult) const override
+			{
+				FFrontendQueryKey NewKey = Step->Map(Entry);
+				OutResult.FindOrAdd(MoveTemp(NewKey)).Add(MoveTemp(Entry));
+			}
+		};
+
+		struct FMultiMapStepExecuter : TMapStepExecuterBase<IFrontendQueryMultiMapStep>
+		{
+			using TMapStepExecuterBase<IFrontendQueryMultiMapStep>::TMapStepExecuterBase;
+
+			virtual void ExecuteStep(FFrontendQueryEntry&& Entry, FFrontendQuerySelection& OutResult) const override
+			{
+				TArray<FFrontendQueryKey> NewKeys = Step->Map(Entry);
+				for (FFrontendQueryKey& NewKey : NewKeys)
+				{
+					OutResult.FindOrAdd(MoveTemp(NewKey)).Add(MoveTemp(Entry));
+				}
+			}
+		};
+
 		struct FReduceStepExecuter : TStepExecuterBase<IFrontendQueryReduceStep>
 		{
 			using FKey = FFrontendQueryKey;
@@ -776,6 +822,11 @@ namespace Metasound
 	{
 	}
 
+	FFrontendQueryStep::FFrontendQueryStep(FMultiMapFunction&& InFunc)
+		: StepExecuter(MakeUnique<FrontendQueryPrivate::FMultiMapStepExecuter>(MakeUnique<FrontendQueryPrivate::FMultiMapFunctionFrontendQueryStep>(MoveTemp(InFunc))))
+	{
+	}
+
 	FFrontendQueryStep::FFrontendQueryStep(FReduceFunction&& InFunc)
 	:	StepExecuter(MakeUnique<FrontendQueryPrivate::FReduceStepExecuter>(MakeUnique<FrontendQueryPrivate::FReduceFunctionFrontendQueryStep>(MoveTemp(InFunc))))
 	{
@@ -816,6 +867,11 @@ namespace Metasound
 	{
 	}
 
+	FFrontendQueryStep::FFrontendQueryStep(TUniquePtr<IFrontendQueryMultiMapStep>&& InStep)
+		: StepExecuter(MakeUnique<FrontendQueryPrivate::FMultiMapStepExecuter>(MoveTemp(InStep)))
+	{
+	}
+
 	FFrontendQueryStep::FFrontendQueryStep(TUniquePtr<IFrontendQueryReduceStep>&& InStep)
 	:	StepExecuter(MakeUnique<FrontendQueryPrivate::FReduceStepExecuter>(MoveTemp(InStep)))
 	{
diff --git a/Engine/Plugins/Runtime/Metasound/Source/MetasoundFrontend/Private/MetasoundFrontendSearchEngineCore.cpp b/Engine/Plugins/Runtime/Metasound/Source/MetasoundFrontend/Private/MetasoundFrontendSearchEngineCore.cpp
--- a/Engine/Plugins/Runtime/Metasound/Source/MetasoundFrontend/Private/MetasoundFrontendSearchEngineCore.cpp
+++ b/Engine/Plugins/Runtime/Metasound/Source/MetasoundFrontend/Private/MetasoundFrontendSearchEngineCore.cpp
@@ -12,468 +12,495 @@
 #include "MetasoundFrontendRegistryTransaction.h"
 #include "MetasoundTrace.h"
 
-namespace Metasound
+namespace Metasound::Frontend
 {
-	namespace Frontend
+	namespace SearchEngineQuerySteps
 	{
-		namespace SearchEngineQuerySteps
+		FFrontendQueryKey CreateKeyFromFullClassNameAndMajorVersion(const FMetasoundFrontendClassName& InClassName, int32 InMajorVersion)
 		{
-			FFrontendQueryKey CreateKeyFromFullClassNameAndMajorVersion(const FMetasoundFrontendClassName& InClassName, int32 InMajorVersion)
-			{
-				return FFrontendQueryKey(FString::Format(TEXT("{0}_v{1}"), {InClassName.ToString(), InMajorVersion}));
-			}
+			return FFrontendQueryKey(FString::Format(TEXT("{0}_v{1}"), {InClassName.ToString(), InMajorVersion}));
+		}
 
-			FFrontendQueryKey FMapClassToFullClassNameAndMajorVersion::Map(const FFrontendQueryEntry& InEntry) const
-			{
-				if (ensure(InEntry.Value.IsType<FMetasoundFrontendClass>()))
-				{
-					const FMetasoundFrontendClassMetadata& Metadata = InEntry.Value.Get<FMetasoundFrontendClass>().Metadata;
-					return CreateKeyFromFullClassNameAndMajorVersion(Metadata.GetClassName(), Metadata.GetVersion().Major);
-				}
-				return FFrontendQueryKey();
-			}
+		FFrontendQueryKey CreateKeyFromUClassPathName(const FTopLevelAssetPath& InClassPath)
+		{
+			return FFrontendQueryKey(InClassPath.ToString());
+		}
 
-			FInterfaceRegistryTransactionSource::FInterfaceRegistryTransactionSource()
-			: TransactionStream(FInterfaceRegistry::Get().CreateTransactionStream())
+		FFrontendQueryKey FMapClassToFullClassNameAndMajorVersion::Map(const FFrontendQueryEntry& InEntry) const
+		{
+			if (ensure(InEntry.Value.IsType<FMetasoundFrontendClass>()))
 			{
+				const FMetasoundFrontendClassMetadata& Metadata = InEntry.Value.Get<FMetasoundFrontendClass>().Metadata;
+				return CreateKeyFromFullClassNameAndMajorVersion(Metadata.GetClassName(), Metadata.GetVersion().Major);
 			}
+			return FFrontendQueryKey();
+		}
+
+		FInterfaceRegistryTransactionSource::FInterfaceRegistryTransactionSource()
+		: TransactionStream(FInterfaceRegistry::Get().CreateTransactionStream())
+		{
+		}
 
-			void FInterfaceRegistryTransactionSource::Stream(TArray<FFrontendQueryValue>& OutEntries)
+		void FInterfaceRegistryTransactionSource::Stream(TArray<FFrontendQueryValue>& OutEntries)
+		{
+			auto AddValue = [&OutEntries](const FInterfaceRegistryTransaction& InTransaction)
 			{
-				auto AddValue = [&OutEntries](const FInterfaceRegistryTransaction& InTransaction)
-				{
-					OutEntries.Emplace(TInPlaceType<FInterfaceRegistryTransaction>(), InTransaction);
-				};
+				OutEntries.Emplace(TInPlaceType<FInterfaceRegistryTransaction>(), InTransaction);
+			};
 
-				if (TransactionStream.IsValid())
-				{
-					TransactionStream->Stream(AddValue);
-				}
+			if (TransactionStream.IsValid())
+			{
+				TransactionStream->Stream(AddValue);
 			}
+		}
 
-			FFrontendQueryKey FMapInterfaceRegistryTransactionsToInterfaceRegistryKeys::Map(const FFrontendQueryEntry& InEntry) const
+		FFrontendQueryKey FMapInterfaceRegistryTransactionsToInterfaceRegistryKeys::Map(const FFrontendQueryEntry& InEntry) const
+		{
+			FInterfaceRegistryKey RegistryKey;
+			if (ensure(InEntry.Value.IsType<FInterfaceRegistryTransaction>()))
 			{
-				FInterfaceRegistryKey RegistryKey;
-				if (ensure(InEntry.Value.IsType<FInterfaceRegistryTransaction>()))
-				{
-					RegistryKey = InEntry.Value.Get<FInterfaceRegistryTransaction>().GetInterfaceRegistryKey();
-				}
-				return FFrontendQueryKey{RegistryKey};
+				RegistryKey = InEntry.Value.Get<FInterfaceRegistryTransaction>().GetInterfaceRegistryKey();
 			}
+			return FFrontendQueryKey{RegistryKey};
+		}
 
-			FFrontendQueryKey FMapInterfaceRegistryTransactionToInterfaceName::Map(const FFrontendQueryEntry& InEntry) const
+		FFrontendQueryKey FMapInterfaceRegistryTransactionToInterfaceName::Map(const FFrontendQueryEntry& InEntry) const
+		{
+			FInterfaceRegistryKey RegistryKey;
+			if (ensure(InEntry.Value.IsType<FInterfaceRegistryTransaction>()))
 			{
-				FInterfaceRegistryKey RegistryKey;
-				if (ensure(InEntry.Value.IsType<FInterfaceRegistryTransaction>()))
-				{
-					return FFrontendQueryKey{InEntry.Value.Get<FInterfaceRegistryTransaction>().GetInterfaceVersion().Name};
-				}
-				return FFrontendQueryKey();
+				return FFrontendQueryKey{InEntry.Value.Get<FInterfaceRegistryTransaction>().GetInterfaceVersion().Name};
 			}
+			return FFrontendQueryKey();
+		}
 
-			void FReduceInterfaceRegistryTransactionsToCurrentStatus::Reduce(const FFrontendQueryKey& InKey, FFrontendQueryPartition& InOutEntries) const
-			{
-				// Get most recent transaction
-				const FFrontendQueryEntry* FinalEntry = Algo::MaxElementBy(InOutEntries, GetTransactionTimestamp);
+		void FReduceInterfaceRegistryTransactionsToCurrentStatus::Reduce(const FFrontendQueryKey& InKey, FFrontendQueryPartition& InOutEntries) const
+		{
+			// Get most recent transaction
+			const FFrontendQueryEntry* FinalEntry = Algo::MaxElementBy(InOutEntries, GetTransactionTimestamp);
 
-				// Check if most recent transaction is valid and not an unregister transaction.
-				if (IsValidTransactionOfType(FInterfaceRegistryTransaction::ETransactionType::InterfaceRegistration, FinalEntry))
-				{
-					FFrontendQueryEntry Entry = *FinalEntry;
-					InOutEntries.Reset();
-					InOutEntries.Add(Entry);
-				}
-				else
-				{
-					InOutEntries.Reset();
-				}
+			// Check if most recent transaction is valid and not an unregister transaction.
+			if (IsValidTransactionOfType(FInterfaceRegistryTransaction::ETransactionType::InterfaceRegistration, FinalEntry))
+			{
+				FFrontendQueryEntry Entry = *FinalEntry;
+				InOutEntries.Reset();
+				InOutEntries.Add(Entry);
 			}
-
-			FInterfaceRegistryTransaction::FTimeType FReduceInterfaceRegistryTransactionsToCurrentStatus::GetTransactionTimestamp(const FFrontendQueryEntry& InEntry)
+			else
 			{
-				if (ensure(InEntry.Value.IsType<FInterfaceRegistryTransaction>()))
-				{
-					return InEntry.Value.Get<FInterfaceRegistryTransaction>().GetTimestamp();
-				}
-				return 0;
+				InOutEntries.Reset();
 			}
+		}
 
-			bool FReduceInterfaceRegistryTransactionsToCurrentStatus::IsValidTransactionOfType(FInterfaceRegistryTransaction::ETransactionType InType, const FFrontendQueryEntry* InEntry)
+		FInterfaceRegistryTransaction::FTimeType FReduceInterfaceRegistryTransactionsToCurrentStatus::GetTransactionTimestamp(const FFrontendQueryEntry& InEntry)
+		{
+			if (ensure(InEntry.Value.IsType<FInterfaceRegistryTransaction>()))
 			{
-				if (nullptr != InEntry)
-				{
-					if (InEntry->Value.IsType<FInterfaceRegistryTransaction>())
-					{
-						return InEntry->Value.Get<FInterfaceRegistryTransaction>().GetTransactionType() == InType;
-					}
-				}
-				return false;
+				return InEntry.Value.Get<FInterfaceRegistryTransaction>().GetTimestamp();
 			}
+			return 0;
+		}
 
-			void FTransformInterfaceRegistryTransactionToInterface::Transform(FFrontendQueryEntry::FValue& InValue) const
+		bool FReduceInterfaceRegistryTransactionsToCurrentStatus::IsValidTransactionOfType(FInterfaceRegistryTransaction::ETransactionType InType, const FFrontendQueryEntry* InEntry)
+		{
+			if (nullptr != InEntry)
 			{
-				FMetasoundFrontendInterface Interface;
-				if (ensure(InValue.IsType<FInterfaceRegistryTransaction>()))
+				if (InEntry->Value.IsType<FInterfaceRegistryTransaction>())
 				{
-					FInterfaceRegistryKey RegistryKey = InValue.Get<FInterfaceRegistryTransaction>().GetInterfaceRegistryKey();
-					IInterfaceRegistry::Get().FindInterface(RegistryKey, Interface);
+					return InEntry->Value.Get<FInterfaceRegistryTransaction>().GetTransactionType() == InType;
 				}
-				InValue.Set<FMetasoundFrontendInterface>(MoveTemp(Interface));
 			}
+			return false;
+		}
 
-			FFrontendQueryKey FMapInterfaceToInterfaceNameAndMajorVersion::GetKey(const FString& InName, int32 InMajorVersion)
+		void FTransformInterfaceRegistryTransactionToInterface::Transform(FFrontendQueryEntry::FValue& InValue) const
+		{
+			FMetasoundFrontendInterface Interface;
+			if (ensure(InValue.IsType<FInterfaceRegistryTransaction>()))
 			{
-				return FFrontendQueryKey{FString::Format(TEXT("{0}_{1}"), { InName, InMajorVersion })};
+				FInterfaceRegistryKey RegistryKey = InValue.Get<FInterfaceRegistryTransaction>().GetInterfaceRegistryKey();
+				IInterfaceRegistry::Get().FindInterface(RegistryKey, Interface);
 			}
+			InValue.Set<FMetasoundFrontendInterface>(MoveTemp(Interface));
+		}
 
-			FFrontendQueryKey FMapInterfaceToInterfaceNameAndMajorVersion::Map(const FFrontendQueryEntry& InEntry) const
+		void FTransformInterfaceToInterfaceVersion::Transform(FFrontendQueryEntry::FValue& InValue) const
+		{
+			const FMetasoundFrontendInterface* Interface = InValue.TryGet<FMetasoundFrontendInterface>();
+			if (ensure(Interface))
 			{
-				FString Key;
-				if (ensure(InEntry.Value.IsType<FMetasoundFrontendInterface>()))
-				{
-					const FMetasoundFrontendInterface& Interface = InEntry.Value.Get<FMetasoundFrontendInterface>();
-
-					return FMapInterfaceToInterfaceNameAndMajorVersion::GetKey(Interface.Version.Name.ToString(), Interface.Version.Number.Major);
-				}
-				return FFrontendQueryKey{Key};
+				InValue.Set<FMetasoundFrontendVersion>(Interface->Version);
 			}
+		}
 
-			void FReduceInterfacesToHighestVersion::Reduce(const FFrontendQueryKey& InKey, FFrontendQueryPartition& InOutEntries) const
-			{
-				FFrontendQueryEntry* HighestVersionEntry = nullptr;
-				FMetasoundFrontendVersionNumber HighestVersion = FMetasoundFrontendVersionNumber::GetInvalid();
+		FFrontendQueryKey FMapInterfaceToInterfaceNameAndMajorVersion::GetKey(const FString& InName, int32 InMajorVersion)
+		{
+			return FFrontendQueryKey{FString::Format(TEXT("{0}_{1}"), { InName, InMajorVersion })};
+		}
 
-				for (FFrontendQueryEntry& Entry : InOutEntries)
-				{
-					if (ensure(Entry.Value.IsType<FMetasoundFrontendInterface>()))
-					{
-						const FMetasoundFrontendVersionNumber& VersionNumber = Entry.Value.Get<FMetasoundFrontendInterface>().Version.Number;
-						if (VersionNumber > HighestVersion)
-						{
-							HighestVersionEntry = &Entry;
-							HighestVersion = VersionNumber;
-						}
-					}
-				}
+		FFrontendQueryKey FMapInterfaceToInterfaceNameAndMajorVersion::Map(const FFrontendQueryEntry& InEntry) const
+		{
+			FString Key;
+			if (ensure(InEntry.Value.IsType<FMetasoundFrontendInterface>()))
+			{
+				const FMetasoundFrontendInterface& Interface = InEntry.Value.Get<FMetasoundFrontendInterface>();
 
-				if (HighestVersionEntry)
-				{
-					FFrontendQueryEntry Entry = *HighestVersionEntry;
-					InOutEntries.Reset();
-					InOutEntries.Add(Entry);
-				}
-				else
-				{
-					InOutEntries.Reset();
-				}
+				return FMapInterfaceToInterfaceNameAndMajorVersion::GetKey(Interface.Version.Name.ToString(), Interface.Version.Number.Major);
 			}
+			return FFrontendQueryKey{Key};
+		}
 
-			bool FRemoveDeprecatedClasses::Filter(const FFrontendQueryEntry& InEntry) const
-			{
-				if (InEntry.Value.IsType<FMetasoundFrontendClass>())
-				{
-					return !InEntry.Value.Get<FMetasoundFrontendClass>().Metadata.GetIsDeprecated();
-				}
-				return false;
-			}
+		void FReduceInterfacesToHighestVersion::Reduce(const FFrontendQueryKey& InKey, FFrontendQueryPartition& InOutEntries) const
+		{
+			FFrontendQueryEntry* HighestVersionEntry = nullptr;
+			FMetasoundFrontendVersionNumber HighestVersion = FMetasoundFrontendVersionNumber::GetInvalid();
 
-			class FMapNodeRegistrationEventsToClassAndMajorVersion : public IFrontendQueryMapStep
+			for (FFrontendQueryEntry& Entry : InOutEntries)
 			{
-			public:
-				virtual FFrontendQueryKey Map(const FFrontendQueryEntry& InEntry) const override
+				if (ensure(Entry.Value.IsType<FMetasoundFrontendInterface>()))
 				{
-					if (InEntry.Value.IsType<FNodeRegistryTransaction>())
-					{
-						const FNodeClassInfo& Info = InEntry.Value.Get<FNodeRegistryTransaction>().GetNodeClassInfo();
-						return CreateKeyFromFullClassNameAndMajorVersion(Info.ClassName, Info.Version.Major);
-					}
-					else
+					const FMetasoundFrontendVersionNumber& VersionNumber = Entry.Value.Get<FMetasoundFrontendInterface>().Version.Number;
+					if (VersionNumber > HighestVersion)
 					{
-						return FFrontendQueryKey();
+						HighestVersionEntry = &Entry;
+						HighestVersion = VersionNumber;
 					}
 				}
-			};
-
-			FFrontendQueryKey FMapInterfaceToInterfaceName::Map(const FFrontendQueryEntry& InEntry) const
-			{
-				FName Key;
-				if (ensure(InEntry.Value.IsType<FMetasoundFrontendInterface>()))
-				{
-					Key = InEntry.Value.Get<FMetasoundFrontendInterface>().Version.Name;
-				}
-				return FFrontendQueryKey{Key};
 			}
 
-			void FTransformInterfaceRegistryTransactionToInterfaceVersion::Transform(FFrontendQueryEntry::FValue& InValue) const
+			if (HighestVersionEntry)
 			{
-				FMetasoundFrontendVersion Version;
-				if (ensure(InValue.IsType<FInterfaceRegistryTransaction>()))
-				{
-					Version = InValue.Get<FInterfaceRegistryTransaction>().GetInterfaceVersion();
-				}
-				InValue.Set<FMetasoundFrontendVersion>(Version);
+				FFrontendQueryEntry Entry = *HighestVersionEntry;
+				InOutEntries.Reset();
+				InOutEntries.Add(Entry);
 			}
-
-			TArray<FMetasoundFrontendClass> BuildArrayOfClassesFromPartition(const FFrontendQueryPartition& InPartition)
+			else
 			{
-				TArray<FMetasoundFrontendClass> Result;
-				for (const FFrontendQueryEntry& Entry : InPartition)
-				{
-					check(Entry.Value.IsType<FMetasoundFrontendClass>());
-					Result.Add(Entry.Value.Get<FMetasoundFrontendClass>());
-				}
-				return Result;
+				InOutEntries.Reset();
 			}
+		}
 
-			FMetasoundFrontendClass BuildSingleClassFromPartition(const FFrontendQueryPartition& InPartition)
+		TArray<FFrontendQueryKey> FMapUClassToDefaultInterface::Map(const FFrontendQueryEntry& InEntry) const
+		{
+			TArray<FFrontendQueryKey> Keys;
+			if (const FMetasoundFrontendInterface* Interface = InEntry.Value.TryGet<FMetasoundFrontendInterface>())
 			{
-				FMetasoundFrontendClass MetaSoundClass;
-				if (ensureMsgf(InPartition.Num() > 0, TEXT("Cannot retrieve class from empty partition")))
+				auto IsDefaultInterface = [](const FMetasoundFrontendInterfaceUClassOptions& InOptions)
 				{
-					MetaSoundClass = InPartition.CreateConstIterator()->Value.Get<FMetasoundFrontendClass>();
-				}
-				return MetaSoundClass;
-			}
-
-			TArray<FMetasoundFrontendInterface> BuildArrayOfInterfacesFromPartition(const FFrontendQueryPartition& InPartition)
-			{
-				TArray<FMetasoundFrontendInterface> Result;
-				for (const FFrontendQueryEntry& Entry : InPartition)
+					return InOptions.bIsDefault;
+				};
+				auto GetUClassQueryKey = [](const FMetasoundFrontendInterfaceUClassOptions& InOptions)
 				{
-					check(Entry.Value.IsType<FMetasoundFrontendInterface>());
-					Result.Add(Entry.Value.Get<FMetasoundFrontendInterface>());
-				}
-				return Result;
+					return CreateKeyFromUClassPathName(InOptions.ClassPath);
+				};
+				Algo::TransformIf(Interface->UClassOptions, Keys, IsDefaultInterface, GetUClassQueryKey);
 			}
 
-			FMetasoundFrontendInterface BuildSingleInterfaceFromPartition(const FFrontendQueryPartition& InPartition)
-			{
-				FMetasoundFrontendInterface MetaSoundInterface;
-				if (ensureMsgf(InPartition.Num() > 0, TEXT("Cannot retrieve class from empty partition")))
-				{
-					MetaSoundInterface = InPartition.CreateConstIterator()->Value.Get<FMetasoundFrontendInterface>();
-				}
-				return MetaSoundInterface;
-			}
+			return Keys;
+		}
 
-			TArray<FMetasoundFrontendVersion> BuildArrayOfVersionsFromPartition(const FFrontendQueryPartition& InPartition)
+		bool FRemoveDeprecatedClasses::Filter(const FFrontendQueryEntry& InEntry) const
+		{
+			if (InEntry.Value.IsType<FMetasoundFrontendClass>())
 			{
-				TArray<FMetasoundFrontendVersion> Result;
-				for (const FFrontendQueryEntry& Entry : InPartition)
-				{
-					check(Entry.Value.IsType<FMetasoundFrontendVersion>());
-					Result.Add(Entry.Value.Get<FMetasoundFrontendVersion>());
-				}
-				return Result;
+				return !InEntry.Value.Get<FMetasoundFrontendClass>().Metadata.GetIsDeprecated();
 			}
-		} // namespace SearchEngineQuerySteps
-
-
-		FFrontendQuery FFindClassWithHighestMinorVersionQueryPolicy::CreateQuery()
-		{
-			using namespace SearchEngineQuerySteps;
-			FFrontendQuery Query;
-			Query.AddStep<FNodeClassRegistrationEvents>()
-				.AddStep<FMapRegistrationEventsToNodeRegistryKeys>()
-				.AddStep<FReduceRegistrationEventsToCurrentStatus>()
-				.AddStep<FMapNodeRegistrationEventsToClassAndMajorVersion>();
-			return Query;
+			return false;
 		}
 
-		FNodeRegistryKey FFindClassWithHighestMinorVersionQueryPolicy::BuildResult(const FFrontendQueryPartition& InPartition)
+		class FMapNodeRegistrationEventsToClassAndMajorVersion : public IFrontendQueryMapStep
 		{
-			using namespace SearchEngineQuerySteps;
-
-			FMetasoundFrontendVersionNumber HighestVersionNumber{-1, -1};
-			FNodeRegistryKey Key = NodeRegistryKey::GetInvalid();
-			for (const FFrontendQueryEntry& InEntry : InPartition)
+		public:
+			virtual FFrontendQueryKey Map(const FFrontendQueryEntry& InEntry) const override
 			{
 				if (InEntry.Value.IsType<FNodeRegistryTransaction>())
 				{
-					const FNodeRegistryTransaction& Transaction = InEntry.Value.Get<FNodeRegistryTransaction>();
-					if (FNodeRegistryTransaction::ETransactionType::NodeRegistration == Transaction.GetTransactionType())
-					{
-						const FNodeClassInfo& NodeClassInfo = Transaction.GetNodeClassInfo();
-						if (NodeClassInfo.Version > HighestVersionNumber)
-						{
-							HighestVersionNumber = NodeClassInfo.Version;
-							Key = Transaction.GetNodeRegistryKey();
-						}
-					}
+					const FNodeClassInfo& Info = InEntry.Value.Get<FNodeRegistryTransaction>().GetNodeClassInfo();
+					return CreateKeyFromFullClassNameAndMajorVersion(Info.ClassName, Info.Version.Major);
+				}
+				else
+				{
+					return FFrontendQueryKey();
 				}
 			}
+		};
 
-			return Key;
-		}
-
-		FFrontendQuery FFindAllRegisteredInterfacesWithNameQueryPolicy::CreateQuery()
+		FFrontendQueryKey FMapInterfaceToInterfaceName::Map(const FFrontendQueryEntry& InEntry) const
 		{
-			using namespace SearchEngineQuerySteps;
-			FFrontendQuery Query;
-			Query.AddStep<FInterfaceRegistryTransactionSource>()
-				.AddStep<FMapInterfaceRegistryTransactionsToInterfaceRegistryKeys>()
-				.AddStep<FReduceInterfaceRegistryTransactionsToCurrentStatus>()
-				.AddStep<FMapInterfaceRegistryTransactionToInterfaceName>()
-				.AddStep<FTransformInterfaceRegistryTransactionToInterfaceVersion>();
-			return Query;
+			FName Key;
+			if (ensure(InEntry.Value.IsType<FMetasoundFrontendInterface>()))
+			{
+				Key = InEntry.Value.Get<FMetasoundFrontendInterface>().Version.Name;
+			}
+			return FFrontendQueryKey{Key};
 		}
 
-		TArray<FMetasoundFrontendVersion> FFindAllRegisteredInterfacesWithNameQueryPolicy::BuildResult(const FFrontendQueryPartition& InPartition)
+		void FTransformInterfaceRegistryTransactionToInterfaceVersion::Transform(FFrontendQueryEntry::FValue& InValue) const
 		{
-			using namespace SearchEngineQuerySteps;
-			return BuildArrayOfVersionsFromPartition(InPartition);
+			FMetasoundFrontendVersion Version;
+			if (ensure(InValue.IsType<FInterfaceRegistryTransaction>()))
+			{
+				Version = InValue.Get<FInterfaceRegistryTransaction>().GetInterfaceVersion();
+			}
+			InValue.Set<FMetasoundFrontendVersion>(Version);
 		}
 
-		FFrontendQuery FFindInterfaceWithHighestVersionQueryPolicy::CreateQuery()
+		TArray<FMetasoundFrontendClass> BuildArrayOfClassesFromPartition(const FFrontendQueryPartition& InPartition)
 		{
-			using namespace SearchEngineQuerySteps;
-			FFrontendQuery Query;
-			Query.AddStep<FInterfaceRegistryTransactionSource>()
-				.AddStep<FMapInterfaceRegistryTransactionsToInterfaceRegistryKeys>()
-				.AddStep<FReduceInterfaceRegistryTransactionsToCurrentStatus>()
-				.AddStep<FTransformInterfaceRegistryTransactionToInterface>()
-				.AddStep<FMapInterfaceToInterfaceName>()
-				.AddStep<FReduceInterfacesToHighestVersion>();
-			return Query;
+			TArray<FMetasoundFrontendClass> Result;
+			for (const FFrontendQueryEntry& Entry : InPartition)
+			{
+				check(Entry.Value.IsType<FMetasoundFrontendClass>());
+				Result.Add(Entry.Value.Get<FMetasoundFrontendClass>());
+			}
+			return Result;
 		}
 
-		FMetasoundFrontendInterface FFindInterfaceWithHighestVersionQueryPolicy::BuildResult(const FFrontendQueryPartition& InPartition)
+		FMetasoundFrontendClass BuildSingleClassFromPartition(const FFrontendQueryPartition& InPartition)
 		{
-			using namespace SearchEngineQuerySteps;
-			return BuildSingleInterfaceFromPartition(InPartition);
+			FMetasoundFrontendClass MetaSoundClass;
+			if (ensureMsgf(InPartition.Num() > 0, TEXT("Cannot retrieve class from empty partition")))
+			{
+				MetaSoundClass = InPartition.CreateConstIterator()->Value.Get<FMetasoundFrontendClass>();
+			}
+			return MetaSoundClass;
 		}
 
-		FFrontendQuery FFindAllInterfacesQueryPolicy::CreateQuery()
+		TArray<FMetasoundFrontendInterface> BuildArrayOfInterfacesFromPartition(const FFrontendQueryPartition& InPartition)
 		{
-			using namespace SearchEngineQuerySteps;
-			FFrontendQuery Query;
-			Query.AddStep<FInterfaceRegistryTransactionSource>()
-				.AddStep<FMapInterfaceRegistryTransactionsToInterfaceRegistryKeys>()
-				.AddStep<FReduceInterfaceRegistryTransactionsToCurrentStatus>()
-				.AddStep<FTransformInterfaceRegistryTransactionToInterface>()
-				.AddStep<FMapInterfaceToInterfaceNameAndMajorVersion>()
-				.AddStep<FReduceInterfacesToHighestVersion>()
-				.AddStep<FMapToNull>();
-			return Query;
+			TArray<FMetasoundFrontendInterface> Result;
+			for (const FFrontendQueryEntry& Entry : InPartition)
+			{
+				check(Entry.Value.IsType<FMetasoundFrontendInterface>());
+				Result.Add(Entry.Value.Get<FMetasoundFrontendInterface>());
+			}
+			return Result;
 		}
 
-		TArray<FMetasoundFrontendInterface> FFindAllInterfacesQueryPolicy::BuildResult(const FFrontendQueryPartition& InPartition)
+		FMetasoundFrontendInterface BuildSingleInterfaceFromPartition(const FFrontendQueryPartition& InPartition)
 		{
-			using namespace SearchEngineQuerySteps;
-			return BuildArrayOfInterfacesFromPartition(InPartition);
+			FMetasoundFrontendInterface MetaSoundInterface;
+			if (ensureMsgf(InPartition.Num() > 0, TEXT("Cannot retrieve class from empty partition")))
+			{
+				MetaSoundInterface = InPartition.CreateConstIterator()->Value.Get<FMetasoundFrontendInterface>();
+			}
+			return MetaSoundInterface;
 		}
 
-		FFrontendQuery FFindAllInterfacesIncludingAllVersionsQueryPolicy::CreateQuery()
+		TArray<FMetasoundFrontendVersion> BuildArrayOfVersionsFromPartition(const FFrontendQueryPartition& InPartition)
 		{
-			using namespace SearchEngineQuerySteps;
-			FFrontendQuery Query;
-			Query.AddStep<FInterfaceRegistryTransactionSource>()
-				.AddStep<FMapInterfaceRegistryTransactionsToInterfaceRegistryKeys>()
-				.AddStep<FReduceInterfaceRegistryTransactionsToCurrentStatus>()
-				.AddStep<FTransformInterfaceRegistryTransactionToInterface>()
-				.AddStep<FMapToNull>();
-			return Query;
-		}
-
-		TArray<FMetasoundFrontendInterface> FFindAllInterfacesIncludingAllVersionsQueryPolicy::BuildResult(const FFrontendQueryPartition& InPartition)
+			TArray<FMetasoundFrontendVersion> Result;
+			for (const FFrontendQueryEntry& Entry : InPartition)
 			{
-				using namespace SearchEngineQuerySteps;
-				return BuildArrayOfInterfacesFromPartition(InPartition);
+				check(Entry.Value.IsType<FMetasoundFrontendVersion>());
+				Result.Add(Entry.Value.Get<FMetasoundFrontendVersion>());
 			}
-
-		void FSearchEngineCore::Prime()
-		{
-			METASOUND_LLM_SCOPE;
-			METASOUND_TRACE_CPUPROFILER_EVENT_SCOPE(metasound::FSearchEngineCore::Prime);
-
-			FindClassWithHighestMinorVersionQuery.Prime();
-			FindAllRegisteredInterfacesWithNameQuery.Prime();
-			FindInterfaceWithHighestVersionQuery.Prime();
-			FindAllInterfacesIncludingAllVersionsQuery.Prime();
-			FindAllInterfacesQuery.Prime();
+			return Result;
 		}
+	} // namespace SearchEngineQuerySteps
 
-		bool FSearchEngineCore::FindClassWithHighestMinorVersion(const FMetasoundFrontendClassName& InName, int32 InMajorVersion, FMetasoundFrontendClass& OutClass)
-		{
-			METASOUND_TRACE_CPUPROFILER_EVENT_SCOPE(metasound::FSearchEngineCore::FindClassWithHighestMinorVersion);
-			using namespace SearchEngineQuerySteps;
 
-			const FFrontendQueryKey QueryKey = CreateKeyFromFullClassNameAndMajorVersion(InName, InMajorVersion);
-			FNodeRegistryKey NodeRegistryKey; 
-			if (FindClassWithHighestMinorVersionQuery.UpdateAndFindResult(QueryKey, NodeRegistryKey))
+	FFrontendQuery FFindClassWithHighestMinorVersionQueryPolicy::CreateQuery()
+	{
+		using namespace SearchEngineQuerySteps;
+		FFrontendQuery Query;
+		Query.AddStep<FNodeClassRegistrationEvents>()
+			.AddStep<FMapRegistrationEventsToNodeRegistryKeys>()
+			.AddStep<FReduceRegistrationEventsToCurrentStatus>()
+			.AddStep<FMapNodeRegistrationEventsToClassAndMajorVersion>();
+		return Query;
+	}
+
+	FNodeRegistryKey FFindClassWithHighestMinorVersionQueryPolicy::BuildResult(const FFrontendQueryPartition& InPartition)
+	{
+		using namespace SearchEngineQuerySteps;
+
+		FMetasoundFrontendVersionNumber HighestVersionNumber{-1, -1};
+		FNodeRegistryKey Key = NodeRegistryKey::GetInvalid();
+		for (const FFrontendQueryEntry& InEntry : InPartition)
+		{
+			if (InEntry.Value.IsType<FNodeRegistryTransaction>())
 			{
-				if (FMetasoundFrontendRegistryContainer* NodeRegistry = FMetasoundFrontendRegistryContainer::Get())
+				const FNodeRegistryTransaction& Transaction = InEntry.Value.Get<FNodeRegistryTransaction>();
+				if (FNodeRegistryTransaction::ETransactionType::NodeRegistration == Transaction.GetTransactionType())
 				{
-					return NodeRegistry->FindFrontendClassFromRegistered(NodeRegistryKey, OutClass);
+					const FNodeClassInfo& NodeClassInfo = Transaction.GetNodeClassInfo();
+					if (NodeClassInfo.Version > HighestVersionNumber)
+					{
+						HighestVersionNumber = NodeClassInfo.Version;
+						Key = Transaction.GetNodeRegistryKey();
+					}
 				}
 			}
-			return false;
 		}
 
-		TArray<FMetasoundFrontendInterface> FSearchEngineCore::FindUClassDefaultInterfaces(FName InUClassName)
-		{
-			METASOUND_TRACE_CPUPROFILER_EVENT_SCOPE(metasound::FSearchEngineCore::FindUClassDefaultInterfaces);
-
-			constexpr bool bIncludeAllVersions = false;
-			const FFrontendQueryKey NullKey;
-			TArray<FMetasoundFrontendInterface> DefaultInterfaces = FindAllInterfaces(bIncludeAllVersions);
-
-			// All default interfaces are filtered here because a UClass name cannot effectively 
-			// be used as a "Key" for this query.
-			//
-			// We cannot know all the UClasses associated with each default interface given only a UClassName
-			// because we need to inspect UClass inheritance to handle derived UClasses
-			// which should be compatible with their parent's default interfaces.
-			//
-			// This module does not have access to the UClass inheritance structure.
-			// Instead, this information is accessed in the IInterfaceRegistryEntry
-			// concrete implementation.
-			auto DoesNotSupportUClass = [&InUClassName](const FMetasoundFrontendInterface& InInterface)
-			{
-				FInterfaceRegistryKey RegistryKey = GetInterfaceRegistryKey(InInterface);
-				const IInterfaceRegistryEntry* RegEntry = IInterfaceRegistry::Get().FindInterfaceRegistryEntry(RegistryKey);
-				if (ensure(RegEntry))
-				{
-					return !RegEntry->IsDefault(InUClassName);
-				}
+		return Key;
+	}
 
-				return true;
-			};
+	FFrontendQuery FFindAllRegisteredInterfacesWithNameQueryPolicy::CreateQuery()
+	{
+		using namespace SearchEngineQuerySteps;
+		FFrontendQuery Query;
+		Query.AddStep<FInterfaceRegistryTransactionSource>()
+			.AddStep<FMapInterfaceRegistryTransactionsToInterfaceRegistryKeys>()
+			.AddStep<FReduceInterfaceRegistryTransactionsToCurrentStatus>()
+			.AddStep<FMapInterfaceRegistryTransactionToInterfaceName>()
+			.AddStep<FTransformInterfaceRegistryTransactionToInterfaceVersion>();
+		return Query;
+	}
+
+	TArray<FMetasoundFrontendVersion> FFindAllRegisteredInterfacesWithNameQueryPolicy::BuildResult(const FFrontendQueryPartition& InPartition)
+	{
+		using namespace SearchEngineQuerySteps;
+		return BuildArrayOfVersionsFromPartition(InPartition);
+	}
 
-			DefaultInterfaces.RemoveAllSwap(DoesNotSupportUClass);
+	FFrontendQuery FFindInterfaceWithHighestVersionQueryPolicy::CreateQuery()
+	{
+		using namespace SearchEngineQuerySteps;
+		FFrontendQuery Query;
+		Query.AddStep<FInterfaceRegistryTransactionSource>()
+			.AddStep<FMapInterfaceRegistryTransactionsToInterfaceRegistryKeys>()
+			.AddStep<FReduceInterfaceRegistryTransactionsToCurrentStatus>()
+			.AddStep<FTransformInterfaceRegistryTransactionToInterface>()
+			.AddStep<FMapInterfaceToInterfaceName>()
+			.AddStep<FReduceInterfacesToHighestVersion>();
+		return Query;
+	}
+
+	FMetasoundFrontendInterface FFindInterfaceWithHighestVersionQueryPolicy::BuildResult(const FFrontendQueryPartition& InPartition)
+	{
+		using namespace SearchEngineQuerySteps;
+		return BuildSingleInterfaceFromPartition(InPartition);
+	}
 
-			return DefaultInterfaces;
-		}
+	FFrontendQuery FFindAllInterfacesQueryPolicy::CreateQuery()
+	{
+		using namespace SearchEngineQuerySteps;
+		FFrontendQuery Query;
+		Query.AddStep<FInterfaceRegistryTransactionSource>()
+			.AddStep<FMapInterfaceRegistryTransactionsToInterfaceRegistryKeys>()
+			.AddStep<FReduceInterfaceRegistryTransactionsToCurrentStatus>()
+			.AddStep<FTransformInterfaceRegistryTransactionToInterface>()
+			.AddStep<FMapInterfaceToInterfaceNameAndMajorVersion>()
+			.AddStep<FReduceInterfacesToHighestVersion>()
+			.AddStep<FMapToNull>();
+		return Query;
+	}
+
+	TArray<FMetasoundFrontendInterface> FFindAllInterfacesQueryPolicy::BuildResult(const FFrontendQueryPartition& InPartition)
+	{
+		using namespace SearchEngineQuerySteps;
+		return BuildArrayOfInterfacesFromPartition(InPartition);
+	}
 
-		TArray<FMetasoundFrontendVersion> FSearchEngineCore::FindAllRegisteredInterfacesWithName(FName InInterfaceName)
-		{
-			METASOUND_TRACE_CPUPROFILER_EVENT_SCOPE(metasound::FSearchEngineCore::FindAllRegisteredInterfacesWithName);
+	FFrontendQuery FFindAllInterfacesIncludingAllVersionsQueryPolicy::CreateQuery()
+	{
+		using namespace SearchEngineQuerySteps;
+		FFrontendQuery Query;
+		Query.AddStep<FInterfaceRegistryTransactionSource>()
+			.AddStep<FMapInterfaceRegistryTransactionsToInterfaceRegistryKeys>()
+			.AddStep<FReduceInterfaceRegistryTransactionsToCurrentStatus>()
+			.AddStep<FTransformInterfaceRegistryTransactionToInterface>()
+			.AddStep<FMapToNull>();
+		return Query;
+	}
+
+	TArray<FMetasoundFrontendInterface> FFindAllInterfacesIncludingAllVersionsQueryPolicy::BuildResult(const FFrontendQueryPartition& InPartition)
+	{
+		using namespace SearchEngineQuerySteps;
+		return BuildArrayOfInterfacesFromPartition(InPartition);
+	}
 
-			const FFrontendQueryKey Key{InInterfaceName};
-			return FindAllRegisteredInterfacesWithNameQuery.UpdateAndFindResult(Key);
-		}
+	FFrontendQuery FFindAllDefaultInterfaceVersionsForUClassQueryPolicy::CreateQuery()
+	{
+		using namespace SearchEngineQuerySteps;
+		FFrontendQuery Query;
+		Query.AddStep<FInterfaceRegistryTransactionSource>()
+			.AddStep<FMapInterfaceRegistryTransactionsToInterfaceRegistryKeys>()
+			.AddStep<FReduceInterfaceRegistryTransactionsToCurrentStatus>()
+			.AddStep<FTransformInterfaceRegistryTransactionToInterface>()
+			.AddStep<FMapUClassToDefaultInterface>()
+			.AddStep<FTransformInterfaceToInterfaceVersion>();
+		return Query;
+	}
+
+	TArray<FMetasoundFrontendVersion> FFindAllDefaultInterfaceVersionsForUClassQueryPolicy::BuildResult(const FFrontendQueryPartition& InPartition)
+	{
+		using namespace SearchEngineQuerySteps;
+		return BuildArrayOfVersionsFromPartition(InPartition);
+	}
 
-		bool FSearchEngineCore::FindInterfaceWithHighestVersion(FName InInterfaceName, FMetasoundFrontendInterface& OutInterface)
-		{
-			METASOUND_TRACE_CPUPROFILER_EVENT_SCOPE(metasound::FSearchEngineCore::FindInterfaceWithHighestVersion);
+	void FSearchEngineCore::Prime()
+	{
+		METASOUND_LLM_SCOPE;
+		METASOUND_TRACE_CPUPROFILER_EVENT_SCOPE(metasound::FSearchEngineCore::Prime);
 
-			const FFrontendQueryKey Key{InInterfaceName};
-			return FindInterfaceWithHighestVersionQuery.UpdateAndFindResult(Key, OutInterface);
-		}
+		FindClassWithHighestMinorVersionQuery.Prime();
+		FindAllRegisteredInterfacesWithNameQuery.Prime();
+		FindInterfaceWithHighestVersionQuery.Prime();
+		FindAllDefaultInterfaceVersionsForUClassQuery.Prime();
+	}
 
-		TArray<FMetasoundFrontendInterface> FSearchEngineCore::FindAllInterfaces(bool bInIncludeAllVersions)
-		{
-			METASOUND_TRACE_CPUPROFILER_EVENT_SCOPE(metasound::FSearchEngine::FindAllInterfaces);
+	bool FSearchEngineCore::FindClassWithHighestMinorVersion(const FMetasoundFrontendClassName& InName, int32 InMajorVersion, FMetasoundFrontendClass& OutClass)
+	{
+		METASOUND_TRACE_CPUPROFILER_EVENT_SCOPE(metasound::FSearchEngineCore::FindClassWithHighestMinorVersion);
+		using namespace SearchEngineQuerySteps;
 
-			const FFrontendQueryKey NullKey;
-			if (bInIncludeAllVersions)
-			{
-				return FindAllInterfacesIncludingAllVersionsQuery.UpdateAndFindResult(NullKey);
-			}
-			else
+		const FFrontendQueryKey QueryKey = CreateKeyFromFullClassNameAndMajorVersion(InName, InMajorVersion);
+		FNodeRegistryKey NodeRegistryKey; 
+		if (FindClassWithHighestMinorVersionQuery.UpdateAndFindResult(QueryKey, NodeRegistryKey))
+		{
+			if (FMetasoundFrontendRegistryContainer* NodeRegistry = FMetasoundFrontendRegistryContainer::Get())
 			{
-				return FindAllInterfacesQuery.UpdateAndFindResult(NullKey);
+				return NodeRegistry->FindFrontendClassFromRegistered(NodeRegistryKey, OutClass);
 			}
 		}
-	} // namespace Frontend
-} // namespace Metasound
+		return false;
+	}
+
+	TArray<FMetasoundFrontendInterface> FSearchEngineCore::FindUClassDefaultInterfaces(FName InUClassName)
+	{
+		using namespace SearchEngineQuerySteps;
+
+		METASOUND_TRACE_CPUPROFILER_EVENT_SCOPE(metasound::FSearchEngineCore::FindUClassDefaultInterfaces);
+
+		TArray<FMetasoundFrontendInterface> Interfaces;
+		TArray<FMetasoundFrontendVersion> InterfaceVersions = FindUClassDefaultInterfaceVersions(FTopLevelAssetPath(InUClassName.ToString()));
+		Algo::Transform(InterfaceVersions, Interfaces, [](const FMetasoundFrontendVersion& Version)
+		{
+			FMetasoundFrontendInterface Interface;
+			const FInterfaceRegistryKey Key = GetInterfaceRegistryKey(Version);
+			ensure(IInterfaceRegistry::Get().FindInterface(Key, Interface));
+			return Interface;
+		});
+
+		return Interfaces;
+	}
+
+	TArray<FMetasoundFrontendVersion> FSearchEngineCore::FindUClassDefaultInterfaceVersions(const FTopLevelAssetPath& InUClassPath)
+	{
+		using namespace SearchEngineQuerySteps;
+
+		METASOUND_TRACE_CPUPROFILER_EVENT_SCOPE(metasound::FSearchEngineCore::FindUClassDefaultInterfaceVersions);
+
+		const FFrontendQueryKey Key = CreateKeyFromUClassPathName(InUClassPath);
+		return FindAllDefaultInterfaceVersionsForUClassQuery.UpdateAndFindResult(Key);
+	}
+
+	TArray<FMetasoundFrontendVersion> FSearchEngineCore::FindAllRegisteredInterfacesWithName(FName InInterfaceName)
+	{
+		METASOUND_TRACE_CPUPROFILER_EVENT_SCOPE(metasound::FSearchEngineCore::FindAllRegisteredInterfacesWithName);
+
+		const FFrontendQueryKey Key{InInterfaceName};
+		return FindAllRegisteredInterfacesWithNameQuery.UpdateAndFindResult(Key);
+	}
+
+	bool FSearchEngineCore::FindInterfaceWithHighestVersion(FName InInterfaceName, FMetasoundFrontendInterface& OutInterface)
+	{
+		METASOUND_TRACE_CPUPROFILER_EVENT_SCOPE(metasound::FSearchEngineCore::FindInterfaceWithHighestVersion);
+
+		const FFrontendQueryKey Key{InInterfaceName};
+		return FindInterfaceWithHighestVersionQuery.UpdateAndFindResult(Key, OutInterface);
+	}
+} // namespace Metasound::Frontend
diff --git a/Engine/Plugins/Runtime/Metasound/Source/MetasoundFrontend/Private/MetasoundFrontendSearchEngineCore.h b/Engine/Plugins/Runtime/Metasound/Source/MetasoundFrontend/Private/MetasoundFrontendSearchEngineCore.h
--- a/Engine/Plugins/Runtime/Metasound/Source/MetasoundFrontend/Private/MetasoundFrontendSearchEngineCore.h
+++ b/Engine/Plugins/Runtime/Metasound/Source/MetasoundFrontend/Private/MetasoundFrontendSearchEngineCore.h
@@ -58,12 +58,33 @@ namespace Metasound
 				static bool IsValidTransactionOfType(FInterfaceRegistryTransaction::ETransactionType InType, const FFrontendQueryEntry* InEntry);
 			};
 
+			class FMapUClassToDefaultInterface : public IFrontendQueryMultiMapStep
+			{
+			public:
+				virtual TArray<FFrontendQueryKey> Map(const FFrontendQueryEntry& InEntry) const override;
+			};
+
+			class FReduceInterfaceRegistryTransactionsToDefault : public IFrontendQueryReduceStep
+			{
+			public:
+				virtual void Reduce(const FFrontendQueryKey& InKey, FFrontendQueryPartition& InOutEntries) const override;
+			private:
+				static FInterfaceRegistryTransaction::FTimeType GetTransactionTimestamp(const FFrontendQueryEntry& InEntry);
+				static bool IsValidTransactionOfType(FInterfaceRegistryTransaction::ETransactionType InType, const FFrontendQueryEntry* InEntry);
+			};
+
 			class FTransformInterfaceRegistryTransactionToInterface : public IFrontendQueryTransformStep
 			{
 			public:
 				virtual void Transform(FFrontendQueryEntry::FValue& InValue) const override;
 			};
 
+			class FTransformInterfaceToInterfaceVersion : public IFrontendQueryTransformStep
+			{
+			public:
+				virtual void Transform(FFrontendQueryEntry::FValue& InValue) const override;
+			};
+
 			class FMapInterfaceToInterfaceNameAndMajorVersion : public IFrontendQueryMapStep
 			{
 			public:
@@ -280,6 +301,15 @@ namespace Metasound
 			static TArray<FMetasoundFrontendInterface> BuildResult(const FFrontendQueryPartition& InPartition);
 		};
 
+		// Policy for finding all default interface versions for a given UClass
+		struct FFindAllDefaultInterfaceVersionsForUClassQueryPolicy
+		{
+			using ResultType = TArray<FMetasoundFrontendVersion>;
+
+			static FFrontendQuery CreateQuery();
+			static TArray<FMetasoundFrontendVersion> BuildResult(const FFrontendQueryPartition& InPartition);
+		};
+
 		/** Supports essential search engine functionality needed for runtime. */
 		class FSearchEngineCore : public ISearchEngine
 		{
@@ -296,18 +326,15 @@ namespace Metasound
 			virtual bool FindClassWithHighestMinorVersion(const FMetasoundFrontendClassName& InName, int32 InMajorVersion, FMetasoundFrontendClass& OutClass) override;
 
 			virtual TArray<FMetasoundFrontendInterface> FindUClassDefaultInterfaces(FName InUClassName) override;
+			virtual TArray<FMetasoundFrontendVersion> FindUClassDefaultInterfaceVersions(const FTopLevelAssetPath& InUClassPath) override;
 			virtual TArray<FMetasoundFrontendVersion> FindAllRegisteredInterfacesWithName(FName InInterfaceName) override;
 			virtual bool FindInterfaceWithHighestVersion(FName InInterfaceName, FMetasoundFrontendInterface& OutInterface) override;
-			virtual TArray<FMetasoundFrontendInterface> FindAllInterfaces(bool bInIncludeAllVersions) override;
 
 		private:
-
 			TSearchEngineQuery<FFindClassWithHighestMinorVersionQueryPolicy> FindClassWithHighestMinorVersionQuery;
 			TSearchEngineQuery<FFindAllRegisteredInterfacesWithNameQueryPolicy> FindAllRegisteredInterfacesWithNameQuery;
+			TSearchEngineQuery<FFindAllDefaultInterfaceVersionsForUClassQueryPolicy> FindAllDefaultInterfaceVersionsForUClassQuery;
 			TSearchEngineQuery<FFindInterfaceWithHighestVersionQueryPolicy> FindInterfaceWithHighestVersionQuery;
-			TSearchEngineQuery<FFindAllInterfacesQueryPolicy> FindAllInterfacesQuery;
-			TSearchEngineQuery<FFindAllInterfacesIncludingAllVersionsQueryPolicy> FindAllInterfacesIncludingAllVersionsQuery;
 		};
 	}
 }
-
diff --git a/Engine/Plugins/Runtime/Metasound/Source/MetasoundFrontend/Private/MetasoundFrontendSearchEngineEditorOnly.cpp b/Engine/Plugins/Runtime/Metasound/Source/MetasoundFrontend/Private/MetasoundFrontendSearchEngineEditorOnly.cpp
--- a/Engine/Plugins/Runtime/Metasound/Source/MetasoundFrontend/Private/MetasoundFrontendSearchEngineEditorOnly.cpp
+++ b/Engine/Plugins/Runtime/Metasound/Source/MetasoundFrontend/Private/MetasoundFrontendSearchEngineEditorOnly.cpp
@@ -124,6 +124,8 @@ namespace Metasound
 				FindClassesWithNameUnsortedQuery.Prime();
 				FindClassesWithNameSortedQuery.Prime();
 				FindClassWithHighestVersionQuery.Prime();
+				FindAllInterfacesIncludingAllVersionsQuery.Prime();
+				FindAllInterfacesQuery.Prime();
 			}
 		}
 
@@ -164,6 +166,21 @@ namespace Metasound
 			const FFrontendQueryKey Key{InName.GetFullName()};
 			return FindClassWithHighestVersionQuery.UpdateAndFindResult(Key, OutClass);
 		}
+
+		TArray<FMetasoundFrontendInterface> FSearchEngineEditorOnly::FindAllInterfaces(bool bInIncludeAllVersions)
+		{
+			METASOUND_TRACE_CPUPROFILER_EVENT_SCOPE(metasound::FSearchEngineEditorOnly::FindAllInterfaces);
+
+			const FFrontendQueryKey NullKey;
+			if (bInIncludeAllVersions)
+			{
+				return FindAllInterfacesIncludingAllVersionsQuery.UpdateAndFindResult(NullKey);
+			}
+			else
+			{
+				return FindAllInterfacesQuery.UpdateAndFindResult(NullKey);
+			}
+		}
 	}
 }
 
diff --git a/Engine/Plugins/Runtime/Metasound/Source/MetasoundFrontend/Private/MetasoundFrontendSearchEngineEditorOnly.h b/Engine/Plugins/Runtime/Metasound/Source/MetasoundFrontend/Private/MetasoundFrontendSearchEngineEditorOnly.h
--- a/Engine/Plugins/Runtime/Metasound/Source/MetasoundFrontend/Private/MetasoundFrontendSearchEngineEditorOnly.h
+++ b/Engine/Plugins/Runtime/Metasound/Source/MetasoundFrontend/Private/MetasoundFrontendSearchEngineEditorOnly.h
@@ -80,6 +80,7 @@ namespace Metasound
 			virtual TArray<FMetasoundFrontendClass> FindAllClasses(bool bInIncludeAllVersions) override;
 			virtual TArray<FMetasoundFrontendClass> FindClassesWithName(const FMetasoundFrontendClassName& InName, bool bInSortByVersion) override;
 			virtual bool FindClassWithHighestVersion(const FMetasoundFrontendClassName& InName, FMetasoundFrontendClass& OutClass) override;
+			virtual TArray<FMetasoundFrontendInterface> FindAllInterfaces(bool bInIncludeAllVersions) override;
 
 		private:
 
@@ -88,6 +89,8 @@ namespace Metasound
 			TSearchEngineQuery<FFindClassesWithNameUnsortedQueryPolicy> FindClassesWithNameUnsortedQuery;
 			TSearchEngineQuery<FFindClassesWithNameSortedQueryPolicy> FindClassesWithNameSortedQuery;
 			TSearchEngineQuery<FFindClassWithHighestVersionQueryPolicy> FindClassWithHighestVersionQuery;
+			TSearchEngineQuery<FFindAllInterfacesQueryPolicy> FindAllInterfacesQuery;
+			TSearchEngineQuery<FFindAllInterfacesIncludingAllVersionsQueryPolicy> FindAllInterfacesIncludingAllVersionsQuery;
 		};
 	}
 }
diff --git a/Engine/Plugins/Runtime/Metasound/Source/MetasoundFrontend/Public/Interfaces/MetasoundFrontendInterfaceRegistry.h b/Engine/Plugins/Runtime/Metasound/Source/MetasoundFrontend/Public/Interfaces/MetasoundFrontendInterfaceRegistry.h
--- a/Engine/Plugins/Runtime/Metasound/Source/MetasoundFrontend/Public/Interfaces/MetasoundFrontendInterfaceRegistry.h
+++ b/Engine/Plugins/Runtime/Metasound/Source/MetasoundFrontend/Public/Interfaces/MetasoundFrontendInterfaceRegistry.h
@@ -23,25 +23,15 @@ namespace Metasound::Frontend
 		// MetaSound Interface definition
 		virtual const FMetasoundFrontendInterface& GetInterface() const = 0;
 
-		// Whether or not to immediately apply the given interface
-		// when a MetaSound asset of the supported UClass is created.
-		virtual bool IsDefault(FName InUClassName) const = 0;
-
 		UE_DEPRECATED(5.3, "UClass name must now be provided to determine if interface is to be apply by default.")
 		virtual bool IsDefault() const { return false; }
 
-		// Whether or not the interface can be added or removed by the editor.
-		virtual bool EditorCanAddOrRemove(FName InUClassName) const = 0;
-
-		UE_DEPRECATED(5.3, "UClass name must now be provided to determine if editor can directly add or remove interface.")
+		UE_DEPRECATED(5.3, "Option (renamed 'IsModifiable') is now data driven and accessibly from the interface (via GetInterface) per UClass option. If no class options found, addition/removal is now considered to be implicitly true")
 		virtual bool EditorCanAddOrRemove() const { return false; }
 
 		// Name of routing system used to update interface inputs (ex. ParameterInterface or DataReference).
 		virtual FName GetRouterName() const = 0;
 
-		// Name of UClass supported by the given interface.
-		virtual bool UClassIsSupported(FName InUClassName) const = 0;
-
 		// How to update a given document if versioning is required to this interface from a deprecated version.
 		virtual bool UpdateRootGraphInterface(FDocumentHandle InDocument) const = 0;
 	};
@@ -82,7 +72,7 @@ namespace Metasound::Frontend
 		virtual ~IInterfaceRegistry() = default;
 
 		// Register an interface
-		virtual void RegisterInterface(TUniquePtr<IInterfaceRegistryEntry>&& InEntry) = 0;
+		virtual bool RegisterInterface(TUniquePtr<IInterfaceRegistryEntry>&& InEntry) = 0;
 
 		// Find an interface entry with the given key. Returns null if entry not found with given key.
 		virtual const IInterfaceRegistryEntry* FindInterfaceRegistryEntry(const FInterfaceRegistryKey& InKey) const = 0;
diff --git a/Engine/Plugins/Runtime/Metasound/Source/MetasoundFrontend/Public/Interfaces/MetasoundFrontendSourceInterface.h b/Engine/Plugins/Runtime/Metasound/Source/MetasoundFrontend/Public/Interfaces/MetasoundFrontendSourceInterface.h
--- a/Engine/Plugins/Runtime/Metasound/Source/MetasoundFrontend/Public/Interfaces/MetasoundFrontendSourceInterface.h
+++ b/Engine/Plugins/Runtime/Metasound/Source/MetasoundFrontend/Public/Interfaces/MetasoundFrontendSourceInterface.h
@@ -23,9 +23,7 @@ namespace Metasound
 
 			METASOUNDFRONTEND_API const FMetasoundFrontendVersion& GetVersion();
 
-			UE_DEPRECATED(5.3, "ParameterInterface is now only registered for given UClasses at registration time, not in interface definition.")
 			METASOUNDFRONTEND_API Audio::FParameterInterfacePtr CreateInterface(const UClass& InClass);
-			METASOUNDFRONTEND_API Audio::FParameterInterfacePtr CreateInterface();
 		}
 
 		namespace SourceInterfaceV1_0
@@ -51,9 +49,7 @@ namespace Metasound
 
 			METASOUNDFRONTEND_API const FMetasoundFrontendVersion& GetVersion();
 
-			UE_DEPRECATED(5.3, "ParameterInterface is now only registered for given UClasses at registration time, not in interface definition.")
 			METASOUNDFRONTEND_API Audio::FParameterInterfacePtr CreateInterface(const UClass& InClass);
-			METASOUNDFRONTEND_API Audio::FParameterInterfacePtr CreateInterface();
 		}
 
 		namespace SourceInterface
@@ -75,9 +71,7 @@ namespace Metasound
 
 			METASOUNDFRONTEND_API const FMetasoundFrontendVersion& GetVersion();
 
-			UE_DEPRECATED(5.3, "ParameterInterface is now only registered for given UClasses at registration time, not in interface definition.")
 			METASOUNDFRONTEND_API Audio::FParameterInterfacePtr CreateInterface(const UClass& InClass);
-			METASOUNDFRONTEND_API Audio::FParameterInterfacePtr CreateInterface();
 
 			class METASOUNDFRONTEND_API FUpdateInterface : public Frontend::IDocumentTransform
 			{
diff --git a/Engine/Plugins/Runtime/Metasound/Source/MetasoundFrontend/Public/MetasoundDocumentInterface.h b/Engine/Plugins/Runtime/Metasound/Source/MetasoundFrontend/Public/MetasoundDocumentInterface.h
--- a/Engine/Plugins/Runtime/Metasound/Source/MetasoundFrontend/Public/MetasoundDocumentInterface.h
+++ b/Engine/Plugins/Runtime/Metasound/Source/MetasoundFrontend/Public/MetasoundDocumentInterface.h
@@ -25,6 +25,7 @@ class METASOUNDFRONTEND_API IMetaSoundDocumentInterface : public IInterface
 
 public:
 	virtual const FMetasoundFrontendDocument& GetDocument() const = 0;
+	virtual const UClass& GetBaseMetaSoundUClass() const = 0;
 
 private:
 	virtual FMetasoundFrontendDocument& GetDocument() = 0;
diff --git a/Engine/Plugins/Runtime/Metasound/Source/MetasoundFrontend/Public/MetasoundFrontendDocument.h b/Engine/Plugins/Runtime/Metasound/Source/MetasoundFrontend/Public/MetasoundFrontendDocument.h
--- a/Engine/Plugins/Runtime/Metasound/Source/MetasoundFrontend/Public/MetasoundFrontendDocument.h
+++ b/Engine/Plugins/Runtime/Metasound/Source/MetasoundFrontend/Public/MetasoundFrontendDocument.h
@@ -853,6 +853,7 @@ struct METASOUNDFRONTEND_API FMetasoundFrontendClassInput : public FMetasoundFro
 	FMetasoundFrontendClassInput() = default;
 
 	FMetasoundFrontendClassInput(const FMetasoundFrontendClassVertex& InOther);
+	FMetasoundFrontendClassInput(const Audio::FParameterInterface::FInput& InInput);
 
 	// Default value for this input.
 	UPROPERTY(EditAnywhere, Category = Parameters)
@@ -881,11 +882,8 @@ struct METASOUNDFRONTEND_API FMetasoundFrontendClassOutput : public FMetasoundFr
 	GENERATED_BODY()
 
 	FMetasoundFrontendClassOutput() = default;
-
-	FMetasoundFrontendClassOutput(const FMetasoundFrontendClassVertex& InOther)
-	:	FMetasoundFrontendClassVertex(InOther)
-	{
-	}
+	FMetasoundFrontendClassOutput(const FMetasoundFrontendClassVertex& InOther);
+	FMetasoundFrontendClassOutput(const Audio::FParameterInterface::FOutput& Output);
 };
 
 USTRUCT()
@@ -893,6 +891,9 @@ struct METASOUNDFRONTEND_API FMetasoundFrontendClassEnvironmentVariable
 {
 	GENERATED_BODY()
 
+	FMetasoundFrontendClassEnvironmentVariable() = default;
+	FMetasoundFrontendClassEnvironmentVariable(const Audio::FParameterInterface::FEnvironmentVariable& InVariable);
+
 	// Name of environment variable.
 	UPROPERTY()
 	FName Name;
@@ -1146,14 +1147,47 @@ struct METASOUNDFRONTEND_API FMetasoundFrontendInterfaceBinding
 };
 
 
+// Options used to restrict a corresponding UClass that interface may be applied to.
+// If unspecified, interface is assumed to be applicable to any arbitrary UClass.
+USTRUCT()
+struct METASOUNDFRONTEND_API FMetasoundFrontendInterfaceUClassOptions
+{
+	GENERATED_BODY()
+
+	FMetasoundFrontendInterfaceUClassOptions() = default;
+	FMetasoundFrontendInterfaceUClassOptions(const Audio::FParameterInterface::FClassOptions& InOptions);
+	FMetasoundFrontendInterfaceUClassOptions(const FTopLevelAssetPath& InClassPath, bool bInIsModifiable = true, bool bInIsDefault = false);
+
+	UPROPERTY()
+	FTopLevelAssetPath ClassPath;
+
+	UPROPERTY()
+	bool bIsModifiable = true;
+
+	UPROPERTY()
+	bool bIsDefault = false;
+};
+
+
+// Definition of an interface that an FMetasoundFrontendClass adheres to in part or full.
 USTRUCT()
 struct METASOUNDFRONTEND_API FMetasoundFrontendInterface : public FMetasoundFrontendClassInterface
 {
 	GENERATED_BODY()
 
+	FMetasoundFrontendInterface() = default;
+	FMetasoundFrontendInterface(Audio::FParameterInterfacePtr InInterface);
+
 	// Name and version number of the interface
 	UPROPERTY()
 	FMetasoundFrontendVersion Version;
+
+	// If specified, options used to restrict a corresponding UClass that interface may be
+	// applied to.  If unspecified, interface is assumed to be applicable to any arbitrary UClass.
+	UPROPERTY()
+	TArray<FMetasoundFrontendInterfaceUClassOptions> UClassOptions;
+
+	const FMetasoundFrontendInterfaceUClassOptions* FindClassOptions(const FTopLevelAssetPath& InClassPath) const;
 };
 
 
diff --git a/Engine/Plugins/Runtime/Metasound/Source/MetasoundFrontend/Public/MetasoundFrontendDocumentBuilder.h b/Engine/Plugins/Runtime/Metasound/Source/MetasoundFrontend/Public/MetasoundFrontendDocumentBuilder.h
--- a/Engine/Plugins/Runtime/Metasound/Source/MetasoundFrontend/Public/MetasoundFrontendDocumentBuilder.h
+++ b/Engine/Plugins/Runtime/Metasound/Source/MetasoundFrontend/Public/MetasoundFrontendDocumentBuilder.h
@@ -63,6 +63,11 @@ public:
 		return Document;
 	}
 
+	virtual const UClass& GetBaseMetaSoundUClass() const final override
+	{
+		return *UMetaSoundBuilderDocument::StaticClass();
+	}
+
 private:
 	virtual FMetasoundFrontendDocument& GetDocument() override
 	{
@@ -139,9 +144,10 @@ public:
 	TArray<const FMetasoundFrontendVertex*> FindNodeOutputs(const FGuid& InNodeID, FName TypeName = FName()) const;
 
 	const FMetasoundFrontendDocument& GetDocument() const;
+	const IMetaSoundDocumentInterface& GetDocumentInterface() const;
 
 	static void InitGraphClassMetadata(FMetasoundFrontendClassMetadata& InOutMetadata, bool bResetVersion = false);
-	void InitDocument(FName UClassName);
+	void InitDocument();
 	void InitNodeLocations();
 
 	bool IsNodeInputConnected(const FGuid& InNodeID, const FGuid& InVertexID) const;
diff --git a/Engine/Plugins/Runtime/Metasound/Source/MetasoundFrontend/Public/MetasoundFrontendQuery.h b/Engine/Plugins/Runtime/Metasound/Source/MetasoundFrontend/Public/MetasoundFrontendQuery.h
--- a/Engine/Plugins/Runtime/Metasound/Source/MetasoundFrontend/Public/MetasoundFrontendQuery.h
+++ b/Engine/Plugins/Runtime/Metasound/Source/MetasoundFrontend/Public/MetasoundFrontendQuery.h
@@ -133,6 +133,14 @@ namespace Metasound
 			virtual FFrontendQueryKey Map(const FFrontendQueryEntry& InEntry) const = 0;
 	};
 
+	/** Interface for a query step which maps entries to multiple keys. */
+	class IFrontendQueryMultiMapStep : public IFrontendQueryStep
+	{
+	public:
+		virtual ~IFrontendQueryMultiMapStep() = default;
+		virtual TArray<FFrontendQueryKey> Map(const FFrontendQueryEntry& InEntry) const = 0;
+	};
+
 	/** Interface for a query step which reduces entries with the same key. */
 	class IFrontendQueryReduceStep : public IFrontendQueryStep
 	{
@@ -224,6 +232,7 @@ namespace Metasound
 		using FStreamFunction = TUniqueFunction<void (TArray<FFrontendQueryValue>&)>;
 		using FTransformFunction = TFunction<void (FFrontendQueryEntry::FValue&)>;
 		using FMapFunction = TFunction<FFrontendQueryKey (const FFrontendQueryEntry&)>;
+		using FMultiMapFunction = TFunction<TArray<FFrontendQueryKey>(const FFrontendQueryEntry&)>;
 		using FReduceFunction = TFunction<void (const FFrontendQueryKey& InKey, FFrontendQueryPartition& InOutEntries)>;
 		using FFilterFunction = TFunction<bool (const FFrontendQueryEntry&)>;
 		using FScoreFunction = TFunction<float (const FFrontendQueryEntry&)>;
@@ -234,6 +243,7 @@ namespace Metasound
 		FFrontendQueryStep(FStreamFunction&& InFunc);
 		FFrontendQueryStep(FTransformFunction&& InFunc);
 		FFrontendQueryStep(FMapFunction&& InFunc);
+		FFrontendQueryStep(FMultiMapFunction&& InFunc);
 		FFrontendQueryStep(FReduceFunction&& InFunc);
 		FFrontendQueryStep(FFilterFunction&& InFilt);
 		FFrontendQueryStep(FScoreFunction&& InScore);
@@ -244,6 +254,7 @@ namespace Metasound
 		FFrontendQueryStep(TUniquePtr<IFrontendQueryStreamStep>&& InStep);
 		FFrontendQueryStep(TUniquePtr<IFrontendQueryTransformStep>&& InStep);
 		FFrontendQueryStep(TUniquePtr<IFrontendQueryMapStep>&& InStep);
+		FFrontendQueryStep(TUniquePtr<IFrontendQueryMultiMapStep>&& InStep);
 		FFrontendQueryStep(TUniquePtr<IFrontendQueryReduceStep>&& InStep);
 		FFrontendQueryStep(TUniquePtr<IFrontendQueryFilterStep>&& InStep);
 		FFrontendQueryStep(TUniquePtr<IFrontendQueryScoreStep>&& InStep);
diff --git a/Engine/Plugins/Runtime/Metasound/Source/MetasoundFrontend/Public/MetasoundFrontendSearchEngine.h b/Engine/Plugins/Runtime/Metasound/Source/MetasoundFrontend/Public/MetasoundFrontendSearchEngine.h
--- a/Engine/Plugins/Runtime/Metasound/Source/MetasoundFrontend/Public/MetasoundFrontendSearchEngine.h
+++ b/Engine/Plugins/Runtime/Metasound/Source/MetasoundFrontend/Public/MetasoundFrontendSearchEngine.h
@@ -31,9 +31,11 @@ namespace Metasound::Frontend
 
 			virtual bool FindInterfaceWithHighestVersion(FName InInterfaceName, FMetasoundFrontendInterface& OutInterface) = 0;
 
+			UE_DEPRECATED(5.3, "Use ISearchEngine::FindUClassDefaultInterfaceVersions using TopLevelAssetPath instead.")
 			virtual TArray<FMetasoundFrontendInterface> FindUClassDefaultInterfaces(FName InUClassName) = 0;
 
-			virtual TArray<FMetasoundFrontendInterface> FindAllInterfaces(bool bInIncludeAllVersions = false) = 0;
+			/** Returns all interfaces that are to be added to a given document when it is initialized on an object with the given class**/
+			virtual TArray<FMetasoundFrontendVersion> FindUClassDefaultInterfaceVersions(const FTopLevelAssetPath& InUClassPath) = 0;
 
 #if WITH_EDITORONLY_DATA
 			/** Find all FMetasoundFrontendClasses.
@@ -48,7 +50,10 @@ namespace Metasound::Frontend
 
 			/** Find the highest version of a class with the given ClassName. Returns false if not found, true if found. */
 			virtual bool FindClassWithHighestVersion(const FMetasoundFrontendClassName& InName, FMetasoundFrontendClass& OutClass) = 0;
-#endif
+
+			/** Returns array with all registered interfaces. Optionally, include interface versions that are not the highest version. */
+			virtual TArray<FMetasoundFrontendInterface> FindAllInterfaces(bool bInIncludeAllVersions = false) = 0;
+#endif // WITH_EDITORONLY_DATA
 
 		protected:
 			ISearchEngine() = default;
diff --git a/Engine/Source/Runtime/AudioExtensions/Private/IAudioParameterInterfaceRegistry.cpp b/Engine/Source/Runtime/AudioExtensions/Private/IAudioParameterInterfaceRegistry.cpp
--- a/Engine/Source/Runtime/AudioExtensions/Private/IAudioParameterInterfaceRegistry.cpp
+++ b/Engine/Source/Runtime/AudioExtensions/Private/IAudioParameterInterfaceRegistry.cpp
@@ -42,7 +42,7 @@ namespace Audio
 	)
 		: NamePrivate(InName)
 		, VersionPrivate(InVersion)
-		, SupportedUClassNames({ InType.GetPathName() })
+		, UClassOptions({ { InType.GetClassPathName() } })
 	{
 	}
 
@@ -52,17 +52,6 @@ namespace Audio
 	{
 	}
 
-	FParameterInterface::FParameterInterface(FName InName, const FVersion& InVersion, const TArray<UClass*>& InClasses)
-		: NamePrivate(InName)
-		, VersionPrivate(InVersion)
-	{
-		Algo::Transform(InClasses, SupportedUClassNames, [](const UClass* Class)
-		{
-			check(Class);
-			return Class->GetPathName();
-		});
-	}
-
 	FName FParameterInterface::GetName() const
 	{
 		return NamePrivate;
@@ -78,6 +67,11 @@ namespace Audio
 		return *UObject::StaticClass();
 	}
 
+	const TArray<FParameterInterface::FClassOptions>& FParameterInterface::GetUClassOptions() const
+	{
+		return UClassOptions;
+	}
+
 	const TArray<FParameterInterface::FInput>& FParameterInterface::GetInputs() const
 	{
 		return Inputs;
@@ -96,9 +90,9 @@ namespace Audio
 	TArray<const UClass*> FParameterInterface::FindSupportedUClasses() const
 	{
 		TArray<const UClass*> SupportedUClasses;
-		for (const FString& Name : SupportedUClassNames)
+		for (const FClassOptions& Options : UClassOptions)
 		{
-			if (const UClass* Class = FindFirstObject<const UClass>(*Name, EFindFirstObjectOptions::ExactClass, ELogVerbosity::Warning, TEXT("FParameterInterface::FindSupportedUClasses")))
+			if (const UClass* Class = FindObject<const UClass>(Options.ClassPath))
 			{
 				SupportedUClasses.Add(Class);
 			}
diff --git a/Engine/Source/Runtime/AudioExtensions/Public/IAudioParameterInterfaceRegistry.h b/Engine/Source/Runtime/AudioExtensions/Public/IAudioParameterInterfaceRegistry.h
--- a/Engine/Source/Runtime/AudioExtensions/Public/IAudioParameterInterfaceRegistry.h
+++ b/Engine/Source/Runtime/AudioExtensions/Public/IAudioParameterInterfaceRegistry.h
@@ -21,83 +21,166 @@ namespace Audio
 	// Forward Declarations
 	class IAudioParameterInterfaceRegistry;
 
+	// Interface for parameterizing data provided to or coming from an executable audio unit
+	// (ex. Sound like a MetaSoundSource, arbitrary DSP graph like a MetaSoundPatch, etc.).
+	// Can be used generically for processing either logically at a control or Game Thread
+	// tick rate (ex. SoundCues), or by an underlying DSP operation via the Audio Render
+	// Thread or delegated task (ex. MetaSounds, custom SoundGenerator, etc.)
 	struct AUDIOEXTENSIONS_API FParameterInterface
 	{
+		// Version of interface (higher numbers are more recent)
 		struct AUDIOEXTENSIONS_API FVersion
 		{
 			const int32 Major = 1;
 			const int32 Minor = 0;
 		};
 
+		// Input of interface
 		struct AUDIOEXTENSIONS_API FInput
 		{
+			// Name to be displayed in editor or tools
 			const FText DisplayName;
+
+			// Description to be displayed in editor or tools
 			const FText Description;
+
+			// FName describing the type of the data.  May be
+			// interpreted solely by a plugin or implementation
+			// (ex. MetaSounds).  If blank, type is assumed to be
+			// that described by the InitValue AudioParameter's
+			// corresponding ParamType. If provided, DataType
+			// must be constructible using the InitValue.
 			const FName DataType;
 
+			// Initial value of the given parameter
 			const FAudioParameter InitValue;
+
+			// Text to display in the editor or tools if the consuming
+			// system of the given input parameter is not implemented (to avoid
+			// passing data using the parameter system when not operated on).
 			const FText RequiredText;
 
+			// Visual sort order of the given input with respect to other inputs either
+			// within the given interface or listed among other unrelated inputs.
 			const int32 SortOrderIndex = 0;
 		};
 
+		// Output of interface
 		struct AUDIOEXTENSIONS_API FOutput
 		{
+			// Name to be displayed in editor or tools
 			const FText DisplayName;
+
+			// Description to be displayed in editor or tools
 			const FText Description;
+
+			// FName describing the type of the data.  May be
+			// interpreted solely by a plugin or implementation
+			// (ex. MetaSounds).  If blank, type is assumed to be
+			// that described by ParamType.
 			const FName DataType;
 
+			// Name of output parameter used as a runtime identifier
 			const FName ParamName;
+
+			// Text to display in the editor or tools if the consuming
+			// system of the given input parameter is not implemented (to avoid
+			// passing data using the parameter system when not operated on).
 			const FText RequiredText = FText();
+
+			// Type of output parameter used as a runtime identifier if unspecified by the DataType.
 			const EAudioParameterType ParamType = EAudioParameterType::None;
 
+			// Visual sort order of the given input with respect to other outputs either
+			// within the given interface or listed among other unrelated outputs.
 			const int32 SortOrderIndex = 0;
 		};
 
+		// Read-only variable that cannot be modified by the sound instance,
+		// and maybe shared amongst instances.
 		struct AUDIOEXTENSIONS_API FEnvironmentVariable
 		{
+			// Name to be displayed in editor or tools
 			const FText DisplayName;
+
+			// Description to be displayed in editor or tools
 			const FText Description;
+
+			// FName describing the type of the data.  May be
+			// interpreted solely by a plugin or implementation
+			// (ex. MetaSounds).  If blank, type is assumed to be
+			// that described by ParamType.
 			const FName DataType;
 
+			// Name of variable used as a runtime identifier
 			const FName ParamName;
+
+			// Type of variable used as a runtime identifier if unspecified by the DataType.
 			const EAudioParameterType ParamType = EAudioParameterType::None;
 		};
 
+		// Options used to restrict a corresponding UClass that interface may be applied to.
+		struct AUDIOEXTENSIONS_API FClassOptions
+		{
+			// Path to restricted UClass
+			const FTopLevelAssetPath ClassPath;
+
+			// Whether or not the class may be directly modifiable on an asset implementing
+			// the given interface (added, removed, etc.)
+			const bool bIsModifiable = true;
+
+			// Whether or not the interface should be immediately added to the given class
+			// type on creation.
+			const bool bIsDefault = false;
+		};
+
 		FParameterInterface() = default;
 
 		// Constructor used for parameter interface not limited to any particular UClass types
 		FParameterInterface(FName InName, const FVersion& InVersion);
 
-		// Constructor used for parameter interface with support for explicit UClass types
-		FParameterInterface(FName InName, const FVersion& InVersion, const TArray<UClass*>& InClasses);
-
-		UE_DEPRECATED(5.3, "Use FParameterInterface ctor that supports array of classes.")
+		UE_DEPRECATED(5.3, "Set UClassOptions to determine what options apply for a given UClass (if any).")
 		FParameterInterface(FName InName, const FVersion& InVersion, const UClass& InClass);
 
+		// Returns name of interface
 		FName GetName() const;
+
+		// Returns version of interface
 		const FVersion& GetVersion() const;
 
 		UE_DEPRECATED(5.3, "Use FParameterInterface::FindSupportedUClasses instead")
 		const UClass& GetType() const;
 
 		TArray<const UClass*> FindSupportedUClasses() const;
+
+		// If specified, options used to restrict a corresponding UClass that interface may be
+		// applied to.  If unspecified, interface is assumed to be applicable to any arbitrary UClass.
+		const TArray<FClassOptions>& GetUClassOptions() const;
+
+		// Returns read-only array of inputs
 		const TArray<FInput>& GetInputs() const;
+
+		// Returns read-only array of outputs
 		const TArray<FOutput>& GetOutputs() const;
+
+		// Returns read-only array of environment variables
 		const TArray<FEnvironmentVariable>& GetEnvironment() const;
 
 	private:
 		FName NamePrivate;
 		FVersion VersionPrivate;
-		TArray<FString> SupportedUClassNames;
 
 	protected:
 		TArray<FInput> Inputs;
 		TArray<FOutput> Outputs;
 		TArray<FEnvironmentVariable> Environment;
+		TArray<FClassOptions> UClassOptions;
 	};
 	using FParameterInterfacePtr = TSharedPtr<FParameterInterface, ESPMode::ThreadSafe>;
 
+	// Registry of engine-defined audio parameter interfaces, used to parameterize data provided
+	// to or coming from an executable audio unit (ex. Sound like a MetaSoundSource, arbitrary DSP
+	// graph like a MetaSoundPatch, etc.).
 	class AUDIOEXTENSIONS_API IAudioParameterInterfaceRegistry
 	{
 		static TUniquePtr<IAudioParameterInterfaceRegistry> Instance;
@@ -107,8 +190,13 @@ namespace Audio
 
 		virtual ~IAudioParameterInterfaceRegistry() = default;
 
+		// Iterate all registered interfaces
 		virtual void IterateInterfaces(TFunction<void(FParameterInterfacePtr)> InFunction) const = 0;
+
+		// Execute a given function when an interface is registered
 		virtual void OnRegistration(TUniqueFunction<void(FParameterInterfacePtr)>&& InFunction) = 0;
+
+		// Registers an interface
 		virtual void RegisterInterface(FParameterInterfacePtr InInterface) = 0;
 
 	protected:
