Workaround in the metasound builder for source interfaces not being added automatically

diff --git a/Engine/Plugins/Runtime/Metasound/Source/MetasoundEngine/Private/MetasoundBuilderSubsystem.cpp b/Engine/Plugins/Runtime/Metasound/Source/MetasoundEngine/Private/MetasoundBuilderSubsystem.cpp
--- a/Engine/Plugins/Runtime/Metasound/Source/MetasoundEngine/Private/MetasoundBuilderSubsystem.cpp
+++ b/Engine/Plugins/Runtime/Metasound/Source/MetasoundEngine/Private/MetasoundBuilderSubsystem.cpp
@@ -724,6 +724,19 @@ void UMetaSoundSourceBuilder::SetFormat(EMetaSoundOutputAudioFormat OutputFormat
 	OutResult = bSuccess ? EMetaSoundBuilderResult::Succeeded : EMetaSoundBuilderResult::Failed;
 }
 
+void UMetaSoundSourceBuilder::AddSourceInterfaces(EMetaSoundBuilderResult& OutResult)
+{
+	using namespace Metasound::Engine;
+	using namespace Metasound::Frontend;
+
+	const UClass& BaseMetaSoundClass = *UMetaSoundSource::StaticClass();
+	const FTopLevelAssetPath BaseMetaSoundClassPath = BaseMetaSoundClass.GetClassPathName();
+	TArray<FMetasoundFrontendVersion> InitVersions = ISearchEngine::Get().FindUClassDefaultInterfaceVersions(BaseMetaSoundClassPath);
+
+	const bool bSuccess = Builder.ModifyInterfaces({ }, InitVersions);
+	OutResult = bSuccess ? EMetaSoundBuilderResult::Succeeded : EMetaSoundBuilderResult::Failed;
+}
+
 UMetaSoundPatchBuilder* UMetaSoundBuilderSubsystem::CreatePatchBuilder(FName BuilderName, EMetaSoundBuilderResult& OutResult)
 {
 	return &Metasound::Engine::BuilderSubsystemPrivate::CreateTransientBuilder<UMetaSoundPatchBuilder>();
@@ -749,6 +762,15 @@ UMetaSoundSourceBuilder* UMetaSoundBuilderSubsystem::CreateSourceBuilder(
 
 	UMetaSoundSourceBuilder& NewBuilder = CreateTransientBuilder<UMetaSoundSourceBuilder>();
 
+	// Temp workaround - the builder currently doesn't have any interfaces added automatically, so we add them here
+	NewBuilder.AddSourceInterfaces(OutResult);
+
+	if (OutResult == EMetaSoundBuilderResult::Failed)
+	{
+		UE_LOG(LogMetaSound, Error, TEXT("Builder '%s' Creation Error: Failed to add source document interfaces."), *BuilderName.ToString());
+		return nullptr;
+	}
+
 	bool bFormatSet = OutputFormat == EMetaSoundOutputAudioFormat::Mono;
 	if (!bFormatSet)
 	{
@@ -756,7 +778,7 @@ UMetaSoundSourceBuilder* UMetaSoundBuilderSubsystem::CreateSourceBuilder(
 		NewBuilder.SetFormat(OutputFormat, Result);
 		bFormatSet = Result == EMetaSoundBuilderResult::Succeeded;
 	}
-
+	
 	if (bFormatSet)
 	{
 		const Metasound::Engine::FOutputAudioFormatInfoPair* FormatInfo = NewBuilder.FindOutputAudioFormatInfo();
diff --git a/Engine/Plugins/Runtime/Metasound/Source/MetasoundEngine/Public/MetasoundBuilderSubsystem.h b/Engine/Plugins/Runtime/Metasound/Source/MetasoundEngine/Public/MetasoundBuilderSubsystem.h
--- a/Engine/Plugins/Runtime/Metasound/Source/MetasoundEngine/Public/MetasoundBuilderSubsystem.h
+++ b/Engine/Plugins/Runtime/Metasound/Source/MetasoundEngine/Public/MetasoundBuilderSubsystem.h
@@ -382,6 +382,8 @@ public:
 	UFUNCTION(BlueprintCallable, Category = "Audio|MetaSound|Builder", meta = (ExpandEnumAsExecs = "OutResult"))
 	void SetFormat(EMetaSoundOutputAudioFormat OutputFormat, EMetaSoundBuilderResult& OutResult);
 
+	void AddSourceInterfaces(EMetaSoundBuilderResult& OutResult);
+
 	const Metasound::Engine::FOutputAudioFormatInfoPair* FindOutputAudioFormatInfo() const;
 
 private:
