Remove transmittable input nodes and move updating parameters to the metasound generator

diff --git a/Engine/Plugins/Runtime/Metasound/Source/MetasoundEngine/Private/MetasoundSource.cpp b/Engine/Plugins/Runtime/Metasound/Source/MetasoundEngine/Private/MetasoundSource.cpp
--- a/Engine/Plugins/Runtime/Metasound/Source/MetasoundEngine/Private/MetasoundSource.cpp
+++ b/Engine/Plugins/Runtime/Metasound/Source/MetasoundEngine/Private/MetasoundSource.cpp
@@ -59,6 +59,38 @@ namespace Metasound
 
 			return RegOptions;
 		}
+
+		class FParameterRouter
+		{
+		public:
+
+			using FAudioDeviceIDAndInstanceID = TTuple<Audio::DeviceID, uint64>;
+
+			TSharedPtr<TSpscQueue<FMetaSoundParameterTransmitter::FParameter>> FindOrCreateDataChannel(Audio::DeviceID InDeviceID, uint64 InstanceID)
+			{
+				FScopeLock Lock(&DataChannelMapCS);
+
+				FAudioDeviceIDAndInstanceID Key = {InDeviceID, InstanceID};
+
+				if (TWeakPtr<TSpscQueue<FMetaSoundParameterTransmitter::FParameter>>* WeakChannel = DataChannels.Find(Key))
+				{
+					TSharedPtr<TSpscQueue<FMetaSoundParameterTransmitter::FParameter>> Channel = WeakChannel->Pin();
+					if (Channel.IsValid())
+					{
+						return Channel;
+					}
+				}
+
+				TSharedPtr<TSpscQueue<FMetaSoundParameterTransmitter::FParameter>> NewChannel = MakeShared<TSpscQueue<FMetaSoundParameterTransmitter::FParameter>>();
+				DataChannels.Add(Key, NewChannel);
+				return NewChannel;
+			}
+
+		private:
+
+			FCriticalSection DataChannelMapCS;
+			TSortedMap<FAudioDeviceIDAndInstanceID, TWeakPtr<TSpscQueue<FMetaSoundParameterTransmitter::FParameter>>> DataChannels;
+		};
 	} // namespace SourcePrivate
 } // namespace Metasound
 
@@ -501,6 +533,7 @@ ISoundGeneratorPtr UMetaSoundSource::CreateSoundGenerator(const FSoundGeneratorI
 	using namespace Metasound;
 	using namespace Metasound::Frontend;
 	using namespace Metasound::Engine;
+	using namespace Metasound::SourcePrivate;
 
 	METASOUND_LLM_SCOPE;
 	METASOUND_TRACE_CPUPROFILER_EVENT_SCOPE(UMetaSoundSource::CreateSoundGenerator);
@@ -508,6 +541,8 @@ ISoundGeneratorPtr UMetaSoundSource::CreateSoundGenerator(const FSoundGeneratorI
 	SampleRate = InParams.SampleRate;
 	FOperatorSettings InSettings = GetOperatorSettings(static_cast<FSampleRate>(SampleRate));
 	FMetasoundEnvironment Environment = CreateEnvironment(InParams);
+	FParameterRouter& Router = GetParameterRouter();
+	TSharedPtr<TSpscQueue<FMetaSoundParameterTransmitter::FParameter>> DataChannel = Router.FindOrCreateDataChannel(InParams.AudioDeviceID, InParams.InstanceID);
 
 	TSharedPtr<const IGraph, ESPMode::ThreadSafe> MetasoundGraph = GetRuntimeData().Graph;
 	if (!MetasoundGraph.IsValid())
@@ -520,6 +555,7 @@ ISoundGeneratorPtr UMetaSoundSource::CreateSoundGenerator(const FSoundGeneratorI
 	// Graph analyzer currently only enabled for preview sounds (but can theoretically be supported for all sounds)
 	BuilderSettings.bPopulateInternalDataReferences = InParams.bIsPreviewSound;
 
+	constexpr bool bBuildSynchronous = false;
 	FMetasoundGeneratorInitParams InitParams =
 	{
 		InSettings,
@@ -528,7 +564,9 @@ ISoundGeneratorPtr UMetaSoundSource::CreateSoundGenerator(const FSoundGeneratorI
 		Environment,
 		GetName(),
 		GetOutputAudioChannelOrder(),
-		MoveTemp(InDefaultParameters)
+		MoveTemp(InDefaultParameters),
+		bBuildSynchronous,
+		DataChannel
 	};
 
 	TSharedPtr<FMetasoundGenerator> Generator = MakeShared<FMetasoundGenerator>(MoveTemp(InitParams));
@@ -811,15 +849,20 @@ bool UMetaSoundSource::IsOneShot() const
 
 TSharedPtr<Audio::IParameterTransmitter> UMetaSoundSource::CreateParameterTransmitter(Audio::FParameterTransmitterInitParams&& InParams) const
 {
+	using namespace Metasound;
+	using namespace Metasound::SourcePrivate;
+
 	METASOUND_LLM_SCOPE;
 
-	Metasound::FMetaSoundParameterTransmitter::FInitParams InitParams(GetOperatorSettings(InParams.SampleRate), InParams.InstanceID, MoveTemp(InParams.DefaultParams));
-	InitParams.DebugMetaSoundName = GetFName();
+	const FRuntimeData& RuntimeData = GetRuntimeData();
 
-	for (const FSendInfoAndVertexName& InfoAndName : FMetasoundAssetBase::GetSendInfos(InParams.InstanceID))
-	{
-		InitParams.Infos.Add(InfoAndName.SendInfo);
-	}
+	TArray<FName> ValidParameters;
+	Algo::Transform(RuntimeData.TransmittableInputs, ValidParameters, [](const FMetasoundFrontendClassInput& InputVertex) { return InputVertex.Name; });
+	FParameterRouter& Router = GetParameterRouter();
+	TSharedPtr<TSpscQueue<FMetaSoundParameterTransmitter::FParameter>> DataChannel = Router.FindOrCreateDataChannel(InParams.AudioDeviceID, InParams.InstanceID);
+
+	Metasound::FMetaSoundParameterTransmitter::FInitParams InitParams(GetOperatorSettings(InParams.SampleRate), InParams.InstanceID, MoveTemp(InParams.DefaultParams), MoveTemp(ValidParameters), DataChannel);
+	InitParams.DebugMetaSoundName = GetFName();
 
 	return MakeShared<Metasound::FMetaSoundParameterTransmitter>(MoveTemp(InitParams));
 }
@@ -924,4 +967,14 @@ TWeakPtr<Metasound::FMetasoundGenerator> UMetaSoundSource::GetGeneratorForAudioC
 	}
 	return *Result;
 }
+
+Metasound::SourcePrivate::FParameterRouter& UMetaSoundSource::GetParameterRouter()
+{
+	using namespace Metasound::SourcePrivate;
+
+	static FParameterRouter Router;
+	return Router;
+}
+
+
 #undef LOCTEXT_NAMESPACE // MetaSound
diff --git a/Engine/Plugins/Runtime/Metasound/Source/MetasoundEngine/Public/MetasoundSource.h b/Engine/Plugins/Runtime/Metasound/Source/MetasoundEngine/Public/MetasoundSource.h
--- a/Engine/Plugins/Runtime/Metasound/Source/MetasoundEngine/Public/MetasoundSource.h
+++ b/Engine/Plugins/Runtime/Metasound/Source/MetasoundEngine/Public/MetasoundSource.h
@@ -24,9 +24,17 @@ namespace Metasound
 {
 	struct FMetaSoundEngineAssetHelper;
 	class FMetasoundGenerator;
-
+	namespace SourcePrivate
+	{
+		class FParameterRouter;
+	}
 } // namespace Metasound
 
+namespace Audio
+{
+	using DeviceID = uint32;
+}
+
 DECLARE_TS_MULTICAST_DELEGATE_TwoParams(FOnGeneratorInstanceCreated, uint64, TSharedPtr<Metasound::FMetasoundGenerator>);
 DECLARE_TS_MULTICAST_DELEGATE_TwoParams(FOnGeneratorInstanceDestroyed, uint64, TSharedPtr<Metasound::FMetasoundGenerator>);
 
@@ -231,6 +239,8 @@ private:
 
 	bool IsParameterValid(const FAudioParameter& InParameter, const TMap<FName, FMetasoundFrontendVertex>& InInputNameVertexMap) const;
 
+	static Metasound::SourcePrivate::FParameterRouter& GetParameterRouter();
+
 	Metasound::FOperatorSettings GetOperatorSettings(Metasound::FSampleRate InSampleRate) const;
 	Metasound::FMetasoundEnvironment CreateEnvironment() const;
 	Metasound::FMetasoundEnvironment CreateEnvironment(const FSoundGeneratorInitParams& InParams) const;
diff --git a/Engine/Plugins/Runtime/Metasound/Source/MetasoundFrontend/Private/MetasoundAssetBase.cpp b/Engine/Plugins/Runtime/Metasound/Source/MetasoundFrontend/Private/MetasoundAssetBase.cpp
--- a/Engine/Plugins/Runtime/Metasound/Source/MetasoundFrontend/Private/MetasoundAssetBase.cpp
+++ b/Engine/Plugins/Runtime/Metasound/Source/MetasoundFrontend/Private/MetasoundAssetBase.cpp
@@ -598,12 +598,6 @@ bool FMetasoundAssetBase::AddingReferenceCausesLoop(const FSoftObjectPath& InRef
 	return bCausesLoop;
 }
 
-Metasound::FSendAddress FMetasoundAssetBase::CreateSendAddress(uint64 InInstanceID, const Metasound::FVertexName& InVertexName, const FName& InDataTypeName) const
-
-{
-	return Metasound::FSendAddress(InVertexName, InDataTypeName, InInstanceID);
-}
-
 void FMetasoundAssetBase::ConvertFromPreset()
 {
 	using namespace Metasound::Frontend;
@@ -637,7 +631,9 @@ TArray<FMetasoundAssetBase::FSendInfoAndVertexName> FMetasoundAssetBase::GetSend
 	{
 		FSendInfoAndVertexName Info;
 
+PRAGMA_DISABLE_DEPRECATION_WARNINGS
 		Info.SendInfo.Address = FMetaSoundParameterTransmitter::CreateSendAddressFromInstanceID(InInstanceID, Vertex.Name, Vertex.TypeName);
+PRAGMA_ENABLE_DEPRECATION_WARNINGS
 		Info.SendInfo.ParameterName = Vertex.Name;
 		Info.SendInfo.TypeName = Vertex.TypeName;
 		Info.VertexName = Vertex.Name;
@@ -649,23 +645,6 @@ TArray<FMetasoundAssetBase::FSendInfoAndVertexName> FMetasoundAssetBase::GetSend
 	return SendInfos;
 }
 
-Metasound::Frontend::FNodeHandle FMetasoundAssetBase::AddInputPinForSendAddress(const Metasound::FMetaSoundParameterTransmitter::FSendInfo& InSendInfo, Metasound::Frontend::FGraphHandle InGraph) const
-{
-	FMetasoundFrontendClassInput Description;
-	FGuid VertexID = FGuid::NewGuid();
-
-	Description.Name = InSendInfo.Address.GetChannelName();
-	Description.TypeName = Metasound::GetMetasoundDataTypeName<Metasound::FSendAddress>();
-	Description.VertexID = VertexID;
-	Description.DefaultLiteral.Set(InSendInfo.Address.GetChannelName().ToString());
-
-#if WITH_EDITOR
-	Description.Metadata.SetDescription(FText::GetEmpty());
-#endif // WITH_EDITOR
-
-	return InGraph->AddInputVertex(Description);
-}
-
 #if WITH_EDITOR
 FText FMetasoundAssetBase::GetDisplayName(FString&& InTypeName) const
 {
diff --git a/Engine/Plugins/Runtime/Metasound/Source/MetasoundFrontend/Private/MetasoundFrontendDataTypeRegistry.cpp b/Engine/Plugins/Runtime/Metasound/Source/MetasoundFrontend/Private/MetasoundFrontendDataTypeRegistry.cpp
--- a/Engine/Plugins/Runtime/Metasound/Source/MetasoundFrontend/Private/MetasoundFrontendDataTypeRegistry.cpp
+++ b/Engine/Plugins/Runtime/Metasound/Source/MetasoundFrontend/Private/MetasoundFrontendDataTypeRegistry.cpp
@@ -184,7 +184,6 @@ namespace Metasound
 					InputParams.InstanceID = InParams.InstanceID;
 					InputParams.NodeName = InParams.NodeName;
 					InputParams.VertexName = InParams.VertexName;
-					InputParams.bEnableTransmission = false;
 
 					return DataTypeEntry->CreateInputNode(MoveTemp(InputParams));
 				}
@@ -225,7 +224,6 @@ namespace Metasound
 					InputParams.InstanceID = InParams.InstanceID;
 					InputParams.NodeName = InParams.NodeName;
 					InputParams.VertexName = InParams.VertexName;
-					InputParams.bEnableTransmission = false;
 
 					return DataTypeEntry->CreateConstructorInputNode(MoveTemp(InputParams));
 				}
@@ -525,6 +523,7 @@ namespace Metasound
 				virtual TSharedPtr<IDataChannel, ESPMode::ThreadSafe> CreateDataChannel(const FName& InDataType, const FOperatorSettings& InOperatorSettings) const override;
 
 				virtual const IParameterAssignmentFunction& GetRawAssignmentFunction(const FName& InDataType) const override;
+				virtual FLiteralAssignmentFunction GetLiteralAssignmentFunction(const FName& InDataType) const override;
 
 				virtual bool GetFrontendInputClass(const FName& InDataType, FMetasoundFrontendClass& OutClass) const override;
 				virtual bool GetFrontendConstructorInputClass(const FName& InDataType, FMetasoundFrontendClass& OutClass) const override;
@@ -928,6 +927,15 @@ namespace Metasound
 				return NoOp;
 			}
 
+			FLiteralAssignmentFunction FDataTypeRegistry::GetLiteralAssignmentFunction(const FName& InDataType) const
+			{
+				if (const IDataTypeRegistryEntry* Entry = FindDataTypeEntry(InDataType))
+				{
+					return Entry->GetLiteralAssignmentFunction();
+				}
+				return nullptr;
+			}
+
 			bool FDataTypeRegistry::GetFrontendInputClass(const FName& InDataType, FMetasoundFrontendClass& OutClass) const
 			{
 				if (const IDataTypeRegistryEntry* Entry = FindDataTypeEntry(InDataType))
@@ -1132,6 +1140,46 @@ namespace Metasound
 			static MetasoundFrontendDataTypeRegistryPrivate::FDataTypeRegistry Registry;
 			return Registry;
 		}
+
+		void CreateAndBindDefaults(const FOperatorSettings& InOperatorSettings, FInputVertexInterfaceData& OutVertexData)
+		{
+			using namespace MetasoundVertexDataPrivate;
+
+			auto VertexAccessTypeToDataReferenceAccessType = [](EVertexAccessType InVertexAccessType) -> EDataReferenceAccessType
+			{
+				switch(InVertexAccessType)
+				{
+					case EVertexAccessType::Value:
+						return EDataReferenceAccessType::Value;
+
+					case EVertexAccessType::Reference:
+					default:
+						return EDataReferenceAccessType::Write;
+				}
+			};
+
+			IDataTypeRegistry& DataTypeRegistry = IDataTypeRegistry::Get();
+
+			for (const TBinding<FInputDataVertex>& Binding : OutVertexData)
+			{
+				// Attempt to create default data reference from the literal stored
+				// on the input vertex.
+				const FInputDataVertex& InputVertex = Binding.GetVertex();
+				EDataReferenceAccessType AccessType = VertexAccessTypeToDataReferenceAccessType(InputVertex.AccessType);
+				TOptional<FAnyDataReference> DataRef = DataTypeRegistry.CreateDataReference(InputVertex.DataTypeName, AccessType, InputVertex.GetDefaultLiteral(), InOperatorSettings);
+
+				if (DataRef.IsSet())
+				{
+					// Set as vertex data reference.
+					OutVertexData.BindVertex(InputVertex.VertexName, *DataRef);
+				}
+				else
+				{
+					// All inputs should have creatable defaults.
+					UE_LOG(LogMetaSound, Warning, TEXT("Failed to create default data reference for vertex %s of data type %s"), *InputVertex.VertexName.ToString(), *InputVertex.DataTypeName.ToString());
+				}
+			}
+		}
 	}
 }
 
diff --git a/Engine/Plugins/Runtime/Metasound/Source/MetasoundFrontend/Private/MetasoundFrontendGraph.cpp b/Engine/Plugins/Runtime/Metasound/Source/MetasoundFrontend/Private/MetasoundFrontendGraph.cpp
--- a/Engine/Plugins/Runtime/Metasound/Source/MetasoundFrontend/Private/MetasoundFrontendGraph.cpp
+++ b/Engine/Plugins/Runtime/Metasound/Source/MetasoundFrontend/Private/MetasoundFrontendGraph.cpp
@@ -220,8 +220,7 @@ namespace Metasound
 						InNode.Name,
 						InNode.GetID(),
 						InputVertex.Name,
-						MoveTemp(Literal),
-						bEnableTransmission
+						MoveTemp(Literal)
 					};
 
 					ensureAlwaysMsgf(InOwningGraphClassInput.AccessType != EMetasoundFrontendVertexAccessType::Unset, TEXT("Graph Class Input cannot be set to access type of 'Unset'"));
diff --git a/Engine/Plugins/Runtime/Metasound/Source/MetasoundFrontend/Private/MetasoundInputNode.cpp b/Engine/Plugins/Runtime/Metasound/Source/MetasoundFrontend/Private/MetasoundInputNode.cpp
--- a/Engine/Plugins/Runtime/Metasound/Source/MetasoundFrontend/Private/MetasoundInputNode.cpp
+++ b/Engine/Plugins/Runtime/Metasound/Source/MetasoundFrontend/Private/MetasoundInputNode.cpp
@@ -31,10 +31,20 @@ namespace Metasound
 		{
 		}
 
-		void FNonExecutableInputOperatorBase::Bind(FVertexInterfaceData& InVertexData) const
+		void FNonExecutableInputOperatorBase::Bind(FVertexInterfaceData& InOutVertexData) const
 		{
-			InVertexData.GetInputs().BindVertex(VertexName, DataRef);
-			InVertexData.GetOutputs().BindVertex(VertexName, DataRef);
+			BindInputs(InOutVertexData.GetInputs());
+			BindOutputs(InOutVertexData.GetOutputs());
+		}
+
+		void FNonExecutableInputOperatorBase::BindInputs(FInputVertexInterfaceData& InOutVertexData) const
+		{
+			InOutVertexData.BindVertex(VertexName, DataRef);
+		}
+
+		void FNonExecutableInputOperatorBase::BindOutputs(FOutputVertexInterfaceData& InOutVertexData) const
+		{
+			InOutVertexData.BindVertex(VertexName, DataRef);
 		}
 
 		IOperator::FExecuteFunction FNonExecutableInputOperatorBase::GetExecuteFunction()
@@ -42,7 +52,12 @@ namespace Metasound
 			return nullptr;
 		}
 
-		IOperator::FResetFunction FNonExecutableInputPassThroughOperator::GetResetFunction() 
+		IOperator::FPostExecuteFunction FNonExecutableInputOperatorBase::GetPostExecuteFunction()
+		{
+			return nullptr;
+		}
+
+		IOperator::FResetFunction FNonExecutableInputOperatorBase::GetResetFunction() 
 		{
 			return nullptr;
 		}
diff --git a/Engine/Plugins/Runtime/Metasound/Source/MetasoundFrontend/Private/MetasoundOutputNode.cpp b/Engine/Plugins/Runtime/Metasound/Source/MetasoundFrontend/Private/MetasoundOutputNode.cpp
--- a/Engine/Plugins/Runtime/Metasound/Source/MetasoundFrontend/Private/MetasoundOutputNode.cpp
+++ b/Engine/Plugins/Runtime/Metasound/Source/MetasoundFrontend/Private/MetasoundOutputNode.cpp
@@ -36,16 +36,6 @@ namespace Metasound
 			InVertexData.GetOutputs().BindVertex(VertexName, DataReference);
 		}
 
-		IOperator::FExecuteFunction FOutputOperator::GetExecuteFunction()
-		{
-			return nullptr;
-		}
-
-		IOperator::FResetFunction FOutputOperator::GetResetFunction()
-		{
-			return nullptr;
-		}
-
 		FOutputNode::FOutputOperatorFactory::FOutputOperatorFactory(const FVertexName& InVertexName, const FName& InDataTypeName, EVertexAccessType InVertexAccessType)
 		: VertexName(InVertexName)
 		, DataTypeName(InDataTypeName)
diff --git a/Engine/Plugins/Runtime/Metasound/Source/MetasoundFrontend/Private/MetasoundParameterPack.cpp b/Engine/Plugins/Runtime/Metasound/Source/MetasoundFrontend/Private/MetasoundParameterPack.cpp
--- a/Engine/Plugins/Runtime/Metasound/Source/MetasoundFrontend/Private/MetasoundParameterPack.cpp
+++ b/Engine/Plugins/Runtime/Metasound/Source/MetasoundFrontend/Private/MetasoundParameterPack.cpp
@@ -23,7 +23,9 @@ TSharedPtr<Audio::IProxyData> UMetasoundParameterPack::CreateProxyData(const Aud
 
 Metasound::FSendAddress UMetasoundParameterPack::CreateSendAddressFromEnvironment(const Metasound::FMetasoundEnvironment& InEnvironment)
 {
+	PRAGMA_DISABLE_DEPRECATION_WARNINGS
 	return Metasound::FMetaSoundParameterTransmitter::CreateSendAddressFromEnvironment(InEnvironment, MetasoundParameterPack::RoutingName, Metasound::GetMetasoundDataTypeName<FMetasoundParameterStorageWrapper>());
+	PRAGMA_ENABLE_DEPRECATION_WARNINGS
 }
 
 FMetasoundFrontendClassInput UMetasoundParameterPack::GetClassInput()
diff --git a/Engine/Plugins/Runtime/Metasound/Source/MetasoundFrontend/Private/MetasoundParameterTransmitter.cpp b/Engine/Plugins/Runtime/Metasound/Source/MetasoundFrontend/Private/MetasoundParameterTransmitter.cpp
--- a/Engine/Plugins/Runtime/Metasound/Source/MetasoundFrontend/Private/MetasoundParameterTransmitter.cpp
+++ b/Engine/Plugins/Runtime/Metasound/Source/MetasoundFrontend/Private/MetasoundParameterTransmitter.cpp
@@ -19,6 +19,7 @@ namespace Metasound
 		TEXT("0: Disabled (default), !0: Enabled"),
 		ECVF_Default);
 
+
 	namespace Frontend
 	{
 		FLiteral ConvertParameterToLiteral(FAudioParameter&& InValue)
@@ -231,7 +232,9 @@ namespace Metasound
 			InstanceID = InEnvironment.GetValue<uint64>(SourceInterface::Environment::TransmitterID);
 		}
 
+		PRAGMA_DISABLE_DEPRECATION_WARNINGS
 		return CreateSendAddressFromInstanceID(InstanceID, InVertexName, InTypeName);
+		PRAGMA_ENABLE_DEPRECATION_WARNINGS
 	}
 
 	FSendAddress FMetaSoundParameterTransmitter::CreateSendAddressFromInstanceID(uint64 InInstanceID, const FVertexName& InVertexName, const FName& InTypeName)
@@ -241,27 +244,21 @@ namespace Metasound
 
 	FMetaSoundParameterTransmitter::FMetaSoundParameterTransmitter(FMetaSoundParameterTransmitter::FInitParams&& InInitParams)
 		: Audio::FParameterTransmitterBase(MoveTemp(InInitParams.DefaultParams))
-		, OperatorSettings(MoveTemp(InInitParams.OperatorSettings))
 		, InstanceID(InInitParams.InstanceID)
 		, DebugMetaSoundName(InInitParams.DebugMetaSoundName)
-		, SendInfos(MoveTemp(InInitParams.Infos))
+		, AvailableParameterNames(MoveTemp(InInitParams.ValidParameterNames))
+		, DataChannel(MoveTemp(InInitParams.DataChannel))
 	{
 	}
 
 	void FMetaSoundParameterTransmitter::OnDeleteActiveSound()
 	{
-		for (const FSendInfo& SendInfo : SendInfos)
-		{
-			if (InputSends.Remove(SendInfo.ParameterName))
-			{
-				// Only unregister the data channel if we had a sender using that 
-				// data channel. This protects against removing the data channel 
-				// multiple times. Multiple removals of data channels has caused
-				// race conditions between newly created transmitters and transmitters
-				// being cleaned up.
-				FDataTransmissionCenter::Get().UnregisterDataChannel(SendInfo.Address);
-			}
-		}
+		// Only reset the data channel if we have a valid data channel.
+		// This protects against removing the data channel 
+		// multiple times. Multiple removals of data channels has caused
+		// race conditions between newly created transmitters and transmitters
+		// being cleaned up.
+		DataChannel.Reset();
 
 		Audio::FParameterTransmitterBase::OnDeleteActiveSound();
 	}
@@ -301,18 +298,14 @@ namespace Metasound
 
 	bool FMetaSoundParameterTransmitter::SetParameterWithLiteral(FName InParameterName, const FLiteral& InLiteral)
 	{
-		if (ISender* Sender = FindSender(InParameterName))
-		{
-			return Sender->PushLiteral(InLiteral);
-		}
-
-		// If no sender exists for parameter name, attempt to add one.
-		if (const FSendInfo* SendInfo = FindSendInfo(InParameterName))
+		if (AvailableParameterNames.Contains(InParameterName))
 		{
-			if (ISender* Sender = AddSender(*SendInfo))
+			if (DataChannel.IsValid())
 			{
-				return Sender->PushLiteral(InLiteral);
+				DataChannel->Enqueue(FParameter{InParameterName, InLiteral});
 			}
+
+			return true;
 		}
 
 		// Enable / disable via CVAR to avoid log spam.
@@ -323,35 +316,4 @@ namespace Metasound
 
 		return false;
 	}
-
-	const FMetaSoundParameterTransmitter::FSendInfo* FMetaSoundParameterTransmitter::FindSendInfo(const FName& InParameterName) const
-	{
-		return SendInfos.FindByPredicate([&](const FSendInfo& Info) { return Info.ParameterName == InParameterName; });
-	}
-
-	ISender* FMetaSoundParameterTransmitter::FindSender(const FName& InParameterName)
-	{
-		if (TUniquePtr<ISender>* SenderPtrPtr = InputSends.Find(InParameterName))
-		{
-			return SenderPtrPtr->Get();
-		}
-		return nullptr;
-	}
-
-	ISender* FMetaSoundParameterTransmitter::AddSender(const FSendInfo& InInfo)
-	{
-		// TODO: likely want to remove this and opt for different protocols having different behaviors.
-		const float DelayTimeInSeconds = 0.1f; // This not used for non-audio routing.
-		const FSenderInitParams InitParams = { OperatorSettings, DelayTimeInSeconds };
-
-		TUniquePtr<ISender> Sender = FDataTransmissionCenter::Get().RegisterNewSender(InInfo.Address, InitParams);
-		if (ensureMsgf(Sender.IsValid(), TEXT("Failed to create sender [Address:%s]"), *InInfo.Address.ToString()))
-		{
-			ISender* Ptr = Sender.Get();
-			InputSends.Add(InInfo.ParameterName, MoveTemp(Sender));
-			return Ptr;
-		}
-
-		return nullptr;
-	}
 }
diff --git a/Engine/Plugins/Runtime/Metasound/Source/MetasoundFrontend/Public/MetasoundAssetBase.h b/Engine/Plugins/Runtime/Metasound/Source/MetasoundFrontend/Public/MetasoundAssetBase.h
--- a/Engine/Plugins/Runtime/Metasound/Source/MetasoundFrontend/Public/MetasoundAssetBase.h
+++ b/Engine/Plugins/Runtime/Metasound/Source/MetasoundFrontend/Public/MetasoundAssetBase.h
@@ -181,6 +181,7 @@ protected:
 #endif
 
 	// Get information for communicating asynchronously with MetaSound running instance.
+	UE_DEPRECATED(5.3, "MetaSounds no longer communicate using FSendInfo.")
 	TArray<FSendInfoAndVertexName> GetSendInfos(uint64 InInstanceID) const;
 
 #if WITH_EDITORONLY_DATA
@@ -235,6 +236,4 @@ private:
 	FRuntimeData CachedRuntimeData;
 
 	TSharedPtr<Metasound::IGraph, ESPMode::ThreadSafe> BuildMetasoundDocument(const FMetasoundFrontendDocument& InPreprocessDoc, const TSet<FName>& InTransmittableInputNames) const;
-	Metasound::FSendAddress CreateSendAddress(uint64 InInstanceID, const Metasound::FVertexName& InVertexName, const FName& InDataTypeName) const;
-	Metasound::Frontend::FNodeHandle AddInputPinForSendAddress(const Metasound::FMetaSoundParameterTransmitter::FSendInfo& InSendInfo, Metasound::Frontend::FGraphHandle InGraph) const;
 };
diff --git a/Engine/Plugins/Runtime/Metasound/Source/MetasoundFrontend/Public/MetasoundDataTypeRegistrationMacro.h b/Engine/Plugins/Runtime/Metasound/Source/MetasoundFrontend/Public/MetasoundDataTypeRegistrationMacro.h
--- a/Engine/Plugins/Runtime/Metasound/Source/MetasoundFrontend/Public/MetasoundDataTypeRegistrationMacro.h
+++ b/Engine/Plugins/Runtime/Metasound/Source/MetasoundFrontend/Public/MetasoundDataTypeRegistrationMacro.h
@@ -2,7 +2,6 @@
 
 #pragma once
 
-#include "Templates/Casts.h"
 
 #include "IAudioProxyInitializer.h"
 #include "MetasoundArrayNodesRegistration.h"
@@ -27,6 +26,8 @@
 #include "MetasoundVariableNodes.h"
 #include "MetasoundParameterPackFixedArray.h"
 
+#include "Templates/Casts.h"
+
 #include <type_traits>
 
 namespace Metasound
@@ -43,6 +44,26 @@ namespace Metasound
 				static constexpr bool bOnlySupportsDeprecatedProxyPtr = bIsParsableWithDeprecatedPtr && !bIsParsableWithSharedProxyPtr;
 		};
 
+		// ExecutableDataTypes are deprecated in 5.3 in favor of TPostExecutableDataTypes.
+		// This class triggers a deprecation warning when a TExecutableDataType
+		// is registered. 
+		template<typename DataType>
+		struct TExecutableDataTypeDeprecation
+		{
+			TExecutableDataTypeDeprecation()
+			{
+				if constexpr (TExecutableDataType<DataType>::bIsExecutable)
+				{
+					TriggerDeprecationMessage();
+				}
+			}
+
+			void TriggerDeprecationMessage() 
+			{
+				UE_LOG(LogMetaSound, Warning, TEXT("TExecutableDataType<> is deprecated in favor of TPostExecutableDataType<>. Please update your code for data type (%s) as TExecutableDataType<> will be removed in future releases"), *GetMetasoundDataTypeString<DataType>())
+			}
+		};
+
 		// Returns the Array version of a literal type if it exists.
 		template<ELiteralType LiteralType>
 		struct TLiteralArrayEnum 
@@ -310,6 +331,11 @@ namespace Metasound
 				// data type constructor. 
 				UE_LOG(LogMetaSound, Warning, TEXT("MetaSound data type \"%s\" supports construction from deprecated TUniquePtr<Audio::IProxyData>. Please update the constructor to accept a \"const TSharedPtr<Audio::IProxyData>& \""), TDataReferenceTypeInfo<TDataType>::TypeName);
 			}
+			
+			// TExecutableDataTypes are deprecated as of 5.3. This call triggers
+			// a deprecation warning in case the TExecutableDataType<> template
+			// was specialized.
+			MetasoundDataTypeRegistrationPrivate::TExecutableDataTypeDeprecation<TDataType>();
 
 			bAlreadyRegisteredThisDataType = true;
 
@@ -386,6 +412,11 @@ namespace Metasound
 					return RawAssignmentFunction;
 				}
 
+				virtual Frontend::FLiteralAssignmentFunction GetLiteralAssignmentFunction() const override
+				{
+					return LiteralAssignmentFunction;
+				}
+
 			protected:
 				Frontend::FDataTypeRegistryInfo Info;
 				FMetasoundFrontendClass InputClass;
@@ -399,14 +430,12 @@ namespace Metasound
 				FMetasoundFrontendClass VariableDeferredAccessorClass;
 				TSharedPtr<Frontend::IEnumDataTypeInterface> EnumInterface;
 				Frontend::IParameterAssignmentFunction RawAssignmentFunction;
+				Frontend::FLiteralAssignmentFunction LiteralAssignmentFunction = nullptr;
 			};
 
 			class FDataTypeRegistryEntry : public FDataTypeRegistryEntryBase
 			{
-			public:
-
-				FDataTypeRegistryEntry()
-					: FDataTypeRegistryEntryBase(CreateDataTypeInfo<TDataType, PreferredArgType, UClassToUse>(), GetEnumDataTypeInterface<TDataType>())
+				void InitRawAssignmentFunction()
 				{
 					if constexpr (HasRawParameterAssignmentOp<TDataType>().value)
 					{
@@ -438,7 +467,24 @@ namespace Metasound
 								};
 						}
 					}
+				}
+
+				void InitLiteralAssignmentFunction()
+				{
+					if constexpr (std::is_copy_assignable_v<TDataType> && bIsParsable)
+					{
+						if constexpr (!TIsArrayType<TDataType>::Value)
+						{
+							this->LiteralAssignmentFunction = [](const FOperatorSettings& InOperatorSettings, const FLiteral& InLiteral, FAnyDataReference& OutDataRef)
+							{
+								*OutDataRef.GetWritableValue<TDataType>() = TDataTypeLiteralFactory<TDataType>::CreateExplicitArgs(InOperatorSettings, InLiteral);
+							};
+						}
+					}
+				}
 
+				void InitFrontendNodeClasses()
+				{
 					// Create class info using prototype node
 					// TODO: register nodes with static class info instead of prototype instance.
 
@@ -475,6 +521,17 @@ namespace Metasound
 					}
 				}
 
+
+			public:
+
+				FDataTypeRegistryEntry()
+					: FDataTypeRegistryEntryBase(CreateDataTypeInfo<TDataType, PreferredArgType, UClassToUse>(), GetEnumDataTypeInterface<TDataType>())
+				{
+					InitRawAssignmentFunction();
+					InitLiteralAssignmentFunction();
+					InitFrontendNodeClasses();
+				}
+
 				virtual TUniquePtr<INode> CreateInputNode(FInputNodeConstructorParams&& InParams) const override
 				{
 					if constexpr (bIsParsable)
@@ -491,7 +548,6 @@ namespace Metasound
 				{
 					if constexpr (bIsParsable && bIsConstructorType)
 					{
-						checkf(!InParams.bEnableTransmission, TEXT("Cannot enable transmission on a constructor input."));
 						return MakeUnique<TInputNode<TDataType, EVertexAccessType::Value>>(MoveTemp(InParams));
 					}
 					else
diff --git a/Engine/Plugins/Runtime/Metasound/Source/MetasoundFrontend/Public/MetasoundFrontendDataTypeRegistry.h b/Engine/Plugins/Runtime/Metasound/Source/MetasoundFrontend/Public/MetasoundFrontendDataTypeRegistry.h
--- a/Engine/Plugins/Runtime/Metasound/Source/MetasoundFrontend/Public/MetasoundFrontendDataTypeRegistry.h
+++ b/Engine/Plugins/Runtime/Metasound/Source/MetasoundFrontend/Public/MetasoundFrontendDataTypeRegistry.h
@@ -20,6 +20,10 @@ namespace Metasound
 	{
 		using IParameterAssignmentFunction = TUniqueFunction<void(const void*,void*)>;
 
+		// Function signature for assigning a FLiteral to a FAnyDataReference 
+		using FLiteralAssignmentFunction = void(*)(const FOperatorSettings& InOperatorSettings, const FLiteral& InLiteral, FAnyDataReference& OutDataRef);
+
+
 		/** FDataTypeRegsitryInfo contains runtime inspectable behavior of a registered
 		 * MetaSound data type.
 		 */
@@ -70,14 +74,6 @@ namespace Metasound
 
 			// If provided in registration call, UClass this datatype was registered with.
 			UClass* ProxyGeneratorClass = nullptr;
-
-
-			// Returns if DataType supports array parsing.
-			UE_DEPRECATED(5.1, "Use bIsArrayType or bIsArrayParseable members instead")
-			FORCEINLINE bool IsArrayType() const
-			{
-				return bIsArrayType;
-			}
 		};
 
 		/** Interface for metadata of a registered MetaSound enum type. */
@@ -189,6 +185,9 @@ namespace Metasound
 			/* Get function that should be used when a parameter pack contains a setting this data type. */ 
 			virtual const IParameterAssignmentFunction& GetRawAssignmentFunction() const = 0;
 
+			/* Get a function that can be used to assign a FLiteral to a FAnyDataReference. */
+			virtual FLiteralAssignmentFunction GetLiteralAssignmentFunction() const = 0;
+
 			/** Create a init variable node for this data type. 
 			 *
 			 *  @param InInitParams - Contains a literal used to create the variable.
@@ -276,6 +275,9 @@ namespace Metasound
 			/* Get function that should be used when a parameter pack contains a setting for the specified data type. */ 
 			virtual const IParameterAssignmentFunction& GetRawAssignmentFunction(const FName& InDataType) const = 0;
 
+			/* Get a function that can be used to assign a FLiteral to a FAnyDataReference. */
+			virtual FLiteralAssignmentFunction GetLiteralAssignmentFunction(const FName& InDataType) const = 0;
+
 			/** Return an FMetasoundFrontendClass representing an input node of the data type. */
 			virtual bool GetFrontendInputClass(const FName& InDataType, FMetasoundFrontendClass& OutClass) const = 0;
 
@@ -315,5 +317,13 @@ namespace Metasound
 			virtual TUniquePtr<INode> CreateVariableAccessorNode(const FName& InDataType, const FNodeInitData& InParams) const = 0;
 			virtual TUniquePtr<INode> CreateVariableDeferredAccessorNode(const FName& InDataType, const FNodeInitData& InParams) const = 0;
 		};
+
+		/** Create data references for all inputs on the FInputVertexInterfaceData 
+		 * using the default literals which live on the FInputDataVertexes.
+		 *
+		 * If the vertex has access type EVertexAccessType::Reference, a writable
+		 * data reference will be bound, otherwise a value data reference will be bound.
+		 */
+		METASOUNDFRONTEND_API void CreateAndBindDefaults(const FOperatorSettings& InOperatorSettings, FInputVertexInterfaceData& OutVertexData);
 	}
 }
diff --git a/Engine/Plugins/Runtime/Metasound/Source/MetasoundFrontend/Public/MetasoundInputNode.h b/Engine/Plugins/Runtime/Metasound/Source/MetasoundFrontend/Public/MetasoundInputNode.h
--- a/Engine/Plugins/Runtime/Metasound/Source/MetasoundFrontend/Public/MetasoundInputNode.h
+++ b/Engine/Plugins/Runtime/Metasound/Source/MetasoundFrontend/Public/MetasoundInputNode.h
@@ -36,13 +36,21 @@ namespace Metasound
 		class METASOUNDFRONTEND_API FNonExecutableInputOperatorBase : public FInputOperatorBase
 		{	
 		public:
-			virtual void Bind(FVertexInterfaceData& InVertexData) const override;
+			virtual void Bind(FVertexInterfaceData& InOutVertexData) const override;
+
 			virtual IOperator::FExecuteFunction GetExecuteFunction() override;
+			virtual IOperator::FPostExecuteFunction GetPostExecuteFunction() override;
+			virtual IOperator::FResetFunction GetResetFunction() override;
 
 		protected:
 			FNonExecutableInputOperatorBase(const FVertexName& InVertexName, FAnyDataReference&& InDataRef);
 
 		private:
+			void BindInputs(FInputVertexInterfaceData& InOutVertexData) const;
+		protected:
+			void BindOutputs(FOutputVertexInterfaceData& InOutVertexData) const;
+
+
 			FVertexName VertexName;
 			FAnyDataReference DataRef;
 		};
@@ -56,7 +64,12 @@ namespace Metasound
 			: FNonExecutableInputOperatorBase(InVertexName, FAnyDataReference{InDataRef})
 			{
 			}
-			virtual IOperator::FResetFunction GetResetFunction() override;
+
+			template<typename DataType>
+			FNonExecutableInputPassThroughOperator(const FVertexName& InVertexName, const TDataWriteReference<DataType>& InDataRef)
+			: FNonExecutableInputPassThroughOperator(InVertexName, TDataReadReference<DataType>(InDataRef))
+			{
+			}
 		};
 
 
@@ -77,55 +90,8 @@ namespace Metasound
 			: FNonExecutableInputOperatorBase(InVertexName, FAnyDataReference{TDataValueReferenceLiteralFactory<DataType>::CreateExplicitArgs(InSettings, InLiteral)})
 			{
 			}
-
-			virtual IOperator::FResetFunction GetResetFunction() { return nullptr; }
-		};
-
-
-		template<typename DataType> 
-		class TNonExecutableInputOperator : public FInputOperatorBase
-		{
-		public:
-			TNonExecutableInputOperator(const FVertexName& InVertexName, const FOperatorSettings& InSettings, const FLiteral& InLiteral)
-			: VertexName(InVertexName)
-			, DataRef(TDataWriteReferenceLiteralFactory<DataType>::CreateExplicitArgs(InSettings, InLiteral))
-			, Literal(InLiteral)
-			{
-			}
-
-			virtual void Bind(FVertexInterfaceData& InVertexData) const override
-			{
-				InVertexData.GetInputs().BindWriteVertex(VertexName, DataRef);
-				InVertexData.GetOutputs().BindReadVertex(VertexName, DataRef);
-			}
-
-			virtual IOperator::FExecuteFunction GetExecuteFunction() override
-			{
-				return nullptr;
-			}
-
-			virtual IOperator::FResetFunction GetResetFunction() override
-			{
-				return &Reset;
-			}
-
-		private:
-			static void Reset(IOperator* InOperator, const IOperator::FResetParams& InParams)
-			{
-				using FOperator = TNonExecutableInputOperator<DataType>;
-
-				FOperator* Operator = static_cast<FOperator*>(InOperator);
-				check(nullptr != Operator);
-
-				*Operator->DataRef = TDataTypeLiteralFactory<DataType>::CreateExplicitArgs(InParams.OperatorSettings, Operator->Literal);
-			}
-
-			FVertexName VertexName;
-			TDataWriteReference<DataType> DataRef;
-			FLiteral Literal;
 		};
 
-		/** A writable input and a readable output. */
 		template<typename DataType>
 		class TExecutableInputOperator : public FInputOperatorBase
 		{
@@ -134,20 +100,17 @@ namespace Metasound
 			using FDataWriteReference = TDataWriteReference<DataType>;
 			using FDataWriteReferenceFactory = TDataWriteReferenceLiteralFactory<DataType>;
 
-			TExecutableInputOperator(const FVertexName& InDataReferenceName, const FOperatorSettings& InSettings, const FLiteral& InLiteral)
+			TExecutableInputOperator(const FVertexName& InDataReferenceName, TDataWriteReference<DataType> InValue)
 				: DataReferenceName(InDataReferenceName)
-				// Executable DataTypes require a copy of the output to operate on whereas non-executable
-				// types do not. Avoid copy by assigning to reference for non-executable types.
-				, InputValue(FDataWriteReferenceFactory::CreateExplicitArgs(InSettings, InLiteral))
-				, OutputValue(FDataWriteReferenceFactory::CreateExplicitArgs(InSettings, InLiteral))
-				, Literal(InLiteral)
+				, InputValue(InValue)
+				, OutputValue(FDataWriteReferenceFactory::CreateNew(*InValue))
 			{
 			}
 
 			virtual void Bind(FVertexInterfaceData& InOutVertexData) const override
 			{
-				InOutVertexData.GetInputs().BindWriteVertex(DataReferenceName, InputValue);
-				InOutVertexData.GetOutputs().BindReadVertex(DataReferenceName, OutputValue);
+				BindInputs(InOutVertexData.GetInputs());
+				BindOutputs(InOutVertexData.GetOutputs());
 			}
 
 			virtual FExecuteFunction GetExecuteFunction() override
@@ -157,20 +120,20 @@ namespace Metasound
 
 			virtual FResetFunction GetResetFunction() override
 			{
-				return &Reset;
+				return nullptr;
 			}
 
 		private:
-
-			static void Reset(IOperator* InOperator, const IOperator::FResetParams& InParams)
+			void BindInputs(FInputVertexInterfaceData& InOutVertexData) const
 			{
-				using FExecutableInputOperator = TExecutableInputOperator<DataType>;
+				InOutVertexData.BindWriteVertex(DataReferenceName, InputValue);
+			}
 
-				FExecutableInputOperator* Operator = static_cast<FExecutableInputOperator*>(InOperator);
-				check(nullptr != Operator);
+		protected:
 
-				*Operator->InputValue = TDataTypeLiteralFactory<DataType>::CreateExplicitArgs(InParams.OperatorSettings, Operator->Literal);
-				*Operator->OutputValue = *Operator->InputValue;
+			void BindOutputs(FOutputVertexInterfaceData& InOutVertexData) const
+			{
+				InOutVertexData.BindReadVertex(DataReferenceName, OutputValue);
 			}
 
 			static void Execute(IOperator* InOperator)
@@ -187,118 +150,164 @@ namespace Metasound
 
 			FDataWriteReference InputValue;
 			FDataWriteReference OutputValue;
-			FLiteral Literal;
-		};
 
+		};
 
-		/** TInputReceiverOperator provides support for transmittable inputs. */
 		template<typename DataType>
-		class TInputReceiverOperator : public FInputOperatorBase
+		class TResetableExecutableInputOperator : public TExecutableInputOperator<DataType>
 		{
 		public:
-			using FDataReadReference = TDataReadReference<DataType>;
 			using FDataWriteReference = TDataWriteReference<DataType>;
-			using FInputReceiverOperator = TInputReceiverOperator<DataType>;
+			using FDataWriteReferenceFactory = TDataWriteReferenceLiteralFactory<DataType>;
 
-			/** Construct an TInputReceiverOperator with the name of the vertex, value reference associated with input, SendAddress, & Receiver
-			 */
-			TInputReceiverOperator(const FVertexName& InDataReferenceName, FDataReadReference InDataReference, FSendAddress&& InSendAddress, TReceiverPtr<DataType>&& InReceiver)
-				: DataReferenceName(InDataReferenceName) 
-				, InputValue(InDataReference)
-				, OutputValue(FDataWriteReference::CreateNew(*InDataReference))
-				, SendAddress(MoveTemp(InSendAddress))
-				, Receiver(MoveTemp(InReceiver))
+			TResetableExecutableInputOperator(const FVertexName& InDataReferenceName, const FOperatorSettings& InSettings, const FLiteral& InLiteral)
+			: TExecutableInputOperator<DataType>(InDataReferenceName, FDataWriteReferenceFactory::CreateExplicitArgs(InSettings, InLiteral))
+			, Literal(InLiteral)
 			{
 			}
 
-			virtual ~TInputReceiverOperator()
+			virtual IOperator::FResetFunction GetResetFunction() override
+			{
+				return &Reset;
+			}
+
+		private:
+
+			static void Reset(IOperator* InOperator, const IOperator::FResetParams& InParams)
+			{
+				using FResetableExecutableInputOperator = TResetableExecutableInputOperator<DataType>;
+
+				FResetableExecutableInputOperator* Operator = static_cast<FResetableExecutableInputOperator*>(InOperator);
+				check(nullptr != Operator);
+
+				*Operator->InputValue = TDataTypeLiteralFactory<DataType>::CreateExplicitArgs(InParams.OperatorSettings, Operator->Literal);
+				*Operator->OutputValue = *Operator->InputValue;
+			}
+
+			FLiteral Literal;
+		};
+
+		template<typename DataType>
+		class TPostExecutableInputOperator : public FInputOperatorBase
+		{
+			static_assert(TPostExecutableDataType<DataType>::bIsPostExecutable, "TPostExecutableInputOperator should only be used with post executable data types");
+			static_assert(!TExecutableDataType<DataType>::bIsExecutable, "A data type cannot be Executable and PostExecutable");
+
+		public:
+			using FDataWriteReference = TDataWriteReference<DataType>;
+			using FDataWriteReferenceFactory = TDataWriteReferenceLiteralFactory<DataType>;
+
+			TPostExecutableInputOperator(const FVertexName& InDataReferenceName, TDataWriteReference<DataType> InValue)
+				: DataReferenceName(InDataReferenceName)
+				, Value(InValue)
 			{
-				Receiver.Reset();
-				FDataTransmissionCenter::Get().UnregisterDataChannelIfUnconnected(SendAddress);
 			}
 
 			virtual void Bind(FVertexInterfaceData& InOutVertexData) const override
 			{
-				InOutVertexData.GetInputs().BindReadVertex(DataReferenceName, InputValue);
-				InOutVertexData.GetOutputs().BindReadVertex(DataReferenceName, OutputValue);
+				BindInputs(InOutVertexData.GetInputs());
+				BindOutputs(InOutVertexData.GetOutputs());
 			}
 
 			virtual FExecuteFunction GetExecuteFunction() override
 			{
-				return &StaticExecute;
+				return nullptr;
+			}
+
+			virtual FPostExecuteFunction GetPostExecuteFunction() override
+			{
+				return &PostExecute;
 			}
 
 			virtual FResetFunction GetResetFunction() override
 			{
-				return &StaticReset;
+				return nullptr;
 			}
 
 		private:
-			void Execute()
+			void BindInputs(FInputVertexInterfaceData& InOutVertexData) const
 			{
-				DataType& OutputData = *OutputValue;
-
-				bool bHasNewData = Receiver->CanPop();
-				if (bHasNewData)
-				{
-					Receiver->Pop(OutputData);
-					bHasNotReceivedData = false;
-				}
+				InOutVertexData.BindWriteVertex(DataReferenceName, Value);
+			}
 
-				if (bHasNotReceivedData)
-				{
-					OutputData = *InputValue;
-					bHasNewData = true;
-				}
+		protected:
 
-				if constexpr (TExecutableDataType<DataType>::bIsExecutable)
-				{
-					TExecutableDataType<DataType>::ExecuteInline(OutputData, bHasNewData);
-				}
+			void BindOutputs(FOutputVertexInterfaceData& InOutVertexData) const
+			{
+				InOutVertexData.BindReadVertex(DataReferenceName, Value);
 			}
 
-			static void StaticExecute(IOperator* InOperator)
+			static void PostExecute(IOperator* InOperator)
 			{
-				using FOperator = TInputReceiverOperator<DataType>;
+				using FPostExecutableInputOperator = TPostExecutableInputOperator<DataType>;
 
-				FOperator* Operator = static_cast<FOperator*>(InOperator);
-				check(nullptr != Operator);
-				Operator->Execute();
+				FPostExecutableInputOperator* DerivedOperator = static_cast<FPostExecutableInputOperator*>(InOperator);
+				check(nullptr != DerivedOperator);
+
+				TPostExecutableDataType<DataType>::PostExecute(*(DerivedOperator->Value));
 			}
 
-			void Reset(const IOperator::FResetParams& InParams)
+			FVertexName DataReferenceName;
+			FDataWriteReference Value;
+
+		};
+
+		template<typename DataType>
+		class TResetablePostExecutableInputOperator : public TPostExecutableInputOperator<DataType>
+		{
+		public:
+			using FDataWriteReference = TDataWriteReference<DataType>;
+			using FDataWriteReferenceFactory = TDataWriteReferenceLiteralFactory<DataType>;
+
+			TResetablePostExecutableInputOperator(const FVertexName& InDataReferenceName, const FOperatorSettings& InSettings, const FLiteral& InLiteral)
+			: TPostExecutableInputOperator<DataType>(InDataReferenceName, FDataWriteReferenceFactory::CreateExplicitArgs(InSettings, InLiteral))
+			, Literal(InLiteral)
 			{
-				*OutputValue = *InputValue;
-				bHasNotReceivedData = true;
 			}
 
-			static void StaticReset(IOperator* InOperator, const IOperator::FResetParams& InParams)
+			virtual IOperator::FResetFunction GetResetFunction() override
 			{
-				using FOperator = TInputReceiverOperator<DataType>;
+				return &Reset;
+			}
+
+		private:
+
+			static void Reset(IOperator* InOperator, const IOperator::FResetParams& InParams)
+			{
+				using FResetablePostExecutableInputOperator = TResetablePostExecutableInputOperator<DataType>;
 
-				FOperator* Operator = static_cast<FOperator*>(InOperator);
+				FResetablePostExecutableInputOperator* Operator = static_cast<FResetablePostExecutableInputOperator*>(InOperator);
 				check(nullptr != Operator);
-				Operator->Reset(InParams);
+
+				*Operator->Value = TDataTypeLiteralFactory<DataType>::CreateExplicitArgs(InParams.OperatorSettings, Operator->Literal);
 			}
 
-			FVertexName DataReferenceName;
-			FDataReadReference InputValue;
-			FDataWriteReference OutputValue;
-			FSendAddress SendAddress;
-			TReceiverPtr<DataType> Receiver;
-			bool bHasNotReceivedData = true;
+			FLiteral Literal;
 		};
+
+		/** Non owning input operator that may need execution. */
+		template<typename DataType, EVertexAccessType VertexAccess=EVertexAccessType::Reference>
+		using TNonOwningInputOperator = std::conditional_t<
+			TExecutableDataType<DataType>::bIsExecutable,
+			TExecutableInputOperator<DataType>, // Use this input operator if the data type is not owned by the input node but needs execution.
+			std::conditional_t<
+				TPostExecutableDataType<DataType>::bIsPostExecutable,
+				TPostExecutableInputOperator<DataType>, // Use this input operator if the data type is not owned by the input node but needs post execution.
+				MetasoundInputNodePrivate::FNonExecutableInputPassThroughOperator // Use this input operator if the data type is not owned by the input node and is not executable, nor post executable.
+			>
+		>;
 	}
 
-	/** Choose input operator based upon data type and access type */
+	/** Owning input operator that may need execution. */
 	template<typename DataType, EVertexAccessType VertexAccess=EVertexAccessType::Reference>
 	using TInputOperator = std::conditional_t<
-		VertexAccess == EVertexAccessType::Value,
-		MetasoundInputNodePrivate::TInputValueOperator<DataType>,
+		VertexAccess == EVertexAccessType::Value || (!TExecutableDataType<DataType>::bIsExecutable && !TPostExecutableDataType<DataType>::bIsPostExecutable),
+		MetasoundInputNodePrivate::TInputValueOperator<DataType>, // Use this input operator if the data type is owned by the input node and is not executable, nor post executable.
 		std::conditional_t<
 			TExecutableDataType<DataType>::bIsExecutable,
-			MetasoundInputNodePrivate::TExecutableInputOperator<DataType>,
-			MetasoundInputNodePrivate::TNonExecutableInputOperator<DataType>
+			MetasoundInputNodePrivate::TResetableExecutableInputOperator<DataType>, // Use this input operator if the data type is owned by the input node and is executable.
+			MetasoundInputNodePrivate::TResetablePostExecutableInputOperator<DataType> // Use this input operator if the data type is owned by the input node and is post executable.
+
 		>
 	>;
 
@@ -367,9 +376,7 @@ namespace Metasound
 			}
 
 		public:
-			explicit FInputNodeOperatorFactory(FLiteral&& InLiteral, bool bInEnableTransmission)
-			: Literal(MoveTemp(InLiteral))
-			, bEnableTransmission(bInEnableTransmission)
+			explicit FInputNodeOperatorFactory()
 			{
 			}
 
@@ -379,47 +386,30 @@ namespace Metasound
 
 				using FInputNodeType = TInputNode<DataType, VertexAccess>;
 
-				checkf(!(bEnableTransmission && bIsValueVertexAccess), TEXT("Input cannot enable transmission for vertex with access 'EVertexAccessType::Reference'"));
-
 				const FInputNodeType& InputNode = static_cast<const FInputNodeType&>(InParams.Node);
 				const FVertexName& VertexKey = InputNode.GetVertexName();
 
-				if (bEnableTransmission)
+				if (const FAnyDataReference* Ref = InParams.InputData.FindDataReference(VertexKey))
 				{
-					const FName DataTypeName = GetMetasoundDataTypeName<DataType>();
-					FSendAddress SendAddress = FMetaSoundParameterTransmitter::CreateSendAddressFromEnvironment(InParams.Environment, VertexKey, DataTypeName);
-					TReceiverPtr<DataType> Receiver = FDataTransmissionCenter::Get().RegisterNewReceiver<DataType>(SendAddress, FReceiverInitParams{ InParams.OperatorSettings });
-					if (Receiver.IsValid())
+					if constexpr (bIsReferenceVertexAccess)
 					{
-						// Transmittable input value
-						if (const FAnyDataReference* Ref = InParams.InputData.FindDataReference(VertexKey))
-						{
-							return MakeUnique<TInputReceiverOperator<DataType>>(VertexKey, CreatePassThroughDataReference(*Ref), MoveTemp(SendAddress), MoveTemp(Receiver));
-						}
-						else
+						if (EDataReferenceAccessType::Write == Ref->GetAccessType())
 						{
-							FDataReference DataRef = FDataReferenceFactory::CreateExplicitArgs(InParams.OperatorSettings, Literal);
-							return MakeUnique<TInputReceiverOperator<DataType>>(VertexKey, DataRef, MoveTemp(SendAddress), MoveTemp(Receiver));
+							return MakeUnique<TNonOwningInputOperator<DataType, VertexAccess>>(VertexKey, Ref->GetDataWriteReference<DataType>());
 						}
 					}
-
-					AddBuildError<FInputReceiverInitializationError>(OutResults.Errors, InParams.Node, VertexKey, DataTypeName);
-					return nullptr;
-				}
-
-				if (const FAnyDataReference* Ref = InParams.InputData.FindDataReference(VertexKey))
-				{
 					// Pass through input value
 					return MakeUnique<TPassThroughOperator<DataType, VertexAccess>>(VertexKey, CreatePassThroughDataReference(*Ref));
 				}
-
-				// Owned input value
-				return MakeUnique<TInputOperator<DataType, VertexAccess>>(VertexKey, InParams.OperatorSettings, Literal);
+				else
+				{
+					const FLiteral& Literal = InputNode.GetVertexInterface().GetInputInterface()[VertexKey].GetDefaultLiteral();
+					// Owned input value
+					return MakeUnique<TInputOperator<DataType, VertexAccess>>(VertexKey, InParams.OperatorSettings, Literal);
+				}
 			}
 
 		private:
-			FLiteral Literal;
-			bool bEnableTransmission = false;
 		};
 
 
@@ -427,11 +417,11 @@ namespace Metasound
 		// If true, this node can be instantiated by the Frontend.
 		static constexpr bool bCanRegister = bIsSupportedInput;
 
-		static FVertexInterface DeclareVertexInterface(const FVertexName& InVertexName)
+		static FVertexInterface CreateVertexInterface(const FVertexName& InVertexName, const FLiteral& InLiteral)
 		{
 			return FVertexInterface(
 				FInputVertexInterface(
-					FInputDataVertex(InVertexName, GetMetasoundDataTypeName<DataType>(), FDataVertexMetadata{ FText::GetEmpty() }, VertexAccess)
+					FInputDataVertex(InVertexName, GetMetasoundDataTypeName<DataType>(), FDataVertexMetadata{ FText::GetEmpty() }, VertexAccess, InLiteral)
 				),
 				FOutputVertexInterface(
 					FOutputDataVertex(InVertexName, GetMetasoundDataTypeName<DataType>(), FDataVertexMetadata{ FText::GetEmpty() }, VertexAccess)
@@ -439,6 +429,18 @@ namespace Metasound
 			);
 		}
 
+		static FVertexInterface CreateDefaultVertexInterface(const FVertexName& InVertexName)
+		{
+			return CreateVertexInterface(InVertexName, FLiteral());
+		}
+
+
+		UE_DEPRECATED(5.3, "Use TInputNode<>::CreateDefaultVertexInterface or TInputNode<>::CreateVertexInterface instead.")
+		static FVertexInterface DeclareVertexInterface(const FVertexName& InVertexName)
+		{
+			return CreateDefaultVertexInterface(InVertexName);
+		}
+
 		static FNodeClassMetadata GetNodeInfo(const FVertexName& InVertexName)
 		{
 			FNodeClassMetadata Info;
@@ -449,7 +451,7 @@ namespace Metasound
 			Info.Description = METASOUND_LOCTEXT("Metasound_InputNodeDescription", "Input into the parent Metasound graph.");
 			Info.Author = PluginAuthor;
 			Info.PromptIfMissing = PluginNodeMissingPrompt;
-			Info.DefaultInterface = DeclareVertexInterface(InVertexName);
+			Info.DefaultInterface = CreateDefaultVertexInterface(InVertexName);
 
 			return Info;
 		}
@@ -459,8 +461,8 @@ namespace Metasound
 		explicit TInputNode(FInputNodeConstructorParams&& InParams)
 		:	FNode(InParams.NodeName, InParams.InstanceID, GetNodeInfo(InParams.VertexName))
 		,	VertexName(InParams.VertexName)
-		,	Interface(DeclareVertexInterface(InParams.VertexName))
-		,	Factory(MakeShared<FInputNodeOperatorFactory>(MoveTemp(InParams.InitParam), InParams.bEnableTransmission))
+		,	Interface(CreateVertexInterface(InParams.VertexName, InParams.InitParam))
+		,	Factory(MakeShared<FInputNodeOperatorFactory>())
 		{
 		}
 
diff --git a/Engine/Plugins/Runtime/Metasound/Source/MetasoundFrontend/Public/MetasoundNodeConstructorParams.h b/Engine/Plugins/Runtime/Metasound/Source/MetasoundFrontend/Public/MetasoundNodeConstructorParams.h
--- a/Engine/Plugins/Runtime/Metasound/Source/MetasoundFrontend/Public/MetasoundNodeConstructorParams.h
+++ b/Engine/Plugins/Runtime/Metasound/Source/MetasoundFrontend/Public/MetasoundNodeConstructorParams.h
@@ -39,10 +39,12 @@ namespace Metasound
 		FGuid InstanceID;
 		FVertexName VertexName;
 		FLiteral InitParam = FLiteral::CreateInvalid();
+
+		UE_DEPRECATED(5.3, "Transmission is no longer enabled on input nodes. In order set values on a MetaSound input, access them through the bound FInputVertexInterfaceData.")
 		bool bEnableTransmission = false;
 	};
 
 	using FOutputNodeConstructorParams = FDefaultNamedVertexNodeConstructorParams;
 	using FLiteralNodeConstructorParams = FDefaultLiteralNodeConstructorParams;
 	using FVariableNodeConstructorParams = FDefaultLiteralNodeConstructorParams;
-} // namespace Metasound
\ No newline at end of file
+} // namespace Metasound
diff --git a/Engine/Plugins/Runtime/Metasound/Source/MetasoundFrontend/Public/MetasoundOutputNode.h b/Engine/Plugins/Runtime/Metasound/Source/MetasoundFrontend/Public/MetasoundOutputNode.h
--- a/Engine/Plugins/Runtime/Metasound/Source/MetasoundFrontend/Public/MetasoundOutputNode.h
+++ b/Engine/Plugins/Runtime/Metasound/Source/MetasoundFrontend/Public/MetasoundOutputNode.h
@@ -4,6 +4,7 @@
 #include "Internationalization/Text.h"
 #include "MetasoundBuildError.h"
 #include "MetasoundBuilderInterface.h"
+#include "MetasoundExecutableOperator.h"
 #include "MetasoundFrontendDataTypeTraits.h"
 #include "MetasoundNodeInterface.h"
 #include "MetasoundNodeRegistrationMacro.h"
@@ -46,7 +47,7 @@ namespace Metasound
 
 	namespace OutputNodePrivate
 	{
-		class METASOUNDFRONTEND_API FOutputOperator : public IOperator
+		class METASOUNDFRONTEND_API FOutputOperator : public FNoOpOperator
 		{
 			public:
 				FOutputOperator(const FVertexName& InVertexName, const FAnyDataReference& InDataReference);
@@ -55,8 +56,6 @@ namespace Metasound
 				virtual FDataReferenceCollection GetInputs() const override;
 				virtual FDataReferenceCollection GetOutputs() const override;
 				virtual void Bind(FVertexInterfaceData& InVertexData) const;
-				virtual FExecuteFunction GetExecuteFunction() override;
-				virtual FResetFunction GetResetFunction() override;
 
 			private:
 				FVertexName VertexName;
diff --git a/Engine/Plugins/Runtime/Metasound/Source/MetasoundFrontend/Public/MetasoundParameterPack.h b/Engine/Plugins/Runtime/Metasound/Source/MetasoundFrontend/Public/MetasoundParameterPack.h
--- a/Engine/Plugins/Runtime/Metasound/Source/MetasoundFrontend/Public/MetasoundParameterPack.h
+++ b/Engine/Plugins/Runtime/Metasound/Source/MetasoundFrontend/Public/MetasoundParameterPack.h
@@ -413,6 +413,7 @@ public:
 	
 	// A couple of utilities for use by MetasoundAssetBase and MetasoundGenerator to set
 	// up the routing for parameter packs
+	UE_DEPRECATED(5.3, "FSendAddress are no longer used to communicate with MetaSound instances.")
 	static Metasound::FSendAddress CreateSendAddressFromEnvironment(const Metasound::FMetasoundEnvironment& InEnvironment);
 	static FMetasoundFrontendClassInput GetClassInput();
 
diff --git a/Engine/Plugins/Runtime/Metasound/Source/MetasoundFrontend/Public/MetasoundParameterTransmitter.h b/Engine/Plugins/Runtime/Metasound/Source/MetasoundFrontend/Public/MetasoundParameterTransmitter.h
--- a/Engine/Plugins/Runtime/Metasound/Source/MetasoundFrontend/Public/MetasoundParameterTransmitter.h
+++ b/Engine/Plugins/Runtime/Metasound/Source/MetasoundFrontend/Public/MetasoundParameterTransmitter.h
@@ -2,6 +2,7 @@
 #pragma once
 
 #include "AudioParameterControllerInterface.h"
+#include "Containers/SpscQueue.h"
 #include "IAudioParameterTransmitter.h"
 #include "MetasoundDataReference.h"
 #include "MetasoundFrontendLiteral.h"
@@ -16,6 +17,8 @@ namespace Metasound
 {
 	extern METASOUNDFRONTEND_API int32 MetaSoundParameterEnableWarningOnIgnoredParameterCVar;
 
+	class FMetasoundGenerator;
+
 	namespace Frontend
 	{
 		METASOUNDFRONTEND_API FLiteral ConvertParameterToLiteral(FAudioParameter&& InValue);
@@ -23,6 +26,7 @@ namespace Metasound
 		METASOUNDFRONTEND_API FName ConvertParameterToDataType(EAudioParameterType InParameterType);
 	}
 
+
 	/** FMetaSoundParameterTransmitter provides a communication interface for 
 	 * sending values to a MetaSound instance. It relies on the send/receive transmission
 	 * system to ferry data from the transmitter to the MetaSound instance. Data will
@@ -49,7 +53,13 @@ namespace Metasound
 			/** Type name of parameter on MetaSound instance. */
 			FName TypeName;
 		};
-
+		
+		
+		struct FParameter
+		{
+			FName Name;
+			FLiteral Value;
+		};
 
 		/** Initialization parameters for a FMetaSoundParameterTransmitter. */
 		struct FInitParams
@@ -61,28 +71,49 @@ namespace Metasound
 			/** ID of the MetaSound instance.  */
 			uint64 InstanceID;
 
+			/** Available input parameters on MetaSound instance. */
+			TArray<FName> ValidParameterNames;
+			
 			/** Available input parameters on MetaSound instance. */
 			TArray<FSendInfo> Infos;
 
 			/** Name of MetaSound used to log parameter related errors. */
 			FName DebugMetaSoundName;
 
-			/** Default Audio Parameters set when transmitter is initialized */
+			/** Default Audio Parameters set when transmitter is initialized. */
 			TArray<FAudioParameter> DefaultParams;
 
+			/** Shared queue used to communicate with running MetaSound instance. */
+			TSharedPtr<TSpscQueue<FParameter>> DataChannel;
+
+			UE_DEPRECATED(5.3, "This constructor is deprecated and will be removed.")
 			FInitParams(const FOperatorSettings& InSettings, uint64 InInstanceID, TArray<FAudioParameter>&& InDefaultParams, const TArray<FSendInfo>& InInfos=TArray<FSendInfo>())
 			: OperatorSettings(InSettings)
 			, InstanceID(InInstanceID)
-			, Infos(InInfos)
 			, DefaultParams(MoveTemp(InDefaultParams))
+			{
+				for (const FSendInfo& SendInfo : InInfos)
+				{
+					ValidParameterNames.Add(SendInfo.ParameterName);
+				}
+			}
+
+			FInitParams(const FOperatorSettings& InSettings, uint64 InInstanceID, TArray<FAudioParameter>&& InDefaultParams, TArray<FName>&& InValidParameterNames, TSharedPtr<TSpscQueue<FParameter>> InDataChannel)
+			: OperatorSettings(InSettings)
+			, InstanceID(InInstanceID)
+			, ValidParameterNames(MoveTemp(InValidParameterNames))
+			, DefaultParams(MoveTemp(InDefaultParams))
+			, DataChannel(InDataChannel)
 			{
 			}
 		};
 
 		/** Creates a unique send address using the given MetaSound environment. */
+		UE_DEPRECATED(5.3, "This function is no longer used. Send addresses are no longer used when communicating with MetaSound inputs")
 		static FSendAddress CreateSendAddressFromEnvironment(const FMetasoundEnvironment& InEnvironment, const FVertexName& InVertexName, const FName& InTypeName);
 		
 		/** Creates a unique send address using the given InstanceID. */
+		UE_DEPRECATED(5.3, "This function is no longer used. Send addresses are no longer used when communicating with MetaSound inputs")
 		static FSendAddress CreateSendAddressFromInstanceID(uint64 InInstanceID, const FVertexName& InVertexName, const FName& InTypeName);
 
 		FMetaSoundParameterTransmitter(FMetaSoundParameterTransmitter::FInitParams&& InInitParams);
@@ -109,20 +140,10 @@ namespace Metasound
 		bool SetParameterWithLiteral(FName InParameterName, const FLiteral& InValue);
 
 	private:
-		// Find FSendInfo by parameter name. 
-		const FSendInfo* FindSendInfo(const FName& InParameterName) const;
-
-		// Find ISender by parameter name. 
-		ISender* FindSender(const FName& InParameterName);
 
-		// Create and store a new ISender for the given FSendInfo.
-		ISender* AddSender(const FSendInfo& InInfo);
-
-		FOperatorSettings OperatorSettings;
 		uint64 InstanceID;
 		FName DebugMetaSoundName;
-		TArray<FSendInfo> SendInfos;
-
-		TMap<FName, TUniquePtr<ISender>> InputSends;
+		TArray<FName> AvailableParameterNames;
+		TSharedPtr<TSpscQueue<FParameter>> DataChannel;
 	};
 }
diff --git a/Engine/Plugins/Runtime/Metasound/Source/MetasoundFrontend/Public/MetasoundTrigger.h b/Engine/Plugins/Runtime/Metasound/Source/MetasoundFrontend/Public/MetasoundTrigger.h
--- a/Engine/Plugins/Runtime/Metasound/Source/MetasoundFrontend/Public/MetasoundTrigger.h
+++ b/Engine/Plugins/Runtime/Metasound/Source/MetasoundFrontend/Public/MetasoundTrigger.h
@@ -189,24 +189,13 @@ namespace Metasound
 	};
 
 	template<>
-	struct TExecutableDataType<FTrigger>
+	struct TPostExecutableDataType<FTrigger>
 	{
-		static constexpr bool bIsExecutable = true;
+		static constexpr bool bIsPostExecutable = true;
 
-		static void Execute(FTrigger& InData, FTrigger& OutData)
+		static void PostExecute(FTrigger& InOutData)
 		{
-			OutData = InData;
-			InData.AdvanceBlock();
-		}
-
-		static void ExecuteInline(FTrigger& OutData, bool bInUpdated)
-		{
-			// Only advance (i.e. flush old triggers) if not updated.  If
-			// updated, leave trigger alone as it has new information.
-			if (!bInUpdated)
-			{
-				OutData.AdvanceBlock();
-			}
+			InOutData.AdvanceBlock();
 		}
 	};
 
diff --git a/Engine/Plugins/Runtime/Metasound/Source/MetasoundGenerator/Private/MetasoundGenerator.cpp b/Engine/Plugins/Runtime/Metasound/Source/MetasoundGenerator/Private/MetasoundGenerator.cpp
--- a/Engine/Plugins/Runtime/Metasound/Source/MetasoundGenerator/Private/MetasoundGenerator.cpp
+++ b/Engine/Plugins/Runtime/Metasound/Source/MetasoundGenerator/Private/MetasoundGenerator.cpp
@@ -195,22 +195,23 @@ namespace Metasound
 	{
 		using namespace Audio;
 		using namespace Frontend;
+		using namespace MetasoundGeneratorPrivate;
 
 		METASOUND_TRACE_CPUPROFILER_EVENT_SCOPE_TEXT(*FString::Printf(TEXT("AsyncMetaSoundBuilder::DoWork %s"), *InitParams.MetaSoundName));
 
 		FBuildResults BuildResults;
 
 		// Create an instance of the new graph
-		TUniquePtr<IOperator> GraphOperator = BuildGraphOperator(MoveTemp(InitParams.DefaultParameters), BuildResults);
+		FGraphOperatorAndInputs GraphOperatorAndInputs = BuildGraphOperator(MoveTemp(InitParams.DefaultParameters), BuildResults);
 		LogBuildErrors(BuildResults);
 
-		if (GraphOperator.IsValid())
+		if (GraphOperatorAndInputs.Operator.IsValid())
 		{
 			// Create graph analyzer
 			TUniquePtr<FGraphAnalyzer> GraphAnalyzer = BuildGraphAnalyzer(MoveTemp(BuildResults.InternalDataReferences));
 	
 			// Collect data for generator
-			FMetasoundGeneratorData GeneratorData = BuildGeneratorData(InitParams, MoveTemp(GraphOperator), MoveTemp(GraphAnalyzer));
+			FMetasoundGeneratorData GeneratorData = BuildGeneratorData(InitParams, MoveTemp(GraphOperatorAndInputs), MoveTemp(GraphAnalyzer));
 
 			Generator->SetPendingGraph(MoveTemp(GeneratorData), bTriggerGenerator);
 		}
@@ -225,18 +226,22 @@ namespace Metasound
 		InitParams.Release();
 	}
 
-	FMetasoundGeneratorData FAsyncMetaSoundBuilder::BuildGeneratorData(const FMetasoundGeneratorInitParams& InInitParams, TUniquePtr<IOperator> InGraphOperator, TUniquePtr<Frontend::FGraphAnalyzer> InAnalyzer) const
+	MetasoundGeneratorPrivate::FMetasoundGeneratorData FAsyncMetaSoundBuilder::BuildGeneratorData(const FMetasoundGeneratorInitParams& InInitParams, FAsyncMetaSoundBuilder::FGraphOperatorAndInputs&& InGraphOperatorAndInputs, TUniquePtr<Frontend::FGraphAnalyzer> InAnalyzer) const
 	{
 		using namespace Audio;
 		using namespace Frontend;
+		using namespace MetasoundGeneratorPrivate;
 
 		METASOUND_TRACE_CPUPROFILER_EVENT_SCOPE_TEXT(TEXT("AsyncMetaSoundBuilder::BuildGeneratorData"));
 
-		checkf(InGraphOperator.IsValid(), TEXT("Graph operator must be a valid object"));
+		checkf(InGraphOperatorAndInputs.Operator.IsValid(), TEXT("Graph operator must be a valid object"));
 
 		// Gather relevant input and output references
 		FVertexInterfaceData VertexData(InInitParams.Graph->GetVertexInterface());
-		InGraphOperator->Bind(VertexData);
+		InGraphOperatorAndInputs.Operator->Bind(VertexData);
+
+		// Replace input data with writable inputs
+		VertexData.GetInputs() = InGraphOperatorAndInputs.Inputs;
 
 		// Get inputs
 		FTriggerWriteRef PlayTrigger = VertexData.GetInputs().GetOrConstructDataWriteReference<FTrigger>(SourceInterface::Inputs::OnPlay, InInitParams.OperatorSettings, false);
@@ -255,18 +260,31 @@ namespace Metasound
 
 		// Create the parameter setter map so parameter packs can be cracked
 		// open and distributed as appropriate...
-		TMap<FName, FParameterSetter> ParameterSetters;
+		const IDataTypeRegistry& DataTypeRegistry = IDataTypeRegistry::Get();
+		FParameterSetterSortedMap ParameterSetters;
+		TMap<FName, FParameterPackSetter> ParameterPackSetters;
 		FInputVertexInterfaceData& GraphInputs = VertexData.GetInputs();
-		for (auto InputIterator = GraphInputs.begin(); InputIterator != GraphInputs.end(); ++InputIterator)
+		for (const MetasoundVertexDataPrivate::TBinding<FInputDataVertex>& Binding : GraphInputs)
 		{
-			const FInputDataVertex& InputVertex = (*InputIterator).GetVertex();
-			const Frontend::IParameterAssignmentFunction& Setter = IDataTypeRegistry::Get().GetRawAssignmentFunction(InputVertex.DataTypeName);
-			if (Setter)
+			// Only assign inputs that are writable. 
+			if (EDataReferenceAccessType::Write == Binding.GetAccessType())
 			{
-				FParameterSetter ParameterSetter(InputVertex.DataTypeName,
-					(*InputIterator).GetDataReference()->GetRaw(),
-					Setter);
-				ParameterSetters.Add(InputVertex.VertexName, ParameterSetter);
+				if (const FAnyDataReference* DataRef = Binding.GetDataReference())
+				{
+					const FInputDataVertex& InputVertex = Binding.GetVertex();
+					const Frontend::IParameterAssignmentFunction& PackSetter = DataTypeRegistry.GetRawAssignmentFunction(InputVertex.DataTypeName);
+					if (PackSetter)
+					{
+						FParameterPackSetter ParameterPackSetter(InputVertex.DataTypeName, DataRef->GetRaw(), PackSetter);
+						ParameterPackSetters.Add(InputVertex.VertexName, ParameterPackSetter);
+					}
+
+					Frontend::FLiteralAssignmentFunction LiteralSetter = DataTypeRegistry.GetLiteralAssignmentFunction(InputVertex.DataTypeName);
+					if (LiteralSetter)
+					{
+						ParameterSetters.Add(InputVertex.VertexName, FParameterSetter{LiteralSetter, *DataRef});
+					}
+				}
 			}
 		}
 
@@ -274,9 +292,10 @@ namespace Metasound
 		return FMetasoundGeneratorData 
 		{
 			InInitParams.OperatorSettings,
-			MoveTemp(InGraphOperator),
+			MoveTemp(InGraphOperatorAndInputs.Operator),
 			MoveTemp(VertexData),
 			MoveTemp(ParameterSetters),
+			MoveTemp(ParameterPackSetters),
 			MoveTemp(InAnalyzer),
 			MoveTemp(OutputBuffers),
 			MoveTemp(PlayTrigger),
@@ -284,15 +303,29 @@ namespace Metasound
 		};
 	}
 
-	TUniquePtr<IOperator> FAsyncMetaSoundBuilder::BuildGraphOperator(TArray<FAudioParameter>&& InParameters, FBuildResults& OutBuildResults) const
+	FAsyncMetaSoundBuilder::FGraphOperatorAndInputs FAsyncMetaSoundBuilder::BuildGraphOperator(TArray<FAudioParameter>&& InParameters, FBuildResults& OutBuildResults) const
 	{
 		METASOUND_TRACE_CPUPROFILER_EVENT_SCOPE_TEXT(TEXT("AsyncMetaSoundBuilder::BuildGraphOperator"));
 		using namespace Frontend;
 
-		// Set input data based on the input parameters and the input interface
+		// Choose which type of data reference access to create depending upon the access of the vertex.
+		auto VertexAccessTypeToDataReferenceAccessType = [](EVertexAccessType InVertexAccessType) -> EDataReferenceAccessType
+		{
+			switch(InVertexAccessType)
+			{
+				case EVertexAccessType::Value:
+					return EDataReferenceAccessType::Value;
+
+				case EVertexAccessType::Reference:
+				default:
+					return EDataReferenceAccessType::Write;
+			}
+		};
+
 		const FInputVertexInterface& InputInterface = InitParams.Graph->GetVertexInterface().GetInputInterface();
 		FInputVertexInterfaceData InputData(InputInterface);
 
+		// Set input data based on the input parameters and the input interface
 		IDataTypeRegistry& DataRegistry = IDataTypeRegistry::Get();
 		for (FAudioParameter& Parameter : InParameters)
 		{
@@ -300,7 +333,8 @@ namespace Metasound
 			if (const FInputDataVertex* InputVertex = InputInterface.Find(ParamName))
 			{
 				FLiteral Literal = Frontend::ConvertParameterToLiteral(MoveTemp(Parameter));
-				TOptional<FAnyDataReference> DataReference = DataRegistry.CreateDataReference(InputVertex->DataTypeName, EDataReferenceAccessType::Value, Literal, InitParams.OperatorSettings);
+
+				TOptional<FAnyDataReference> DataReference = DataRegistry.CreateDataReference(InputVertex->DataTypeName, VertexAccessTypeToDataReferenceAccessType(InputVertex->AccessType), Literal, InitParams.OperatorSettings);
 
 				if (DataReference)
 				{
@@ -315,7 +349,23 @@ namespace Metasound
 			{
 				UE_LOG(LogMetaSound, Warning, TEXT("Failed to set initial input parameter %s on graph in MetaSoundSource [%s]"), *ParamName.ToString(), *InitParams.MetaSoundName);
 			}
+		}
 
+		// Set any remaining inputs to their default values.
+		for (const MetasoundVertexDataPrivate::TBinding<FInputDataVertex>& Binding : InputData)
+		{
+			// Only create data reference if something does not already exist. 
+			if (!Binding.IsBound())
+			{
+				const FInputDataVertex& InputVertex = Binding.GetVertex();
+				EDataReferenceAccessType AccessType = VertexAccessTypeToDataReferenceAccessType(InputVertex.AccessType);
+				TOptional<FAnyDataReference> DataReference = DataRegistry.CreateDataReference(InputVertex.DataTypeName, VertexAccessTypeToDataReferenceAccessType(InputVertex.AccessType), InputVertex.GetDefaultLiteral(), InitParams.OperatorSettings);
+
+				if (DataReference)
+				{
+					InputData.BindVertex(InputVertex.VertexName, *DataReference);
+				}
+			}
 		}
 
 		// Reset as elements in array have been moved.
@@ -323,7 +373,11 @@ namespace Metasound
 
 		// Create an instance of the new graph
 		FBuildGraphOperatorParams BuildParams { *InitParams.Graph, InitParams.OperatorSettings, InputData, InitParams.Environment };
-		return FOperatorBuilder(InitParams.BuilderSettings).BuildGraphOperator(BuildParams, OutBuildResults);
+		FGraphOperatorAndInputs OpAndInputs;
+		OpAndInputs.Operator = FOperatorBuilder(InitParams.BuilderSettings).BuildGraphOperator(BuildParams, OutBuildResults);
+		OpAndInputs.Inputs = InputData;
+
+		return OpAndInputs;
 	}
 
 	TUniquePtr<Frontend::FGraphAnalyzer> FAsyncMetaSoundBuilder::BuildGraphAnalyzer(TMap<FGuid, FDataReferenceCollection>&& InInternalDataReferences) const
@@ -383,6 +437,7 @@ namespace Metasound
 		, bPendingGraphTrigger(true)
 		, bIsNewGraphPending(false)
 		, bIsWaitingForFirstGraph(true)
+		, ParameterQueue(MoveTemp(InParams.DataChannel))
 #if ENABLE_METASOUND_GENERATOR_RENDER_TIMING
 		, RenderTimer(MakeUnique<MetasoundGeneratorPrivate::FRenderTimer>(InParams.OperatorSettings, 1. /* AnalysisPeriod */))
 #endif // if ENABLE_METASOUND_GENERATOR_RENDER_TIMING
@@ -393,7 +448,9 @@ namespace Metasound
 		NumSamplesPerExecute = NumChannels * NumFramesPerExecute;
 
 		// Create the routing for parameter packs
+PRAGMA_DISABLE_DEPRECATION_WARNINGS
 		ParameterPackSendAddress = UMetasoundParameterPack::CreateSendAddressFromEnvironment(InParams.Environment);
+PRAGMA_ENABLE_DEPRECATION_WARNINGS
 		ParameterPackReceiver = FDataTransmissionCenter::Get().RegisterNewReceiver<FMetasoundParameterStorageWrapper>(ParameterPackSendAddress, FReceiverInitParams{ InParams.OperatorSettings });
 
 		BuilderTask = MakeUnique<FBuilderTask>(this, MoveTemp(InParams), true /* bTriggerGenerator */);
@@ -426,8 +483,10 @@ namespace Metasound
 		FDataTransmissionCenter::Get().UnregisterDataChannelIfUnconnected(ParameterPackSendAddress);
 	}
 
-	void FMetasoundGenerator::SetPendingGraph(FMetasoundGeneratorData&& InData, bool bTriggerGraph)
+	void FMetasoundGenerator::SetPendingGraph(MetasoundGeneratorPrivate::FMetasoundGeneratorData&& InData, bool bTriggerGraph)
 	{
+		using namespace MetasoundGeneratorPrivate;
+
 		FScopeLock SetPendingGraphLock(&PendingGraphMutex);
 		{
 			PendingGraphData = MakeUnique<FMetasoundGeneratorData>(MoveTemp(InData));
@@ -459,7 +518,7 @@ namespace Metasound
 		return false;
 	}
 
-	void FMetasoundGenerator::SetGraph(TUniquePtr<FMetasoundGeneratorData>&& InData, bool bTriggerGraph)
+	void FMetasoundGenerator::SetGraph(TUniquePtr<MetasoundGeneratorPrivate::FMetasoundGeneratorData>&& InData, bool bTriggerGraph)
 	{
 		if (!InData.IsValid())
 		{
@@ -519,6 +578,7 @@ namespace Metasound
 		GraphAnalyzer = MoveTemp(InData->GraphAnalyzer);
 
 		ParameterSetters = MoveTemp(InData->ParameterSetters);
+		ParameterPackSetters = MoveTemp(InData->ParameterPackSetters);
 
 		if (bTriggerGraph)
 		{
@@ -634,7 +694,7 @@ namespace Metasound
 
 		while (NumSamplesRemaining > 0)
 		{
-			UnpackAndTransmitUpdatedParameters();
+			ApplyPendingUpdatesToInputs();
 
 			{
 #if ENABLE_METASOUND_GENERATOR_RENDER_TIMING
@@ -684,6 +744,8 @@ namespace Metasound
 
 				FMemory::Memcpy(OverflowBuffer.GetData(), &InterleavedAudioBuffer.GetData()[ThisLoopNumSamplesWritten], OverflowCount * sizeof(float));
 			}
+
+			RootExecuter.PostExecute();
 		}
 
 		return NumSamplesWritten;
@@ -763,25 +825,39 @@ namespace Metasound
 		// TODO: memcpy for single channel. 
 	}
 
-	void FMetasoundGenerator::UnpackAndTransmitUpdatedParameters()
+	void FMetasoundGenerator::ApplyPendingUpdatesToInputs()
 	{
 		using namespace Frontend;
+		using namespace MetasoundGeneratorPrivate;
 
-		TUniqueFunction<void(FMetasoundParameterPackStorage*)> ProcessPack = [&](FMetasoundParameterPackStorage* Pack)
+		auto ProcessPack = [&](FMetasoundParameterPackStorage* Pack)
 			{
 				for (auto Walker = Pack->begin(); Walker != Pack->end(); ++Walker)
 				{
-					if (const FParameterSetter* ParameterSetter = ParameterSetters.Find(Walker->Name))
+					if (const FParameterPackSetter* ParameterPackSetter = ParameterPackSetters.Find(Walker->Name))
 					{
-						if (ParameterSetter->DataType == Walker->TypeName)
+						if (ParameterPackSetter->DataType == Walker->TypeName)
 						{
-							ParameterSetter->SetParameterWithPayload(Walker->GetPayload());
+							ParameterPackSetter->SetParameterWithPayload(Walker->GetPayload());
 						}
 					}
 				}
 			};
 
-		// First handle packs that have come from the IAudioParameterInterface system...
+		// Handle parameters from the FMetasoundParameterTransmitter
+		if (ParameterQueue.IsValid())
+		{
+			TOptional<FMetaSoundParameterTransmitter::FParameter> Parameter;
+			while ((Parameter = ParameterQueue->Dequeue()))
+			{
+				if (FParameterSetter* Setter = ParameterSetters.Find(Parameter->Name))
+				{
+					Setter->Assign(OperatorSettings, Parameter->Value, Setter->DataReference);
+				}
+			}
+		}
+
+		// Handle parameter packs that have come from the IAudioParameterInterface system...
 		if (ParameterPackReceiver.IsValid())
 		{
 			FMetasoundParameterStorageWrapper Pack;
@@ -792,7 +868,7 @@ namespace Metasound
 			}
 		}
 
-		// Now handle packs that came from QueueParameterPack...
+		// Handle parameter packs that came from QueueParameterPack...
 		TSharedPtr<FMetasoundParameterPackStorage> QueuedParameterPack;
 		while (ParameterPackQueue.Dequeue(QueuedParameterPack))
 		{
diff --git a/Engine/Plugins/Runtime/Metasound/Source/MetasoundGenerator/Public/MetasoundGenerator.h b/Engine/Plugins/Runtime/Metasound/Source/MetasoundGenerator/Public/MetasoundGenerator.h
--- a/Engine/Plugins/Runtime/Metasound/Source/MetasoundGenerator/Public/MetasoundGenerator.h
+++ b/Engine/Plugins/Runtime/Metasound/Source/MetasoundGenerator/Public/MetasoundGenerator.h
@@ -27,6 +27,48 @@ namespace Metasound
 #if ENABLE_METASOUND_GENERATOR_RENDER_TIMING
 		struct FRenderTimer;
 #endif // if ENABLE_METASOUND_GENERATOR_RENDER_TIMING
+
+		struct FParameterSetter
+		{
+			Frontend::FLiteralAssignmentFunction Assign;
+			FAnyDataReference DataReference;
+		};
+
+		// In order to use FName as a key to a TSortedMap you have to explicitly 
+		// choose which comparison implementation you want to use.  Declaring the
+		// type here helps minimize the confusion over why there are so many 
+		// template arguments. 
+		using FParameterSetterSortedMap  = TSortedMap<FName, FParameterSetter, FDefaultAllocator, FNameFastLess>;
+
+		// A struct that provides a method of pushing "raw" data from a parameter pack into a specific metasound input node.
+		struct FParameterPackSetter
+		{
+			FName DataType;
+			void* Destination;
+			const Frontend::IParameterAssignmentFunction& Setter;
+			FParameterPackSetter(FName InDataType, void* InDestination, const Frontend::IParameterAssignmentFunction& InSetter)
+				: DataType(InDataType)
+				, Destination(InDestination)
+				, Setter(InSetter)
+			{}
+			void SetParameterWithPayload(const void* ParameterPayload) const
+			{
+				Setter(ParameterPayload, Destination);
+			}
+		};
+
+		struct FMetasoundGeneratorData
+		{
+			FOperatorSettings OperatorSettings;
+			TUniquePtr<IOperator> GraphOperator;
+			FVertexInterfaceData VertexInterfaceData;
+			FParameterSetterSortedMap ParameterSetters;
+			TMap<FName, FParameterPackSetter> ParameterPackSetters;
+			TUniquePtr<Frontend::FGraphAnalyzer> GraphAnalyzer;
+			TArray<TDataReadReference<FAudioBuffer>> OutputBuffers;
+			FTriggerWriteRef TriggerOnPlayRef;
+			FTriggerReadRef TriggerOnFinishRef;
+		};
 	}
 
 	// Struct needed for building the metasound graph
@@ -40,39 +82,11 @@ namespace Metasound
 		TArray<FVertexName> AudioOutputNames;
 		TArray<FAudioParameter> DefaultParameters;
 		bool bBuildSynchronous = false;
+		TSharedPtr<TSpscQueue<FMetaSoundParameterTransmitter::FParameter>> DataChannel;
 
 		void Release();
 	};
 
-	// A struct that provides a method of pushing "raw" data from a parameter pack into a specific metasound input node.
-	struct FParameterSetter
-	{
-		FName DataType;
-		void* Destination;
-		const Frontend::IParameterAssignmentFunction& Setter;
-		FParameterSetter(FName InDataType, void* InDestination, const Frontend::IParameterAssignmentFunction& InSetter)
-			: DataType(InDataType)
-			, Destination(InDestination)
-			, Setter(InSetter)
-		{}
-		void SetParameterWithPayload(const void* ParameterPayload) const
-		{
-			Setter(ParameterPayload, Destination);
-		}
-	};
-
-	struct FMetasoundGeneratorData
-	{
-		FOperatorSettings OperatorSettings;
-		TUniquePtr<IOperator> GraphOperator;
-		FVertexInterfaceData VertexInterfaceData;
-		TMap<FName, FParameterSetter> ParameterSetters;
-		TUniquePtr<Frontend::FGraphAnalyzer> GraphAnalyzer;
-		TArray<TDataReadReference<FAudioBuffer>> OutputBuffers;
-		FTriggerWriteRef TriggerOnPlayRef;
-		FTriggerReadRef TriggerOnFinishRef;
-	};
-
 	class FMetasoundGenerator;
 
 	class FAsyncMetaSoundBuilder : public FNonAbandonableTask
@@ -90,8 +104,13 @@ namespace Metasound
 		}
 
 	private:
-		FMetasoundGeneratorData BuildGeneratorData(const FMetasoundGeneratorInitParams& InInitParams, TUniquePtr<IOperator> InOperator, TUniquePtr<Frontend::FGraphAnalyzer> InAnalyzer) const;
-		TUniquePtr<IOperator> BuildGraphOperator(TArray<FAudioParameter>&& InParameters, FBuildResults& OutBuildResults) const;
+		struct FGraphOperatorAndInputs
+		{
+			TUniquePtr<IOperator> Operator;
+			FInputVertexInterfaceData Inputs;
+		};
+		MetasoundGeneratorPrivate::FMetasoundGeneratorData BuildGeneratorData(const FMetasoundGeneratorInitParams& InInitParams, FGraphOperatorAndInputs&& InOperatorAndInputs, TUniquePtr<Frontend::FGraphAnalyzer> InAnalyzer) const;
+		FGraphOperatorAndInputs BuildGraphOperator(TArray<FAudioParameter>&& InParameters, FBuildResults& OutBuildResults) const;
 		TUniquePtr<Frontend::FGraphAnalyzer> BuildGraphAnalyzer(TMap<FGuid, FDataReferenceCollection>&& InInternalDataReferences) const;
 		void LogBuildErrors(const FBuildResults& InBuildResults) const;
 		TArray<FAudioBufferReadRef> FindOutputAudioBuffers(const FVertexInterfaceData& InVertexData) const;
@@ -225,6 +244,13 @@ namespace Metasound
 
 #endif // if ENABLE_METASOUND_GENERATOR_RENDER_TIMING
 		
+
+		// Called when a new graph has been "compiled" and set up as this generator's graph.
+		FOnSetGraph OnSetGraph;
+
+	private:
+		friend class FAsyncMetaSoundBuilder;
+
 		/** Update the current graph operator with a new graph operator. The number of channels
 		 * of InGraphOutputAudioRef must match the existing number of channels reported by
 		 * GetNumChannels() in order for this function to successfully replace the graph operator.
@@ -232,18 +258,13 @@ namespace Metasound
 		 * @param InData - Metasound data of built graph.
 		 * @param bTriggerGraph - If true, "OnPlay" will be triggered on the new graph.
 		 */
-		void SetPendingGraph(FMetasoundGeneratorData&& InData, bool bTriggerGraph);
+		void SetPendingGraph(MetasoundGeneratorPrivate::FMetasoundGeneratorData&& InData, bool bTriggerGraph);
 		void SetPendingGraphBuildFailed();
 
-		// Called when a new graph has been "compiled" and set up as this generator's graph.
-		FOnSetGraph OnSetGraph;
-
-	private:
-
 		bool UpdateGraphIfPending();
 
 		// Internal set graph after checking compatibility.
-		void SetGraph(TUniquePtr<FMetasoundGeneratorData>&& InData, bool bTriggerGraph);
+		void SetGraph(TUniquePtr<MetasoundGeneratorPrivate::FMetasoundGeneratorData>&& InData, bool bTriggerGraph);
 
 		// Fill OutAudio with data in InBuffer, up to maximum number of samples.
 		// Returns the number of samples used.
@@ -252,7 +273,7 @@ namespace Metasound
 		// Metasound creates deinterleaved audio while sound generator requires interleaved audio.
 		void InterleaveGeneratedAudio();
 		
-		void UnpackAndTransmitUpdatedParameters();
+		void ApplyPendingUpdatesToInputs();
 
 		FExecuter RootExecuter;
 		FVertexInterfaceData VertexInterfaceData;
@@ -282,13 +303,16 @@ namespace Metasound
 		TUniquePtr<FBuilderTask> BuilderTask;
 
 		FCriticalSection PendingGraphMutex;
-		TUniquePtr<FMetasoundGeneratorData> PendingGraphData;
+		TUniquePtr<MetasoundGeneratorPrivate::FMetasoundGeneratorData> PendingGraphData;
 		bool bPendingGraphTrigger;
 		bool bIsNewGraphPending;
 		bool bIsWaitingForFirstGraph;
 
 		TUniquePtr<Frontend::FGraphAnalyzer> GraphAnalyzer;
 
+		TSharedPtr<TSpscQueue<FMetaSoundParameterTransmitter::FParameter>> ParameterQueue;
+		MetasoundGeneratorPrivate::FParameterSetterSortedMap ParameterSetters;
+
 		// These next items are needed to provide a destination for the FAudioDevice, etc. to
 		// send parameter packs to. Every playing metasound will have a parameter destination
 		// that can accept parameter packs.
@@ -297,7 +321,7 @@ namespace Metasound
 		
 		// This map provides setters for all of the input nodes in the metasound graph. 
 		// It is used when processing named parameters in a parameter pack.
-		TMap<FName, FParameterSetter> ParameterSetters;
+		TMap<FName, MetasoundGeneratorPrivate::FParameterPackSetter> ParameterPackSetters;
 
 		// While parameter packs may arrive via the IAudioParameterInterface system,
 		// a faster method of sending parameters is via the QueueParameterPack function 
diff --git a/Engine/Plugins/Runtime/Metasound/Source/MetasoundGraphCore/Private/MetasoundGraphOperator.cpp b/Engine/Plugins/Runtime/Metasound/Source/MetasoundGraphCore/Private/MetasoundGraphOperator.cpp
--- a/Engine/Plugins/Runtime/Metasound/Source/MetasoundGraphCore/Private/MetasoundGraphOperator.cpp
+++ b/Engine/Plugins/Runtime/Metasound/Source/MetasoundGraphCore/Private/MetasoundGraphOperator.cpp
@@ -22,6 +22,12 @@ namespace Metasound
 				bIsOperatorInAnyStack = true;
 			}
 
+			if (FPostExecuteFunction PostExecuteFunc = InOperator->GetPostExecuteFunction())
+			{
+				PostExecuteStack.Emplace(*InOperator, PostExecuteFunc);
+				bIsOperatorInAnyStack = true;
+			}
+
 			if (FResetFunction ResetFunc = InOperator->GetResetFunction())
 			{
 				ResetStack.Emplace(*InOperator, ResetFunc);
@@ -65,6 +71,19 @@ namespace Metasound
 		InVertexData.Bind(VertexData);
 	}
 
+	IOperator::FPostExecuteFunction FGraphOperator::GetPostExecuteFunction()
+	{
+		return &StaticPostExecute;
+	}
+
+	void FGraphOperator::StaticPostExecute(IOperator* InOperator)
+	{
+		FGraphOperator* GraphOperator = static_cast<FGraphOperator*>(InOperator);
+		check(GraphOperator);
+
+		GraphOperator->PostExecute();
+	}
+
 	void FGraphOperator::Execute()
 	{
 		FExecuteEntry* StackPtr = ExecuteStack.GetData();
@@ -75,6 +94,16 @@ namespace Metasound
 		}
 	}
 
+	void FGraphOperator::PostExecute()
+	{
+		FPostExecuteEntry* StackPtr = PostExecuteStack.GetData();
+		const int32 Num = PostExecuteStack.Num();
+		for (int32 i = 0; i < Num; i++)
+		{
+			StackPtr[i].PostExecute();
+		}
+	}
+
 	void FGraphOperator::Reset(const FGraphOperator::FResetParams& InParams)
 	{
 		FResetEntry* StackPtr = ResetStack.GetData();
@@ -91,11 +120,23 @@ namespace Metasound
 	{
 		check(Function);
 	}
+
 	void FGraphOperator::FExecuteEntry::Execute()
 	{
 		Function(Operator);
 	}
 
+	FGraphOperator::FPostExecuteEntry::FPostExecuteEntry(IOperator& InOperator, FPostExecuteFunction InFunc)
+	: Operator(&InOperator)
+	, Function(InFunc)
+	{
+		check(Function);
+	}
+
+	void FGraphOperator::FPostExecuteEntry::PostExecute()
+	{
+		Function(Operator);
+	}
 
 	FGraphOperator::FResetEntry::FResetEntry(IOperator& InOperator, FResetFunction InFunc)
 	: Operator(&InOperator)
@@ -103,9 +144,12 @@ namespace Metasound
 	{
 		check(Function);
 	}
+
 	void FGraphOperator::FResetEntry::Reset(const FGraphOperator::FResetParams& InParams)
 	{
 		Function(Operator, InParams);
 	}
 
+			static void StaticPostExecute(IOperator* Operator);
+
 }
diff --git a/Engine/Plugins/Runtime/Metasound/Source/MetasoundGraphCore/Public/MetasoundDataReference.h b/Engine/Plugins/Runtime/Metasound/Source/MetasoundGraphCore/Public/MetasoundDataReference.h
--- a/Engine/Plugins/Runtime/Metasound/Source/MetasoundGraphCore/Public/MetasoundDataReference.h
+++ b/Engine/Plugins/Runtime/Metasound/Source/MetasoundGraphCore/Public/MetasoundDataReference.h
@@ -649,6 +649,26 @@ namespace Metasound
 			return nullptr;
 		}
 
+		/** Return a non-const pointer to the data.
+		 * 
+		 * If the vertex is bound with this will return a valid pointer. 
+		 * Otherwise it will return a nullptr.  This method will assert
+		 * if the reference is bound, but not writable or if the data type
+		 * does match.  
+		 */
+		template<typename DataType>
+		DataType* GetWritableValue()
+		{
+			if (DataRefPtr.IsValid())
+			{
+				check(IsDataReferenceOfType<DataType>(*DataRefPtr));
+				checkf(EDataReferenceAccessType::Write == AccessType, TEXT("Writable values can only be obtained with data references which have writable access"));
+
+				return static_cast<TDataWriteReference<DataType>*>(DataRefPtr.Get())->Get();
+			}
+			return nullptr;
+		}
+
 		/** Get access to a TDataValueReference. 
 		 *
 		 * This method will return a valid TDataValueReference of the templated data
diff --git a/Engine/Plugins/Runtime/Metasound/Source/MetasoundGraphCore/Public/MetasoundExecutableOperator.h b/Engine/Plugins/Runtime/Metasound/Source/MetasoundGraphCore/Public/MetasoundExecutableOperator.h
--- a/Engine/Plugins/Runtime/Metasound/Source/MetasoundGraphCore/Public/MetasoundExecutableOperator.h
+++ b/Engine/Plugins/Runtime/Metasound/Source/MetasoundGraphCore/Public/MetasoundExecutableOperator.h
@@ -97,6 +97,17 @@ namespace Metasound
 		}
 	};
 
+	template<class DataType>
+	struct TPostExecutableDataType
+	{
+		static constexpr bool bIsPostExecutable = false;
+
+		/*
+		 * static void PostExecute(DataType& InOutData) { ... }
+		 */
+	};
+
+
 	/** Convenience class for supporting the IOperator interface's GetExecuteFunction virtual member function.
 	 *
 	 * Derived classes should inherit from this template class as well as implement a void Execute() member
@@ -112,7 +123,6 @@ namespace Metasound
 	 * 	  }
 	 * 	};
 	 */
-
 	template<class DerivedOperatorType>
 	class TExecutableOperator : public IOperator
 	{
@@ -130,6 +140,11 @@ namespace Metasound
 			return &TExecutableOperator<DerivedOperatorType>::ExecuteFunction;
 		}
 
+		virtual FPostExecuteFunction GetPostExecuteFunction() override
+		{
+			return nullptr;
+		}
+
 	private:
 
 		static void ExecuteFunction(IOperator* InOperator)
@@ -142,27 +157,36 @@ namespace Metasound
 		}
 	};
 
+	/** FNoOpOperator is for IOperators which do not perform any execution.
+	 * Their only behavior is to perform operations on constructor or through
+	 * the Bind(...) methods.
+	 */
 	class FNoOpOperator : public IOperator
 	{
 		public:
 			virtual ~FNoOpOperator() {}
 
-			FExecuteFunction GetExecuteFunction() override
+			virtual FExecuteFunction GetExecuteFunction() override
+			{
+				return nullptr;
+			}
+
+			virtual FResetFunction GetResetFunction() override
 			{
 				return nullptr;
 			}
 
-			FResetFunction GetResetFunction() override
+			virtual FPostExecuteFunction GetPostExecuteFunction() override
 			{
 				return nullptr;
 			}
 
-			FDataReferenceCollection GetInputs() const override
+			virtual FDataReferenceCollection GetInputs() const override
 			{
 				return FDataReferenceCollection{};
 			}
 
-			FDataReferenceCollection GetOutputs() const override
+			virtual FDataReferenceCollection GetOutputs() const override
 			{
 				return FDataReferenceCollection{};
 			}
@@ -176,17 +200,12 @@ namespace Metasound
 	{
 		public:
 			using FOperatorPtr = TUniquePtr<IOperator>;
-			using FExecuteFunction = IOperator::FExecuteFunction;
 
 			FExecuter()
-			:	ExecuteFunction(&FExecuter::NoOpExecute)
-			,	ResetFunction(&FExecuter::NoOpReset)
 			{
 			}
 
 			FExecuter(FOperatorPtr InOperator)
-			:	ExecuteFunction(&FExecuter::NoOpExecute)
-			,	ResetFunction(&FExecuter::NoOpReset)
 			{
 				SetOperator(MoveTemp(InOperator));
 			}
@@ -194,52 +213,46 @@ namespace Metasound
 			void SetOperator(FOperatorPtr InOperator)
 			{
 				Operator = MoveTemp(InOperator);
+
 				ExecuteFunction = nullptr;
+				PostExecuteFunction = nullptr;
 				ResetFunction = nullptr;
 
 				if (Operator.IsValid())
 				{
-					if (FExecuteFunction Func = Operator->GetExecuteFunction())
-					{
-						ExecuteFunction = Func;
-					}
-
-					if (FResetFunction Func = Operator->GetResetFunction())
-					{
-						ResetFunction = Func;
-					}
-				}
-
-				if (!ExecuteFunction)
-				{
-					ExecuteFunction = &FExecuter::NoOpExecute;
+					ExecuteFunction = Operator->GetExecuteFunction();
+					PostExecuteFunction = Operator->GetPostExecuteFunction();
+					ResetFunction = Operator->GetResetFunction();
 				}
+			}
 
-				if (!ResetFunction)
+			void Execute()
+			{
+				if (ExecuteFunction)
 				{
-					ResetFunction = &FExecuter::NoOpReset;
+					ExecuteFunction(Operator.Get());
 				}
 			}
 
-			void Execute()
+			void PostExecute()
 			{
-				// ExecuteFunction is always non-null. If the contained operator
-				// does not provide a valid execution function, then a No-Op function
-				// is utilized. 
-				ExecuteFunction(Operator.Get());
+				if (PostExecuteFunction)
+				{
+					PostExecuteFunction(Operator.Get());
+				}
 			}
 
 			void Reset(const IOperator::FResetParams& InParams)
 			{
-				// ResetFunction is always non-null. If the contained operator
-				// does not provide a valid initialization function, then a No-Op function
-				// is utilized. 
-				ResetFunction(Operator.Get(), InParams);
+				if (ResetFunction)
+				{
+					ResetFunction(Operator.Get(), InParams);
+				}
 			}
 
 			bool IsNoOp()
 			{
-				return (ExecuteFunction == &FExecuter::NoOpExecute) && (ResetFunction == &FExecuter::NoOpReset);
+				return (nullptr == ExecuteFunction) && (nullptr == ResetFunction) && (nullptr == PostExecuteFunction);
 			}
 
 			bool IsValid() const
@@ -271,12 +284,17 @@ namespace Metasound
 				return EmptyCollection;
 			}
 
-			virtual FExecuteFunction GetExecuteFunction() override
+			virtual IOperator::FExecuteFunction GetExecuteFunction() override
 			{
 				return ExecuteFunction;
 			}
 
-			virtual FResetFunction GetResetFunction() override
+			virtual IOperator::FPostExecuteFunction GetPostExecuteFunction() override
+			{
+				return PostExecuteFunction;
+			}
+
+			virtual IOperator::FResetFunction GetResetFunction() override
 			{
 				return ResetFunction;
 			}
@@ -284,15 +302,8 @@ namespace Metasound
 		private:
 			FOperatorPtr Operator;
 
-			FExecuteFunction ExecuteFunction;
-			FResetFunction ResetFunction;
-
-			static void NoOpExecute(IOperator*)
-			{
-			}
-
-			static void NoOpReset(IOperator*, const IOperator::FResetParams&)
-			{
-			}
+			IOperator::FExecuteFunction ExecuteFunction = nullptr;
+			IOperator::FPostExecuteFunction PostExecuteFunction = nullptr;
+			IOperator::FResetFunction ResetFunction = nullptr;
 	};
 }
diff --git a/Engine/Plugins/Runtime/Metasound/Source/MetasoundGraphCore/Public/MetasoundGraphOperator.h b/Engine/Plugins/Runtime/Metasound/Source/MetasoundGraphCore/Public/MetasoundGraphOperator.h
--- a/Engine/Plugins/Runtime/Metasound/Source/MetasoundGraphCore/Public/MetasoundGraphOperator.h
+++ b/Engine/Plugins/Runtime/Metasound/Source/MetasoundGraphCore/Public/MetasoundGraphOperator.h
@@ -42,7 +42,10 @@ namespace Metasound
 			// Bind the graph's interface data references to FVertexInterfaceData.
 			virtual void Bind(FVertexInterfaceData& InOutVertexData) const override;
 
+			virtual IOperator::FPostExecuteFunction GetPostExecuteFunction() override;
+
 			void Execute();
+			void PostExecute();
 			void Reset(const FResetParams& InParams);
 
 		private:
@@ -50,6 +53,8 @@ namespace Metasound
 			FGraphOperator& operator=(const FGraphOperator&) = delete;
 			FGraphOperator(const FGraphOperator&) = delete;
 
+			static void StaticPostExecute(IOperator* Operator);
+
 			struct FExecuteEntry
 			{
 				FExecuteEntry(IOperator& InOperator, FExecuteFunction InFunc);
@@ -59,6 +64,15 @@ namespace Metasound
 				FExecuteFunction Function;	
 			};
 
+			struct FPostExecuteEntry
+			{
+				FPostExecuteEntry(IOperator& InOperator, FPostExecuteFunction InFunc);
+				void PostExecute();
+
+				IOperator* Operator;
+				FPostExecuteFunction Function;	
+			};
+
 			struct FResetEntry
 			{
 				FResetEntry(IOperator& InOperator, FResetFunction InFunc);
@@ -69,6 +83,7 @@ namespace Metasound
 			};
 
 			TArray<FExecuteEntry> ExecuteStack;
+			TArray<FPostExecuteEntry> PostExecuteStack;
 			TArray<FResetEntry> ResetStack;
 			TArray<TUniquePtr<IOperator>> ActiveOperators;
 			FVertexInterfaceData VertexData;
diff --git a/Engine/Plugins/Runtime/Metasound/Source/MetasoundGraphCore/Public/MetasoundLiteralNode.h b/Engine/Plugins/Runtime/Metasound/Source/MetasoundGraphCore/Public/MetasoundLiteralNode.h
--- a/Engine/Plugins/Runtime/Metasound/Source/MetasoundGraphCore/Public/MetasoundLiteralNode.h
+++ b/Engine/Plugins/Runtime/Metasound/Source/MetasoundGraphCore/Public/MetasoundLiteralNode.h
@@ -1,10 +1,10 @@
 // Copyright Epic Games, Inc. All Rights Reserved.
 #pragma once
 
-#include "CoreMinimal.h"
 #include "Internationalization/Text.h"
 #include "MetasoundBuilderInterface.h"
 #include "MetasoundExecutableOperator.h"
+#include "MetasoundInputNode.h"
 #include "MetasoundNode.h"
 #include "MetasoundNodeInterface.h"
 #include "MetasoundOperatorInterface.h"
@@ -22,7 +22,7 @@ namespace Metasound
 	}
 
 	template<typename DataType>
-	class TLiteralOperator : public IOperator
+	class UE_DEPRECATED(5.3, "The TLiteralOperator will be removed.") TLiteralOperator : public FNoOpOperator
 	{
 		static_assert(!TExecutableDataType<DataType>::bIsExecutable, "TLiteralOperator is only suitable for non-executable data types");
 
@@ -40,30 +40,6 @@ namespace Metasound
 			InVertexData.GetOutputs().BindValueVertex(METASOUND_GET_PARAM_NAME(OutputValue), Value);
 		}
 
-		virtual FDataReferenceCollection GetInputs() const override
-		{
-			FDataReferenceCollection Inputs;
-			return Inputs;
-		}
-
-		virtual FDataReferenceCollection GetOutputs() const override
-		{
-			using namespace LiteralNodeNames;
-			FDataReferenceCollection Outputs;
-			Outputs.AddDataReference(METASOUND_GET_PARAM_NAME(OutputValue), FAnyDataReference{Value});
-			return Outputs;
-		}
-
-		virtual FExecuteFunction GetExecuteFunction() override
-		{
-			return nullptr;
-		}
-
-		virtual FResetFunction GetResetFunction() override
-		{
-			return nullptr;
-		}
-
 	private:
 		TDataValueReference<DataType> Value;
 	};
@@ -74,7 +50,7 @@ namespace Metasound
 	 * Note: this is supporting for FTrigger literals which are deprecated.
 	 */
 	template<typename DataType>
-	class TExecutableLiteralOperator : public TExecutableOperator<TExecutableLiteralOperator<DataType>>
+	class UE_DEPRECATED(5.3, "The TExecutableLiteralOperator will be removed.") TExecutableLiteralOperator : public TExecutableOperator<TExecutableLiteralOperator<DataType>>
 	{
 		static_assert(TExecutableDataType<DataType>::bIsExecutable, "TExecutableLiteralOperator is only suitable for executable data types");
 
@@ -145,6 +121,18 @@ namespace Metasound
 		virtual TUniquePtr<IOperator> CreateOperator(const FBuildOperatorParams& InParams, FBuildResults& OutResults) override;
 
 	private:
+
+		class FLiteralOperator : public TInputOperator<DataType>
+		{
+		public:
+			using TInputOperator<DataType>::TInputOperator;
+			virtual void Bind(FVertexInterfaceData& InVertexData) const override
+			{
+				// Do not bind inputs as there are no inputs exposed from the literal node. 
+				TInputOperator<DataType>::BindOutputs(InVertexData.GetOutputs());
+			}
+		};
+
 		FLiteral InitParam;
 	};
 
@@ -155,6 +143,7 @@ namespace Metasound
 		// Executable data types handle Triggers which need to be advanced every
 		// block.
 		static constexpr bool bIsExecutableDataType = TExecutableDataType<DataType>::bIsExecutable;
+		static constexpr bool bIsPostExecutableDataType = TPostExecutableDataType<DataType>::bIsPostExecutable;
 
 	public:
 
@@ -167,7 +156,7 @@ namespace Metasound
 				, METASOUND_GET_PARAM_DISPLAYNAME(OutputValue) // display name
 			};
 
-			if constexpr (bIsExecutableDataType)
+			if constexpr (bIsExecutableDataType || bIsPostExecutableDataType)
 			{
 				return FVertexInterface(
 					FInputVertexInterface(),
@@ -244,22 +233,8 @@ namespace Metasound
 	template<typename DataType>
 	TUniquePtr<IOperator> TLiteralOperatorLiteralFactory<DataType>::CreateOperator(const FBuildOperatorParams& InParams, FBuildResults& OutResults)
 	{
-		using FLiteralNodeType = TLiteralNode<DataType>;
-
-		static constexpr bool bIsExecutableDataType = TExecutableDataType<DataType>::bIsExecutable;
-
-		if constexpr (bIsExecutableDataType)
-		{
-			// Create write reference by calling compatible constructor with literal.
-			TDataWriteReference<DataType> DataRef = TDataWriteReferenceLiteralFactory<DataType>::CreateExplicitArgs(InParams.OperatorSettings, InitParam);
-			return MakeUnique<TExecutableLiteralOperator<DataType>>(InParams.OperatorSettings, InitParam);
-		}
-		else
-		{
-			// Create value reference by calling compatible constructor with literal.
-			TDataValueReference DataRef = TDataValueReferenceLiteralFactory<DataType>::CreateExplicitArgs(InParams.OperatorSettings, InitParam);
-			return MakeUnique<TLiteralOperator<DataType>>(DataRef);
-		}
+		using namespace LiteralNodeNames;
+		return MakeUnique<FLiteralOperator>(METASOUND_GET_PARAM_NAME(OutputValue), InParams.OperatorSettings, InitParam);
 	}
 
 } // namespace Metasound
diff --git a/Engine/Plugins/Runtime/Metasound/Source/MetasoundGraphCore/Public/MetasoundOperatorInterface.h b/Engine/Plugins/Runtime/Metasound/Source/MetasoundGraphCore/Public/MetasoundOperatorInterface.h
--- a/Engine/Plugins/Runtime/Metasound/Source/MetasoundGraphCore/Public/MetasoundOperatorInterface.h
+++ b/Engine/Plugins/Runtime/Metasound/Source/MetasoundGraphCore/Public/MetasoundOperatorInterface.h
@@ -90,5 +90,21 @@ namespace Metasound
 		 * executed.
 		 */
 		virtual FExecuteFunction GetExecuteFunction() = 0;
+
+		/** Pointer to post execute function for an operator.
+		 *
+		 * @param IOperator* - The operator associated with the function pointer.
+		 */
+		using FPostExecuteFunction = void(*)(IOperator*);
+
+		/** Return the FPostExecute function to call during graph post execution.
+		 *
+		 * The IOperator* argument to the FPostExecutionFunction will be the same IOperator instance
+		 * which returned the execution function.
+		 *
+		 * nullptr return values are valid and signal an IOperator which does not need to be
+		 * post executed.
+		 */
+		virtual FPostExecuteFunction GetPostExecuteFunction() = 0;
 	};
 }
diff --git a/Engine/Plugins/Runtime/Metasound/Source/MetasoundGraphCore/Public/MetasoundVertex.h b/Engine/Plugins/Runtime/Metasound/Source/MetasoundGraphCore/Public/MetasoundVertex.h
--- a/Engine/Plugins/Runtime/Metasound/Source/MetasoundGraphCore/Public/MetasoundVertex.h
+++ b/Engine/Plugins/Runtime/Metasound/Source/MetasoundGraphCore/Public/MetasoundVertex.h
@@ -100,88 +100,6 @@ namespace Metasound
 	/** FInputDataVertex */
 	class FInputDataVertex : public FDataVertex
 	{
-		// FLiteralFactory supports creation of an FLiteral for a vertex given a
-		// copyable literal value type.
-		//
-		// An FLiteral cannot be assigned or copy-constructed because it can contain
-		// a TUniquePtr<> which does not support the assignment operator and copy constructor.
-		//
-		// For factory can create a literal given a copyable type. The factory can
-		// also be copied.
-		struct FLiteralFactory
-		{
-			struct ILiteralFactoryImpl
-			{
-				virtual ~ILiteralFactoryImpl() = default;
-				virtual FLiteral CreateLiteral() const = 0;
-				virtual TUniquePtr<ILiteralFactoryImpl> Clone() const = 0;
-			};
-
-			// Create a literal with a copyable type.
-			template<typename LiteralValueType>
-			struct TLiteralFactoryImpl : ILiteralFactoryImpl
-			{
-				static_assert(std::is_copy_constructible<LiteralValueType>::value && std::is_copy_assignable<LiteralValueType>::value, "Literals can only be assigned for copyable types");
-
-				TLiteralFactoryImpl(const LiteralValueType& InValue)
-				: LiteralValue(InValue)
-				{
-				}
-				
-				FLiteral CreateLiteral() const override
-				{
-					return FLiteral(LiteralValue);
-				}
-
-				TUniquePtr<ILiteralFactoryImpl> Clone() const override
-				{
-					return MakeUnique<TLiteralFactoryImpl<LiteralValueType>>(*this);
-				}
-			private:
-				LiteralValueType LiteralValue;
-			};
-
-			FLiteralFactory()
-			: FactoryImpl(nullptr)
-			{
-			}
-
-			template<typename ValueType>
-			FLiteralFactory(const ValueType& InValue)
-			: FactoryImpl(MakeUnique<TLiteralFactoryImpl<ValueType>>(InValue))
-			{
-			}
-
-			FLiteralFactory(const FLiteralFactory& InOther)
-			{
-				*this = InOther;
-			}
-
-			FLiteralFactory& operator=(const FLiteralFactory& InOther)
-			{
-				FactoryImpl.Reset();
-				if (InOther.FactoryImpl.IsValid())
-				{
-					FactoryImpl = InOther.FactoryImpl->Clone();
-				}
-				return *this;
-			}
-
-			FLiteral CreateLiteral() const
-			{
-				if (FactoryImpl.IsValid())
-				{
-					return FactoryImpl->CreateLiteral();
-				}
-
-				return FLiteral::CreateInvalid();
-			}
-
-		private:
-
-			TUniquePtr<ILiteralFactoryImpl> FactoryImpl;
-		};
-
 	public:
 
 		FInputDataVertex() = default;
@@ -189,31 +107,28 @@ namespace Metasound
 		/** Construct an FInputDataVertex. */
 		FInputDataVertex(const FVertexName& InVertexName, const FName& InDataTypeName, const FDataVertexMetadata& InMetadata, EVertexAccessType InAccessType=EVertexAccessType::Reference)
 			: FDataVertex(InVertexName, InDataTypeName, InMetadata, InAccessType)
-			, LiteralFactory(FLiteral::FNone{})
+			, Literal(FLiteral::FNone{})
 		{
 		}
 
 		/** Construct an FInputDataVertex with a default literal. */
 		template<typename LiteralValueType>
-		UE_DEPRECATED(5.1, "Use constructor with EVertexAccessType argument.")
-		FInputDataVertex(const FVertexName& InVertexName, const FName& InDataTypeName, const FDataVertexMetadata& InMetadata, const LiteralValueType& InLiteralValue)
-			: FDataVertex(InVertexName, InDataTypeName, InMetadata, EVertexAccessType::Reference)
-			, LiteralFactory(InLiteralValue)
+		FInputDataVertex(const FVertexName& InVertexName, const FName& InDataTypeName, const FDataVertexMetadata& InMetadata, EVertexAccessType InAccessType, const LiteralValueType& InLiteralValue)
+			: FDataVertex(InVertexName, InDataTypeName, InMetadata, InAccessType)
+			, Literal(InLiteralValue)
 		{
 		}
 
-		/** Construct an FInputDataVertex with a default literal. */
-		template<typename LiteralValueType>
-		FInputDataVertex(const FVertexName& InVertexName, const FName& InDataTypeName, const FDataVertexMetadata& InMetadata, EVertexAccessType InAccessType, const LiteralValueType& InLiteralValue)
+		FInputDataVertex(const FVertexName& InVertexName, const FName& InDataTypeName, const FDataVertexMetadata& InMetadata, EVertexAccessType InAccessType, const FLiteral& InLiteral)
 			: FDataVertex(InVertexName, InDataTypeName, InMetadata, InAccessType)
-			, LiteralFactory(InLiteralValue)
+			, Literal(InLiteral)
 		{
 		}
 
 		/** Returns the default literal associated with this input. */
 		FLiteral GetDefaultLiteral() const 
 		{
-			return LiteralFactory.CreateLiteral();
+			return Literal;
 		}
 
 		friend bool METASOUNDGRAPHCORE_API operator==(const FInputDataVertex& InLHS, const FInputDataVertex& InRHS);
@@ -222,7 +137,7 @@ namespace Metasound
 
 	private:
 
-		FLiteralFactory LiteralFactory;
+		FLiteral Literal;
 	};
 
 	/** Create a FInputDataVertex with a templated MetaSound data type. */
diff --git a/Engine/Plugins/Runtime/Metasound/Source/MetasoundGraphCore/Public/MetasoundVertexData.h b/Engine/Plugins/Runtime/Metasound/Source/MetasoundGraphCore/Public/MetasoundVertexData.h
--- a/Engine/Plugins/Runtime/Metasound/Source/MetasoundGraphCore/Public/MetasoundVertexData.h
+++ b/Engine/Plugins/Runtime/Metasound/Source/MetasoundGraphCore/Public/MetasoundVertexData.h
@@ -77,7 +77,7 @@ namespace Metasound
 			{
 				if (Data.IsSet())
 				{
-					Data->GetAccessType();
+					return Data->GetAccessType();
 				}
 				return EDataReferenceAccessType::None;
 			}
@@ -633,7 +633,7 @@ namespace Metasound
 
 
 	/** A vertex interface with optionally bound data. */
-	class FVertexInterfaceData
+	class METASOUNDGRAPHCORE_API FVertexInterfaceData
 	{
 	public:
 
diff --git a/Engine/Plugins/Runtime/Metasound/Source/MetasoundStandardNodes/Private/MetasoundAudioFormats.cpp b/Engine/Plugins/Runtime/Metasound/Source/MetasoundStandardNodes/Private/MetasoundAudioFormats.cpp
--- a/Engine/Plugins/Runtime/Metasound/Source/MetasoundStandardNodes/Private/MetasoundAudioFormats.cpp
+++ b/Engine/Plugins/Runtime/Metasound/Source/MetasoundStandardNodes/Private/MetasoundAudioFormats.cpp
@@ -98,7 +98,7 @@ namespace Metasound
 	class METASOUNDSTANDARDNODES_API TOutputNode<FStereoAudioFormat> : public FNode
 	{
 		// FOutputOperator primarly used to report inputs and outputs. Has no execute function.
-		class FOutputOperator : public IOperator
+		class FOutputOperator : public FNoOpOperator
 		{
 		public:
 			FOutputOperator(const FVertexName& InOutputName, TDataReadReference<FAudioBuffer> InLeft, TDataReadReference<FAudioBuffer> InRight, TDataReadReference<FStereoAudioFormat> InStereo)
@@ -132,16 +132,6 @@ namespace Metasound
 				return Outputs;
 			}
 
-			virtual FExecuteFunction GetExecuteFunction() override
-			{
-				return nullptr;
-			}
-
-			virtual FResetFunction GetResetFunction() override
-			{
-				return nullptr;
-			}
-
 		private:
 
 			FVertexName OutputName;
@@ -247,7 +237,7 @@ namespace Metasound
 	class METASOUNDSTANDARDNODES_API TInputNode<FStereoAudioFormat> : public FNode
 	{
 		// Noop operator. Used to return inputs / outputs for debugging.
-		class FInputOperator : public IOperator
+		class FInputOperator : public FNoOpOperator
 		{
 		public:
 
@@ -282,16 +272,6 @@ namespace Metasound
 				return Outputs;
 			}
 
-			virtual FExecuteFunction GetExecuteFunction() override
-			{
-				return nullptr;
-			}
-
-			virtual FResetFunction GetResetFunction() override
-			{
-				return nullptr;
-			}
-
 		private:
 
 			FVertexName InputName;
@@ -396,7 +376,7 @@ namespace Metasound
 	template<>
 	class METASOUNDSTANDARDNODES_API TOutputNode<FMonoAudioFormat> : public FNode
 	{
-		class FOutputOperator : public IOperator
+		class FOutputOperator : public FNoOpOperator
 		{
 		public:
 			FOutputOperator(const FVertexName& InOutputName, TDataReadReference<FAudioBuffer> InCenter, TDataReadReference<FMonoAudioFormat> InMono)
@@ -533,7 +513,7 @@ namespace Metasound
 	template<>
 	class METASOUNDSTANDARDNODES_API TInputNode<FMonoAudioFormat> : public FNode
 	{
-		class FInputOperator : public IOperator
+		class FInputOperator : public FNoOpOperator
 		{
 		public:
 
diff --git a/Engine/Plugins/Runtime/Metasound/Source/MetasoundStandardNodes/Private/MetasoundNoiseGenerator.cpp b/Engine/Plugins/Runtime/Metasound/Source/MetasoundStandardNodes/Private/MetasoundNoiseGenerator.cpp
--- a/Engine/Plugins/Runtime/Metasound/Source/MetasoundStandardNodes/Private/MetasoundNoiseGenerator.cpp
+++ b/Engine/Plugins/Runtime/Metasound/Source/MetasoundStandardNodes/Private/MetasoundNoiseGenerator.cpp
@@ -54,6 +54,7 @@ namespace Metasound
 		virtual FDataReferenceCollection GetInputs() const override;
 		virtual FDataReferenceCollection GetOutputs() const override;
 
+		virtual IOperator::FPostExecuteFunction GetPostExecuteFunction() override { return nullptr; }
 
 	protected:
 		FInt32ReadRef Seed;
diff --git a/Engine/Source/Runtime/AudioExtensions/Public/IAudioParameterTransmitter.h b/Engine/Source/Runtime/AudioExtensions/Public/IAudioParameterTransmitter.h
--- a/Engine/Source/Runtime/AudioExtensions/Public/IAudioParameterTransmitter.h
+++ b/Engine/Source/Runtime/AudioExtensions/Public/IAudioParameterTransmitter.h
@@ -14,6 +14,8 @@ class UObject;
 
 namespace Audio
 {
+	using DeviceID = uint32;
+
 	/** Data passed to CreateParameterTransmitter. */
 	struct AUDIOEXTENSIONS_API FParameterTransmitterInitParams
 	{
@@ -24,6 +26,10 @@ namespace Audio
 		float SampleRate = 0.0f;
 
 		TArray<FAudioParameter> DefaultParams;
+
+		// Audio device ID
+		DeviceID AudioDeviceID = INDEX_NONE;
+
 	};
 
 	// Reference collector functionality for legacy parameter system
diff --git a/Engine/Source/Runtime/Engine/Private/AudioDevice.cpp b/Engine/Source/Runtime/Engine/Private/AudioDevice.cpp
--- a/Engine/Source/Runtime/Engine/Private/AudioDevice.cpp
+++ b/Engine/Source/Runtime/Engine/Private/AudioDevice.cpp
@@ -5168,7 +5168,8 @@ void FAudioDevice::AddNewActiveSoundInternal(const FActiveSound& InNewActiveSoun
 			{
 				OutActiveSound.GetAudioComponentID(),
 				GetSampleRate(),
-				MoveTemp(DefaultParams)
+				MoveTemp(DefaultParams),
+				DeviceID
 			};
 
 			OutActiveSound.InstanceTransmitter = Sound->CreateParameterTransmitter(MoveTemp(TransmitterInitParams));
diff --git a/Engine/Source/Runtime/Engine/Private/SoundNodeWavePlayer.cpp b/Engine/Source/Runtime/Engine/Private/SoundNodeWavePlayer.cpp
--- a/Engine/Source/Runtime/Engine/Private/SoundNodeWavePlayer.cpp
+++ b/Engine/Source/Runtime/Engine/Private/SoundNodeWavePlayer.cpp
@@ -148,6 +148,7 @@ void USoundNodeWavePlayer::ParseNodes( FAudioDevice* AudioDevice, const UPTRINT
 			Params.DefaultParams = ActiveSound.GetTransmitter()->GetParameters();
 			Params.InstanceID = ActiveSound.GetAudioComponentID();
 			Params.SampleRate = AudioDevice->GetSampleRate();
+			Params.AudioDeviceID = AudioDevice->DeviceID;
 
 			if (FSoundCueParameterTransmitter* SoundCueTransmitter = static_cast<FSoundCueParameterTransmitter*>(ActiveSound.GetTransmitter()))
 			{
@@ -291,4 +292,4 @@ int32 USoundNodeWavePlayer::GetMaxChildNodes() const
 }
 
 
-#undef LOCTEXT_NAMESPACE
\ No newline at end of file
+#undef LOCTEXT_NAMESPACE
