[Metasound] Fix races in UMetasoundOutputSubsystem

The generator created/destroyed events get called on the audio thread, so avoid manipulating containers used on the game thread in the event.

#rnx

diff --git a/Engine/Plugins/Runtime/Metasound/Source/MetasoundEngine/Private/MetasoundEngineModule.cpp b/Engine/Plugins/Runtime/Metasound/Source/MetasoundEngine/Private/MetasoundEngineModule.cpp
--- a/Engine/Plugins/Runtime/Metasound/Source/MetasoundEngine/Private/MetasoundEngineModule.cpp
+++ b/Engine/Plugins/Runtime/Metasound/Source/MetasoundEngine/Private/MetasoundEngineModule.cpp
@@ -73,16 +73,16 @@ class FMetasoundEngineModule : public IMetasoundEngineModule
 		METASOUND_PRIVATE_REGISTER_GENERATOR_OUTPUT_WATCHER_TYPE_OPERATIONS(FString, "")
 
 		// Register passthrough output analyzers
-		UMetaSoundOutputSubsystem::RegisterPassthroughAnalyzerForType(
+		UMetasoundGeneratorHandle::RegisterPassthroughAnalyzerForType(
 			GetMetasoundDataTypeName<float>(),
 			Frontend::FVertexAnalyzerForwardFloat::GetAnalyzerName());
-		UMetaSoundOutputSubsystem::RegisterPassthroughAnalyzerForType(
+		UMetasoundGeneratorHandle::RegisterPassthroughAnalyzerForType(
 			GetMetasoundDataTypeName<int32>(),
 			Frontend::FVertexAnalyzerForwardInt::GetAnalyzerName());
-		UMetaSoundOutputSubsystem::RegisterPassthroughAnalyzerForType(
+		UMetasoundGeneratorHandle::RegisterPassthroughAnalyzerForType(
 			GetMetasoundDataTypeName<bool>(),
 			Frontend::FVertexAnalyzerForwardBool::GetAnalyzerName());
-		UMetaSoundOutputSubsystem::RegisterPassthroughAnalyzerForType(
+		UMetasoundGeneratorHandle::RegisterPassthroughAnalyzerForType(
 			GetMetasoundDataTypeName<FString>(),
 			Frontend::FVertexAnalyzerForwardString::GetAnalyzerName());
 
diff --git a/Engine/Plugins/Runtime/Metasound/Source/MetasoundEngine/Private/MetasoundGeneratorHandle.cpp b/Engine/Plugins/Runtime/Metasound/Source/MetasoundEngine/Private/MetasoundGeneratorHandle.cpp
--- a/Engine/Plugins/Runtime/Metasound/Source/MetasoundEngine/Private/MetasoundGeneratorHandle.cpp
+++ b/Engine/Plugins/Runtime/Metasound/Source/MetasoundEngine/Private/MetasoundGeneratorHandle.cpp
@@ -9,8 +9,13 @@
 #include "MetasoundTrace.h"
 #include "MetasoundParameterPack.h"
 
+TMap<FName, FName> UMetasoundGeneratorHandle::PassthroughAnalyzers{};
+
 UMetasoundGeneratorHandle* UMetasoundGeneratorHandle::CreateMetaSoundGeneratorHandle(UAudioComponent* OnComponent)
 {
+	METASOUND_LLM_SCOPE;
+	METASOUND_TRACE_CPUPROFILER_EVENT_SCOPE(UMetasoundGeneratorHandle::CreateMetaSoundGeneratorHandle);
+	
 	if (!OnComponent)
 	{
 		return nullptr;
@@ -21,11 +26,31 @@ UMetasoundGeneratorHandle* UMetasoundGeneratorHandle::CreateMetaSoundGeneratorHa
 	return Result;
 }
 
+void UMetasoundGeneratorHandle::BeginDestroy()
+{
+	Super::BeginDestroy();
+	DetachGeneratorDelegates();
+}
+
+bool UMetasoundGeneratorHandle::IsValid() const
+{
+	return AudioComponent.IsValid();
+}
+
+uint64 UMetasoundGeneratorHandle::GetAudioComponentId() const
+{
+	if (AudioComponent.IsValid())
+	{
+		return AudioComponent->GetAudioComponentID();
+	}
+
+	return INDEX_NONE;
+}
+
 void UMetasoundGeneratorHandle::ClearCachedData()
 {
 	DetachGeneratorDelegates();
 	AudioComponent        = nullptr;
-	AudioComponentId      = 0;
 	CachedMetasoundSource = nullptr;
 	CachedGeneratorPtr    = nullptr;
 	CachedParameterPack   = nullptr;
@@ -37,13 +62,12 @@ void UMetasoundGeneratorHandle::SetAudioComponent(UAudioComponent* InAudioCompon
 	{
 		ClearCachedData();
 		AudioComponent   = InAudioComponent;
-		AudioComponentId = InAudioComponent->GetAudioComponentID();
 	}
 }
 
 void UMetasoundGeneratorHandle::CacheMetasoundSource()
 {
-	if (!AudioComponent)
+	if (!AudioComponent.IsValid())
 	{
 		return;
 	}
@@ -58,7 +82,7 @@ void UMetasoundGeneratorHandle::CacheMetasoundSource()
 	CachedGeneratorPtr    = nullptr;
 	CachedMetasoundSource = CurrentMetasoundSource;
 
-	if (CachedMetasoundSource)
+	if (CachedMetasoundSource.IsValid())
 	{
 		AttachGeneratorDelegates();
 	}
@@ -156,7 +180,7 @@ void UMetasoundGeneratorHandle::DetachGeneratorDelegates()
 {
 	// First detach any callbacks that tell us when a generator is created or destroyed
 	// for the UMetasoundSource of interest...
-	if (CachedMetasoundSource)
+	if (CachedMetasoundSource.IsValid())
 	{
 	CachedMetasoundSource->OnGeneratorInstanceCreated.Remove(GeneratorCreatedDelegateHandle);
 	GeneratorCreatedDelegateHandle.Reset();
@@ -176,20 +200,24 @@ void UMetasoundGeneratorHandle::DetachGeneratorDelegates()
 TSharedPtr<Metasound::FMetasoundGenerator> UMetasoundGeneratorHandle::PinGenerator()
 {
 	TSharedPtr<Metasound::FMetasoundGenerator> PinnedGenerator = CachedGeneratorPtr.Pin();
-	if (PinnedGenerator.IsValid() || !CachedMetasoundSource)
+	if (PinnedGenerator.IsValid() || !CachedMetasoundSource.IsValid())
 	{
 		return PinnedGenerator;
 	}
 
 	// The first attempt to pin failed, so reach out to the MetaSoundSource and see if it has a 
 	// generator for our AudioComponent...
-	CachedGeneratorPtr = CachedMetasoundSource->GetGeneratorForAudioComponent(AudioComponentId);
+	check(AudioComponent.IsValid()); // expect the audio component to still be valid if the generator is.
+	CachedGeneratorPtr = CachedMetasoundSource->GetGeneratorForAudioComponent(AudioComponent->GetAudioComponentID());
 	PinnedGenerator    = CachedGeneratorPtr.Pin();
 	return PinnedGenerator;
 }
 
 bool UMetasoundGeneratorHandle::ApplyParameterPack(UMetasoundParameterPack* Pack)
 {
+	METASOUND_LLM_SCOPE;
+	METASOUND_TRACE_CPUPROFILER_EVENT_SCOPE(UMetasoundGeneratorHandle::ApplyParameterPack);
+	
 	if (!Pack)
 	{
 		return false;
@@ -246,24 +274,139 @@ bool UMetasoundGeneratorHandle::RemoveGraphSetCallback(const FDelegateHandle& Ha
 	return OnGeneratorsGraphChanged.Remove(Handle);
 }
 
-void UMetasoundGeneratorHandle::BeginDestroy()
+bool UMetasoundGeneratorHandle::WatchOutput(
+	FName OutputName,
+	const FOnMetasoundOutputValueChanged& OnOutputValueChanged,
+	FName AnalyzerName,
+	FName AnalyzerOutputName)
 {
-	Super::BeginDestroy();
-	DetachGeneratorDelegates();
+	METASOUND_LLM_SCOPE;
+	METASOUND_TRACE_CPUPROFILER_EVENT_SCOPE(UMetasoundGeneratorHandle::WatchOutput);
+
+	if (!IsValid())
+	{
+		return false;
+	}
+
+	CacheMetasoundSource();
+
+	if (!CachedMetasoundSource.IsValid())
+	{
+		return false;
+	}
+
+	// Find the node id and type name
+	const Metasound::Frontend::FNodeHandle Node = CachedMetasoundSource->GetRootGraphHandle()->GetOutputNodeWithName(OutputName);
+
+	if (!Node->IsValid())
+	{
+		return false;
+	}
+
+	const FGuid NodeId = Node->GetID();
+
+	// We expect an output node to have exactly one output
+	if (!ensure(Node->GetNumOutputs() == 1))
+	{
+		return false;
+	}
+
+	const FName TypeName = Node->GetOutputs()[0]->GetDataType();
+
+	// Make the analyzer address
+	Metasound::Frontend::FAnalyzerAddress AnalyzerAddress;
+	AnalyzerAddress.DataType = TypeName;
+	AnalyzerAddress.InstanceID = AudioComponent->GetAudioComponentID();
+	AnalyzerAddress.OutputName = OutputName;
+
+	// If no analyzer name was provided, try to find a passthrough analyzer
+	if (AnalyzerName.IsNone())
+	{
+		if (!PassthroughAnalyzers.Contains(TypeName))
+		{
+			return false;
+		}
+
+		AnalyzerName = PassthroughAnalyzers[TypeName];
+		AnalyzerOutputName = "Value";
+	}
+	AnalyzerAddress.AnalyzerName = AnalyzerName;
+	
+	AnalyzerAddress.AnalyzerInstanceID = FGuid::NewGuid();
+	AnalyzerAddress.NodeID = NodeId;
+
+	// if we already have a generator, go ahead and make the analyzer and watcher
+	if (const TSharedPtr<Metasound::FMetasoundGenerator> PinnedGenerator = PinGenerator())
+	{
+		CreateAnalyzerAndWatcher(PinnedGenerator, MoveTemp(AnalyzerAddress));
+	}
+	// otherwise enqueue it for later
+	else
+	{
+		OutputAnalyzersToAdd.Enqueue(MoveTemp(AnalyzerAddress));
+	}
+	
+	// either way, add the delegate to the map
+	OutputListenerMap.FindOrAdd(OutputName).FindOrAdd(AnalyzerOutputName).AddUnique(OnOutputValueChanged);
+
+	return true;
+}
+
+void UMetasoundGeneratorHandle::RegisterPassthroughAnalyzerForType(FName TypeName, FName AnalyzerName)
+{
+	check(!PassthroughAnalyzers.Contains(TypeName));
+	PassthroughAnalyzers.Add(TypeName, AnalyzerName);
+}
+
+void UMetasoundGeneratorHandle::UpdateWatchers()
+{
+	METASOUND_LLM_SCOPE;
+	METASOUND_TRACE_CPUPROFILER_EVENT_SCOPE(UMetasoundGeneratorHandle::UpdateWatchers);
+
+	for (auto& Watcher : OutputWatchers)
+	{
+		Watcher.Update([this](FName OutputName, const FMetaSoundOutput& Output)
+		{
+			if (TMap<FName, FOnOutputValueChangedMulticast>* Map = OutputListenerMap.Find(OutputName))
+			{
+				if (const FOnOutputValueChangedMulticast* Delegate = Map->Find(Output.Name))
+				{
+					Delegate->Broadcast(OutputName, Output);
+				}
+			}
+		});
+	}
 }
 
 void UMetasoundGeneratorHandle::OnSourceCreatedAGenerator(uint64 InAudioComponentId, TSharedPtr<Metasound::FMetasoundGenerator> InGenerator)
 {
 	check(IsInGameThread());
-	if (InAudioComponentId == AudioComponentId)
+
+	if (!AudioComponent.IsValid())
+	{
+		return;
+	}
+	
+	if (InAudioComponentId == AudioComponent->GetAudioComponentID())
 	{
 		CachedGeneratorPtr = InGenerator;
 		if (InGenerator)
 		{
+			// If there is a parameter pack to apply, apply it
 			if (CachedParameterPack)
 			{
 				InGenerator->QueueParameterPack(CachedParameterPack);
 			}
+
+			// If there are analyzers to create, create them
+			{
+				Metasound::Frontend::FAnalyzerAddress Address;
+				while (OutputAnalyzersToAdd.Dequeue(Address))
+				{
+					CreateAnalyzerAndWatcher(InGenerator, MoveTemp(Address));
+				}
+			}
+			
 			OnGeneratorHandleAttached.Broadcast();
 			// If anyone has told us they are interested in being notified when a generator's 
 			// graph has changed go ahead and set that up now...
@@ -275,7 +418,13 @@ void UMetasoundGeneratorHandle::OnSourceCreatedAGenerator(uint64 InAudioComponen
 void UMetasoundGeneratorHandle::OnSourceDestroyedAGenerator(uint64 InAudioComponentId, TSharedPtr<Metasound::FMetasoundGenerator> InGenerator)
 {
 	check(IsInGameThread());
-	if (InAudioComponentId == AudioComponentId)
+
+	if (!AudioComponent.IsValid())
+	{
+		return;
+	}
+	
+	if (InAudioComponentId == AudioComponent->GetAudioComponentID())
 	{
 		if (CachedGeneratorPtr.IsValid())
 		{
@@ -285,4 +434,27 @@ void UMetasoundGeneratorHandle::OnSourceDestroyedAGenerator(uint64 InAudioCompon
 	}
 }
 
+void UMetasoundGeneratorHandle::CreateAnalyzerAndWatcher(
+	const TSharedPtr<Metasound::FMetasoundGenerator> Generator,
+	Metasound::Frontend::FAnalyzerAddress&& AnalyzerAddress)
+{
+	if (!IsValid())
+	{
+		return;
+	}
+
+	// Create the analyzer (will skip if there's already one)
+	Generator->AddOutputVertexAnalyzer(AnalyzerAddress);
+	
+	// Create the watcher
+	if (!OutputWatchers.ContainsByPredicate(
+		[&AnalyzerAddress](const Metasound::Private::FMetasoundOutputWatcher& Watcher)
+		{
+			return AnalyzerAddress.OutputName == Watcher.Name;
+		}))
+	{
+		OutputWatchers.Emplace(MoveTemp(AnalyzerAddress), Generator->OperatorSettings);
+	}
+}
+
 
diff --git a/Engine/Plugins/Runtime/Metasound/Source/MetasoundEngine/Private/MetasoundOutputSubsystem.cpp b/Engine/Plugins/Runtime/Metasound/Source/MetasoundEngine/Private/MetasoundOutputSubsystem.cpp
--- a/Engine/Plugins/Runtime/Metasound/Source/MetasoundEngine/Private/MetasoundOutputSubsystem.cpp
+++ b/Engine/Plugins/Runtime/Metasound/Source/MetasoundEngine/Private/MetasoundOutputSubsystem.cpp
@@ -3,86 +3,51 @@
 #include "MetasoundOutputSubsystem.h"
 
 #include "MetasoundGenerator.h"
+#include "MetasoundGeneratorHandle.h"
 #include "MetasoundSource.h"
 #include "MetasoundTrace.h"
-#include "Analysis/MetasoundFrontendAnalyzerAddress.h"
 #include "Components/AudioComponent.h"
 #include "ProfilingDebugging/CpuProfilerTrace.h"
 
-TMap<FName, FName> UMetaSoundOutputSubsystem::PassthroughAnalyzers{};
-
 bool UMetaSoundOutputSubsystem::WatchOutput(
 	UAudioComponent* AudioComponent,
 	const FName OutputName,
 	const FOnMetasoundOutputValueChanged& OnOutputValueChanged,
-	FName AnalyzerName,
-	FName AnalyzerOutputName)
+	const FName AnalyzerName,
+	const FName AnalyzerOutputName)
 {
 	METASOUND_TRACE_CPUPROFILER_EVENT_SCOPE(UMetasoundOutputSubsystem::WatchOutput);
-	
-	FGeneratorInfo* GeneratorInfo = FindOrAddGeneratorInfo(AudioComponent);
-
-	if (nullptr == GeneratorInfo)
-	{
-		return false;
-	}
-	
-	check(GeneratorInfo->Source.IsValid());
-
-	// Find the node id and type name
-	const Metasound::Frontend::FNodeHandle Node = GeneratorInfo->Source->GetRootGraphHandle()->GetOutputNodeWithName(OutputName);
 
-	if (!Node->IsValid())
+	if (nullptr == AudioComponent)
 	{
 		return false;
 	}
 
-	const FGuid NodeId = Node->GetID();
+	UMetasoundGeneratorHandle* Handle = nullptr;
 
-	// We expect an output node to only have one output
-	if (Node->GetNumOutputs() != 1)
+	// Try to find an existing handle
+	const uint64 AudioComponentId = AudioComponent->GetAudioComponentID();
+	if (const TObjectPtr<UMetasoundGeneratorHandle>* FoundHandle = TrackedGenerators.FindByPredicate(
+		[AudioComponentId](const TObjectPtr<UMetasoundGeneratorHandle> ExistingHandle)
+		{
+			return ExistingHandle->IsValid() && ExistingHandle->GetAudioComponentId() == AudioComponentId;
+		}))
 	{
-		return false;
+		Handle = *FoundHandle;
 	}
-
-	const FName TypeName = Node->GetOutputs()[0]->GetDataType();
-
-	// Make the analyzer address
-	Metasound::Frontend::FAnalyzerAddress AnalyzerAddress;
-	AnalyzerAddress.DataType = TypeName;
-	AnalyzerAddress.InstanceID = AudioComponent->GetAudioComponentID();
-	AnalyzerAddress.OutputName = OutputName;
-
-	if (AnalyzerName.IsNone())
+	// Create a new one
+	else
 	{
-		if (!PassthroughAnalyzers.Contains(TypeName))
-		{
-			return false;
-		}
-
-		AnalyzerName = PassthroughAnalyzers[TypeName];
-		AnalyzerOutputName = "Value";
+		Handle = UMetasoundGeneratorHandle::CreateMetaSoundGeneratorHandle(AudioComponent);
+		TrackedGenerators.Add(Handle);
 	}
-	AnalyzerAddress.AnalyzerName = AnalyzerName;
-	
-	AnalyzerAddress.AnalyzerInstanceID = FGuid::NewGuid();
-	AnalyzerAddress.NodeID = NodeId;
 
-	// if we already have a generator, go ahead and make the analyzer and watcher
-	if (const TSharedPtr<Metasound::FMetasoundGenerator, ESPMode::ThreadSafe> PinnedGenerator = GeneratorInfo->Generator.Pin())
+	if (nullptr == Handle)
 	{
-		CreateAnalyzerAndWatcher(PinnedGenerator, MoveTemp(AnalyzerAddress));
-	}
-	// otherwise enqueue it for later
-	else
-	{
-		GeneratorInfo->AnalyzersToCreate.Enqueue(MoveTemp(AnalyzerAddress));
+		return false;
 	}
 
-	// either way, add the delegate to the map
-	GeneratorInfo->OutputChangedDelegates.FindOrAdd(OutputName).FindOrAdd(AnalyzerOutputName).AddUnique(OnOutputValueChanged);
-	
-	return true;
+	return Handle->WatchOutput(OutputName, OnOutputValueChanged, AnalyzerName, AnalyzerOutputName);
 }
 
 bool UMetaSoundOutputSubsystem::IsTickable() const
@@ -93,22 +58,16 @@ bool UMetaSoundOutputSubsystem::IsTickable() const
 void UMetaSoundOutputSubsystem::Tick(float DeltaTime)
 {
 	METASOUND_TRACE_CPUPROFILER_EVENT_SCOPE(UMetasoundOutputSubsystem::Tick);
-	 
-	for (auto GeneratorIt = TrackedGenerators.CreateIterator(); GeneratorIt; ++GeneratorIt)
+
+	for (auto It = TrackedGenerators.CreateIterator(); It; ++It)
 	{
-		if (!GeneratorIt->IsValid())
+		if (!(*It)->IsValid())
 		{
-			GeneratorIt.RemoveCurrent();
+			It.RemoveCurrent();
 		}
 		else
 		{
-			for (Metasound::Private::FMetasoundOutputWatcher& Watcher : GeneratorIt->OutputWatchers)
-			{
-				Watcher.Update([&GeneratorIt](FName OutputName, const FMetaSoundOutput& Output)
-				{
-					GeneratorIt->HandleOutputChanged(OutputName, Output);
-				});
-			}
+			(*It)->UpdateWatchers();
 		}
 	}
 }
@@ -117,160 +76,3 @@ TStatId UMetaSoundOutputSubsystem::GetStatId() const
 {
 	RETURN_QUICK_DECLARE_CYCLE_STAT(UMetasoundGeneratorAccessSubsystem, STATGROUP_Tickables);
 }
-
-void UMetaSoundOutputSubsystem::RegisterPassthroughAnalyzerForType(const FName TypeName, const FName AnalyzerName)
-{
-	check(!PassthroughAnalyzers.Contains(TypeName));
-	PassthroughAnalyzers.Add(TypeName, AnalyzerName);
-}
-
-void UMetaSoundOutputSubsystem::FGeneratorInfo::HandleOutputChanged(FName OutputName, const FMetaSoundOutput& Output)
-{
-	if (TMap<FName, FOnOutputValueChangedMulticast>* Map = OutputChangedDelegates.Find(OutputName))
-	{
-		if (const FOnOutputValueChangedMulticast* Delegate = Map->Find(Output.Name))
-		{
-			Delegate->Broadcast(Output);
-		}
-	}
-}
-
-bool UMetaSoundOutputSubsystem::FGeneratorInfo::IsValid() const
-{
-	return AudioComponent.IsValid() && Source.IsValid();
-}
-
-UMetaSoundOutputSubsystem::FGeneratorInfo* UMetaSoundOutputSubsystem::FindOrAddGeneratorInfo(UAudioComponent* AudioComponent)
-{
-	if (!ensure(nullptr != AudioComponent))
-	{
-		return  nullptr;
-	}
-	
-	const int64 AudioComponentId = AudioComponent->GetAudioComponentID();
-
-	// already in the list...
-	if (FGeneratorInfo* Info = TrackedGenerators.FindByPredicate([AudioComponentId](const FGeneratorInfo& GeneratorInfo)
-	{
-		return GeneratorInfo.AudioComponent.IsValid()
-		&& GeneratorInfo.AudioComponent->GetAudioComponentID() == AudioComponentId;
-	}))
-	{
-		if (Info->Source.IsValid())
-		{
-			return Info;
-		}
-		// Source is no longer valid, remove info
-		OnGeneratorDestroyed(AudioComponentId, Info->Generator.Pin());
-	}
-
-	UMetaSoundSource* Source = Cast<UMetaSoundSource>(AudioComponent->GetSound());
-	
-	if (nullptr == Source)
-	{
-		return nullptr;
-	}
-
-	// set up a new entry
-	TrackedGenerators.AddDefaulted();
-	FGeneratorInfo& GeneratorInfo = TrackedGenerators.Last();
-	GeneratorInfo.AudioComponent = AudioComponent;
-	GeneratorInfo.Source = Source;
-
-	// attempt to get the generator
-	GeneratorInfo.Generator = GeneratorInfo.Source->GetGeneratorForAudioComponent(AudioComponentId);
-
-	// if that didn't work, set up the generator creation callback
-	if (nullptr == GeneratorInfo.Generator)
-	{
-		GeneratorInfo.OnCreatedHandle = GeneratorInfo.Source->OnGeneratorInstanceCreated.AddUObject(
-			this,
-			&UMetaSoundOutputSubsystem::OnGeneratorCreated);
-	}
-	
-	// always set up the generator destruction callback
-	GeneratorInfo.OnDestroyedHandle = GeneratorInfo.Source->OnGeneratorInstanceDestroyed.AddUObject(
-		this,
-		&UMetaSoundOutputSubsystem::OnGeneratorDestroyed);
-
-	return &GeneratorInfo;
-}
-
-void UMetaSoundOutputSubsystem::OnGeneratorCreated(
-	const uint64 InAudioComponentId,
-	const TSharedPtr<Metasound::FMetasoundGenerator> Generator)
-{
-	for (auto GeneratorIt = TrackedGenerators.CreateIterator(); GeneratorIt; ++GeneratorIt)
-	{
-		if (!GeneratorIt->IsValid())
-		{
-			GeneratorIt.RemoveCurrent();
-		}
-		else if (GeneratorIt->AudioComponent->GetAudioComponentID() == InAudioComponentId)
-		{
-			GeneratorIt->Generator = Generator;
-			GeneratorIt->Source->OnGeneratorInstanceCreated.Remove(GeneratorIt->OnCreatedHandle);
-
-			// create analyzers
-			Metasound::Frontend::FAnalyzerAddress Address;
-			while (GeneratorIt->AnalyzersToCreate.Dequeue(Address))
-			{
-				CreateAnalyzerAndWatcher(Generator, MoveTemp(Address));
-			}
-		}
-	}
-}
-
-void UMetaSoundOutputSubsystem::OnGeneratorDestroyed(
-	const uint64 InAudioComponentId,
-	TSharedPtr<Metasound::FMetasoundGenerator>)
-{
-	for (auto GeneratorIt = TrackedGenerators.CreateIterator(); GeneratorIt; ++GeneratorIt)
-	{
-		if (!GeneratorIt->IsValid())
-		{
-			GeneratorIt.RemoveCurrent();
-		}
-		else if (GeneratorIt->AudioComponent.IsValid() && GeneratorIt->AudioComponent->GetAudioComponentID() == InAudioComponentId)
-		{
-			GeneratorIt->Source->OnGeneratorInstanceDestroyed.Remove(GeneratorIt->OnDestroyedHandle);
-			GeneratorIt.RemoveCurrent();
-			return;
-		}
-	}
-}
-
-void UMetaSoundOutputSubsystem::CreateAnalyzerAndWatcher(
-	const TSharedPtr<Metasound::FMetasoundGenerator>& Generator,
-	Metasound::Frontend::FAnalyzerAddress&& AnalyzerAddress)
-{
-	if (!Generator.IsValid())
-	{
-		return;
-	}
-	
-	// Find the generator info and create the watcher
-	{
-		FGeneratorInfo* GeneratorInfo = TrackedGenerators.FindByPredicate([&Generator](const FGeneratorInfo& Info)
-		{
-			return Info.Generator.HasSameObject(Generator.Get());
-		});
-
-		if (nullptr == GeneratorInfo)
-		{
-			return;
-		}
-
-		if (!GeneratorInfo->OutputWatchers.ContainsByPredicate(
-			[&AnalyzerAddress](const Metasound::Private::FMetasoundOutputWatcher& Watcher)
-			{
-				return AnalyzerAddress.OutputName == Watcher.Name;
-			}))
-		{
-			GeneratorInfo->OutputWatchers.Emplace(MoveTemp(AnalyzerAddress), Generator->OperatorSettings);
-		}
-	}
-
-	// Create the analyzer (will skip if there's already one)
-	Generator->AddOutputVertexAnalyzer(AnalyzerAddress);
-}
diff --git a/Engine/Plugins/Runtime/Metasound/Source/MetasoundEngine/Public/MetasoundGeneratorHandle.h b/Engine/Plugins/Runtime/Metasound/Source/MetasoundEngine/Public/MetasoundGeneratorHandle.h
--- a/Engine/Plugins/Runtime/Metasound/Source/MetasoundEngine/Public/MetasoundGeneratorHandle.h
+++ b/Engine/Plugins/Runtime/Metasound/Source/MetasoundEngine/Public/MetasoundGeneratorHandle.h
@@ -1,8 +1,11 @@
 // Copyright Epic Games, Inc. All Rights Reserved.
 
 #pragma once
-#include "Templates/SharedPointer.h"
+
+#include "MetasoundOutputWatcher.h"
 #include "MetasoundParameterPack.h"
+#include "Containers/SpscQueue.h"
+#include "Templates/SharedPointer.h"
 
 #include "MetasoundGeneratorHandle.generated.h"
 
@@ -14,6 +17,8 @@ namespace Metasound
 	class FMetasoundGenerator;
 }
 
+DECLARE_DYNAMIC_DELEGATE_TwoParams(FOnMetasoundOutputValueChanged, FName, OutputName, const FMetaSoundOutput&, Output);
+
 UCLASS(BlueprintType,Category="MetaSound")
 class METASOUNDENGINE_API UMetasoundGeneratorHandle : public UObject
 {
@@ -24,7 +29,17 @@ public:
 	UFUNCTION(BlueprintCallable, Category="MetaSound")
 	static UMetasoundGeneratorHandle* CreateMetaSoundGeneratorHandle(UAudioComponent* OnComponent);
 
-	void BeginDestroy() override;
+	virtual void BeginDestroy() override;
+
+	bool IsValid() const;
+
+	/**
+	 * Get the id for the UAudioComponent associated with this handle.
+	 * NOTE: Be sure to check IsValid() before expecting a valid return from this method.
+	 *
+	 * @returns The audio component's id, or INDEX_NONE if the component is no longer valid.
+	 */
+	uint64 GetAudioComponentId() const;
 
 	// UMetasoundGeneratorHandle shields its "clients" from "cross thread" issues
 	// related to callbacks coming in the audio control or rendering threads that 
@@ -56,8 +71,36 @@ public:
 	FDelegateHandle AddGraphSetCallback(const UMetasoundGeneratorHandle::FOnSetGraph& Delegate);
 	bool RemoveGraphSetCallback(const FDelegateHandle& Handle);
 
-private:
+	/**
+	 * Watch an output value.
+	 *
+	 * @param OutputName - The user-specified name of the output in the Metasound
+	 * @param OnOutputValueChanged - The event to fire when the output's value changes
+	 * @param AnalyzerName - (optional) The name of the analyzer to use on the output, defaults to a passthrough
+	 * @param AnalyzerOutputName - (optional) The name of the output on the analyzer to watch, defaults to the passthrough output
+	 * @returns true if the watch setup succeeded, false otherwise
+	 */
+	UFUNCTION(BlueprintCallable, Category="MetaSoundOutput", meta=(AdvancedDisplay = "2"))
+	bool WatchOutput(
+		FName OutputName,
+		const FOnMetasoundOutputValueChanged& OnOutputValueChanged,
+		FName AnalyzerName = NAME_None,
+		FName AnalyzerOutputName = NAME_None);
 
+	/**
+	 * Map a type name to a passthrough analyzer name to use as a default for UMetasoundOutputSubsystem::WatchOutput()
+	 *
+	 * @param TypeName - The type name returned from GetMetasoundDataTypeName()
+	 * @param AnalyzerName - The name of the analyzer to use
+	 */
+	static void RegisterPassthroughAnalyzerForType(FName TypeName, FName AnalyzerName);
+
+	/**
+	 * Update any watched outputs
+	 */
+	void UpdateWatchers();
+	
+private:
 	void SetAudioComponent(UAudioComponent* InAudioComponent);
 	void CacheMetasoundSource();
 	void ClearCachedData();
@@ -81,13 +124,9 @@ private:
 	 */
 	void OnSourceCreatedAGenerator(uint64 InAudioComponentId, TSharedPtr<Metasound::FMetasoundGenerator> InGenerator);
 	void OnSourceDestroyedAGenerator(uint64 InAudioComponentId, TSharedPtr<Metasound::FMetasoundGenerator> InGenerator);
-
-
-	UPROPERTY(Transient)
-	TObjectPtr<UAudioComponent> AudioComponent;
-	uint64 AudioComponentId;
-	UPROPERTY(Transient)
-	TObjectPtr<UMetaSoundSource> CachedMetasoundSource;
+	
+	TWeakObjectPtr<UAudioComponent> AudioComponent;
+	TWeakObjectPtr<UMetaSoundSource> CachedMetasoundSource;
 	TWeakPtr<Metasound::FMetasoundGenerator> CachedGeneratorPtr;
 	FSharedMetasoundParameterStoragePtr CachedParameterPack;
 
@@ -102,4 +141,14 @@ private:
 	FDelegateHandle GeneratorCreatedDelegateHandle;
 	FDelegateHandle GeneratorDestroyedDelegateHandle;
 	FDelegateHandle GeneratorGraphChangedDelegateHandle;
+
+	DECLARE_DYNAMIC_MULTICAST_DELEGATE_TwoParams(FOnOutputValueChangedMulticast, FName, Name, const FMetaSoundOutput&, Output);
+	static TMap<FName, FName> PassthroughAnalyzers;
+	// This is a map of output names on the generator to a map of output names on an analyzer
+	// to delegates listening to that analyzer output. We do this because analyzers can have more
+	// than one output.
+	TMap<FName, TMap<FName, FOnOutputValueChangedMulticast>> OutputListenerMap;
+	TSpscQueue<Metasound::Frontend::FAnalyzerAddress> OutputAnalyzersToAdd;
+	void CreateAnalyzerAndWatcher(const TSharedPtr<Metasound::FMetasoundGenerator> Generator, Metasound::Frontend::FAnalyzerAddress&& AnalyzerAddress);
+	TArray<Metasound::Private::FMetasoundOutputWatcher> OutputWatchers;
 };
\ No newline at end of file
diff --git a/Engine/Plugins/Runtime/Metasound/Source/MetasoundEngine/Public/MetasoundOutputSubsystem.h b/Engine/Plugins/Runtime/Metasound/Source/MetasoundEngine/Public/MetasoundOutputSubsystem.h
--- a/Engine/Plugins/Runtime/Metasound/Source/MetasoundEngine/Public/MetasoundOutputSubsystem.h
+++ b/Engine/Plugins/Runtime/Metasound/Source/MetasoundEngine/Public/MetasoundOutputSubsystem.h
@@ -2,28 +2,12 @@
 
 #pragma once
 
-#include "MetasoundOutput.h"
-#include "MetasoundOutputWatcher.h"
-#include "Containers/Queue.h"
-#include "Containers/Map.h"
+#include "MetasoundGeneratorHandle.h"
 #include "Subsystems/WorldSubsystem.h"
 
 #include "MetasoundOutputSubsystem.generated.h"
 
-namespace Metasound
-{
-	namespace Frontend
-	{
-		class FAnalyzerAddress;
-	}
-
-	class FMetasoundGenerator;
-}
-
 class UAudioComponent;
-class UMetaSoundSource;
-
-DECLARE_DYNAMIC_DELEGATE_OneParam(FOnMetasoundOutputValueChanged, const FMetaSoundOutput&, Output);
 
 /**
  * Provides access to a playing Metasound generator's outputs
@@ -60,39 +44,7 @@ public:
 	virtual TStatId GetStatId() const override;
 	/** End UTickableWorldSubsystem */
 
-	/**
-	 * Map a type name to a passthrough analyzer name to use as a default for UMetasoundOutputSubsystem::WatchOutput()
-	 *
-	 * @param TypeName - The type name returned from GetMetasoundDataTypeName()
-	 * @param AnalyzerName - The name of the analyzer to use
-	 */
-	static void RegisterPassthroughAnalyzerForType(FName TypeName, FName AnalyzerName);
-
 private:
-	static TMap<FName, FName> PassthroughAnalyzers;
-
-	DECLARE_DYNAMIC_MULTICAST_DELEGATE_OneParam(FOnOutputValueChangedMulticast, const FMetaSoundOutput&, Output);
-	
-	struct FGeneratorInfo
-	{
-		TWeakObjectPtr<UAudioComponent> AudioComponent;
-		TWeakObjectPtr<UMetaSoundSource> Source;
-		TWeakPtr<Metasound::FMetasoundGenerator, ESPMode::ThreadSafe> Generator;
-		FDelegateHandle OnCreatedHandle;
-		FDelegateHandle OnDestroyedHandle;
-		TQueue<Metasound::Frontend::FAnalyzerAddress> AnalyzersToCreate;
-		TMap<FName, TMap<FName, FOnOutputValueChangedMulticast>> OutputChangedDelegates;
-		TArray<Metasound::Private::FMetasoundOutputWatcher> OutputWatchers;
-		void HandleOutputChanged(FName OutputName, const FMetaSoundOutput& Output);
-		bool IsValid() const;
-	};
-
-	FGeneratorInfo* FindOrAddGeneratorInfo(UAudioComponent* AudioComponent);
-	void OnGeneratorCreated(uint64 InAudioComponentId, TSharedPtr<Metasound::FMetasoundGenerator> Generator);
-	void OnGeneratorDestroyed(uint64 InAudioComponentId, TSharedPtr<Metasound::FMetasoundGenerator> Generator);
-	void CreateAnalyzerAndWatcher(
-		const TSharedPtr<Metasound::FMetasoundGenerator>& Generator,
-		Metasound::Frontend::FAnalyzerAddress&& AnalyzerAddress);
-
-	TArray<FGeneratorInfo> TrackedGenerators;
+	UPROPERTY(Transient)
+	TArray<TObjectPtr<UMetasoundGeneratorHandle>> TrackedGenerators;
 };
diff --git a/Engine/Plugins/Runtime/Metasound/Source/MetasoundEngine/Public/MetasoundOutputWatcher.h b/Engine/Plugins/Runtime/Metasound/Source/MetasoundEngine/Public/MetasoundOutputWatcher.h
--- a/Engine/Plugins/Runtime/Metasound/Source/MetasoundEngine/Public/MetasoundOutputWatcher.h
+++ b/Engine/Plugins/Runtime/Metasound/Source/MetasoundEngine/Public/MetasoundOutputWatcher.h
@@ -25,9 +25,6 @@ namespace Metasound::Private
 		 */
 		void Update(TFunctionRef<void(FName, const FMetaSoundOutput&)> OnOutputChanged);
 
-		using FOutputInitFn = TFunction<void(FMetaSoundOutput&)>;
-		using FOutputUpdateFn = TFunction<bool(FMetasoundOutputWatcher&, FMetaSoundOutput&)>;
-
 		/**
 		 * Stores the typed operations for a specific supported type
 		 */
