Fix for OneShot MetaSound interface not being modifiable in editor due to copy pasta of UClass options
#rnx

diff --git a/Engine/Plugins/Runtime/Metasound/Source/MetasoundFrontend/Private/Interfaces/MetasoundFrontendSourceInterface.cpp b/Engine/Plugins/Runtime/Metasound/Source/MetasoundFrontend/Private/Interfaces/MetasoundFrontendSourceInterface.cpp
--- a/Engine/Plugins/Runtime/Metasound/Source/MetasoundFrontend/Private/Interfaces/MetasoundFrontendSourceInterface.cpp
+++ b/Engine/Plugins/Runtime/Metasound/Source/MetasoundFrontend/Private/Interfaces/MetasoundFrontendSourceInterface.cpp
@@ -14,9 +14,8 @@ namespace Metasound::Frontend
 {
 	namespace SourceInterfacePrivate
 	{
-		TArray<Audio::FParameterInterface::FClassOptions> GetDefaultSourceClassOptions(const FTopLevelAssetPath& InClassPath)
+		TArray<Audio::FParameterInterface::FClassOptions> GetDefaultSourceClassOptions(const FTopLevelAssetPath& InClassPath, bool bIsModifiable)
 		{
-			constexpr bool bIsModifiable = false;
 			constexpr bool bIsDefault = true;
 			return TArray<Audio::FParameterInterface::FClassOptions>
 			{
@@ -46,7 +45,8 @@ namespace Metasound::Frontend
 				FInterface(const FTopLevelAssetPath& InClassPath)
 					: FParameterInterface(SourceOneShotInterface::GetVersion().Name, SourceOneShotInterface::GetVersion().Number.ToInterfaceVersion())
 				{
-					UClassOptions = SourceInterfacePrivate::GetDefaultSourceClassOptions(InClassPath);
+					constexpr bool bIsModifiable = true;
+					UClassOptions = SourceInterfacePrivate::GetDefaultSourceClassOptions(InClassPath, bIsModifiable);
 
 					Outputs =
 					{
@@ -199,7 +199,8 @@ namespace Metasound::Frontend
 				FInterface(const FTopLevelAssetPath& InClassPath)
 					: FParameterInterface(SourceInterface::GetVersion().Name, SourceInterface::GetVersion().Number.ToInterfaceVersion())
 				{
-					UClassOptions = SourceInterfacePrivate::GetDefaultSourceClassOptions(InClassPath);
+					constexpr bool bIsModifiable = false;
+					UClassOptions = SourceInterfacePrivate::GetDefaultSourceClassOptions(InClassPath, bIsModifiable);
 
 					Inputs =
 					{
