Fix class references misguided towards transient MetaSound class instead of MetaSound Builder class causing failures to add/remove interfaces in certain contexts
#rnx

diff --git a/Engine/Plugins/Runtime/Metasound/Source/MetasoundEditor/Private/MetasoundEditor.cpp b/Engine/Plugins/Runtime/Metasound/Source/MetasoundEditor/Private/MetasoundEditor.cpp
--- a/Engine/Plugins/Runtime/Metasound/Source/MetasoundEditor/Private/MetasoundEditor.cpp
+++ b/Engine/Plugins/Runtime/Metasound/Source/MetasoundEditor/Private/MetasoundEditor.cpp
@@ -1769,8 +1769,11 @@ namespace Metasound
 				FMetasoundAssetBase* MetasoundAsset = IMetasoundUObjectRegistry::Get().GetObjectAsAssetBase(Metasound);
 				check(MetasoundAsset);
 
+				const UClass* MetaSoundClass = Metasound->GetClass();
+				check(MetaSoundClass);
+
 				TScriptInterface<IMetaSoundDocumentInterface> DocInterface = MetasoundAsset->GetOwningAsset();
-				FMetaSoundFrontendDocumentBuilder Builder(DocInterface);
+				FMetaSoundFrontendDocumentBuilder Builder(MetaSoundClass->GetClassPathName(), DocInterface);
 				Builder.ConvertFromPreset();
 
 				// Hack until toolbar is polished up & corner text properly dynamically updates
diff --git a/Engine/Plugins/Runtime/Metasound/Source/MetasoundEditor/Private/MetasoundEditorSubsystem.cpp b/Engine/Plugins/Runtime/Metasound/Source/MetasoundEditor/Private/MetasoundEditorSubsystem.cpp
--- a/Engine/Plugins/Runtime/Metasound/Source/MetasoundEditor/Private/MetasoundEditorSubsystem.cpp
+++ b/Engine/Plugins/Runtime/Metasound/Source/MetasoundEditor/Private/MetasoundEditorSubsystem.cpp
@@ -23,7 +23,7 @@ TScriptInterface<IMetaSoundDocumentInterface> UMetaSoundEditorSubsystem::BuildTo
 		constexpr UFactory* Factory = nullptr;
 
 		// Not about to follow this lack of const correctness down a multidecade in the works rabbit hole.
-		UClass& BuilderUClass = const_cast<UClass&>(InBuilder->GetBaseMetaSoundUClass());
+		UClass& BuilderUClass = const_cast<UClass&>(InBuilder->GetBuilderUClass());
 		if (UObject* NewMetaSound = IAssetTools::Get().CreateAsset(AssetName, PackagePath, &BuilderUClass, Factory))
 		{
 			InBuilder->InitNodeLocations();
diff --git a/Engine/Plugins/Runtime/Metasound/Source/MetasoundEngine/Private/MetasoundBuilderSubsystem.cpp b/Engine/Plugins/Runtime/Metasound/Source/MetasoundEngine/Private/MetasoundBuilderSubsystem.cpp
--- a/Engine/Plugins/Runtime/Metasound/Source/MetasoundEngine/Private/MetasoundBuilderSubsystem.cpp
+++ b/Engine/Plugins/Runtime/Metasound/Source/MetasoundEngine/Private/MetasoundBuilderSubsystem.cpp
@@ -478,15 +478,11 @@ void UMetaSoundBuilderBase::CreateTransientDocument()
 {
 	TObjectPtr<UMetaSoundBuilderDocument> NewBuilder = NewObject<UMetaSoundBuilderDocument>();
 	TScriptInterface<IMetaSoundDocumentInterface> MetaSoundSourceDocInterface = NewBuilder;
-	Builder = FMetaSoundFrontendDocumentBuilder(MetaSoundSourceDocInterface);
+	const FTopLevelAssetPath BuilderPath = GetBuilderUClass().GetClassPathName();
+	Builder = FMetaSoundFrontendDocumentBuilder(BuilderPath, MetaSoundSourceDocInterface);
 	Builder.InitDocument();
 }
 
-const UClass& UMetaSoundBuilderBase::GetBaseMetaSoundUClass() const
-{
-	return Builder.GetDocumentInterface().GetBaseMetaSoundUClass();
-}
-
 void UMetaSoundBuilderBase::InitNodeLocations()
 {
 	Builder.InitNodeLocations();
@@ -617,6 +613,11 @@ TScriptInterface<IMetaSoundDocumentInterface> UMetaSoundPatchBuilder::Build(UObj
 	return BuildInternal<UMetaSoundPatch>(Parent, InBuilderOptions);
 }
 
+const UClass& UMetaSoundPatchBuilder::GetBuilderUClass() const
+{
+	return *UMetaSoundPatch::StaticClass();
+}
+
 void UMetaSoundSourceBuilder::Audition(UObject* Parent, UAudioComponent* AudioComponent, FOnCreateAuditionGeneratorHandleDelegate CreateGenerator)
 {
 	using namespace Metasound;
@@ -687,6 +688,11 @@ const Metasound::Engine::FOutputAudioFormatInfoPair* UMetaSoundSourceBuilder::Fi
 	return Algo::FindByPredicate(FormatInfo, Predicate);
 }
 
+const UClass& UMetaSoundSourceBuilder::GetBuilderUClass() const
+{
+	return *UMetaSoundSource::StaticClass();
+}
+
 void UMetaSoundSourceBuilder::SetFormat(EMetaSoundOutputAudioFormat OutputFormat, EMetaSoundBuilderResult& OutResult)
 {
 	using namespace Metasound::Engine;
@@ -724,19 +730,6 @@ void UMetaSoundSourceBuilder::SetFormat(EMetaSoundOutputAudioFormat OutputFormat
 	OutResult = bSuccess ? EMetaSoundBuilderResult::Succeeded : EMetaSoundBuilderResult::Failed;
 }
 
-void UMetaSoundSourceBuilder::AddSourceInterfaces(EMetaSoundBuilderResult& OutResult)
-{
-	using namespace Metasound::Engine;
-	using namespace Metasound::Frontend;
-
-	const UClass& BaseMetaSoundClass = *UMetaSoundSource::StaticClass();
-	const FTopLevelAssetPath BaseMetaSoundClassPath = BaseMetaSoundClass.GetClassPathName();
-	TArray<FMetasoundFrontendVersion> InitVersions = ISearchEngine::Get().FindUClassDefaultInterfaceVersions(BaseMetaSoundClassPath);
-
-	const bool bSuccess = Builder.ModifyInterfaces({ }, InitVersions);
-	OutResult = bSuccess ? EMetaSoundBuilderResult::Succeeded : EMetaSoundBuilderResult::Failed;
-}
-
 UMetaSoundPatchBuilder* UMetaSoundBuilderSubsystem::CreatePatchBuilder(FName BuilderName, EMetaSoundBuilderResult& OutResult)
 {
 	return &Metasound::Engine::BuilderSubsystemPrivate::CreateTransientBuilder<UMetaSoundPatchBuilder>();
@@ -762,9 +755,6 @@ UMetaSoundSourceBuilder* UMetaSoundBuilderSubsystem::CreateSourceBuilder(
 
 	UMetaSoundSourceBuilder& NewBuilder = CreateTransientBuilder<UMetaSoundSourceBuilder>();
 
-	// Temp workaround - the builder currently doesn't have any interfaces added automatically, so we add them here
-	NewBuilder.AddSourceInterfaces(OutResult);
-
 	if (OutResult == EMetaSoundBuilderResult::Failed)
 	{
 		UE_LOG(LogMetaSound, Error, TEXT("Builder '%s' Creation Error: Failed to add source document interfaces."), *BuilderName.ToString());
@@ -869,7 +859,8 @@ UMetaSoundSourceBuilder* UMetaSoundBuilderSubsystem::AttachSourceBuilderToAsset(
 	{
 		TObjectPtr<UMetaSoundSourceBuilder> NewBuilder = NewObject<UMetaSoundSourceBuilder>(InSource);
 		TScriptInterface<IMetaSoundDocumentInterface> MetaSoundSourceDocInterface = InSource;
-		NewBuilder->Builder = FMetaSoundFrontendDocumentBuilder(MetaSoundSourceDocInterface);
+		const FTopLevelAssetPath BuilderPath = NewBuilder->GetBuilderUClass().GetClassPathName();
+		NewBuilder->Builder = FMetaSoundFrontendDocumentBuilder(BuilderPath, MetaSoundSourceDocInterface);
 		return NewBuilder;
 	}
 
diff --git a/Engine/Plugins/Runtime/Metasound/Source/MetasoundEngine/Private/MetasoundSource.cpp b/Engine/Plugins/Runtime/Metasound/Source/MetasoundEngine/Private/MetasoundSource.cpp
--- a/Engine/Plugins/Runtime/Metasound/Source/MetasoundEngine/Private/MetasoundSource.cpp
+++ b/Engine/Plugins/Runtime/Metasound/Source/MetasoundEngine/Private/MetasoundSource.cpp
@@ -141,10 +141,6 @@ void UMetaSoundSource::PostEditChangeOutputFormat()
 {
 	using namespace Metasound::Frontend;
 
-	// Add and/or remove output format interfaces from the root document.
-	TScriptInterface<IMetaSoundDocumentInterface> DocInterface = this;
-	FMetaSoundFrontendDocumentBuilder Builder(DocInterface);
-
 	const UMetaSoundBuilderSubsystem* BuilderSubsystem = GEngine->GetEngineSubsystem<UMetaSoundBuilderSubsystem>();
 	check(BuilderSubsystem);
 	UMetaSoundSourceBuilder* SourceBuilder = BuilderSubsystem->AttachSourceBuilderToAsset(this);
diff --git a/Engine/Plugins/Runtime/Metasound/Source/MetasoundEngine/Public/MetasoundBuilderSubsystem.h b/Engine/Plugins/Runtime/Metasound/Source/MetasoundEngine/Public/MetasoundBuilderSubsystem.h
--- a/Engine/Plugins/Runtime/Metasound/Source/MetasoundEngine/Public/MetasoundBuilderSubsystem.h
+++ b/Engine/Plugins/Runtime/Metasound/Source/MetasoundEngine/Public/MetasoundBuilderSubsystem.h
@@ -272,8 +272,8 @@ public:
 	// Constructs a transient UMetaSoundBuilderDocument & attaches builder to it
 	void CreateTransientDocument();
 
-	// Returns the base MetaSound UClass the builder is operating on (ex. MetaSoundSource, MetaSoundPatch, etc.)
-	const UClass& GetBaseMetaSoundUClass() const;
+	// Returns the base MetaSound UClass the builder is operating on (ex. MetaSoundSource, MetaSoundPatch, etc. not to be confused with the underlying document builder's MetaSound class type)
+	virtual const UClass& GetBuilderUClass() const PURE_VIRTUAL(UMetaSoundBuilderBase::Build, return *UClass::StaticClass(); );
 
 	// Initializes and ensures all nodes have a position (required prior to exporting to an asset if expected to be viewed in the editor).
 	void InitNodeLocations();
@@ -363,6 +363,8 @@ class METASOUNDENGINE_API UMetaSoundPatchBuilder : public UMetaSoundBuilderBase
 public:
 	UFUNCTION(BlueprintCallable, Category = "Audio|MetaSound|Builder", meta = (WorldContext = "Parent"))
 	virtual UPARAM(DisplayName = "MetaSound") TScriptInterface<IMetaSoundDocumentInterface> Build(UObject* Parent, const FMetaSoundBuilderOptions& Options) const override;
+
+	virtual const UClass& GetBuilderUClass() const override;
 };
 
 /** Builder in charge of building a MetaSound Source */
@@ -382,10 +384,10 @@ public:
 	UFUNCTION(BlueprintCallable, Category = "Audio|MetaSound|Builder", meta = (ExpandEnumAsExecs = "OutResult"))
 	void SetFormat(EMetaSoundOutputAudioFormat OutputFormat, EMetaSoundBuilderResult& OutResult);
 
-	void AddSourceInterfaces(EMetaSoundBuilderResult& OutResult);
-
 	const Metasound::Engine::FOutputAudioFormatInfoPair* FindOutputAudioFormatInfo() const;
 
+	virtual const UClass& GetBuilderUClass() const override;
+
 private:
 	TWeakObjectPtr<UMetaSoundSource> AuditionSound;
 };
diff --git a/Engine/Plugins/Runtime/Metasound/Source/MetasoundFrontend/Private/MetasoundFrontendDocumentBuilder.cpp b/Engine/Plugins/Runtime/Metasound/Source/MetasoundFrontend/Private/MetasoundFrontendDocumentBuilder.cpp
--- a/Engine/Plugins/Runtime/Metasound/Source/MetasoundFrontend/Private/MetasoundFrontendDocumentBuilder.cpp
+++ b/Engine/Plugins/Runtime/Metasound/Source/MetasoundFrontend/Private/MetasoundFrontendDocumentBuilder.cpp
@@ -97,11 +97,18 @@ namespace Metasound::Frontend
 } // namespace Metasound::Frontend
 
 FMetaSoundFrontendDocumentBuilder::FMetaSoundFrontendDocumentBuilder()
+	: BuilderClassPath(UMetaSoundDocumentInterface::StaticClass()->GetClassPathName())
 {
 }
 
-FMetaSoundFrontendDocumentBuilder::FMetaSoundFrontendDocumentBuilder(TScriptInterface<IMetaSoundDocumentInterface> InDocumentInterface)
+FMetaSoundFrontendDocumentBuilder::FMetaSoundFrontendDocumentBuilder(const FTopLevelAssetPath& InBuilderClassPath)
+	: BuilderClassPath(InBuilderClassPath)
+{
+}
+
+FMetaSoundFrontendDocumentBuilder::FMetaSoundFrontendDocumentBuilder(const FTopLevelAssetPath& InBuilderClassPath, TScriptInterface<IMetaSoundDocumentInterface> InDocumentInterface)
 	: DocumentInterface(InDocumentInterface)
+	, BuilderClassPath(InBuilderClassPath)
 {
 	ReloadCache();
 }
@@ -113,11 +120,13 @@ FMetaSoundFrontendDocumentBuilder::FMetaSoundFrontendDocumentBuilder(const FMeta
 	{
 		ReloadCache();
 	}	
+	BuilderClassPath = InBuilder.BuilderClassPath;
 }
 
-FMetaSoundFrontendDocumentBuilder::FMetaSoundFrontendDocumentBuilder(FMetaSoundFrontendDocumentBuilder&& InBuilder)
+FMetaSoundFrontendDocumentBuilder::FMetaSoundFrontendDocumentBuilder(const FTopLevelAssetPath& InBuilderClassPath, FMetaSoundFrontendDocumentBuilder&& InBuilder)
 	: DocumentInterface(MoveTemp(InBuilder.DocumentInterface))
 	, DocumentCache(MoveTemp(InBuilder.DocumentCache))
+	, BuilderClassPath(MoveTemp(BuilderClassPath))
 {
 }
 
@@ -126,6 +135,7 @@ FMetaSoundFrontendDocumentBuilder& FMetaSoundFrontendDocumentBuilder::operator=(
 	using namespace Metasound::Frontend;
 	DocumentInterface = MoveTemp(InRHS.DocumentInterface);
 	DocumentCache = MoveTemp(InRHS.DocumentCache);
+	BuilderClassPath = MoveTemp(InRHS.BuilderClassPath);
 	return *this;
 }
 
@@ -137,7 +147,8 @@ FMetaSoundFrontendDocumentBuilder& FMetaSoundFrontendDocumentBuilder::operator=(
 	if (DocumentInterface)
 	{
 		ReloadCache();
-	}	
+	}
+	BuilderClassPath = InRHS.BuilderClassPath;
 	return *this;
 }
 
@@ -525,11 +536,10 @@ bool FMetaSoundFrontendDocumentBuilder::AddInterface(FName InterfaceName)
 		const FInterfaceRegistryKey Key = GetInterfaceRegistryKey(Interface.Version);
 		if (const IInterfaceRegistryEntry* Entry = IInterfaceRegistry::Get().FindInterfaceRegistryEntry(Key))
 		{
-			const FTopLevelAssetPath ClassPath = DocumentInterface->GetBaseMetaSoundUClass().GetClassPathName();
-			const FMetasoundFrontendInterfaceUClassOptions* ClassOptions = Entry->GetInterface().FindClassOptions(ClassPath);
+			const FMetasoundFrontendInterfaceUClassOptions* ClassOptions = Entry->GetInterface().FindClassOptions(BuilderClassPath);
 			if (ClassOptions && !ClassOptions->bIsModifiable)
 			{
-				UE_LOG(LogMetaSound, Error, TEXT("DocumentBuilder failed to add MetaSound Interface '%s' to document: is not set to be modifiable for given UClass '%s'"), *InterfaceName.ToString(), *ClassPath.ToString());
+				UE_LOG(LogMetaSound, Error, TEXT("DocumentBuilder failed to add MetaSound Interface '%s' to document: is not set to be modifiable for given UClass '%s'"), *InterfaceName.ToString(), *BuilderClassPath.ToString());
 				return false;
 			}
 		}
@@ -957,9 +967,7 @@ void FMetaSoundFrontendDocumentBuilder::InitDocument()
 
 	// 3. Add default interfaces for given UClass
 	{
-		const UClass& BaseMetaSoundClass = DocInterface->GetBaseMetaSoundUClass();
-		const FTopLevelAssetPath BaseMetaSoundClassPath = BaseMetaSoundClass.GetClassPathName();
-		TArray<FMetasoundFrontendVersion> InitVersions = ISearchEngine::Get().FindUClassDefaultInterfaceVersions(BaseMetaSoundClassPath);
+		TArray<FMetasoundFrontendVersion> InitVersions = ISearchEngine::Get().FindUClassDefaultInterfaceVersions(BuilderClassPath);
 		FModifyRootGraphInterfaces ModifyRootGraphInterfaces({ }, InitVersions);
 		ModifyRootGraphInterfaces.Transform(Document);
 	}
@@ -1264,11 +1272,10 @@ bool FMetaSoundFrontendDocumentBuilder::RemoveInterface(FName InterfaceName)
 		const FInterfaceRegistryKey Key = GetInterfaceRegistryKey(Interface.Version);
 		if (const IInterfaceRegistryEntry* Entry = IInterfaceRegistry::Get().FindInterfaceRegistryEntry(Key))
 		{
-			const FTopLevelAssetPath ClassPath = DocumentInterface->GetBaseMetaSoundUClass().GetClassPathName();
-			const FMetasoundFrontendInterfaceUClassOptions* ClassOptions = Entry->GetInterface().FindClassOptions(ClassPath);
+			const FMetasoundFrontendInterfaceUClassOptions* ClassOptions = Entry->GetInterface().FindClassOptions(BuilderClassPath);
 			if (ClassOptions && !ClassOptions->bIsModifiable)
 			{
-				UE_LOG(LogMetaSound, Error, TEXT("DocumentBuilder failed to remove MetaSound Interface '%s' to document: is not set to be modifiable for given UClass '%s'"), *InterfaceName.ToString(), *ClassPath.ToString());
+				UE_LOG(LogMetaSound, Error, TEXT("DocumentBuilder failed to remove MetaSound Interface '%s' to document: is not set to be modifiable for given UClass '%s'"), *InterfaceName.ToString(), *BuilderClassPath.ToString());
 				return false;
 			}
 		}
diff --git a/Engine/Plugins/Runtime/Metasound/Source/MetasoundFrontend/Public/MetasoundFrontendDocumentBuilder.h b/Engine/Plugins/Runtime/Metasound/Source/MetasoundFrontend/Public/MetasoundFrontendDocumentBuilder.h
--- a/Engine/Plugins/Runtime/Metasound/Source/MetasoundFrontend/Public/MetasoundFrontendDocumentBuilder.h
+++ b/Engine/Plugins/Runtime/Metasound/Source/MetasoundFrontend/Public/MetasoundFrontendDocumentBuilder.h
@@ -90,9 +90,12 @@ struct METASOUNDFRONTEND_API FMetaSoundFrontendDocumentBuilder
 	GENERATED_BODY()
 
 public:
+	// Exists only to make UObject reflection happy.  Should never be
+	// used directly as builder class should be specified on construction.
 	FMetaSoundFrontendDocumentBuilder();
-	FMetaSoundFrontendDocumentBuilder(FMetaSoundFrontendDocumentBuilder&& InBuilder);
-	FMetaSoundFrontendDocumentBuilder(TScriptInterface<IMetaSoundDocumentInterface> InDocumentInterface);
+	FMetaSoundFrontendDocumentBuilder(const FTopLevelAssetPath& InBuilderClassPath);
+	FMetaSoundFrontendDocumentBuilder(const FTopLevelAssetPath& InBuilderClassPath, FMetaSoundFrontendDocumentBuilder&& InBuilder);
+	FMetaSoundFrontendDocumentBuilder(const FTopLevelAssetPath& InBuilderClassPath, TScriptInterface<IMetaSoundDocumentInterface> InDocumentInterface);
 
 	// Copying actively used builders is highly discouraged (as only one builder should be acting on a document at any given time
 	// to avoid internal cache state to become out-of-date) but must contractually be implemented in this case as a USTRUCT (builders
@@ -180,6 +183,7 @@ private:
 	TScriptInterface<IMetaSoundDocumentInterface> DocumentInterface;
 
 	TUniquePtr<Metasound::Frontend::IDocumentCache> DocumentCache;
+	FTopLevelAssetPath BuilderClassPath;
 
 	FMetasoundFrontendDocument& GetDocument();
 
