Jumped the gun on submit after preflight on last CL, missing fix.
#rnx

diff --git a/Engine/Plugins/Runtime/Metasound/Source/MetasoundEngine/Private/MetasoundBuilderSubsystem.cpp b/Engine/Plugins/Runtime/Metasound/Source/MetasoundEngine/Private/MetasoundBuilderSubsystem.cpp
--- a/Engine/Plugins/Runtime/Metasound/Source/MetasoundEngine/Private/MetasoundBuilderSubsystem.cpp
+++ b/Engine/Plugins/Runtime/Metasound/Source/MetasoundEngine/Private/MetasoundBuilderSubsystem.cpp
@@ -751,34 +751,27 @@ UMetaSoundSourceBuilder* UMetaSoundBuilderSubsystem::CreateSourceBuilder(
 	OnFinishedNodeInput = { };
 	AudioOutNodeInputs.Reset();
 
-	OutResult = EMetaSoundBuilderResult::Failed;
-
 	UMetaSoundSourceBuilder& NewBuilder = CreateTransientBuilder<UMetaSoundSourceBuilder>();
 
-	if (OutResult == EMetaSoundBuilderResult::Failed)
-	{
-		UE_LOG(LogMetaSound, Error, TEXT("Builder '%s' Creation Error: Failed to add source document interfaces."), *BuilderName.ToString());
-		return nullptr;
-	}
-
-	bool bFormatSet = OutputFormat == EMetaSoundOutputAudioFormat::Mono;
-	if (!bFormatSet)
+	OutResult = EMetaSoundBuilderResult::Succeeded;
+	if (OutputFormat != EMetaSoundOutputAudioFormat::Mono)
 	{
-		EMetaSoundBuilderResult Result = EMetaSoundBuilderResult::Failed;
-		NewBuilder.SetFormat(OutputFormat, Result);
-		bFormatSet = Result == EMetaSoundBuilderResult::Succeeded;
+		NewBuilder.SetFormat(OutputFormat, OutResult);
 	}
 	
-	if (bFormatSet)
+	if (OutResult == EMetaSoundBuilderResult::Succeeded)
 	{
-		const Metasound::Engine::FOutputAudioFormatInfoPair* FormatInfo = NewBuilder.FindOutputAudioFormatInfo();
 		TArray<FMetaSoundNodeHandle> AudioOutputNodes;
-		if (FormatInfo)
+		if (const Metasound::Engine::FOutputAudioFormatInfoPair* FormatInfo = NewBuilder.FindOutputAudioFormatInfo())
 		{
 			AudioOutputNodes = NewBuilder.FindInterfaceOutputNodes(FormatInfo->Value.InterfaceVersion.Name, OutResult);
 		}
+		else
+		{
+			OutResult = EMetaSoundBuilderResult::Failed;
+		}
 
-		if (!AudioOutputNodes.IsEmpty())
+		if (OutResult == EMetaSoundBuilderResult::Succeeded)
 		{
 			Algo::Transform(AudioOutputNodes, AudioOutNodeInputs, [&NewBuilder, &BuilderName](const FMetaSoundNodeHandle& AudioOutputNode) -> FMetaSoundBuilderNodeInputHandle
 			{
@@ -801,7 +794,7 @@ UMetaSoundSourceBuilder* UMetaSoundBuilderSubsystem::CreateSourceBuilder(
 	}
 	else
 	{
-		UE_LOG(LogMetaSound, Error, TEXT("Builder '%s' Creation Error: Failed to set output audio output format when initializing."), *BuilderName.ToString());
+		UE_LOG(LogMetaSound, Error, TEXT("Builder '%s' Creation Error: Failed to set output format when initializing."), *BuilderName.ToString());
 		return nullptr;
 	}
 
@@ -824,32 +817,29 @@ UMetaSoundSourceBuilder* UMetaSoundBuilderSubsystem::CreateSourceBuilder(
 		OnPlayNodeOutput = Outputs.Last();
 	}
 
+	if (bIsOneShot)
 	{
-		if (bIsOneShot)
+		FMetaSoundNodeHandle OnFinishedNode = NewBuilder.FindGraphOutputNode(SourceOneShotInterface::Outputs::OnFinished, OutResult);
+		if (OutResult == EMetaSoundBuilderResult::Failed)
 		{
-			FMetaSoundNodeHandle OnFinishedNode = NewBuilder.FindGraphOutputNode(SourceOneShotInterface::Outputs::OnFinished, OutResult);
-			if (OutResult == EMetaSoundBuilderResult::Failed)
-			{
-				UE_LOG(LogMetaSound, Error, TEXT("Builder '%s' Creation Error: Failed to add '%s' interface; interface definition may not be registered."), *BuilderName.ToString(), *SourceOneShotInterface::GetVersion().ToString());
-			}
-
-			TArray<FMetaSoundBuilderNodeInputHandle> Inputs = NewBuilder.FindNodeInputs(OnFinishedNode, OutResult);
-			if (OutResult == EMetaSoundBuilderResult::Failed)
-			{
-				UE_LOG(LogMetaSound, Error, TEXT("Builder '%s' Creation Error: Failed to find input vertex for 'OnFinished' output node when attempting to create MetaSound Source Builder"), *BuilderName.ToString());
-				return nullptr;
-			}
-
-			check(!Inputs.IsEmpty());
-			OnFinishedNodeInput = Inputs.Last();
+			UE_LOG(LogMetaSound, Error, TEXT("Builder '%s' Creation Error: Failed to add '%s' interface; interface definition may not be registered."), *BuilderName.ToString(), *SourceOneShotInterface::GetVersion().ToString());
 		}
-		else
+
+		TArray<FMetaSoundBuilderNodeInputHandle> Inputs = NewBuilder.FindNodeInputs(OnFinishedNode, OutResult);
+		if (OutResult == EMetaSoundBuilderResult::Failed)
 		{
-			NewBuilder.RemoveInterface(SourceOneShotInterface::GetVersion().Name, OutResult);
+			UE_LOG(LogMetaSound, Error, TEXT("Builder '%s' Creation Error: Failed to find input vertex for 'OnFinished' output node when attempting to create MetaSound Source Builder"), *BuilderName.ToString());
+			return nullptr;
 		}
+
+		check(!Inputs.IsEmpty());
+		OnFinishedNodeInput = Inputs.Last();
+	}
+	else
+	{
+		NewBuilder.RemoveInterface(SourceOneShotInterface::GetVersion().Name, OutResult);
 	}
 
-	OutResult = EMetaSoundBuilderResult::Succeeded;
 	return &NewBuilder;
 }
 
