Fix data channel management on MetaSounds
#rnx

diff --git a/Engine/Plugins/Runtime/Metasound/Source/MetasoundEngine/Private/MetasoundSource.cpp b/Engine/Plugins/Runtime/Metasound/Source/MetasoundEngine/Private/MetasoundSource.cpp
--- a/Engine/Plugins/Runtime/Metasound/Source/MetasoundEngine/Private/MetasoundSource.cpp
+++ b/Engine/Plugins/Runtime/Metasound/Source/MetasoundEngine/Private/MetasoundSource.cpp
@@ -62,34 +62,71 @@ namespace Metasound
 
 		class FParameterRouter
 		{
+			struct FQueueState
+			{
+				TWeakPtr<TSpscQueue<FMetaSoundParameterTransmitter::FParameter>> DataChannel;
+				bool bWriterAvailable = true;
+			};
+
 		public:
 
 			using FAudioDeviceIDAndInstanceID = TTuple<Audio::DeviceID, uint64>;
 
-			TSharedPtr<TSpscQueue<FMetaSoundParameterTransmitter::FParameter>> FindOrCreateDataChannel(Audio::DeviceID InDeviceID, uint64 InstanceID)
+			TSharedPtr<TSpscQueue<FMetaSoundParameterTransmitter::FParameter>> FindOrCreateDataChannelForReader(Audio::DeviceID InDeviceID, uint64 InstanceID)
+			{
+				constexpr bool bIsForWriter = false;
+				return FindOrCreateDataChannel(InDeviceID, InstanceID, bIsForWriter);
+			}
+
+			TSharedPtr<TSpscQueue<FMetaSoundParameterTransmitter::FParameter>> FindOrCreateDataChannelForWriter(Audio::DeviceID InDeviceID, uint64 InstanceID)
+			{
+				constexpr bool bIsForWriter = true;
+				return FindOrCreateDataChannel(InDeviceID, InstanceID, bIsForWriter);
+			}
+
+		private:
+
+			TSharedPtr<TSpscQueue<FMetaSoundParameterTransmitter::FParameter>> FindOrCreateDataChannel(Audio::DeviceID InDeviceID, uint64 InstanceID, bool bIsForWriter)
 			{
 				FScopeLock Lock(&DataChannelMapCS);
 
 				FAudioDeviceIDAndInstanceID Key = {InDeviceID, InstanceID};
+				const bool bIsForReader = !bIsForWriter;
 
-				if (TWeakPtr<TSpscQueue<FMetaSoundParameterTransmitter::FParameter>>* WeakChannel = DataChannels.Find(Key))
+				if (FQueueState* State = DataChannels.Find(Key))
 				{
-					TSharedPtr<TSpscQueue<FMetaSoundParameterTransmitter::FParameter>> Channel = WeakChannel->Pin();
-					if (Channel.IsValid())
+					// Allow multiple readers to be returned because FMetaSoundGenerators are recreated when they come out of virtualization.
+					// Only allow a single writer to be returned because FMetaSoundParameterTransmitters are only created once
+					const bool bIsAvailable = bIsForReader || (State->bWriterAvailable && bIsForWriter);
+					if (bIsAvailable)
 					{
-						return Channel;
+						TSharedPtr<TSpscQueue<FMetaSoundParameterTransmitter::FParameter>> Channel = State->DataChannel.Pin();
+						if (Channel.IsValid())
+						{
+							if (bIsForWriter)
+							{
+								State->bWriterAvailable = false;
+							}
+							return Channel;
+						}
 					}
 				}
 
 				TSharedPtr<TSpscQueue<FMetaSoundParameterTransmitter::FParameter>> NewChannel = MakeShared<TSpscQueue<FMetaSoundParameterTransmitter::FParameter>>();
-				DataChannels.Add(Key, NewChannel);
+
+				FQueueState NewState;
+				NewState.DataChannel = NewChannel;
+				if (bIsForWriter)
+				{
+					NewState.bWriterAvailable = false;
+				}
+
+				DataChannels.Add(Key, NewState);
 				return NewChannel;
 			}
 
-		private:
-
 			FCriticalSection DataChannelMapCS;
-			TSortedMap<FAudioDeviceIDAndInstanceID, TWeakPtr<TSpscQueue<FMetaSoundParameterTransmitter::FParameter>>> DataChannels;
+			TSortedMap<FAudioDeviceIDAndInstanceID, FQueueState> DataChannels;
 		};
 	} // namespace SourcePrivate
 } // namespace Metasound
@@ -538,7 +575,7 @@ ISoundGeneratorPtr UMetaSoundSource::CreateSoundGenerator(const FSoundGeneratorI
 	FOperatorSettings InSettings = GetOperatorSettings(static_cast<FSampleRate>(SampleRate));
 	FMetasoundEnvironment Environment = CreateEnvironment(InParams);
 	FParameterRouter& Router = GetParameterRouter();
-	TSharedPtr<TSpscQueue<FMetaSoundParameterTransmitter::FParameter>> DataChannel = Router.FindOrCreateDataChannel(InParams.AudioDeviceID, InParams.InstanceID);
+	TSharedPtr<TSpscQueue<FMetaSoundParameterTransmitter::FParameter>> DataChannel = Router.FindOrCreateDataChannelForReader(InParams.AudioDeviceID, InParams.InstanceID);
 
 	TSharedPtr<const IGraph, ESPMode::ThreadSafe> MetasoundGraph = GetRuntimeData().Graph;
 	if (!MetasoundGraph.IsValid())
@@ -855,7 +892,7 @@ TSharedPtr<Audio::IParameterTransmitter> UMetaSoundSource::CreateParameterTransm
 	TArray<FName> ValidParameters;
 	Algo::Transform(RuntimeData.TransmittableInputs, ValidParameters, [](const FMetasoundFrontendClassInput& InputVertex) { return InputVertex.Name; });
 	FParameterRouter& Router = GetParameterRouter();
-	TSharedPtr<TSpscQueue<FMetaSoundParameterTransmitter::FParameter>> DataChannel = Router.FindOrCreateDataChannel(InParams.AudioDeviceID, InParams.InstanceID);
+	TSharedPtr<TSpscQueue<FMetaSoundParameterTransmitter::FParameter>> DataChannel = Router.FindOrCreateDataChannelForWriter(InParams.AudioDeviceID, InParams.InstanceID);
 
 	Metasound::FMetaSoundParameterTransmitter::FInitParams InitParams(GetOperatorSettings(InParams.SampleRate), InParams.InstanceID, MoveTemp(InParams.DefaultParams), MoveTemp(ValidParameters), DataChannel);
 	InitParams.DebugMetaSoundName = GetFName();
diff --git a/Engine/Plugins/Runtime/Metasound/Source/MetasoundFrontend/Private/MetasoundParameterTransmitter.cpp b/Engine/Plugins/Runtime/Metasound/Source/MetasoundFrontend/Private/MetasoundParameterTransmitter.cpp
--- a/Engine/Plugins/Runtime/Metasound/Source/MetasoundFrontend/Private/MetasoundParameterTransmitter.cpp
+++ b/Engine/Plugins/Runtime/Metasound/Source/MetasoundFrontend/Private/MetasoundParameterTransmitter.cpp
@@ -251,18 +251,7 @@ namespace Metasound
 	{
 	}
 
-	void FMetaSoundParameterTransmitter::OnDeleteActiveSound()
-	{
-		// Only reset the data channel if we have a valid data channel.
-		// This protects against removing the data channel 
-		// multiple times. Multiple removals of data channels has caused
-		// race conditions between newly created transmitters and transmitters
-		// being cleaned up.
-		DataChannel.Reset();
-
-		Audio::FParameterTransmitterBase::OnDeleteActiveSound();
-	}
-	
+
 	bool FMetaSoundParameterTransmitter::Reset()
 	{
 		OnDeleteActiveSound();
diff --git a/Engine/Plugins/Runtime/Metasound/Source/MetasoundFrontend/Public/MetasoundParameterTransmitter.h b/Engine/Plugins/Runtime/Metasound/Source/MetasoundFrontend/Public/MetasoundParameterTransmitter.h
--- a/Engine/Plugins/Runtime/Metasound/Source/MetasoundFrontend/Public/MetasoundParameterTransmitter.h
+++ b/Engine/Plugins/Runtime/Metasound/Source/MetasoundFrontend/Public/MetasoundParameterTransmitter.h
@@ -119,7 +119,6 @@ namespace Metasound
 		FMetaSoundParameterTransmitter(FMetaSoundParameterTransmitter::FInitParams&& InInitParams);
 		virtual ~FMetaSoundParameterTransmitter() = default;
 
-		virtual void OnDeleteActiveSound() override;
 
 		UE_DEPRECATED(5.2, "Use ResetParameters() or OnDeleteActiveSound() instead depending on use case.")
 		bool Reset() override;
