[Metasound] Use generators for node tests and fix up tests.

diff --git a/Engine/Plugins/Runtime/Metasound/Source/MetasoundEngineTest/Private/NodeTestGraphBuilder.cpp b/Engine/Plugins/Runtime/Metasound/Source/MetasoundEngineTest/Private/NodeTestGraphBuilder.cpp
--- a/Engine/Plugins/Runtime/Metasound/Source/MetasoundEngineTest/Private/NodeTestGraphBuilder.cpp
+++ b/Engine/Plugins/Runtime/Metasound/Source/MetasoundEngineTest/Private/NodeTestGraphBuilder.cpp
@@ -5,11 +5,10 @@
 #include "MetasoundFrontendGraph.h"
 #include "MetasoundFrontendSearchEngine.h"
 #include "MetasoundOperatorBuilder.h"
-#include "Interfaces/MetasoundOutputFormatInterfaces.h"
 
 namespace Metasound::Test
 {
-	using namespace Metasound::Frontend;
+	using namespace Frontend;
 
 	FNodeTestGraphBuilder::FNodeTestGraphBuilder()
 	{
@@ -21,7 +20,7 @@ namespace Metasound::Test
 		check(RootGraph->IsValid());
 	}
 
-	FNodeHandle FNodeTestGraphBuilder::AddNode(const FNodeClassName& ClassName, int32 MajorVersion)
+	FNodeHandle FNodeTestGraphBuilder::AddNode(const FNodeClassName& ClassName, int32 MajorVersion) const
 	{
 		check(RootGraph->IsValid());
 
@@ -35,7 +34,7 @@ namespace Metasound::Test
 		return Node;
 	}
 
-	FNodeHandle FNodeTestGraphBuilder::AddInput(const FName& InputName, const FName& TypeName)
+	FNodeHandle FNodeTestGraphBuilder::AddInput(const FName& InputName, const FName& TypeName) const
 	{
 		check(RootGraph->IsValid());
 
@@ -65,34 +64,7 @@ namespace Metasound::Test
 		return NodeHandle;
 	}
 
-	TUniquePtr<IOperator> FNodeTestGraphBuilder::BuildGraph(FSampleRate SampleRate, int32 SamplesPerBlock)
-	{
-		TSet<FName> TransmittableInputNames;
-		const FString UnknownAsset = TEXT("UnknownAsset");
-		if (TUniquePtr<FFrontendGraph> Graph = FFrontendGraphBuilder::CreateGraph(Document, TransmittableInputNames, UnknownAsset))
-		{
-			FOperatorBuilderSettings BuilderSettings;
-			BuilderSettings.bFailOnAnyError = true;
-			BuilderSettings.bPopulateInternalDataReferences = true;
-			BuilderSettings.bValidateEdgeDataTypesMatch = true;
-			BuilderSettings.bValidateNoCyclesInGraph = true;
-			BuilderSettings.bValidateNoDuplicateInputs = true;
-			BuilderSettings.bValidateOperatorOutputsAreBound = true;
-			BuilderSettings.bValidateVerticesExist = true;
-			FOperatorBuilder Builder{ BuilderSettings };
-
-			FOperatorSettings OperatorSettings{SampleRate, static_cast<float>(SampleRate) / SamplesPerBlock};
-			FInputVertexInterfaceData InterfaceData;
-			FMetasoundEnvironment Environment;
-			FBuildGraphOperatorParams BuildParams{ *Graph, OperatorSettings, InterfaceData, Environment };
-			FBuildResults Results;
-			return Builder.BuildGraphOperator(BuildParams, Results);
-		}
-
-		return nullptr;
-	}
-
-	TUniquePtr<FMetasoundGenerator> FNodeTestGraphBuilder::BuildGenerator(FSampleRate SampleRate, int32 SamplesPerBlock)
+	TUniquePtr<FMetasoundGenerator> FNodeTestGraphBuilder::BuildGenerator(FSampleRate SampleRate, int32 SamplesPerBlock) const
 	{
 		// Because a generator is tied to the concept of a source, go ahead and add the "OnPlay" input.
 		// Otherwise we get warnings when we run our tests.
@@ -133,11 +105,11 @@ namespace Metasound::Test
 		return nullptr;
 	}
 
-	TUniquePtr<IOperator> FNodeTestGraphBuilder::MakeSingleNodeGraph(
+	TUniquePtr<FMetasoundGenerator> FNodeTestGraphBuilder::MakeSingleNodeGraph(
 		const FNodeClassName& ClassName,
-		int32 MajorVersion,
-		FSampleRate SampleRate,
-		int32 SamplesPerBlock)
+		const int32 MajorVersion,
+		const FSampleRate SampleRate,
+		const int32 SamplesPerBlock)
 	{
 		FNodeTestGraphBuilder Builder;
 		FNodeHandle NodeHandle = Builder.AddNode(ClassName, MajorVersion);
@@ -185,7 +157,13 @@ namespace Metasound::Test
 			}
 		}
 
+		// have to make an audio output for the generator to do anything
+		if (Builder.AudioOutputNames.IsEmpty())
+		{
+			Builder.AddOutput("AudioOut", GetMetasoundDataTypeName<FAudioBuffer>());
+		}
+
 		// build the graph
-		return Builder.BuildGraph(SampleRate, SamplesPerBlock);
+		return Builder.BuildGenerator(SampleRate, SamplesPerBlock);
 	}
 }
diff --git a/Engine/Plugins/Runtime/Metasound/Source/MetasoundEngineTest/Public/NodeTestGraphBuilder.h b/Engine/Plugins/Runtime/Metasound/Source/MetasoundEngineTest/Public/NodeTestGraphBuilder.h
--- a/Engine/Plugins/Runtime/Metasound/Source/MetasoundEngineTest/Public/NodeTestGraphBuilder.h
+++ b/Engine/Plugins/Runtime/Metasound/Source/MetasoundEngineTest/Public/NodeTestGraphBuilder.h
@@ -6,7 +6,6 @@
 #include "MetasoundFrontendDocument.h"
 #include "MetasoundGenerator.h"
 #include "MetasoundNodeInterface.h"
-#include "MetasoundOperatorInterface.h"
 
 namespace Metasound::Test
 {
@@ -17,21 +16,18 @@ namespace Metasound::Test
 		FNodeTestGraphBuilder();
 
 		/** Add a node to the graph */
-		Frontend::FNodeHandle AddNode(const FNodeClassName& ClassName, int32 MajorVersion);
+		Frontend::FNodeHandle AddNode(const FNodeClassName& ClassName, int32 MajorVersion) const;
 
 		/** Add an input node to the graph */
-		Frontend::FNodeHandle AddInput(const FName& InputName, const FName& TypeName);
+		Frontend::FNodeHandle AddInput(const FName& InputName, const FName& TypeName) const;
 
 		/** Add an output node to the graph */
 		Frontend::FNodeHandle AddOutput(const FName& OutputName, const FName& TypeName);
 
-		/** Build the graph and get the graph's operator to work with */
-		TUniquePtr<IOperator> BuildGraph(FSampleRate SampleRate = 48000, int32 SamplesPerBlock = 256);
-
-		TUniquePtr<FMetasoundGenerator> BuildGenerator(FSampleRate SampleRate = 48000, int32 SamplesPerBlock = 256);
+		TUniquePtr<FMetasoundGenerator> BuildGenerator(FSampleRate SampleRate = 48000, int32 SamplesPerBlock = 256) const;
 
 		/** Helper that will add a single node, wire up the node's inputs and outputs, and hand back the graph's operator */
-		static TUniquePtr<IOperator> MakeSingleNodeGraph(
+		static TUniquePtr<FMetasoundGenerator> MakeSingleNodeGraph(
 			const FNodeClassName& ClassName,
 			int32 MajorVersion,
 			FSampleRate SampleRate = 48000,
