Reduce template use in metasound input and literal nodes.

diff --git a/Engine/Plugins/Runtime/Metasound/Source/MetasoundFrontend/Private/MetasoundInputNode.cpp b/Engine/Plugins/Runtime/Metasound/Source/MetasoundFrontend/Private/MetasoundInputNode.cpp
--- a/Engine/Plugins/Runtime/Metasound/Source/MetasoundFrontend/Private/MetasoundInputNode.cpp
+++ b/Engine/Plugins/Runtime/Metasound/Source/MetasoundFrontend/Private/MetasoundInputNode.cpp
@@ -2,11 +2,15 @@
 
 #include "MetasoundInputNode.h"
 
+#include "Internationalization/Text.h"
 #include "MetasoundDataReference.h"
+#include "MetasoundNodeRegistrationMacro.h"
 #include "MetasoundVertex.h"
 #include "MetasoundVertexData.h"
 #include "UObject/NameTypes.h"
 
+#define LOCTEXT_NAMESPACE "MetasoundFrontend"
+
 namespace Metasound
 {
 	namespace MetasoundInputNodePrivate
@@ -62,4 +66,86 @@ namespace Metasound
 			return nullptr;
 		}
 	}
+
+	FLazyName FInputNode::ConstructorVariant("Constructor");
+	FName FInputNode::GetVariantName(EVertexAccessType InVertexAccess)
+	{
+		if (EVertexAccessType::Value == InVertexAccess)
+		{
+			return ConstructorVariant;
+		}
+		else
+		{
+			return FName();
+		}
+	}
+
+	FVertexInterface FInputNode::CreateVertexInterface(const FVertexName& InVertexName, const FName& InDataTypeName, EVertexAccessType InVertexAccess, const FLiteral& InLiteral)
+	{
+		return FVertexInterface(
+			FInputVertexInterface(
+				FInputDataVertex(InVertexName, InDataTypeName, FDataVertexMetadata{ FText::GetEmpty() }, InVertexAccess, InLiteral)
+			),
+			FOutputVertexInterface(
+				FOutputDataVertex(InVertexName, InDataTypeName, FDataVertexMetadata{ FText::GetEmpty() }, InVertexAccess)
+			)
+		);
+	}
+
+
+	FVertexInterface FInputNode::CreateDefaultVertexInterface(const FVertexName& InVertexName, const FName& InDataTypeName, EVertexAccessType InVertexAccess)
+	{
+		return CreateVertexInterface(InVertexName, InDataTypeName, InVertexAccess, FLiteral());
+	}
+
+
+	FNodeClassMetadata FInputNode::GetNodeMetadata(const FVertexName& InVertexName, const FName& InDataTypeName, EVertexAccessType InVertexAccess)
+	{
+		FNodeClassMetadata Info;
+
+		Info.ClassName = { "Input", InDataTypeName, GetVariantName(InVertexAccess) };
+		Info.MajorVersion = 1;
+		Info.MinorVersion = 0;
+		Info.Description = METASOUND_LOCTEXT("Metasound_InputNodeDescription", "Input into the parent Metasound graph.");
+		Info.Author = PluginAuthor;
+		Info.PromptIfMissing = PluginNodeMissingPrompt;
+		Info.DefaultInterface = CreateDefaultVertexInterface(InVertexName, InDataTypeName, InVertexAccess);
+
+		return Info;
+	}
+
+	FInputNode::FInputNode(FInputNodeConstructorParams&& InParams, const FName& InDataTypeName, EVertexAccessType InVertexAccess, FOperatorFactorySharedRef InFactory)
+	: FNode(InParams.NodeName, InParams.InstanceID, GetNodeMetadata(InParams.VertexName, InDataTypeName, InVertexAccess))
+	, VertexName(InParams.VertexName)
+	, Interface(CreateVertexInterface(InParams.VertexName, InDataTypeName, InVertexAccess, InParams.InitParam))
+	, Factory(InFactory)
+	{
+	}
+
+	const FVertexName& FInputNode::GetVertexName() const
+	{
+		return VertexName;
+	}
+
+	const FVertexInterface& FInputNode::GetVertexInterface() const
+	{
+		return Interface;
+	}
+
+	bool FInputNode::SetVertexInterface(const FVertexInterface& InInterface)
+	{
+		return Interface == InInterface;
+	}
+
+	bool FInputNode::IsVertexInterfaceSupported(const FVertexInterface& InInterface) const
+	{
+		return Interface == InInterface;
+	}
+
+	TSharedRef<IOperatorFactory, ESPMode::ThreadSafe> FInputNode::GetDefaultOperatorFactory() const
+	{
+		return Factory;
+	}
 }
+
+#undef LOCTEXT_NAMESPACE // MetasoundFrontend
diff --git a/Engine/Plugins/Runtime/Metasound/Source/MetasoundFrontend/Private/MetasoundLiteralNode.cpp b/Engine/Plugins/Runtime/Metasound/Source/MetasoundFrontend/Private/MetasoundLiteralNode.cpp
new file mode 100644
--- /dev/null
+++ b/Engine/Plugins/Runtime/Metasound/Source/MetasoundFrontend/Private/MetasoundLiteralNode.cpp
@@ -0,0 +1,70 @@
+// Copyright Epic Games, Inc. All Rights Reserved.
+
+#include "MetasoundLiteralNode.h"
+
+#define LOCTEXT_NAMESPACE "MetasoundFrontend"
+namespace Metasound
+{
+	FVertexInterface FLiteralNode::CreateVertexInterface(const FName& InDataTypeName, EVertexAccessType InAccessType)
+	{
+		using namespace LiteralNodeNames; 
+		const FDataVertexMetadata OutputMetadata
+		{
+			  FText::GetEmpty() // description
+			, METASOUND_GET_PARAM_DISPLAYNAME(OutputValue) // display name
+		};
+
+		return FVertexInterface(
+			FInputVertexInterface(),
+			FOutputVertexInterface(
+				FOutputDataVertex(METASOUND_GET_PARAM_NAME(OutputValue), InDataTypeName, OutputMetadata, InAccessType)
+			)
+		);
+	}
+
+	FNodeClassMetadata FLiteralNode::GetNodeMetadata(const FName& InDataTypeName, EVertexAccessType InAccessType)
+	{
+		FNodeClassMetadata Info;
+
+		Info.ClassName = {"Literal", InDataTypeName, ""};
+		Info.MajorVersion = 1;
+		Info.MinorVersion = 0;
+#if WITH_EDITOR
+		Info.DisplayName = FText::Format(LOCTEXT("Metasound_LiteralNodeDisplayNameFormat", "Literal {0}"), FText::FromName(InDataTypeName));
+		Info.Description = LOCTEXT("Metasound_LiteralNodeDescription", "Literal accessible within a parent Metasound graph.");
+#endif // WITH_EDITOR
+		Info.Author = PluginAuthor;
+		Info.PromptIfMissing = PluginNodeMissingPrompt;
+		Info.DefaultInterface = CreateVertexInterface(InDataTypeName, InAccessType);
+
+		return Info;
+	}
+
+	FLiteralNode::FLiteralNode(const FVertexName& InInstanceName, const FGuid& InInstanceID, const FName& InDataTypeName, EVertexAccessType InAccessType, FOperatorFactorySharedRef InFactory)
+	: FNode(InInstanceName, InInstanceID, GetNodeMetadata(InDataTypeName, InAccessType))
+	, Interface(CreateVertexInterface(InDataTypeName, InAccessType))
+	, Factory(InFactory)
+	{
+	}
+
+	const FVertexInterface& FLiteralNode::GetVertexInterface() const
+	{
+		return Interface;
+	}
+
+	bool FLiteralNode::SetVertexInterface(const FVertexInterface& InInterface)
+	{
+		return Interface == InInterface;
+	}
+
+	bool FLiteralNode::IsVertexInterfaceSupported(const FVertexInterface& InInterface) const
+	{
+		return Interface == InInterface;
+	}
+
+	TSharedRef<IOperatorFactory, ESPMode::ThreadSafe> FLiteralNode::GetDefaultOperatorFactory() const
+	{
+		return Factory;
+	}
+}
+#undef LOCTEXT_NAMESPACE // MetasoundFrontend
diff --git a/Engine/Plugins/Runtime/Metasound/Source/MetasoundFrontend/Public/MetasoundInputNode.h b/Engine/Plugins/Runtime/Metasound/Source/MetasoundFrontend/Public/MetasoundInputNode.h
--- a/Engine/Plugins/Runtime/Metasound/Source/MetasoundFrontend/Public/MetasoundInputNode.h
+++ b/Engine/Plugins/Runtime/Metasound/Source/MetasoundFrontend/Public/MetasoundInputNode.h
@@ -1,7 +1,6 @@
 // Copyright Epic Games, Inc. All Rights Reserved.
 #pragma once
 
-#include "Internationalization/Text.h"
 #include "MetasoundBuilderInterface.h"
 #include "MetasoundBuildError.h"
 #include "MetasoundDataReference.h"
@@ -11,17 +10,10 @@
 #include "MetasoundLiteral.h"
 #include "MetasoundNode.h"
 #include "MetasoundNodeInterface.h"
-#include "MetasoundNodeRegistrationMacro.h"
 #include "MetasoundOperatorInterface.h"
-#include "MetasoundParameterTransmitter.h"
-#include "MetasoundRouter.h"
-#include "MetasoundTrigger.h"
 #include "MetasoundVertexData.h"
 #include "UObject/NameTypes.h"
 
-#define LOCTEXT_NAMESPACE "MetasoundFrontend"
-
-
 namespace Metasound
 {
 	namespace MetasoundInputNodePrivate
@@ -319,10 +311,45 @@ namespace Metasound
 		MetasoundInputNodePrivate::FNonExecutableInputPassThroughOperator
 	>;
 
+	/** FInputNode represents an input to a metasound graph. */
+	class METASOUNDFRONTEND_API FInputNode : public FNode
+	{
+		static FLazyName ConstructorVariant;
+		// Use Variant names to differentiate between normal input nodes and constructor 
+		// input nodes.
+		static FName GetVariantName(EVertexAccessType InVertexAccess);
+
+		static FVertexInterface CreateVertexInterface(const FVertexName& InVertexName, const FName& InDataTypeName, EVertexAccessType InVertexAccess, const FLiteral& InLiteral);
+
+	protected:
+
+		static FVertexInterface CreateDefaultVertexInterface(const FVertexName& InVertexName, const FName& InDataTypeName, EVertexAccessType InVertexAccess);
+
+	public:
+
+		static FNodeClassMetadata GetNodeMetadata(const FVertexName& InVertexName, const FName& InDataTypeName, EVertexAccessType InVertexAccess);
+
+		/* Construct a TInputNode using the TInputOperatorLiteralFactory<> and moving
+		 * InParam to the TInputOperatorLiteralFactory constructor.*/
+		explicit FInputNode(FInputNodeConstructorParams&& InParams, const FName& InDataTypeName, EVertexAccessType InVertexAccess, FOperatorFactorySharedRef InFactory);
+
+		const FVertexName& GetVertexName() const;
+
+		virtual const FVertexInterface& GetVertexInterface() const override;
+		virtual bool SetVertexInterface(const FVertexInterface& InInterface) override;
+		virtual bool IsVertexInterfaceSupported(const FVertexInterface& InInterface) const override;
+		virtual TSharedRef<IOperatorFactory, ESPMode::ThreadSafe> GetDefaultOperatorFactory() const override;
+
+	private:
+		FVertexName VertexName;
+		FVertexInterface Interface;
+		FOperatorFactorySharedRef Factory;
+	};
+
 
 	/** TInputNode represents an input to a metasound graph. */
 	template<typename DataType, EVertexAccessType VertexAccess=EVertexAccessType::Reference>
-	class TInputNode : public FNode
+	class TInputNode : public FInputNode
 	{
 		static constexpr bool bIsConstructorInput = VertexAccess == EVertexAccessType::Value;
 		static constexpr bool bIsSupportedConstructorInput = TIsConstructorVertexSupported<DataType>::Value && bIsConstructorInput;
@@ -331,20 +358,6 @@ namespace Metasound
 
 		static constexpr bool bIsSupportedInput = bIsSupportedConstructorInput || bIsSupportedReferenceInput;
 
-		// Use Variant names to differentiate between normal input nodes and constructor 
-		// input nodes.
-		static FName GetVariantName()
-		{
-			if constexpr (EVertexAccessType::Value == VertexAccess)
-			{
-				return FName("Constructor");
-			}
-			else
-			{
-				return FName();
-			}
-		}
-
 		// Factory for creating input operators. 
 		class FInputNodeOperatorFactory : public IOperatorFactory
 		{
@@ -408,8 +421,6 @@ namespace Metasound
 					return MakeUnique<TInputOperator<DataType, VertexAccess>>(VertexKey, InParams.OperatorSettings, Literal);
 				}
 			}
-
-		private:
 		};
 
 
@@ -417,85 +428,22 @@ namespace Metasound
 		// If true, this node can be instantiated by the Frontend.
 		static constexpr bool bCanRegister = bIsSupportedInput;
 
-		static FVertexInterface CreateVertexInterface(const FVertexName& InVertexName, const FLiteral& InLiteral)
-		{
-			return FVertexInterface(
-				FInputVertexInterface(
-					FInputDataVertex(InVertexName, GetMetasoundDataTypeName<DataType>(), FDataVertexMetadata{ FText::GetEmpty() }, VertexAccess, InLiteral)
-				),
-				FOutputVertexInterface(
-					FOutputDataVertex(InVertexName, GetMetasoundDataTypeName<DataType>(), FDataVertexMetadata{ FText::GetEmpty() }, VertexAccess)
-				)
-			);
-		}
-
-		static FVertexInterface CreateDefaultVertexInterface(const FVertexName& InVertexName)
-		{
-			return CreateVertexInterface(InVertexName, FLiteral());
-		}
-
-
-		UE_DEPRECATED(5.3, "Use TInputNode<>::CreateDefaultVertexInterface or TInputNode<>::CreateVertexInterface instead.")
+		UE_DEPRECATED(5.3, "Access the default vertex interface from the input node metadata.")
 		static FVertexInterface DeclareVertexInterface(const FVertexName& InVertexName)
 		{
-			return CreateDefaultVertexInterface(InVertexName);
+			return CreateDefaultVertexInterface(InVertexName, GetMetasoundDataTypeName<DataType>(), VertexAccess);
 		}
 
 		static FNodeClassMetadata GetNodeInfo(const FVertexName& InVertexName)
 		{
-			FNodeClassMetadata Info;
-
-			Info.ClassName = { "Input", GetMetasoundDataTypeName<DataType>(), GetVariantName() };
-			Info.MajorVersion = 1;
-			Info.MinorVersion = 0;
-			Info.Description = METASOUND_LOCTEXT("Metasound_InputNodeDescription", "Input into the parent Metasound graph.");
-			Info.Author = PluginAuthor;
-			Info.PromptIfMissing = PluginNodeMissingPrompt;
-			Info.DefaultInterface = CreateDefaultVertexInterface(InVertexName);
-
-			return Info;
+			return GetNodeMetadata(InVertexName, GetMetasoundDataTypeName<DataType>(), VertexAccess);
 		}
 
 		/* Construct a TInputNode using the TInputOperatorLiteralFactory<> and moving
 		 * InParam to the TInputOperatorLiteralFactory constructor.*/
 		explicit TInputNode(FInputNodeConstructorParams&& InParams)
-		:	FNode(InParams.NodeName, InParams.InstanceID, GetNodeInfo(InParams.VertexName))
-		,	VertexName(InParams.VertexName)
-		,	Interface(CreateVertexInterface(InParams.VertexName, InParams.InitParam))
-		,	Factory(MakeShared<FInputNodeOperatorFactory>())
-		{
-		}
-
-		const FVertexName& GetVertexName() const
-		{
-			return VertexName;
-		}
-
-		virtual const FVertexInterface& GetVertexInterface() const override
-		{
-			return Interface;
-		}
-
-		virtual bool SetVertexInterface(const FVertexInterface& InInterface) override
-		{
-			return Interface == InInterface;
-		}
-
-		virtual bool IsVertexInterfaceSupported(const FVertexInterface& InInterface) const override
-		{
-			return Interface == InInterface;
-		}
-
-		virtual TSharedRef<IOperatorFactory, ESPMode::ThreadSafe> GetDefaultOperatorFactory() const override
+		:	FInputNode(MoveTemp(InParams), GetMetasoundDataTypeName<DataType>(), VertexAccess, MakeShared<FInputNodeOperatorFactory>())
 		{
-			return Factory;
 		}
-
-	private:
-		FVertexName VertexName;
-
-		FVertexInterface Interface;
-		FOperatorFactorySharedRef Factory;
 	};
 } // namespace Metasound
-#undef LOCTEXT_NAMESPACE // MetasoundFrontend
diff --git a/Engine/Plugins/Runtime/Metasound/Source/MetasoundFrontend/Public/MetasoundLiteralNode.h b/Engine/Plugins/Runtime/Metasound/Source/MetasoundFrontend/Public/MetasoundLiteralNode.h
new file mode 100644
--- /dev/null
+++ b/Engine/Plugins/Runtime/Metasound/Source/MetasoundFrontend/Public/MetasoundLiteralNode.h
@@ -0,0 +1,206 @@
+// Copyright Epic Games, Inc. All Rights Reserved.
+#pragma once
+
+#include "Internationalization/Text.h"
+#include "MetasoundBuilderInterface.h"
+#include "MetasoundExecutableOperator.h"
+#include "MetasoundInputNode.h"
+#include "MetasoundNode.h"
+#include "MetasoundNodeInterface.h"
+#include "MetasoundOperatorInterface.h"
+#include "MetasoundParamHelper.h"
+#include "MetasoundVertex.h"
+
+#define LOCTEXT_NAMESPACE "MetasoundFrontend"
+
+namespace Metasound
+{
+	namespace LiteralNodeNames
+	{
+		METASOUND_PARAM(OutputValue, "Value", "Value")
+	}
+
+	template<typename DataType>
+	class UE_DEPRECATED(5.3, "The TLiteralOperator will be removed.") TLiteralOperator : public FNoOpOperator
+	{
+		static_assert(!TExecutableDataType<DataType>::bIsExecutable, "TLiteralOperator is only suitable for non-executable data types");
+
+	public:
+		TLiteralOperator(TDataValueReference<DataType> InValue)
+		: Value (InValue)
+		{
+		}
+
+		virtual ~TLiteralOperator() = default;
+
+		virtual void Bind(FVertexInterfaceData& InVertexData) const
+		{
+			using namespace LiteralNodeNames;
+			InVertexData.GetOutputs().BindValueVertex(METASOUND_GET_PARAM_NAME(OutputValue), Value);
+		}
+
+	private:
+		TDataValueReference<DataType> Value;
+	};
+
+	/** TExecutableLiteralOperator is used for executable types. The majority of literals
+	 * are constant values, but executable data types cannot be constant values.
+	 * 
+	 * Note: this is supporting for FTrigger literals which are deprecated.
+	 */
+	template<typename DataType>
+	class UE_DEPRECATED(5.3, "The TExecutableLiteralOperator will be removed.") TExecutableLiteralOperator : public TExecutableOperator<TExecutableLiteralOperator<DataType>>
+	{
+		static_assert(TExecutableDataType<DataType>::bIsExecutable, "TExecutableLiteralOperator is only suitable for executable data types");
+
+	public:
+		using FDataWriteReference = TDataWriteReference<DataType>;
+
+		TExecutableLiteralOperator(const FOperatorSettings& InSettings, const FLiteral& InLiteral)
+			// Executable DataTypes require a copy of the output to operate on whereas non-executable
+			// types do not. Avoid copy by assigning to reference for non-executable types.
+			: Literal(InLiteral)
+			, InputValue(TDataWriteReferenceLiteralFactory<DataType>::CreateExplicitArgs(InSettings, InLiteral))
+			, OutputValue(TDataWriteReferenceLiteralFactory<DataType>::CreateExplicitArgs(InSettings, InLiteral))
+		{
+		}
+
+		virtual ~TExecutableLiteralOperator() = default;
+
+		virtual void Bind(FVertexInterfaceData& InVertexData) const
+		{
+			using namespace LiteralNodeNames;
+			InVertexData.GetOutputs().BindReadVertex(METASOUND_GET_PARAM_NAME(OutputValue), OutputValue);
+		}
+
+		virtual FDataReferenceCollection GetInputs() const override
+		{
+			FDataReferenceCollection Inputs;
+			return Inputs;
+		}
+
+		virtual FDataReferenceCollection GetOutputs() const override
+		{
+			using namespace LiteralNodeNames;
+			FDataReferenceCollection Outputs;
+			Outputs.AddDataReadReference<DataType>(METASOUND_GET_PARAM_NAME(OutputValue), OutputValue);
+			return Outputs;
+		}
+
+		void Execute()
+		{
+			TExecutableDataType<DataType>::Execute(*InputValue, *OutputValue);
+		}
+
+		void Reset(const IOperator::FResetParams& InParams)
+		{
+			*OutputValue = *InputValue = TDataTypeLiteralFactory<DataType>::CreateExplicitArgs(InParams.OperatorSettings, Literal);
+		}
+
+	private:
+		FLiteral Literal;
+		FDataWriteReference InputValue;
+		FDataWriteReference OutputValue;
+	};
+
+
+	/** TLiteralOperatorLiteralFactory creates an input by passing it a literal. */
+	template<typename DataType>
+	class TLiteralOperatorLiteralFactory : public IOperatorFactory
+	{
+	public:
+
+		using FDataWriteReference = TDataWriteReference<DataType>;
+
+		TLiteralOperatorLiteralFactory(FLiteral&& InInitParam)
+			: InitParam(MoveTemp(InInitParam))
+		{
+		}
+
+		virtual TUniquePtr<IOperator> CreateOperator(const FBuildOperatorParams& InParams, FBuildResults& OutResults) override;
+
+	private:
+
+		class FLiteralOperator : public TInputOperator<DataType>
+		{
+		public:
+			using TInputOperator<DataType>::TInputOperator;
+			virtual void Bind(FVertexInterfaceData& InVertexData) const override
+			{
+				// Do not bind inputs as there are no inputs exposed from the literal node. 
+				TInputOperator<DataType>::BindOutputs(InVertexData.GetOutputs());
+			}
+		};
+
+		FLiteral InitParam;
+	};
+
+	/** TLiteralNode represents a variable within a metasound graph. */
+	class METASOUNDFRONTEND_API FLiteralNode : public FNode
+	{
+	protected:
+		static FVertexInterface CreateVertexInterface(const FName& InDataTypeName, EVertexAccessType InAccessType);
+
+	public:
+
+		static FNodeClassMetadata GetNodeMetadata(const FName& InDataTypeName, EVertexAccessType InAccessType);
+
+		FLiteralNode(const FVertexName& InInstanceName, const FGuid& InInstanceID, const FName& InDataTypeName, EVertexAccessType InAccessType, FOperatorFactorySharedRef InFactory);
+
+		virtual const FVertexInterface& GetVertexInterface() const override;
+
+		virtual bool SetVertexInterface(const FVertexInterface& InInterface) override;
+
+		virtual bool IsVertexInterfaceSupported(const FVertexInterface& InInterface) const override;
+
+		virtual TSharedRef<IOperatorFactory, ESPMode::ThreadSafe> GetDefaultOperatorFactory() const override;
+
+	private:
+		FVertexInterface Interface;
+		FOperatorFactorySharedRef Factory;
+	};
+
+	/** TLiteralNode represents a variable within a metasound graph. */
+	template<typename DataType>
+	class TLiteralNode : public FLiteralNode
+	{
+		// Executable data types handle Triggers which need to be advanced every
+		// block.
+		static constexpr bool bIsExecutableDataType = TExecutableDataType<DataType>::bIsExecutable;
+		static constexpr bool bIsPostExecutableDataType = TPostExecutableDataType<DataType>::bIsPostExecutable;
+
+		// If the data type is executable, we treat it as a reference since it is non-const
+		static constexpr EVertexAccessType AccessType = (bIsExecutableDataType || bIsPostExecutableDataType) ? EVertexAccessType::Reference : EVertexAccessType::Value;
+
+	public:
+
+		UE_DEPRECATED(5.3, "Use the default interface on the node class metadata")
+		static FVertexInterface DeclareVertexInterface()
+		{
+			return CreateVertexInterface(GetMetasoundDataTypeName<DataType>(), AccessType);
+		}
+
+		static FNodeClassMetadata GetNodeInfo()
+		{
+			return GetNodeMetadata(GetMetasoundDataTypeName<DataType>(), AccessType);
+		}
+
+		/* Construct a TLiteralNode using the TLiteralOperatorLiteralFactory<> and moving
+		 * InParam to the TLiteralOperatorLiteralFactory constructor.*/
+		explicit TLiteralNode(const FVertexName& InInstanceName, const FGuid& InInstanceID, FLiteral&& InParam)
+		:	FLiteralNode(InInstanceName, InInstanceID, GetMetasoundDataTypeName<DataType>(), AccessType, MakeOperatorFactoryRef<TLiteralOperatorLiteralFactory<DataType>>(MoveTemp(InParam)))
+		{
+		}
+
+	};
+
+	template<typename DataType>
+	TUniquePtr<IOperator> TLiteralOperatorLiteralFactory<DataType>::CreateOperator(const FBuildOperatorParams& InParams, FBuildResults& OutResults)
+	{
+		using namespace LiteralNodeNames;
+		return MakeUnique<FLiteralOperator>(METASOUND_GET_PARAM_NAME(OutputValue), InParams.OperatorSettings, InitParam);
+	}
+
+} // namespace Metasound
+#undef LOCTEXT_NAMESPACE // MetasoundFrontend
+
diff --git a/Engine/Plugins/Runtime/Metasound/Source/MetasoundGraphCore/Public/MetasoundLiteralNode.h b/Engine/Plugins/Runtime/Metasound/Source/MetasoundGraphCore/Public/MetasoundLiteralNode.h
deleted file mode 100644
