Speculative fix for parameter list construction issue on certain platforms
#rnx

diff --git a/Engine/Plugins/Runtime/Metasound/Source/MetasoundEngine/Private/MetasoundBuilderSubsystem.cpp b/Engine/Plugins/Runtime/Metasound/Source/MetasoundEngine/Private/MetasoundBuilderSubsystem.cpp
--- a/Engine/Plugins/Runtime/Metasound/Source/MetasoundEngine/Private/MetasoundBuilderSubsystem.cpp
+++ b/Engine/Plugins/Runtime/Metasound/Source/MetasoundEngine/Private/MetasoundBuilderSubsystem.cpp
@@ -200,7 +200,7 @@ TArray<FMetaSoundBuilderNodeInputHandle> UMetaSoundBuilderBase::ConnectNodeOutpu
 	{
 		const FMetasoundFrontendVertex* Vertex = Builder.FindNodeInput(NewEdge->ToNodeID, NewEdge->ToVertexID);
 		checkf(Vertex, TEXT("Edge connection reported success but vertex not found."));
-		return FMetaSoundBuilderNodeInputHandle { NewEdge->ToNodeID, Vertex->VertexID };
+		return FMetaSoundBuilderNodeInputHandle(NewEdge->ToNodeID, Vertex->VertexID);
 	});
 
 	return ConnectedVertices;
@@ -217,7 +217,7 @@ TArray<FMetaSoundBuilderNodeOutputHandle> UMetaSoundBuilderBase::ConnectNodeInpu
 	{
 		const FMetasoundFrontendVertex* Vertex = Builder.FindNodeInput(NewEdge->FromNodeID, NewEdge->FromVertexID);
 		checkf(Vertex, TEXT("Edge connection reported success but vertex not found."));
-		return FMetaSoundBuilderNodeOutputHandle { NewEdge->ToNodeID, Vertex->VertexID };
+		return FMetaSoundBuilderNodeOutputHandle(NewEdge->ToNodeID, Vertex->VertexID);
 	});
 
 	return ConnectedVertices;
@@ -301,7 +301,7 @@ FMetaSoundBuilderNodeInputHandle UMetaSoundBuilderBase::FindNodeInputByName(cons
 		if (const FMetasoundFrontendVertex* Input = InputVertices.FindByPredicate(FindByNamePredicate))
 		{
 			OutResult = EMetaSoundBuilderResult::Succeeded;
-			return FMetaSoundBuilderNodeInputHandle { Node->GetID(), Input->VertexID };
+			return FMetaSoundBuilderNodeInputHandle(Node->GetID(), Input->VertexID);
 		}
 
 		FString NodeClassName = TEXT("N/A");
@@ -335,7 +335,7 @@ TArray<FMetaSoundBuilderNodeInputHandle> UMetaSoundBuilderBase::FindNodeInputsBy
 		TArray<const FMetasoundFrontendVertex*> Vertices = Builder.FindNodeInputs(NodeHandle.NodeID, DataType);
 		Algo::Transform(Vertices, FoundVertices, [&NodeHandle](const FMetasoundFrontendVertex* Vertex)
 		{
-			return FMetaSoundBuilderNodeInputHandle { NodeHandle.NodeID, Vertex->VertexID };
+			return FMetaSoundBuilderNodeInputHandle(NodeHandle.NodeID, Vertex->VertexID);
 		});
 		OutResult = EMetaSoundBuilderResult::Succeeded;
 	}
@@ -358,7 +358,7 @@ FMetaSoundBuilderNodeOutputHandle UMetaSoundBuilderBase::FindNodeOutputByName(co
 		if (const FMetasoundFrontendVertex* Output = OutputVertices.FindByPredicate(FindByNamePredicate))
 		{
 			OutResult = EMetaSoundBuilderResult::Succeeded;
-			return FMetaSoundBuilderNodeOutputHandle { Node->GetID(), Output->VertexID };
+			return FMetaSoundBuilderNodeOutputHandle(Node->GetID(), Output->VertexID);
 		}
 
 		FString NodeClassName = TEXT("N/A");
@@ -391,7 +391,7 @@ TArray<FMetaSoundBuilderNodeOutputHandle> UMetaSoundBuilderBase::FindNodeOutputs
 		TArray<const FMetasoundFrontendVertex*> Vertices = Builder.FindNodeOutputs(NodeHandle.NodeID, DataType);
 		Algo::Transform(Vertices, FoundVertices, [&NodeHandle](const FMetasoundFrontendVertex* Vertex)
 		{
-			return FMetaSoundBuilderNodeOutputHandle { NodeHandle.NodeID, Vertex->VertexID };
+			return FMetaSoundBuilderNodeOutputHandle(NodeHandle.NodeID, Vertex->VertexID);
 		});
 		OutResult = EMetaSoundBuilderResult::Succeeded;
 	}
diff --git a/Engine/Plugins/Runtime/Metasound/Source/MetasoundEngine/Public/MetasoundBuilderSubsystem.h b/Engine/Plugins/Runtime/Metasound/Source/MetasoundEngine/Public/MetasoundBuilderSubsystem.h
--- a/Engine/Plugins/Runtime/Metasound/Source/MetasoundEngine/Public/MetasoundBuilderSubsystem.h
+++ b/Engine/Plugins/Runtime/Metasound/Source/MetasoundEngine/Public/MetasoundBuilderSubsystem.h
@@ -34,12 +34,28 @@ USTRUCT(BlueprintType, meta = (DisplayName = "MetaSound Node Input Handle"))
 struct METASOUNDENGINE_API FMetaSoundBuilderNodeInputHandle : public FMetasoundFrontendVertexHandle
 {
 	GENERATED_BODY()
+
+public:
+	FMetaSoundBuilderNodeInputHandle() = default;
+	FMetaSoundBuilderNodeInputHandle(const FGuid InNodeID, const FGuid& InVertexID)
+	{
+		NodeID = InNodeID;
+		VertexID = InVertexID;
+	}
 };
 
 USTRUCT(BlueprintType, meta = (DisplayName = "MetaSound Node Output Handle"))
 struct METASOUNDENGINE_API FMetaSoundBuilderNodeOutputHandle : public FMetasoundFrontendVertexHandle
 {
 	GENERATED_BODY()
+
+public:
+	FMetaSoundBuilderNodeOutputHandle() = default;
+	FMetaSoundBuilderNodeOutputHandle(const FGuid InNodeID, const FGuid& InVertexID)
+	{
+		NodeID = InNodeID;
+		VertexID = InVertexID;
+	}
 };
 
 USTRUCT(BlueprintType)
