[Metasound] Ensure variadic constructor arguments to TVertexInterfaceGroup are move constructible to the group's base type.

diff --git a/Engine/Plugins/Runtime/Metasound/Source/MetasoundGraphCore/Public/MetasoundVertex.h b/Engine/Plugins/Runtime/Metasound/Source/MetasoundGraphCore/Public/MetasoundVertex.h
--- a/Engine/Plugins/Runtime/Metasound/Source/MetasoundGraphCore/Public/MetasoundVertex.h
+++ b/Engine/Plugins/Runtime/Metasound/Source/MetasoundGraphCore/Public/MetasoundVertex.h
@@ -306,34 +306,7 @@ namespace Metasound
 		};
 
 		static_assert(TIsSupportedVertexType<VertexType>::Value, "VertexType must be derived from FInputDataVertex, FOutputDataVertex, or FEnvironmentVertex");
-
-
-		// Helper struct for determining whether the vertices in a parameter pack
-		// can be used in the TVertexInterfaceGroup constructor
-		template<typename ...>
-		struct TIsConstructibleFromParameterPack
-		{
-			static constexpr bool Value = false;
-		};
-
-		// Helper struct specialization for no template argument. 
-		template<>
-		struct TIsConstructibleFromParameterPack<>
-		{
-			static constexpr bool Value = true;
-		};
-
-
-	
-
-		// Helper struct specialization for unwinding a parameter pack.
-		template<typename T, typename ... Ts>
-		struct TIsConstructibleFromParameterPack<T, Ts...>
-		{
-			static constexpr bool Value = std::is_constructible<VertexType, T>::value && TIsConstructibleFromParameterPack<Ts...>::Value;
-		};
-
-
+		
 		using FContainerType = TArray<VertexType>;
 
 		struct FEqualVertexNamePredicate
@@ -391,21 +364,21 @@ namespace Metasound
 	public:
 
 		using RangedForConstIteratorType = typename FContainerType::RangedForConstIteratorType;
-
+		
 		/** TVertexInterfaceGroup constructor with variadic list of vertex
 		 * models.
 		 */
 		template<
-			typename... VertexTypes, 
-			typename std::enable_if< TIsConstructibleFromParameterPack<VertexTypes...>::Value, int >::type = 0
+			typename... VertexTypes,
+			std::enable_if_t<(std::is_constructible_v<VertexType, VertexTypes&&> && ...), int> = 0
 		>
-		TVertexInterfaceGroup(VertexTypes&&... InVertexs)
+		explicit TVertexInterfaceGroup(VertexTypes&&... InVertices)
 		{
 			// Reserve array to hold exact number of vertices to avoid
 			// over allocation.
 			Vertices.Reserve(sizeof...(VertexTypes));
 
-			CopyInputs(Forward<VertexTypes>(InVertexs)...);
+			CopyInputs(Forward<VertexTypes>(InVertices)...);
 		}
 
 		/** Add a vertex to the group. */
