Fix metasound literal node loctext clash
#rnx

diff --git a/Engine/Plugins/Runtime/Metasound/Source/MetasoundFrontend/Private/MetasoundLiteralNode.cpp b/Engine/Plugins/Runtime/Metasound/Source/MetasoundFrontend/Private/MetasoundLiteralNode.cpp
--- a/Engine/Plugins/Runtime/Metasound/Source/MetasoundFrontend/Private/MetasoundLiteralNode.cpp
+++ b/Engine/Plugins/Runtime/Metasound/Source/MetasoundFrontend/Private/MetasoundLiteralNode.cpp
@@ -2,7 +2,7 @@
 
 #include "MetasoundLiteralNode.h"
 
-#define LOCTEXT_NAMESPACE "MetasoundFrontend"
+#define LOCTEXT_NAMESPACE "MetasoundLiteralNode"
 namespace Metasound
 {
 	FVertexInterface FLiteralNode::CreateVertexInterface(const FName& InDataTypeName, EVertexAccessType InAccessType)
@@ -67,4 +67,4 @@ namespace Metasound
 		return Factory;
 	}
 }
-#undef LOCTEXT_NAMESPACE // MetasoundFrontend
+#undef LOCTEXT_NAMESPACE // MetasoundLiteralNode
diff --git a/Engine/Plugins/Runtime/Metasound/Source/MetasoundFrontend/Public/MetasoundLiteralNode.h b/Engine/Plugins/Runtime/Metasound/Source/MetasoundFrontend/Public/MetasoundLiteralNode.h
--- a/Engine/Plugins/Runtime/Metasound/Source/MetasoundFrontend/Public/MetasoundLiteralNode.h
+++ b/Engine/Plugins/Runtime/Metasound/Source/MetasoundFrontend/Public/MetasoundLiteralNode.h
@@ -11,7 +11,7 @@
 #include "MetasoundParamHelper.h"
 #include "MetasoundVertex.h"
 
-#define LOCTEXT_NAMESPACE "MetasoundFrontend"
+#define LOCTEXT_NAMESPACE "MetasoundLiteralNode"
 
 namespace Metasound
 {
@@ -202,5 +202,5 @@ namespace Metasound
 	}
 
 } // namespace Metasound
-#undef LOCTEXT_NAMESPACE // MetasoundFrontend
+#undef LOCTEXT_NAMESPACE // MetasoundLiteralNode
 
