[Metasound] TVertexInterfaceGroup: use static_assert instead of enable_if in constructor to fix MSVC build

diff --git a/Engine/Plugins/Runtime/Metasound/Source/MetasoundGraphCore/Public/MetasoundVertex.h b/Engine/Plugins/Runtime/Metasound/Source/MetasoundGraphCore/Public/MetasoundVertex.h
--- a/Engine/Plugins/Runtime/Metasound/Source/MetasoundGraphCore/Public/MetasoundVertex.h
+++ b/Engine/Plugins/Runtime/Metasound/Source/MetasoundGraphCore/Public/MetasoundVertex.h
@@ -368,12 +368,13 @@ namespace Metasound
 		/** TVertexInterfaceGroup constructor with variadic list of vertex
 		 * models.
 		 */
-		template<
-			typename... VertexTypes,
-			std::enable_if_t<(std::is_constructible_v<VertexType, VertexTypes&&> && ...), int> = 0
-		>
+		template<typename... VertexTypes>
 		explicit TVertexInterfaceGroup(VertexTypes&&... InVertices)
 		{
+			static_assert(
+				(std::is_constructible_v<VertexType, VertexTypes&&> && ...),
+				"Vertex types must be move constructible from the group's base type");
+			
 			// Reserve array to hold exact number of vertices to avoid
 			// over allocation.
 			Vertices.Reserve(sizeof...(VertexTypes));
