Fix for crash in metasound array random node.
#rnx

diff --git a/Engine/Plugins/Runtime/Metasound/Source/MetasoundFrontend/Public/MetasoundArrayRandomNode.h b/Engine/Plugins/Runtime/Metasound/Source/MetasoundFrontend/Public/MetasoundArrayRandomNode.h
--- a/Engine/Plugins/Runtime/Metasound/Source/MetasoundFrontend/Public/MetasoundArrayRandomNode.h
+++ b/Engine/Plugins/Runtime/Metasound/Source/MetasoundFrontend/Public/MetasoundArrayRandomNode.h
@@ -159,7 +159,7 @@
 			FInt32ReadRef InNoRepeatOrder = InParams.InputDataReferences.GetDataReadReferenceOrConstructWithVertexDefault<int32>(Inputs, METASOUND_GET_PARAM_NAME(InputNoRepeatOrder), InParams.OperatorSettings);
 			FBoolReadRef bInEnableSharedState = InParams.InputDataReferences.GetDataReadReferenceOrConstructWithVertexDefault<bool>(Inputs, METASOUND_GET_PARAM_NAME(InputEnableSharedState), InParams.OperatorSettings);
 
-			return MakeUnique<TArrayRandomGetOperator<ArrayType>>(InParams, InTriggerNext, InTriggerReset, InInputArray, InInputWeightsArray, InSeedValue, InNoRepeatOrder, bInEnableSharedState);
+			return MakeUnique<TArrayRandomGetOperator<ArrayType>>(InParams, InTriggerNext, InTriggerReset, InInputArray, InInputWeightsArray, InSeedValue, InNoRepeatOrder, *bInEnableSharedState);
 		}
 
 		TArrayRandomGetOperator(
@@ -170,17 +170,17 @@
 			const TDataReadReference<WeightsArrayType>& InInputWeightsArray,
 			const FInt32ReadRef& InSeedValue,
 			const FInt32ReadRef& InNoRepeatOrder,
-			const FBoolReadRef& bInEnableSharedState)
+			bool bInEnableSharedState)
 			: TriggerNext(InTriggerNext)
 			, TriggerReset(InTriggerReset)
 			, InputArray(InInputArray)
 			, InputWeightsArray(InInputWeightsArray)
 			, SeedValue(InSeedValue)
 			, NoRepeatOrder(InNoRepeatOrder)
-			, bEnableSharedState(bInEnableSharedState)
 			, TriggerOnNext(FTriggerWriteRef::CreateNew(InParams.OperatorSettings))
 			, TriggerOnReset(FTriggerWriteRef::CreateNew(InParams.OperatorSettings))
 			, OutValue(TDataWriteReferenceFactory<ElementType>::CreateAny(InParams.OperatorSettings))
+			, bEnableSharedState(bInEnableSharedState)
 		{
 			using namespace Frontend;
 
@@ -193,30 +193,39 @@
 
 		virtual FDataReferenceCollection GetInputs() const override
 		{
-			using namespace ArrayNodeRandomGetVertexNames;
-
 			FDataReferenceCollection Inputs;
-			Inputs.AddDataReadReference(METASOUND_GET_PARAM_NAME(InputTriggerNextValue), TriggerNext);
-			Inputs.AddDataReadReference(METASOUND_GET_PARAM_NAME(InputTriggerResetSeed), TriggerReset);
-			Inputs.AddDataReadReference(METASOUND_GET_PARAM_NAME(InputRandomArray), InputArray);
-			Inputs.AddDataReadReference(METASOUND_GET_PARAM_NAME(InputWeights), InputWeightsArray);
-			Inputs.AddDataReadReference(METASOUND_GET_PARAM_NAME(InputSeed), SeedValue);
-			Inputs.AddDataReadReference(METASOUND_GET_PARAM_NAME(InputNoRepeatOrder), NoRepeatOrder);
-			Inputs.AddDataReadReference(METASOUND_GET_PARAM_NAME(InputEnableSharedState), bEnableSharedState);
-
+			// This function is no longer used
+			checkNoEntry();
 			return Inputs;
 		}
 
 		virtual FDataReferenceCollection GetOutputs() const override
 		{
-			using namespace ArrayNodeRandomGetVertexNames;
-
 			FDataReferenceCollection Outputs;
-			Outputs.AddDataReadReference(METASOUND_GET_PARAM_NAME(OutputTriggerOnNext), TriggerOnNext);
-			Outputs.AddDataReadReference(METASOUND_GET_PARAM_NAME(OutputTriggerOnReset), TriggerOnReset);
-			Outputs.AddDataReadReference(METASOUND_GET_PARAM_NAME(ShuffleOutputValue), OutValue);
-
+			// This function is no longer used
+			checkNoEntry();
 			return Outputs;
 		}
 
+		virtual void Bind(FVertexInterfaceData& InOutVertexData) const override
+		{
+			using namespace ArrayNodeRandomGetVertexNames;
+
+			FInputVertexInterfaceData& Inputs = InOutVertexData.GetInputs();
+
+			Inputs.BindReadVertex(METASOUND_GET_PARAM_NAME(InputTriggerNextValue), TriggerNext);
+			Inputs.BindReadVertex(METASOUND_GET_PARAM_NAME(InputTriggerResetSeed), TriggerReset);
+			Inputs.BindReadVertex(METASOUND_GET_PARAM_NAME(InputRandomArray), InputArray);
+			Inputs.BindReadVertex(METASOUND_GET_PARAM_NAME(InputWeights), InputWeightsArray);
+			Inputs.BindReadVertex(METASOUND_GET_PARAM_NAME(InputSeed), SeedValue);
+			Inputs.BindReadVertex(METASOUND_GET_PARAM_NAME(InputNoRepeatOrder), NoRepeatOrder);
+			Inputs.SetValue(METASOUND_GET_PARAM_NAME(InputEnableSharedState), bEnableSharedState);
+
+			FOutputVertexInterfaceData& Outputs = InOutVertexData.GetOutputs();
+
+			Outputs.BindReadVertex(METASOUND_GET_PARAM_NAME(OutputTriggerOnNext), TriggerOnNext);
+			Outputs.BindReadVertex(METASOUND_GET_PARAM_NAME(OutputTriggerOnReset), TriggerOnReset);
+			Outputs.BindReadVertex(METASOUND_GET_PARAM_NAME(ShuffleOutputValue), OutValue);
+		}
+
 		void Reset(const IOperator::FResetParams& InParams)
@@ -240,38 +249,7 @@
 			const ArrayType& InputArrayRef = *InputArray;
 			PrevArraySize = InputArrayRef.Num();
 
-			if (PrevArraySize > 0)
-			{
-				if (*bEnableSharedState)
-				{
-					// Get the environment variable for the unique ID of the sound
-					check(SharedStateUniqueId.IsValid());
-					FSharedStateRandomGetManager& RGM = FSharedStateRandomGetManager::Get();
+			InitializeState(PrevArraySize);
 
-					InitSharedStateArgs Args;
-					Args.SharedStateId = SharedStateUniqueId;
-					Args.Seed = PrevSeedValue;
-					Args.NumElements = PrevArraySize;
-					Args.NoRepeatOrder = PrevNoRepeatOrder;
-					Args.bIsPreviewSound = bIsPreviewSound;
-					Args.Weights = WeightsArray;
-
-					RGM.InitSharedState(Args);
-
-					bSharedStateInitialized = true;
-				}
-				else
-				{
-					ArrayRandomGet = MakeUnique<FArrayRandomGet>(PrevSeedValue, PrevArraySize, WeightsArray, PrevNoRepeatOrder);
-				}
-
-			}
-#if WITH_METASOUND_DEBUG_ENVIRONMENT
-			else
-			{
- 				UE_LOG(LogMetaSound, Verbose, TEXT("Array Random Get: Can't retrieve random elements from an empty array in graph '%s'"), *GraphName);
-			}
-#endif // WITH_METASOUND_DEBUG_ENVIRONMENT
-
 			*OutValue = TDataTypeFactory<ElementType>::CreateAny(InParams.OperatorSettings);
 			TriggerOnNext->Reset();
@@ -280,7 +258,7 @@
 
 		bool UseSharedState() const
 		{
-			return bSharedStateInitialized && *bEnableSharedState;
+			return bEnableSharedState && bSharedStateInitialized;
 		}
 
 		void Execute()
@@ -296,23 +274,19 @@
 				PrevArraySize = InputArrayRef.Num();
 				if (PrevArraySize != 0)
 				{
-					if (!ArrayRandomGet.IsValid())
-					{
-						ArrayRandomGet = MakeUnique<FArrayRandomGet>(PrevSeedValue, PrevArraySize, WeightsArray, PrevNoRepeatOrder);
-					}
-					ArrayRandomGet->Init(PrevSeedValue, PrevArraySize, WeightsArray, PrevNoRepeatOrder);
+					InitializeState(PrevArraySize);
 				}
 			}
 
 			if (PrevArraySize == 0)
 			{
-				if (!bHasLoggedEmptyArrayWarning)
-				{
 #if WITH_METASOUND_DEBUG_ENVIRONMENT
+				if (!bHasLoggedEmptyArrayWarning)
+				{
 					UE_LOG(LogMetaSound, Verbose, TEXT("Array Random Get: empty array input (Graph '%s')"), *GraphName);
-#endif // WITH_METASOUND_DEBUG_ENVIRONMENT
 					bHasLoggedEmptyArrayWarning = true;
 				}
+#endif // WITH_METASOUND_DEBUG_ENVIRONMENT
 				return;
 			}
 
@@ -417,6 +391,42 @@
 		}
 
 	private:
+		void InitializeState(int32 InArraySize)
+		{
+			bSharedStateInitialized = false;
+			if (InArraySize > 0)
+			{
+				if (bEnableSharedState)
+				{
+					// Get the environment variable for the unique ID of the sound
+					check(SharedStateUniqueId.IsValid());
+					FSharedStateRandomGetManager& RGM = FSharedStateRandomGetManager::Get();
 
+					InitSharedStateArgs Args;
+					Args.SharedStateId = SharedStateUniqueId;
+					Args.Seed = PrevSeedValue;
+					Args.NumElements = PrevArraySize;
+					Args.NoRepeatOrder = PrevNoRepeatOrder;
+					Args.bIsPreviewSound = bIsPreviewSound;
+					Args.Weights = WeightsArray;
+
+					RGM.InitSharedState(Args);
+
+					bSharedStateInitialized = true;
+				}
+				else
+				{
+					ArrayRandomGet = MakeUnique<FArrayRandomGet>(PrevSeedValue, PrevArraySize, WeightsArray, PrevNoRepeatOrder);
+				}
+			}
+			else
+			{
+				ArrayRandomGet = MakeUnique<FArrayRandomGet>();
+#if WITH_METASOUND_DEBUG_ENVIRONMENT
+ 				UE_LOG(LogMetaSound, Verbose, TEXT("Array Random Get: Can't retrieve random elements from an empty array in graph '%s'"), *GraphName);
+#endif // WITH_METASOUND_DEBUG_ENVIRONMENT
+			}
+		}
+
 		// Inputs
 		FTriggerReadRef TriggerNext;
@@ -425,7 +435,6 @@
 		TDataReadReference<WeightsArrayType> InputWeightsArray;
 		FInt32ReadRef SeedValue;
 		FInt32ReadRef NoRepeatOrder;
-		FBoolReadRef bEnableSharedState;
 
 		// Outputs
 		FTriggerWriteRef TriggerOnNext;
@@ -434,6 +443,7 @@
 
 #if WITH_METASOUND_DEBUG_ENVIRONMENT
 		FString GraphName;
+		bool bHasLoggedEmptyArrayWarning = false;
 #endif // WITH_METASOUND_DEBUG_ENVIRONMENT
 
 		// Data
@@ -444,7 +454,7 @@
 		FGuid SharedStateUniqueId;
 		int32 PrevArraySize = 0;
 		bool bIsPreviewSound = false;
-		bool bHasLoggedEmptyArrayWarning = false;
+		bool bEnableSharedState = false;
 		bool bSharedStateInitialized = false;
 	};
 
