Add absolute time output to metasound waveplayer node
#9953

diff --git a/Engine/Plugins/Runtime/Metasound/Source/MetasoundEngine/Private/MetasoundWavePlayerNode.cpp b/Engine/Plugins/Runtime/Metasound/Source/MetasoundEngine/Private/MetasoundWavePlayerNode.cpp
--- a/Engine/Plugins/Runtime/Metasound/Source/MetasoundEngine/Private/MetasoundWavePlayerNode.cpp
+++ b/Engine/Plugins/Runtime/Metasound/Source/MetasoundEngine/Private/MetasoundWavePlayerNode.cpp
@@ -44,6 +44,7 @@ namespace Metasound
 		METASOUND_PARAM(OutputCuePointLabel, "Cue Point Label", "The cue point label that was triggered (if there was a label parsed in the imported .wav file).")
 		METASOUND_PARAM(OutputLoopRatio, "Loop Percent", "Returns the current playback location as a ratio of the loop (0-1) if looping is enabled.")
 		METASOUND_PARAM(OutputPlaybackLocation, "Playback Location", "Returns the absolute position of the wave playback as a ratio of wave duration (0-1).")
+		METASOUND_PARAM(OutputPlaybackTime, "Playback Time", "Returns the current absolute playback time of the wave.")
 		METASOUND_PARAM(OutputAudioMono, "Out Mono", "The mono channel audio output.")
 		METASOUND_PARAM(OutputAudioLeft, "Out Left", "The left channel audio output.")
 		METASOUND_PARAM(OutputAudioRight, "Out Right", "The right channel audio output.")
@@ -338,6 +339,7 @@ namespace Metasound
 			, CuePointLabel(FStringWriteRef::CreateNew(TEXT("")))
 			, LoopPercent(FFloatWriteRef::CreateNew(0.0f))
 			, PlaybackLocation(FFloatWriteRef::CreateNew(0.0f))
+			, PlaybackTime(FTimeWriteRef::CreateNew(0.0))
 		{
 			NumOutputChannels = InArgs.OutputAudioVertices.Num();
 
@@ -384,6 +386,7 @@ namespace Metasound
 			OutputDataReferences.AddDataReadReference(METASOUND_GET_PARAM_NAME(OutputCuePointLabel), CuePointLabel);
 			OutputDataReferences.AddDataReadReference(METASOUND_GET_PARAM_NAME(OutputLoopRatio), LoopPercent);
 			OutputDataReferences.AddDataReadReference(METASOUND_GET_PARAM_NAME(OutputPlaybackLocation), PlaybackLocation);
+			OutputDataReferences.AddDataReadReference(METASOUND_GET_PARAM_NAME(OutputPlaybackTime), PlaybackTime);
 
 			check(OutputAudioBuffers.Num() == OutputAudioBufferVertexNames.Num());
 
@@ -445,6 +448,7 @@ namespace Metasound
 			*CuePointLabel = TEXT("");
 			*LoopPercent = 0;
 			*PlaybackLocation = 0;
+			*PlaybackTime = FTime(0.0);
 			for (const FAudioBufferWriteRef& BufferRef : OutputAudioBuffers)
 			{
 				BufferRef->Zero();
@@ -577,6 +581,15 @@ namespace Metasound
 		{
 			*PlaybackLocation = SourceState.GetPlaybackFraction();
 
+			if (WaveProxyReader.IsValid())
+			{
+				*PlaybackTime = FTime::FromSeconds(static_cast<double>(SourceState.GetStartFrameIndex()) / static_cast<double>(FMath::Max(UE_SMALL_NUMBER, WaveProxyReader->GetSampleRate())));
+			}
+			else
+			{
+				*PlaybackTime = FTime(0.0);
+			}
+
 			if (*bLoop)
 			{
 				*LoopPercent = SourceState.GetLoopFraction();
@@ -994,6 +1007,7 @@ namespace Metasound
 		FStringWriteRef CuePointLabel;
 		FFloatWriteRef LoopPercent;
 		FFloatWriteRef PlaybackLocation;
+		FTimeWriteRef PlaybackTime;
 		TArray<FAudioBufferWriteRef> OutputAudioBuffers;
 		TArray<FName> OutputAudioBufferVertexNames;
 
@@ -1089,7 +1103,8 @@ namespace Metasound
 					TOutputDataVertex<int32>(METASOUND_GET_PARAM_NAME_AND_METADATA(OutputCuePointID)),
 					TOutputDataVertex<FString>(METASOUND_GET_PARAM_NAME_AND_METADATA(OutputCuePointLabel)),
 					TOutputDataVertex<float>(METASOUND_GET_PARAM_NAME(OutputLoopRatio), OutputLoopRatioMetadata),
-					TOutputDataVertex<float>(METASOUND_GET_PARAM_NAME_AND_METADATA(OutputPlaybackLocation))
+					TOutputDataVertex<float>(METASOUND_GET_PARAM_NAME_AND_METADATA(OutputPlaybackLocation)),
+					TOutputDataVertex<FTime>(METASOUND_GET_PARAM_NAME_AND_METADATA(OutputPlaybackTime))
 				)
 			);
 
