[Metasound] Re-register output analyzers if sound is stopped and re-started

#rnx

diff --git a/Engine/Plugins/Runtime/Metasound/Source/MetasoundEngine/Private/MetasoundGeneratorHandle.cpp b/Engine/Plugins/Runtime/Metasound/Source/MetasoundEngine/Private/MetasoundGeneratorHandle.cpp
--- a/Engine/Plugins/Runtime/Metasound/Source/MetasoundEngine/Private/MetasoundGeneratorHandle.cpp
+++ b/Engine/Plugins/Runtime/Metasound/Source/MetasoundEngine/Private/MetasoundGeneratorHandle.cpp
@@ -431,6 +431,13 @@ void UMetasoundGeneratorHandle::OnSourceDestroyedAGenerator(uint64 InAudioCompon
 			OnGeneratorHandleDetached.Broadcast();
 		}
 		CachedGeneratorPtr = nullptr;
+
+		// Clear out the watchers, but keep track of what we had in case the sound comes back
+		for (const Metasound::Private::FMetasoundOutputWatcher& Watcher : OutputWatchers)
+		{
+			OutputAnalyzersToAdd.Enqueue(Watcher.GetAddress());
+		}
+		OutputWatchers.Empty();
 	}
 }
 
diff --git a/Engine/Plugins/Runtime/Metasound/Source/MetasoundEngine/Private/MetasoundOutputWatcher.cpp b/Engine/Plugins/Runtime/Metasound/Source/MetasoundEngine/Private/MetasoundOutputWatcher.cpp
--- a/Engine/Plugins/Runtime/Metasound/Source/MetasoundEngine/Private/MetasoundOutputWatcher.cpp
+++ b/Engine/Plugins/Runtime/Metasound/Source/MetasoundEngine/Private/MetasoundOutputWatcher.cpp
@@ -30,6 +30,11 @@ namespace Metasound::Private
 		}
 	}
 
+	Frontend::FAnalyzerAddress FMetasoundOutputWatcher::GetAddress() const
+	{
+		return View.AnalyzerAddress;
+	}
+
 	void FMetasoundOutputWatcher::Update(const TFunctionRef<void(FName, const FMetaSoundOutput&)> OnOutputChanged)
 	{
 		METASOUND_TRACE_CPUPROFILER_EVENT_SCOPE(FMetasoundOutputWatcher::Update);
diff --git a/Engine/Plugins/Runtime/Metasound/Source/MetasoundEngine/Public/MetasoundOutputWatcher.h b/Engine/Plugins/Runtime/Metasound/Source/MetasoundEngine/Public/MetasoundOutputWatcher.h
--- a/Engine/Plugins/Runtime/Metasound/Source/MetasoundEngine/Public/MetasoundOutputWatcher.h
+++ b/Engine/Plugins/Runtime/Metasound/Source/MetasoundEngine/Public/MetasoundOutputWatcher.h
@@ -19,6 +19,12 @@ namespace Metasound::Private
 		FMetasoundOutputWatcher(Frontend::FAnalyzerAddress&& Address, const FOperatorSettings& OperatorSettings);
 
 		const FName Name;
+
+		/**
+		 * Get the address of the analyzer being watched
+		 * @return The analyzer address
+		 */
+		Frontend::FAnalyzerAddress GetAddress() const;
 		
 		/**
 		 * Update the watched outputs
