Fixed mismatched copy constructor/assignment operator in FMetasoundParameterPackProxy.

diff --git a/Engine/Plugins/Runtime/Metasound/Source/MetasoundFrontend/Public/MetasoundParameterPack.h b/Engine/Plugins/Runtime/Metasound/Source/MetasoundFrontend/Public/MetasoundParameterPack.h
--- a/Engine/Plugins/Runtime/Metasound/Source/MetasoundFrontend/Public/MetasoundParameterPack.h
+++ b/Engine/Plugins/Runtime/Metasound/Source/MetasoundFrontend/Public/MetasoundParameterPack.h
@@ -449,8 +449,6 @@ public:
 		: ParameterStoragePtr(Data)
 	{}
 
-	FMetasoundParameterPackProxy(const FMetasoundParameterPackProxy& Other) = default;
-
 	FSharedMetasoundParameterStoragePtr GetParamStorage()
 	{
 		return ParameterStoragePtr;
