[Metasounds] Adding Reset() function to audio bus reader/writer nodes.

diff --git a/Engine/Plugins/Runtime/Metasound/Source/MetasoundEngine/Private/MetasoundAudioBusReaderNode.cpp b/Engine/Plugins/Runtime/Metasound/Source/MetasoundEngine/Private/MetasoundAudioBusReaderNode.cpp
--- a/Engine/Plugins/Runtime/Metasound/Source/MetasoundEngine/Private/MetasoundAudioBusReaderNode.cpp
+++ b/Engine/Plugins/Runtime/Metasound/Source/MetasoundEngine/Private/MetasoundAudioBusReaderNode.cpp
@@ -91,55 +91,30 @@ namespace Metasound
 			using namespace AudioBusReaderNode; 
 			const FDataReferenceCollection& InputCollection = InParams.InputDataReferences;
 
-			bool bHasEnvironmentVars = InParams.Environment.Contains<Audio::FDeviceId>(SourceInterface::Environment::DeviceID) && InParams.Environment.Contains<int32>(SourceInterface::Environment::AudioMixerNumOutputFrames);
+			bool bHasEnvironmentVars = InParams.Environment.Contains<Audio::FDeviceId>(SourceInterface::Environment::DeviceID);
+			bHasEnvironmentVars &= InParams.Environment.Contains<int32>(SourceInterface::Environment::AudioMixerNumOutputFrames);
 			
 			if (bHasEnvironmentVars)
 			{
-				Audio::FDeviceId AudioDeviceId = InParams.Environment.GetValue<Audio::FDeviceId>(SourceInterface::Environment::DeviceID);
-				int32 AudioMixerOutputFrames = InParams.Environment.GetValue<int32>(SourceInterface::Environment::AudioMixerNumOutputFrames);
-
 				FAudioBusAssetReadRef AudioBusIn = InputCollection.GetDataReadReferenceOrConstruct<FAudioBusAsset>(METASOUND_GET_PARAM_NAME(InParamAudioBusInput));
-				
-				return MakeUnique<TAudioBusReaderOperator<NumChannels>>(InParams.OperatorSettings, AudioMixerOutputFrames, AudioDeviceId, AudioBusIn);
+				return MakeUnique<TAudioBusReaderOperator<NumChannels>>(InParams, AudioBusIn);
 			}
 			else
 			{
-				UE_LOG(LogMetaSound, Warning, TEXT("Audio bus reader node requires audio device ID '%s' and audio mixer num output frames '%s' environment variables"), *SourceInterface::Environment::DeviceID.ToString(), *SourceInterface::Environment::AudioMixerNumOutputFrames.ToString());
+				UE_LOG(LogMetaSound, Warning, TEXT("Audio bus reader node requires audio device ID '%s' and audio mixer num output frames '%s' environment variables")
+					, *SourceInterface::Environment::DeviceID.ToString(), *SourceInterface::Environment::AudioMixerNumOutputFrames.ToString());
 				return nullptr;
 			}
 		}
 
-		TAudioBusReaderOperator(const FOperatorSettings& InSettings, int32 InAudioMixerOutputFrames, Audio::FDeviceId InAudioDeviceId, const FAudioBusAssetReadRef& InAudioBusAsset)
-			: AudioBusAsset(InAudioBusAsset)
-			, AudioMixerOutputFrames(InAudioMixerOutputFrames)
-			, AudioDeviceId(InAudioDeviceId)
-			, SampleRate(InSettings.GetSampleRate())
+		TAudioBusReaderOperator(const FCreateOperatorParams& InParams, const FAudioBusAssetReadRef& InAudioBusAsset) : AudioBusAsset(InAudioBusAsset)
 		{
 			for (int32 ChannelIndex = 0; ChannelIndex < NumChannels; ++ChannelIndex)
 			{
-				AudioOutputs.Add(FAudioBufferWriteRef::CreateNew(InSettings));
+				AudioOutputs.Add(FAudioBufferWriteRef::CreateNew(InParams.OperatorSettings));
 			}
-			
-			const FAudioBusProxyPtr& AudioBusProxy = AudioBusAsset->GetAudioBusProxy();
-			if (AudioBusProxy.IsValid())
-			{
-				if (FAudioDeviceManager* ADM = FAudioDeviceManager::Get())
-				{
-					if (FAudioDevice* AudioDevice = ADM->GetAudioDeviceRaw(AudioDeviceId))
-					{
-						// Start the audio bus in case it's not already started
-						AudioBusChannels = AudioBusProxy->NumChannels;
-						Audio::FAudioBusKey AudioBusKey = Audio::FAudioBusKey(AudioBusProxy->AudioBusId);
-						UAudioBusSubsystem* AudioBusSubsystem = AudioDevice->GetSubsystem<UAudioBusSubsystem>();
-						check(AudioBusSubsystem);
-						AudioBusSubsystem->StartAudioBus(AudioBusKey, AudioBusChannels, false);
 
-						// Create a bus patch output with enough room for the number of samples we expect and some buffering
-						int32 BlockSizeFrames = InSettings.GetNumFramesPerBlock();
-						AudioBusPatchOutput = AudioBusSubsystem->AddPatchOutputForAudioBus(AudioBusKey, BlockSizeFrames, AudioBusChannels);
-					}
-				}
-			}
+			Reset(InParams);
 		}
 
 		virtual FDataReferenceCollection GetInputs() const override
@@ -224,6 +199,55 @@ namespace Metasound
 			}
 		}
 
+		void Reset(const IOperator::FResetParams& InParams)
+		{
+			using namespace Frontend;
+			using namespace AudioBusReaderNode;
+
+			bool bHasEnvironmentVars = InParams.Environment.Contains<Audio::FDeviceId>(SourceInterface::Environment::DeviceID);
+			bHasEnvironmentVars &= InParams.Environment.Contains<int32>(SourceInterface::Environment::AudioMixerNumOutputFrames);
+
+			if (bHasEnvironmentVars)
+			{
+				SampleRate = InParams.OperatorSettings.GetSampleRate();
+				AudioDeviceId = InParams.Environment.GetValue<Audio::FDeviceId>(SourceInterface::Environment::DeviceID);
+				AudioMixerOutputFrames = InParams.Environment.GetValue<int32>(SourceInterface::Environment::AudioMixerNumOutputFrames);
+			}
+			else
+			{
+				UE_LOG(LogMetaSound, Warning, TEXT("Audio bus reader node requires audio device ID '%s' and audio mixer num output frames '%s' environment variables")
+					, *SourceInterface::Environment::DeviceID.ToString(), *SourceInterface::Environment::AudioMixerNumOutputFrames.ToString());
+			}
+			
+			for(const FAudioBufferWriteRef& Buffer : AudioOutputs)
+			{
+				Buffer->Zero();
+			}
+
+			const FAudioBusProxyPtr& AudioBusProxy = AudioBusAsset->GetAudioBusProxy();
+			if (AudioBusProxy.IsValid())
+			{
+				if (FAudioDeviceManager* ADM = FAudioDeviceManager::Get())
+				{
+					if (FAudioDevice* AudioDevice = ADM->GetAudioDeviceRaw(AudioDeviceId))
+					{
+						// Start the audio bus in case it's not already started
+						AudioBusChannels = AudioBusProxy->NumChannels;
+						const Audio::FAudioBusKey AudioBusKey(AudioBusProxy->AudioBusId);
+						UAudioBusSubsystem* AudioBusSubsystem = AudioDevice->GetSubsystem<UAudioBusSubsystem>();
+						check(AudioBusSubsystem);
+						AudioBusSubsystem->StartAudioBus(AudioBusKey, AudioBusChannels, false);
+
+						// Create a bus patch output with enough room for the number of samples we expect and some buffering
+						const int32 BlockSizeFrames = InParams.OperatorSettings.GetNumFramesPerBlock();
+						AudioBusPatchOutput = AudioBusSubsystem->AddPatchOutputForAudioBus(AudioBusKey, BlockSizeFrames, AudioBusChannels);
+					}
+				}
+			}
+
+			bFirstBlock = true;
+		}
+
 	private:
 		FAudioBusAssetReadRef AudioBusAsset;
 		TArray<FAudioBufferWriteRef> AudioOutputs;
diff --git a/Engine/Plugins/Runtime/Metasound/Source/MetasoundEngine/Private/MetasoundAudioBusWriterNode.cpp b/Engine/Plugins/Runtime/Metasound/Source/MetasoundEngine/Private/MetasoundAudioBusWriterNode.cpp
--- a/Engine/Plugins/Runtime/Metasound/Source/MetasoundEngine/Private/MetasoundAudioBusWriterNode.cpp
+++ b/Engine/Plugins/Runtime/Metasound/Source/MetasoundEngine/Private/MetasoundAudioBusWriterNode.cpp
@@ -1,7 +1,6 @@
 // Copyright Epic Games, Inc. All Rights Reserved.
 
-#include "AudioBusSubsystem.h"
-#include "AudioDevice.h"
+#include "AudioMixerDevice.h"
 #include "Internationalization/Text.h"
 #include "MediaPacket.h"
 #include "MetasoundAudioBuffer.h"
@@ -77,17 +76,15 @@ namespace Metasound
 		static TUniquePtr<IOperator> CreateOperator(const FCreateOperatorParams& InParams, FBuildErrorArray& OutErrors)
 		{
 			using namespace Frontend;
-
 			using namespace AudioBusWriterNode;
+
 			const FDataReferenceCollection& InputCollection = InParams.InputDataReferences;
 
-			bool bHasEnvironmentVars = InParams.Environment.Contains<Audio::FDeviceId>(SourceInterface::Environment::DeviceID) && InParams.Environment.Contains<int32>(SourceInterface::Environment::AudioMixerNumOutputFrames);
-			
+			bool bHasEnvironmentVars = InParams.Environment.Contains<Audio::FDeviceId>(SourceInterface::Environment::DeviceID);
+			bHasEnvironmentVars &= InParams.Environment.Contains<int32>(SourceInterface::Environment::AudioMixerNumOutputFrames);
+
 			if (bHasEnvironmentVars)
 			{
-				Audio::FDeviceId AudioDeviceId = InParams.Environment.GetValue<Audio::FDeviceId>(SourceInterface::Environment::DeviceID);
-				int32 AudioMixerOutputFrames = InParams.Environment.GetValue<int32>(SourceInterface::Environment::AudioMixerNumOutputFrames);
-
 				FAudioBusAssetReadRef AudioBusIn = InputCollection.GetDataReadReferenceOrConstruct<FAudioBusAsset>(METASOUND_GET_PARAM_NAME(InParamAudioBusOutput));
 
 				TArray<FAudioBufferReadRef> AudioInputs;
@@ -96,22 +93,43 @@ namespace Metasound
 					AudioInputs.Add(InputCollection.GetDataReadReferenceOrConstruct<FAudioBuffer>(METASOUND_GET_PARAM_NAME_WITH_INDEX(InParamAudio, ChannelIndex), InParams.OperatorSettings));
 				}
 
-				return MakeUnique<TAudioBusWriterOperator<NumChannels>>(InParams, AudioMixerOutputFrames, AudioDeviceId, MoveTemp(AudioBusIn), MoveTemp(AudioInputs));
+				return MakeUnique<TAudioBusWriterOperator<NumChannels>>(InParams, MoveTemp(AudioBusIn), MoveTemp(AudioInputs));
 			}
 			else
 			{
-				UE_LOG(LogMetaSound, Warning, TEXT("Audio bus writer node requires audio device ID '%s' and audio mixer num output frames '%s' environment variables"), *SourceInterface::Environment::DeviceID.ToString(), *SourceInterface::Environment::AudioMixerNumOutputFrames.ToString());
+				UE_LOG(LogMetaSound, Warning, TEXT("Audio bus writer node requires audio device ID '%s' and audio mixer num output frames '%s' environment variables")
+					, *SourceInterface::Environment::DeviceID.ToString(), *SourceInterface::Environment::AudioMixerNumOutputFrames.ToString());
 				return nullptr;
 			}
 		}
 
-		TAudioBusWriterOperator(const FCreateOperatorParams& InParams, int32 InAudioMixerOutputFrames, Audio::FDeviceId InAudioDeviceId, FAudioBusAssetReadRef InAudioBusAsset, TArray<FAudioBufferReadRef> InAudioInputs) :
-			AudioBusAsset(MoveTemp(InAudioBusAsset)),
-			AudioInputs(MoveTemp(InAudioInputs)),
-			AudioMixerOutputFrames(InAudioMixerOutputFrames),
-			AudioDeviceId(InAudioDeviceId),
-			SampleRate(InParams.OperatorSettings.GetSampleRate())
+		TAudioBusWriterOperator(const FCreateOperatorParams& InParams, FAudioBusAssetReadRef InAudioBusAsset, TArray<FAudioBufferReadRef> InAudioInputs)
+			: AudioBusAsset(MoveTemp(InAudioBusAsset))
+			, AudioInputs(MoveTemp(InAudioInputs))
 		{
+			Reset(InParams);
+		}
+
+		void Reset(const IOperator::FResetParams& InParams)
+		{
+			using namespace Frontend;
+			using namespace AudioBusWriterNode;
+
+			bool bHasEnvironmentVars = InParams.Environment.Contains<Audio::FDeviceId>(SourceInterface::Environment::DeviceID);
+			bHasEnvironmentVars &= InParams.Environment.Contains<int32>(SourceInterface::Environment::AudioMixerNumOutputFrames);
+
+			if (bHasEnvironmentVars)
+			{
+				SampleRate = InParams.OperatorSettings.GetSampleRate();
+				AudioDeviceId = InParams.Environment.GetValue<Audio::FDeviceId>(SourceInterface::Environment::DeviceID);
+				AudioMixerOutputFrames = InParams.Environment.GetValue<int32>(SourceInterface::Environment::AudioMixerNumOutputFrames);
+			}
+			else
+			{
+				UE_LOG(LogMetaSound, Warning, TEXT("Audio bus writer node requires audio device ID '%s' and audio mixer num output frames '%s' environment variables")
+					, *SourceInterface::Environment::DeviceID.ToString(), *SourceInterface::Environment::AudioMixerNumOutputFrames.ToString());
+			}
+			
 			const FAudioBusProxyPtr& AudioBusProxy = AudioBusAsset->GetAudioBusProxy();
 			if (AudioBusProxy.IsValid())
 			{
@@ -119,21 +137,20 @@ namespace Metasound
 				{
 					if (FAudioDevice* AudioDevice = ADM->GetAudioDeviceRaw(AudioDeviceId))
 					{
-						UAudioBusSubsystem* AudioBusSubsystem = AudioDevice->GetSubsystem<UAudioBusSubsystem>();
-						check(AudioBusSubsystem);
-						Audio::FAudioBusKey AudioBusKey = Audio::FAudioBusKey(AudioBusProxy->AudioBusId);	
-						// Start the audio bus in case it's not already started					
+						// Start the audio bus in case it's not already started
 						AudioBusChannels = AudioBusProxy->NumChannels;
-						AudioBusSubsystem->StartAudioBus(AudioBusKey, AudioBusChannels, false);
+						AudioDevice->StartAudioBus(AudioBusProxy->AudioBusId, AudioBusChannels, false);
 
 						BlockSizeFrames = InParams.OperatorSettings.GetNumFramesPerBlock();
 						InterleavedBuffer.AddZeroed(BlockSizeFrames * AudioBusChannels);
 
 						// Create a bus patch input with enough room for the number of samples we expect and some buffering
-						AudioBusPatchInput = AudioBusSubsystem->AddPatchInputForAudioBus(AudioBusKey, BlockSizeFrames, AudioBusChannels);
+						AudioBusPatchInput = AudioDevice->AddPatchInputForAudioBus(AudioBusProxy->AudioBusId, BlockSizeFrames, AudioBusChannels);
 					}
 				}
 			}
+
+			bFirstBlock = true;
 		}
 
 		virtual FDataReferenceCollection GetInputs() const override
