Adding Metasound Operator Cache

diff --git a/Engine/Plugins/Runtime/Metasound/Source/MetasoundFrontend/Public/MetasoundDataTypeRegistrationMacro.h b/Engine/Plugins/Runtime/Metasound/Source/MetasoundFrontend/Public/MetasoundDataTypeRegistrationMacro.h
--- a/Engine/Plugins/Runtime/Metasound/Source/MetasoundFrontend/Public/MetasoundDataTypeRegistrationMacro.h
+++ b/Engine/Plugins/Runtime/Metasound/Source/MetasoundFrontend/Public/MetasoundDataTypeRegistrationMacro.h
@@ -475,7 +475,7 @@ namespace Metasound
 					{
 						if constexpr (!TIsArrayType<TDataType>::Value)
 						{
-							this->LiteralAssignmentFunction = [](const FOperatorSettings& InOperatorSettings, const FLiteral& InLiteral, FAnyDataReference& OutDataRef)
+							this->LiteralAssignmentFunction = [](const FOperatorSettings& InOperatorSettings, const FLiteral& InLiteral, const FAnyDataReference& OutDataRef)
 							{
 								*OutDataRef.GetWritableValue<TDataType>() = TDataTypeLiteralFactory<TDataType>::CreateExplicitArgs(InOperatorSettings, InLiteral);
 							};
diff --git a/Engine/Plugins/Runtime/Metasound/Source/MetasoundFrontend/Public/MetasoundFrontendDataTypeRegistry.h b/Engine/Plugins/Runtime/Metasound/Source/MetasoundFrontend/Public/MetasoundFrontendDataTypeRegistry.h
--- a/Engine/Plugins/Runtime/Metasound/Source/MetasoundFrontend/Public/MetasoundFrontendDataTypeRegistry.h
+++ b/Engine/Plugins/Runtime/Metasound/Source/MetasoundFrontend/Public/MetasoundFrontendDataTypeRegistry.h
@@ -21,7 +21,7 @@ namespace Metasound
 		using IParameterAssignmentFunction = TUniqueFunction<void(const void*,void*)>;
 
 		// Function signature for assigning a FLiteral to a FAnyDataReference 
-		using FLiteralAssignmentFunction = void(*)(const FOperatorSettings& InOperatorSettings, const FLiteral& InLiteral, FAnyDataReference& OutDataRef);
+		using FLiteralAssignmentFunction = void(*)(const FOperatorSettings& InOperatorSettings, const FLiteral& InLiteral, const FAnyDataReference& OutDataRef);
 
 
 		/** FDataTypeRegsitryInfo contains runtime inspectable behavior of a registered
diff --git a/Engine/Plugins/Runtime/Metasound/Source/MetasoundGenerator/Private/MetasoundGenerator.cpp b/Engine/Plugins/Runtime/Metasound/Source/MetasoundGenerator/Private/MetasoundGenerator.cpp
--- a/Engine/Plugins/Runtime/Metasound/Source/MetasoundGenerator/Private/MetasoundGenerator.cpp
+++ b/Engine/Plugins/Runtime/Metasound/Source/MetasoundGenerator/Private/MetasoundGenerator.cpp
@@ -3,6 +3,8 @@
 
 #include "AudioParameter.h"
 #include "MetasoundFrontendDataTypeRegistry.h"
+#include "MetasoundGeneratorBuilder.h"
+#include "MetasoundGeneratorModuleImpl.h"
 #include "MetasoundOperatorBuilder.h"
 #include "MetasoundOperatorInterface.h"
 #include "MetasoundOutputNode.h"
@@ -16,6 +18,7 @@
 #include "Analysis/MetasoundFrontendVertexAnalyzer.h"
 #include "DSP/FloatArrayMath.h"
 #include "Interfaces/MetasoundFrontendSourceInterface.h"
+#include "ProfilingDebugging/CountersTrace.h"
 
 #ifndef ENABLE_METASOUNDGENERATOR_INVALID_SAMPLE_VALUE_LOGGING 
 #define ENABLE_METASOUNDGENERATOR_INVALID_SAMPLE_VALUE_LOGGING !UE_BUILD_SHIPPING
@@ -26,6 +29,8 @@ namespace Metasound
 	namespace ConsoleVariables
 	{
 		static bool bEnableAsyncMetaSoundGeneratorBuilder = true;
+		static bool bEnableExperimentalOneShotOperatorCache = false;
+		static bool bEnableExperimentalOperatorCache = false;
 #if ENABLE_METASOUNDGENERATOR_INVALID_SAMPLE_VALUE_LOGGING
 		static bool bEnableMetaSoundGeneratorNonFiniteLogging = false;
 		static bool bEnableMetaSoundGeneratorInvalidSampleValueLogging = false;
@@ -137,6 +142,7 @@ namespace Metasound
 			uint64 StartCycle = 0;
 			FRenderTimer* Timer = nullptr;
 		};
+
 #endif // if ENABLE_METASOUND_GENERATOR_RENDER_TIMING
 	}
 }
@@ -148,6 +154,20 @@ FAutoConsoleVariableRef CVarMetaSoundEnableAsyncGeneratorBuilder(
 	TEXT("Default: true"),
 	ECVF_Default);
 
+FAutoConsoleVariableRef CVarMetaSoundEnableExperimentalOneShotOperatorCache(
+	TEXT("au.MetaSound.Experimental.EnableOneShotOperatorCache"),
+	Metasound::ConsoleVariables::bEnableExperimentalOneShotOperatorCache,
+	TEXT("Enables caching of MetaSound operators using the OneShot source interface.\n")
+	TEXT("Default: false"),
+	ECVF_Default);
+
+FAutoConsoleVariableRef CVarMetaSoundEnableExperimentalOperatorCache(
+	TEXT("au.MetaSound.Experimental.EnableOperatorCache"),
+	Metasound::ConsoleVariables::bEnableExperimentalOperatorCache,
+	TEXT("Enables caching of all MetaSound operators.\n")
+	TEXT("Default: false"),
+	ECVF_Default);
+
 #if ENABLE_METASOUNDGENERATOR_INVALID_SAMPLE_VALUE_LOGGING
 
 FAutoConsoleVariableRef CVarMetaSoundEnableGeneratorNonFiniteLogging(
@@ -175,6 +195,14 @@ FAutoConsoleVariableRef CVarMetaSoundGeneratorSampleValueThrehshold(
 
 namespace Metasound
 {
+	namespace MetasoundGeneratorPrivate
+	{
+		bool HasOneShotInterface(const FVertexInterface& Interface)
+		{
+			return Interface.GetOutputInterface().Contains(Frontend::SourceOneShotInterface::Outputs::OnFinished);
+		}
+	}
+
 	void FMetasoundGeneratorInitParams::Release()
 	{
 		Graph.Reset();
@@ -184,246 +212,6 @@ namespace Metasound
 		DefaultParameters = {};
 	}
 
-	FAsyncMetaSoundBuilder::FAsyncMetaSoundBuilder(FMetasoundGenerator* InGenerator, FMetasoundGeneratorInitParams&& InInitParams, bool bInTriggerGenerator)
-		: Generator(InGenerator)
-		, InitParams(MoveTemp(InInitParams))
-		, bTriggerGenerator(bInTriggerGenerator)
-	{
-	}
-
-	void FAsyncMetaSoundBuilder::DoWork()
-	{
-		using namespace Audio;
-		using namespace Frontend;
-		using namespace MetasoundGeneratorPrivate;
-
-		METASOUND_TRACE_CPUPROFILER_EVENT_SCOPE_TEXT(*FString::Printf(TEXT("AsyncMetaSoundBuilder::DoWork %s"), *InitParams.MetaSoundName));
-
-		FBuildResults BuildResults;
-
-		// Create an instance of the new graph
-		FGraphOperatorAndInputs GraphOperatorAndInputs = BuildGraphOperator(MoveTemp(InitParams.DefaultParameters), BuildResults);
-		LogBuildErrors(BuildResults);
-
-		if (GraphOperatorAndInputs.Operator.IsValid())
-		{
-			// Create graph analyzer
-			TUniquePtr<FGraphAnalyzer> GraphAnalyzer = BuildGraphAnalyzer(MoveTemp(BuildResults.InternalDataReferences));
-	
-			// Collect data for generator
-			FMetasoundGeneratorData GeneratorData = BuildGeneratorData(InitParams, MoveTemp(GraphOperatorAndInputs), MoveTemp(GraphAnalyzer));
-
-			Generator->SetPendingGraph(MoveTemp(GeneratorData), bTriggerGenerator);
-		}
-		else 
-		{
-			UE_LOG(LogMetaSound, Error, TEXT("Failed to build Metasound operator from graph in MetasoundSource [%s]"), *InitParams.MetaSoundName);
-			// Set null generator data to inform that generator failed to build. 
-			// Otherwise, generator will continually wait for a new generator.
-			Generator->SetPendingGraphBuildFailed();
-		}
-
-		InitParams.Release();
-	}
-
-	MetasoundGeneratorPrivate::FMetasoundGeneratorData FAsyncMetaSoundBuilder::BuildGeneratorData(const FMetasoundGeneratorInitParams& InInitParams, FAsyncMetaSoundBuilder::FGraphOperatorAndInputs&& InGraphOperatorAndInputs, TUniquePtr<Frontend::FGraphAnalyzer> InAnalyzer) const
-	{
-		using namespace Audio;
-		using namespace Frontend;
-		using namespace MetasoundGeneratorPrivate;
-
-		METASOUND_TRACE_CPUPROFILER_EVENT_SCOPE_TEXT(TEXT("AsyncMetaSoundBuilder::BuildGeneratorData"));
-
-		checkf(InGraphOperatorAndInputs.Operator.IsValid(), TEXT("Graph operator must be a valid object"));
-
-		// Gather relevant input and output references
-		FVertexInterfaceData VertexData(InInitParams.Graph->GetVertexInterface());
-		InGraphOperatorAndInputs.Operator->Bind(VertexData);
-
-		// Replace input data with writable inputs
-		VertexData.GetInputs() = InGraphOperatorAndInputs.Inputs;
-
-		// Get inputs
-		FTriggerWriteRef PlayTrigger = VertexData.GetInputs().GetOrConstructDataWriteReference<FTrigger>(SourceInterface::Inputs::OnPlay, InInitParams.OperatorSettings, false);
-
-		// Get outputs
-		TArray<FAudioBufferReadRef> OutputBuffers = FindOutputAudioBuffers(VertexData);
-		FTriggerReadRef FinishTrigger = TDataReadReferenceFactory<FTrigger>::CreateExplicitArgs(InInitParams.OperatorSettings, false);
-
-		if (InInitParams.Graph->GetVertexInterface().GetOutputInterface().Contains(SourceOneShotInterface::Outputs::OnFinished))
-		{
-			// Only attempt to retrieve the on finished trigger if it exists.
-			// Attempting to retrieve a data reference from a non-existent vertex 
-			// will log an error. 
-			FinishTrigger = VertexData.GetOutputs().GetOrConstructDataReadReference<FTrigger>(SourceOneShotInterface::Outputs::OnFinished, InitParams.OperatorSettings, false);
-		}
-
-		// Create the parameter setter map so parameter packs can be cracked
-		// open and distributed as appropriate...
-		const IDataTypeRegistry& DataTypeRegistry = IDataTypeRegistry::Get();
-		FParameterSetterSortedMap ParameterSetters;
-		TMap<FName, FParameterPackSetter> ParameterPackSetters;
-		FInputVertexInterfaceData& GraphInputs = VertexData.GetInputs();
-		for (const MetasoundVertexDataPrivate::TBinding<FInputDataVertex>& Binding : GraphInputs)
-		{
-			// Only assign inputs that are writable. 
-			if (EDataReferenceAccessType::Write == Binding.GetAccessType())
-			{
-				if (const FAnyDataReference* DataRef = Binding.GetDataReference())
-				{
-					const FInputDataVertex& InputVertex = Binding.GetVertex();
-					const Frontend::IParameterAssignmentFunction& PackSetter = DataTypeRegistry.GetRawAssignmentFunction(InputVertex.DataTypeName);
-					if (PackSetter)
-					{
-						FParameterPackSetter ParameterPackSetter(InputVertex.DataTypeName, DataRef->GetRaw(), PackSetter);
-						ParameterPackSetters.Add(InputVertex.VertexName, ParameterPackSetter);
-					}
-
-					Frontend::FLiteralAssignmentFunction LiteralSetter = DataTypeRegistry.GetLiteralAssignmentFunction(InputVertex.DataTypeName);
-					if (LiteralSetter)
-					{
-						ParameterSetters.Add(InputVertex.VertexName, FParameterSetter{LiteralSetter, *DataRef});
-					}
-				}
-			}
-		}
-
-		// Set data needed for graph
-		return FMetasoundGeneratorData 
-		{
-			InInitParams.OperatorSettings,
-			MoveTemp(InGraphOperatorAndInputs.Operator),
-			MoveTemp(VertexData),
-			MoveTemp(ParameterSetters),
-			MoveTemp(ParameterPackSetters),
-			MoveTemp(InAnalyzer),
-			MoveTemp(OutputBuffers),
-			MoveTemp(PlayTrigger),
-			MoveTemp(FinishTrigger),
-		};
-	}
-
-	FAsyncMetaSoundBuilder::FGraphOperatorAndInputs FAsyncMetaSoundBuilder::BuildGraphOperator(TArray<FAudioParameter>&& InParameters, FBuildResults& OutBuildResults) const
-	{
-		METASOUND_TRACE_CPUPROFILER_EVENT_SCOPE_TEXT(TEXT("AsyncMetaSoundBuilder::BuildGraphOperator"));
-		using namespace Frontend;
-
-		// Choose which type of data reference access to create depending upon the access of the vertex.
-		auto VertexAccessTypeToDataReferenceAccessType = [](EVertexAccessType InVertexAccessType) -> EDataReferenceAccessType
-		{
-			switch(InVertexAccessType)
-			{
-				case EVertexAccessType::Value:
-					return EDataReferenceAccessType::Value;
-
-				case EVertexAccessType::Reference:
-				default:
-					return EDataReferenceAccessType::Write;
-			}
-		};
-
-		const FInputVertexInterface& InputInterface = InitParams.Graph->GetVertexInterface().GetInputInterface();
-		FInputVertexInterfaceData InputData(InputInterface);
-
-		// Set input data based on the input parameters and the input interface
-		IDataTypeRegistry& DataRegistry = IDataTypeRegistry::Get();
-		for (FAudioParameter& Parameter : InParameters)
-		{
-			const FName ParamName = Parameter.ParamName;
-			if (const FInputDataVertex* InputVertex = InputInterface.Find(ParamName))
-			{
-				FLiteral Literal = Frontend::ConvertParameterToLiteral(MoveTemp(Parameter));
-
-				TOptional<FAnyDataReference> DataReference = DataRegistry.CreateDataReference(InputVertex->DataTypeName, VertexAccessTypeToDataReferenceAccessType(InputVertex->AccessType), Literal, InitParams.OperatorSettings);
-
-				if (DataReference)
-				{
-					InputData.BindVertex(ParamName, *DataReference);
-				}
-				else if(MetaSoundParameterEnableWarningOnIgnoredParameterCVar)
-				{
-					UE_LOG(LogMetaSound, Warning, TEXT("Failed to create initial input data reference from parameter %s of type %s on graph in MetaSoundSource [%s]"), *ParamName.ToString(), *InputVertex->DataTypeName.ToString(), *InitParams.MetaSoundName);
-				}
-			}
-			else if(MetaSoundParameterEnableWarningOnIgnoredParameterCVar)
-			{
-				UE_LOG(LogMetaSound, Warning, TEXT("Failed to set initial input parameter %s on graph in MetaSoundSource [%s]"), *ParamName.ToString(), *InitParams.MetaSoundName);
-			}
-		}
-
-		// Set any remaining inputs to their default values.
-		for (const MetasoundVertexDataPrivate::TBinding<FInputDataVertex>& Binding : InputData)
-		{
-			// Only create data reference if something does not already exist. 
-			if (!Binding.IsBound())
-			{
-				const FInputDataVertex& InputVertex = Binding.GetVertex();
-				EDataReferenceAccessType AccessType = VertexAccessTypeToDataReferenceAccessType(InputVertex.AccessType);
-				TOptional<FAnyDataReference> DataReference = DataRegistry.CreateDataReference(InputVertex.DataTypeName, VertexAccessTypeToDataReferenceAccessType(InputVertex.AccessType), InputVertex.GetDefaultLiteral(), InitParams.OperatorSettings);
-
-				if (DataReference)
-				{
-					InputData.BindVertex(InputVertex.VertexName, *DataReference);
-				}
-			}
-		}
-
-		// Reset as elements in array have been moved.
-		InParameters.Reset();
-
-		// Create an instance of the new graph
-		FBuildGraphOperatorParams BuildParams { *InitParams.Graph, InitParams.OperatorSettings, InputData, InitParams.Environment };
-		FGraphOperatorAndInputs OpAndInputs;
-		OpAndInputs.Operator = FOperatorBuilder(InitParams.BuilderSettings).BuildGraphOperator(BuildParams, OutBuildResults);
-		OpAndInputs.Inputs = InputData;
-
-		return OpAndInputs;
-	}
-
-	TUniquePtr<Frontend::FGraphAnalyzer> FAsyncMetaSoundBuilder::BuildGraphAnalyzer(TMap<FGuid, FDataReferenceCollection>&& InInternalDataReferences) const
-	{
-		METASOUND_TRACE_CPUPROFILER_EVENT_SCOPE_TEXT(TEXT("AsyncMetaSoundBuilder::BuildGraphAnalyzer"));
-		using namespace Frontend;
-
-		if (InitParams.BuilderSettings.bPopulateInternalDataReferences)
-		{
-			const uint64 InstanceID = InitParams.Environment.GetValue<uint64>(SourceInterface::Environment::TransmitterID);
-			return MakeUnique<FGraphAnalyzer>(InitParams.OperatorSettings, InstanceID, MoveTemp(InInternalDataReferences));
-		}
-		return nullptr;
-	}
-
-	void FAsyncMetaSoundBuilder::LogBuildErrors(const FBuildResults& InBuildResults) const
-	{
-		// Log build errors
-		for (const IOperatorBuilder::FBuildErrorPtr& Error : InBuildResults.Errors)
-		{
-			if (Error.IsValid())
-			{
-				UE_LOG(LogMetaSound, Warning, TEXT("MetasoundSource [%s] build error [%s] \"%s\""), *InitParams.MetaSoundName, *(Error->GetErrorType().ToString()), *(Error->GetErrorDescription().ToString()));
-			}
-		}
-	}
-
-	TArray<FAudioBufferReadRef> FAsyncMetaSoundBuilder::FindOutputAudioBuffers(const FVertexInterfaceData& InVertexData) const
-	{
-		TArray<FAudioBufferReadRef> OutputBuffers;
-
-		const FOutputVertexInterfaceData& OutputVertexData = InVertexData.GetOutputs();
-
-		// Get output audio buffers.
-		for (const FVertexName& AudioOutputName : InitParams.AudioOutputNames)
-		{
-			if (!OutputVertexData.IsVertexBound(AudioOutputName))
-			{
-				UE_LOG(LogMetaSound, Warning, TEXT("MetasoundSource [%s] does not contain audio output [%s] in output"), *InitParams.MetaSoundName, *AudioOutputName.ToString());
-			}
-			OutputBuffers.Add(OutputVertexData.GetOrConstructDataReadReference<FAudioBuffer>(AudioOutputName, InitParams.OperatorSettings));
-		}
-
-		return OutputBuffers;
-	}
-
 	FMetasoundGenerator::FMetasoundGenerator(FMetasoundGeneratorInitParams&& InParams)
 		: MetasoundName(InParams.MetaSoundName)
 		, OperatorSettings(InParams.OperatorSettings)
@@ -453,20 +241,34 @@ PRAGMA_DISABLE_DEPRECATION_WARNINGS
 PRAGMA_ENABLE_DEPRECATION_WARNINGS
 		ParameterPackReceiver = FDataTransmissionCenter::Get().RegisterNewReceiver<FMetasoundParameterStorageWrapper>(ParameterPackSendAddress, FReceiverInitParams{ InParams.OperatorSettings });
 
-		BuilderTask = MakeUnique<FBuilderTask>(this, MoveTemp(InParams), true /* bTriggerGenerator */);
-		
-		if (!InParams.bBuildSynchronous && Metasound::ConsoleVariables::bEnableAsyncMetaSoundGeneratorBuilder)
+		// attempt to use operator cache instead of building a new operator.
+		bool bDidUseCachedOperator = false;
+		if (ConsoleVariables::bEnableExperimentalOneShotOperatorCache || ConsoleVariables::bEnableExperimentalOperatorCache)
 		{
-			// Build operator asynchronously
-			BuilderTask->StartBackgroundTask(GBackgroundPriorityThreadPool);
+			bUseOperatorCache = ConsoleVariables::bEnableExperimentalOperatorCache || MetasoundGeneratorPrivate::HasOneShotInterface(InParams.Graph->GetVertexInterface());
+
+			if (bUseOperatorCache)
+			{
+				bDidUseCachedOperator = TryUseCachedOperator(InParams, true /* bTriggerGenerator */);
+			}
 		}
-		else
+
+		if (!bDidUseCachedOperator)
 		{
-			// Build operator synchronously
-			BuilderTask->StartSynchronousTask();
-			BuilderTask = nullptr;
-			UpdateGraphIfPending();
-			bIsWaitingForFirstGraph = false;
+			BuilderTask = MakeUnique<FAsyncTask<FAsyncMetaSoundBuilder>>(this, MoveTemp(InParams), true /* bTriggerGenerator */);
+			
+			if (!InParams.bBuildSynchronous && Metasound::ConsoleVariables::bEnableAsyncMetaSoundGeneratorBuilder)
+			{
+				// Build operator asynchronously
+				BuilderTask->StartBackgroundTask(GBackgroundPriorityThreadPool);
+			}
+			else
+			{
+				// Build operator synchronously
+				BuilderTask->StartSynchronousTask();
+				BuilderTask = nullptr;
+				UpdateGraphIfPending();
+			}
 		}
 	}
 
@@ -478,11 +280,80 @@ PRAGMA_ENABLE_DEPRECATION_WARNINGS
 			BuilderTask = nullptr;
 		}
 
+		if (bUseOperatorCache)
+		{
+			ReleaseOperatorToCache();
+		}
+
 		// Remove routing for parameter packs
 		ParameterPackReceiver.Reset();
 		FDataTransmissionCenter::Get().UnregisterDataChannelIfUnconnected(ParameterPackSendAddress);
 	}
 
+	bool FMetasoundGenerator::TryUseCachedOperator(FMetasoundGeneratorInitParams& InInitParams, bool bInTriggerGraph)
+	{
+		using namespace MetasoundGeneratorPrivate;
+		METASOUND_TRACE_CPUPROFILER_EVENT_SCOPE(FMetasoundGenerator::TryUseCachedOperator);
+
+		FMetasoundGeneratorModule& Module = FModuleManager::GetModuleChecked<FMetasoundGeneratorModule>("MetasoundGenerator");
+		TSharedPtr<FOperatorCache> OperatorCache = Module.GetOperatorCache();
+		if (OperatorCache.IsValid())
+		{
+			// See if the cache has an operator with matching OperatorID
+			OperatorID = InInitParams.Graph->GetInstanceID();
+			FOperatorAndInputs GraphOperatorAndInputs = OperatorCache->ClaimCachedOperator(OperatorID);
+			if (GraphOperatorAndInputs.Operator.IsValid())
+			{
+				UE_LOG(LogMetasoundGenerator, VeryVerbose, TEXT("Using cached operator %s for MetaSound %s"), *LexToString(OperatorID), *InInitParams.MetaSoundName); 
+
+				// Apply and default inputs to the operator.
+				GeneratorBuilder::ApplyAudioParameters(OperatorSettings, MoveTemp(InInitParams.DefaultParameters), GraphOperatorAndInputs.Inputs);
+
+				// Reset operator internal state before playing it.
+				if (IOperator::FResetFunction Reset = GraphOperatorAndInputs.Operator->GetResetFunction())
+				{
+					IOperator::FResetParams ResetParams {OperatorSettings, InInitParams.Environment};
+					Reset(GraphOperatorAndInputs.Operator.Get(), ResetParams);
+				}
+				
+				// Create data needed for MetaSound Generator
+				FMetasoundGeneratorData Data = GeneratorBuilder::BuildGeneratorData(InInitParams, MoveTemp(GraphOperatorAndInputs), nullptr);
+
+				SetGraph(MakeUnique<FMetasoundGeneratorData>(MoveTemp(Data)), bInTriggerGraph);
+				
+				return true;
+			}
+		}
+
+		return false;
+	}
+
+	void FMetasoundGenerator::ReleaseOperatorToCache()
+	{
+		using namespace MetasoundGeneratorPrivate;
+		METASOUND_TRACE_CPUPROFILER_EVENT_SCOPE(FMetasoundGenerator::ReleaseOperatorToCache);
+
+		if (OperatorID.IsValid())
+		{
+			FMetasoundGeneratorModule& Module = FModuleManager::GetModuleChecked<FMetasoundGeneratorModule>("MetasoundGenerator");
+			TSharedPtr<FOperatorCache> OperatorCache = Module.GetOperatorCache();
+			if (OperatorCache.IsValid())
+			{
+				TUniquePtr<IOperator> GraphOperator = RootExecuter.ReleaseOperator();
+
+				if (GraphOperator.IsValid())
+				{
+					// Release graph operator and input data to the cache
+					UE_LOG(LogMetasoundGenerator, VeryVerbose, TEXT("Caching operator %s"), *LexToString(OperatorID));
+					OperatorCache->AddOperatorToCache(OperatorID, MoveTemp(GraphOperator), MoveTemp(VertexInterfaceData.GetInputs()));
+				}
+			}
+		}
+
+		// Clear out any internal references to graph data
+		ClearGraph();
+	}
+
 	void FMetasoundGenerator::SetPendingGraph(MetasoundGeneratorPrivate::FMetasoundGeneratorData&& InData, bool bTriggerGraph)
 	{
 		using namespace MetasoundGeneratorPrivate;
@@ -585,9 +456,21 @@ PRAGMA_ENABLE_DEPRECATION_WARNINGS
 			OnPlayTriggerRef->TriggerFrame(0);
 		}
 
+		bIsWaitingForFirstGraph = false;
 		OnSetGraph.Broadcast();
 	}
 
+	void FMetasoundGenerator::ClearGraph()
+	{
+		RootExecuter.ReleaseOperator();
+		VertexInterfaceData = FVertexInterfaceData();
+		GraphOutputAudio.Reset();
+		OnFinishedTriggerRef = OnPlayTriggerRef = TDataWriteReference<FTrigger>::CreateNew(OperatorSettings);
+		GraphAnalyzer.Reset();
+		ParameterSetters.Reset();
+		ParameterPackSetters.Reset();
+	}
+
 	void FMetasoundGenerator::QueueParameterPack(TSharedPtr<FMetasoundParameterPackStorage> ParameterPack)
 	{
 		ParameterPackQueue.Enqueue(ParameterPack);
@@ -655,15 +538,16 @@ PRAGMA_ENABLE_DEPRECATION_WARNINGS
 			return 0;
 		}
 
-		const bool bDidUpdateGraph = UpdateGraphIfPending();
-		bIsWaitingForFirstGraph = bIsWaitingForFirstGraph && !bDidUpdateGraph;
+		UpdateGraphIfPending();
 
 		// Output silent audio if we're still building a graph
 		if (bIsWaitingForFirstGraph)
 		{
+			METASOUND_TRACE_CPUPROFILER_EVENT_SCOPE_TEXT(*FString::Printf(TEXT("MetasoundGenerator::OnGenerateAudio::MissedRenderDeadline %s"), *MetasoundName));
 			FMemory::Memset(OutAudio, 0, sizeof(float)* NumSamplesRemaining);
 			return NumSamplesRemaining;
 		}
+
 		// If no longer pending and executer is no-op, kill the MetaSound.
 		// Covers case where there was an error when building, resulting in
 		// Executer operator being assigned to NoOp.
diff --git a/Engine/Plugins/Runtime/Metasound/Source/MetasoundGenerator/Private/MetasoundGeneratorBuilder.cpp b/Engine/Plugins/Runtime/Metasound/Source/MetasoundGenerator/Private/MetasoundGeneratorBuilder.cpp
new file mode 100644
--- /dev/null
+++ b/Engine/Plugins/Runtime/Metasound/Source/MetasoundGenerator/Private/MetasoundGeneratorBuilder.cpp
@@ -0,0 +1,282 @@
+// Copyright Epic Games, Inc. All Rights Reserved.
+
+#include "MetasoundGeneratorBuilder.h"
+
+#include "AudioParameter.h"
+#include "MetasoundFrontendDataTypeRegistry.h"
+#include "MetasoundOperatorBuilder.h"
+#include "MetasoundOperatorInterface.h"
+#include "MetasoundOutputNode.h"
+#include "MetasoundParameterTransmitter.h"
+#include "MetasoundRouter.h"
+#include "MetasoundTrace.h"
+#include "MetasoundTrigger.h"
+#include "MetasoundVertexData.h"
+#include "Analysis/MetasoundFrontendAnalyzerFactory.h"
+#include "Analysis/MetasoundFrontendAnalyzerRegistry.h"
+#include "Analysis/MetasoundFrontendVertexAnalyzer.h"
+#include "Interfaces/MetasoundFrontendSourceInterface.h"
+
+namespace Metasound
+{
+	FAsyncMetaSoundBuilder::FAsyncMetaSoundBuilder(FMetasoundGenerator* InGenerator, FMetasoundGeneratorInitParams&& InInitParams, bool bInTriggerGenerator)
+		: Generator(InGenerator)
+		, InitParams(MoveTemp(InInitParams))
+		, bTriggerGenerator(bInTriggerGenerator)
+	{
+	}
+
+	void FAsyncMetaSoundBuilder::DoWork()
+	{
+		using namespace Audio;
+		using namespace Frontend;
+		using namespace MetasoundGeneratorPrivate;
+
+		METASOUND_TRACE_CPUPROFILER_EVENT_SCOPE_TEXT(*FString::Printf(TEXT("AsyncMetaSoundBuilder::DoWork %s"), *InitParams.MetaSoundName));
+
+		// Create an instance of the new graph
+		FBuildResults BuildResults;
+		FOperatorAndInputs GraphOperatorAndInputs = GeneratorBuilder::BuildGraphOperator(InitParams, BuildResults);
+		GeneratorBuilder::LogBuildErrors(InitParams.MetaSoundName, BuildResults);
+
+		if (GraphOperatorAndInputs.Operator.IsValid())
+		{
+			// Create graph analyzer
+			TUniquePtr<FGraphAnalyzer> GraphAnalyzer;
+			if (InitParams.BuilderSettings.bPopulateInternalDataReferences)
+			{
+				GraphAnalyzer = GeneratorBuilder::BuildGraphAnalyzer(MoveTemp(BuildResults.InternalDataReferences), InitParams.Environment, InitParams.OperatorSettings);
+			}
+	
+			// Collect data for generator
+			FMetasoundGeneratorData GeneratorData = GeneratorBuilder::BuildGeneratorData(InitParams, MoveTemp(GraphOperatorAndInputs), MoveTemp(GraphAnalyzer));
+
+			Generator->SetPendingGraph(MoveTemp(GeneratorData), bTriggerGenerator);
+		}
+		else 
+		{
+			UE_LOG(LogMetaSound, Error, TEXT("Failed to build Metasound operator from graph in MetasoundSource [%s]"), *InitParams.MetaSoundName);
+			// Set null generator data to inform that generator failed to build. 
+			// Otherwise, generator will continually wait for a new generator.
+			Generator->SetPendingGraphBuildFailed();
+		}
+
+		InitParams.Release();
+	}
+
+	namespace GeneratorBuilder
+	{
+		TArray<FAudioBufferReadRef> FindOutputAudioBuffers(const TArray<FVertexName>& InAudioVertexNames, const FVertexInterfaceData& InVertexData, const FOperatorSettings& InOperatorSettings, const FString& InMetaSoundName)
+		{
+			TArray<FAudioBufferReadRef> OutputBuffers;
+
+			const FOutputVertexInterfaceData& OutputVertexData = InVertexData.GetOutputs();
+
+			// Get output audio buffers.
+			for (const FVertexName& AudioOutputName : InAudioVertexNames)
+			{
+				if (!OutputVertexData.IsVertexBound(AudioOutputName))
+				{
+					UE_LOG(LogMetaSound, Warning, TEXT("MetasoundSource [%s] does not contain audio output [%s] in output"), *InMetaSoundName, *AudioOutputName.ToString());
+				}
+				OutputBuffers.Add(OutputVertexData.GetOrConstructDataReadReference<FAudioBuffer>(AudioOutputName, InOperatorSettings));
+			}
+
+			return OutputBuffers;
+		}
+
+		void LogBuildErrors(const FString& InMetaSoundName, const FBuildResults& InBuildResults)
+		{
+			// Log build errors
+			for (const IOperatorBuilder::FBuildErrorPtr& Error : InBuildResults.Errors)
+			{
+				if (Error.IsValid())
+				{
+					UE_LOG(LogMetaSound, Warning, TEXT("MetasoundSource [%s] build error [%s] \"%s\""), *InMetaSoundName, *(Error->GetErrorType().ToString()), *(Error->GetErrorDescription().ToString()));
+				}
+			}
+		}
+
+		TUniquePtr<Frontend::FGraphAnalyzer> BuildGraphAnalyzer(TMap<FGuid, FDataReferenceCollection>&& InInternalDataReferences, const FMetasoundEnvironment& InEnvironment, const FOperatorSettings& InOperatorSettings)
+		{
+			METASOUND_TRACE_CPUPROFILER_EVENT_SCOPE_TEXT(TEXT("GeneratorBuilder::BuildGraphAnalyzer"));
+			using namespace Frontend;
+
+			const uint64 InstanceID = InEnvironment.GetValue<uint64>(SourceInterface::Environment::TransmitterID);
+			return MakeUnique<FGraphAnalyzer>(InOperatorSettings, InstanceID, MoveTemp(InInternalDataReferences));
+		}
+
+		FOperatorAndInputs BuildGraphOperator(FMetasoundGeneratorInitParams& InInitParams, FBuildResults& OutBuildResults)
+		{
+			METASOUND_TRACE_CPUPROFILER_EVENT_SCOPE_TEXT(TEXT("GeneratorBuilder::BuildGraphOperator"));
+			using namespace Frontend;
+
+			// Choose which type of data reference access to create depending upon the access of the vertex.
+			auto VertexAccessTypeToDataReferenceAccessType = [](EVertexAccessType InVertexAccessType) -> EDataReferenceAccessType
+			{
+				switch(InVertexAccessType)
+				{
+					case EVertexAccessType::Value:
+						return EDataReferenceAccessType::Value;
+
+					case EVertexAccessType::Reference:
+					default:
+						return EDataReferenceAccessType::Write;
+				}
+			};
+
+			const FInputVertexInterface& InputInterface = InInitParams.Graph->GetVertexInterface().GetInputInterface();
+			FInputVertexInterfaceData InputData(InputInterface);
+
+			// Set input data based on the input parameters and the input interface
+			IDataTypeRegistry& DataRegistry = IDataTypeRegistry::Get();
+			for (FAudioParameter& Parameter : InInitParams.DefaultParameters)
+			{
+				const FName ParamName = Parameter.ParamName;
+				if (const FInputDataVertex* InputVertex = InputInterface.Find(ParamName))
+				{
+					FLiteral Literal = Frontend::ConvertParameterToLiteral(MoveTemp(Parameter));
+
+					TOptional<FAnyDataReference> DataReference = DataRegistry.CreateDataReference(InputVertex->DataTypeName, VertexAccessTypeToDataReferenceAccessType(InputVertex->AccessType), Literal, InInitParams.OperatorSettings);
+
+					if (DataReference)
+					{
+						InputData.BindVertex(ParamName, *DataReference);
+					}
+					else if(MetaSoundParameterEnableWarningOnIgnoredParameterCVar)
+					{
+						UE_LOG(LogMetaSound, Warning, TEXT("Failed to create initial input data reference from parameter %s of type %s on graph in MetaSoundSource [%s]"), *ParamName.ToString(), *InputVertex->DataTypeName.ToString(), *InInitParams.MetaSoundName);
+					}
+				}
+				else if(MetaSoundParameterEnableWarningOnIgnoredParameterCVar)
+				{
+					UE_LOG(LogMetaSound, Warning, TEXT("Failed to set initial input parameter %s on graph in MetaSoundSource [%s]"), *ParamName.ToString(), *InInitParams.MetaSoundName);
+				}
+			}
+
+			// Set any remaining inputs to their default values.
+			for (const MetasoundVertexDataPrivate::TBinding<FInputDataVertex>& Binding : InputData)
+			{
+				// Only create data reference if something does not already exist. 
+				if (!Binding.IsBound())
+				{
+					const FInputDataVertex& InputVertex = Binding.GetVertex();
+					EDataReferenceAccessType AccessType = VertexAccessTypeToDataReferenceAccessType(InputVertex.AccessType);
+					TOptional<FAnyDataReference> DataReference = DataRegistry.CreateDataReference(InputVertex.DataTypeName, VertexAccessTypeToDataReferenceAccessType(InputVertex.AccessType), InputVertex.GetDefaultLiteral(), InInitParams.OperatorSettings);
+
+					if (DataReference)
+					{
+						InputData.BindVertex(InputVertex.VertexName, *DataReference);
+					}
+				}
+			}
+
+			// Reset as elements in array have been moved.
+			InInitParams.DefaultParameters.Reset();
+
+			// Create an instance of the new graph
+			FBuildGraphOperatorParams BuildParams { *InInitParams.Graph, InInitParams.OperatorSettings, InputData, InInitParams.Environment };
+			FOperatorAndInputs OpAndInputs;
+			OpAndInputs.Operator = FOperatorBuilder(InInitParams.BuilderSettings).BuildGraphOperator(BuildParams, OutBuildResults);
+			OpAndInputs.Inputs = InputData;
+
+			return OpAndInputs;
+		}
+
+		MetasoundGeneratorPrivate::FMetasoundGeneratorData BuildGeneratorData(const FMetasoundGeneratorInitParams& InInitParams, FOperatorAndInputs&& InGraphOperatorAndInputs, TUniquePtr<Frontend::FGraphAnalyzer> InAnalyzer)
+		{
+			using namespace Audio;
+			using namespace Frontend;
+			using namespace MetasoundGeneratorPrivate;
+
+			METASOUND_TRACE_CPUPROFILER_EVENT_SCOPE_TEXT(TEXT("GeneratorBuilder::BuildGeneratorData"));
+
+			checkf(InGraphOperatorAndInputs.Operator.IsValid(), TEXT("Graph operator must be a valid object"));
+
+			// Gather relevant input and output references
+			FVertexInterfaceData VertexData(InInitParams.Graph->GetVertexInterface());
+			InGraphOperatorAndInputs.Operator->Bind(VertexData);
+
+			// Replace input data with writable inputs
+			VertexData.GetInputs() = InGraphOperatorAndInputs.Inputs;
+
+			// Get inputs
+			FTriggerWriteRef PlayTrigger = VertexData.GetInputs().GetOrConstructDataWriteReference<FTrigger>(SourceInterface::Inputs::OnPlay, InInitParams.OperatorSettings, false);
+
+			// Get outputs
+			TArray<FAudioBufferReadRef> OutputBuffers = FindOutputAudioBuffers(InInitParams.AudioOutputNames, VertexData, InInitParams.OperatorSettings, InInitParams.MetaSoundName);
+			FTriggerReadRef FinishTrigger = TDataReadReferenceFactory<FTrigger>::CreateExplicitArgs(InInitParams.OperatorSettings, false);
+
+			if (InInitParams.Graph->GetVertexInterface().GetOutputInterface().Contains(SourceOneShotInterface::Outputs::OnFinished))
+			{
+				FinishTrigger = VertexData.GetOutputs().GetOrConstructDataReadReference<FTrigger>(SourceOneShotInterface::Outputs::OnFinished, InInitParams.OperatorSettings, false);
+			}
+
+			// Create the parameter setter map so parameter packs can be cracked
+			// open and distributed as appropriate...
+			const IDataTypeRegistry& DataTypeRegistry = IDataTypeRegistry::Get();
+			FParameterSetterSortedMap ParameterSetters;
+			TMap<FName, FParameterPackSetter> ParameterPackSetters;
+			FInputVertexInterfaceData& GraphInputs = VertexData.GetInputs();
+			for (const MetasoundVertexDataPrivate::TBinding<FInputDataVertex>& Binding : GraphInputs)
+			{
+				// Only assign inputs that are writable. 
+				if (EDataReferenceAccessType::Write == Binding.GetAccessType())
+				{
+					if (const FAnyDataReference* DataRef = Binding.GetDataReference())
+					{
+						const FInputDataVertex& InputVertex = Binding.GetVertex();
+						const Frontend::IParameterAssignmentFunction& PackSetter = DataTypeRegistry.GetRawAssignmentFunction(InputVertex.DataTypeName);
+						if (PackSetter)
+						{
+							FParameterPackSetter ParameterPackSetter(InputVertex.DataTypeName, DataRef->GetRaw(), PackSetter);
+							ParameterPackSetters.Add(InputVertex.VertexName, ParameterPackSetter);
+						}
+
+						Frontend::FLiteralAssignmentFunction LiteralSetter = DataTypeRegistry.GetLiteralAssignmentFunction(InputVertex.DataTypeName);
+						if (LiteralSetter)
+						{
+							ParameterSetters.Add(InputVertex.VertexName, FParameterSetter{LiteralSetter, *DataRef});
+						}
+					}
+				}
+			}
+
+			// Set data needed for graph
+			return FMetasoundGeneratorData 
+			{
+				InInitParams.OperatorSettings,
+				MoveTemp(InGraphOperatorAndInputs.Operator),
+				MoveTemp(VertexData),
+				MoveTemp(ParameterSetters),
+				MoveTemp(ParameterPackSetters),
+				MoveTemp(InAnalyzer),
+				MoveTemp(OutputBuffers),
+				MoveTemp(PlayTrigger),
+				MoveTemp(FinishTrigger),
+			};
+		}
+
+		void ApplyAudioParameters(const FOperatorSettings& InOperatorSettings, TArray<FAudioParameter>&& InParameters, FInputVertexInterfaceData& InInterface)
+		{
+			METASOUND_TRACE_CPUPROFILER_EVENT_SCOPE("GeneratorBuilder::ApplyAudioParameters");
+			Frontend::IDataTypeRegistry& DataTypeRegistry = Frontend::IDataTypeRegistry::Get();
+			for (FAudioParameter& Parameter : InParameters)
+			{
+				if (const FAnyDataReference* Ref = InInterface.FindDataReference(Parameter.ParamName))
+				{
+					if (EDataReferenceAccessType::Write == Ref->GetAccessType())
+					{
+						Frontend::FLiteralAssignmentFunction LiteralSetter = DataTypeRegistry.GetLiteralAssignmentFunction(Ref->GetDataTypeName());
+
+						if (LiteralSetter)
+						{
+							FLiteral Literal = Frontend::ConvertParameterToLiteral(MoveTemp(Parameter)); 
+							LiteralSetter(InOperatorSettings, Literal, *Ref);
+						}
+					}
+				}
+			}
+		}
+	}
+}
diff --git a/Engine/Plugins/Runtime/Metasound/Source/MetasoundGenerator/Private/MetasoundGeneratorBuilder.h b/Engine/Plugins/Runtime/Metasound/Source/MetasoundGenerator/Private/MetasoundGeneratorBuilder.h
new file mode 100644
--- /dev/null
+++ b/Engine/Plugins/Runtime/Metasound/Source/MetasoundGenerator/Private/MetasoundGeneratorBuilder.h
@@ -0,0 +1,57 @@
+// Copyright Epic Games, Inc. All Rights Reserved.
+
+#pragma once
+
+#include "MetasoundGenerator.h"
+
+#include "Async/AsyncWork.h"
+#include "Containers/Map.h"
+#include "MetasoundEnvironment.h"
+#include "MetasoundOperatorBuilder.h"
+#include "MetasoundOperatorCache.h"
+#include "MetasoundOperatorSettings.h"
+#include "MetasoundVertex.h"
+#include "MetasoundVertexData.h"
+#include "Misc/Guid.h"
+#include "Templates/UniquePtr.h"
+
+namespace Metasound
+{
+	class FMetasoundGenerator;
+
+	class FAsyncMetaSoundBuilder : public FNonAbandonableTask
+	{
+	public:
+		FAsyncMetaSoundBuilder(FMetasoundGenerator* InGenerator, FMetasoundGeneratorInitParams&& InInitParams, bool bInTriggerGenerator);
+
+		~FAsyncMetaSoundBuilder() = default;
+
+		void DoWork();
+
+		FORCEINLINE TStatId GetStatId() const
+		{
+			RETURN_QUICK_DECLARE_CYCLE_STAT(FAsyncMetaSoundBuilder, STATGROUP_ThreadPoolAsyncTasks);
+		}
+
+	private:
+
+		FMetasoundGenerator* Generator;
+		FMetasoundGeneratorInitParams InitParams;
+		bool bTriggerGenerator;
+	};
+
+	namespace GeneratorBuilder
+	{
+		TArray<FAudioBufferReadRef> FindOutputAudioBuffers(const TArray<FVertexName>& InAudioVertexNames, const FVertexInterfaceData& InVertexData, const FOperatorSettings& InOperatorSettings, const FString& InMetaSoundName);
+
+		void LogBuildErrors(const FString& InMetaSoundName, const FBuildResults& InBuildResults);
+
+		TUniquePtr<Frontend::FGraphAnalyzer> BuildGraphAnalyzer(TMap<FGuid, FDataReferenceCollection>&& InInternalDataReferences, const FMetasoundEnvironment& InEnvironment, const FOperatorSettings& InOperatorSettings);
+
+		FOperatorAndInputs BuildGraphOperator(FMetasoundGeneratorInitParams& InInitParams, FBuildResults& OutBuildResults);
+
+		MetasoundGeneratorPrivate::FMetasoundGeneratorData BuildGeneratorData(const FMetasoundGeneratorInitParams& InInitParams, FOperatorAndInputs&& InGraphOperatorAndInputs, TUniquePtr<Frontend::FGraphAnalyzer> InAnalyzer);
+
+		void ApplyAudioParameters(const FOperatorSettings& InOperatorSettings, TArray<FAudioParameter>&& InParameters, FInputVertexInterfaceData& InInterface);
+	}
+}
diff --git a/Engine/Plugins/Runtime/Metasound/Source/MetasoundGenerator/Private/MetasoundGeneratorModule.cpp b/Engine/Plugins/Runtime/Metasound/Source/MetasoundGenerator/Private/MetasoundGeneratorModule.cpp
--- a/Engine/Plugins/Runtime/Metasound/Source/MetasoundGenerator/Private/MetasoundGeneratorModule.cpp
+++ b/Engine/Plugins/Runtime/Metasound/Source/MetasoundGenerator/Private/MetasoundGeneratorModule.cpp
@@ -1,20 +1,65 @@
 // Copyright Epic Games, Inc. All Rights Reserved.
 
 #include "MetasoundGeneratorModule.h"
+#include "MetasoundGeneratorModuleImpl.h"
 
-#include "CoreMinimal.h"
+#include "Logging/LogMacros.h"
+#include "MetasoundOperatorCache.h"
 #include "Modules/ModuleInterface.h"
 #include "Modules/ModuleManager.h"
 
 DEFINE_LOG_CATEGORY(LogMetasoundGenerator);
 
+
+static FAutoConsoleCommand CommandMetaSoundExperimentalOperatorCacheSetMaxNumOperators(
+	TEXT("au.MetaSound.Experimental.OperatorCache.SetMaxNumOperators"),
+	TEXT("Set the maximum number of operators in the MetaSound operator cache."),
+	FConsoleCommandWithArgsDelegate::CreateStatic([](const TArray< FString >& Args)
+	{
+		using namespace Metasound;
+
+		if (Args.Num() < 1)
+		{
+			return;
+		}
+
+		const int32 MaxNumOperators = FCString::Atoi(*Args[0]);
+
+		if (MaxNumOperators < 0)
+		{
+			return;
+		}
+
+		FMetasoundGeneratorModule& Module = FModuleManager::GetModuleChecked<FMetasoundGeneratorModule>("MetasoundGenerator");
+		TSharedPtr<FOperatorCache> OperatorCache = Module.GetOperatorCache();
+		if (OperatorCache.IsValid())
+		{
+			OperatorCache->SetMaxNumOperators(static_cast<uint32>(MaxNumOperators));
+			UE_LOG(LogMetasoundGenerator, Display, TEXT("Metasound operator cache size set to %d operators."), MaxNumOperators);
+		}
+	})
+);
+
 namespace Metasound
 {
-	class FMetasoundGeneratorModule : public IMetasoundGeneratorModule
+	void FMetasoundGeneratorModule::StartupModule()
 	{
-	};
+		FOperatorCacheSettings OperatorCacheSettings;
+		OperatorCacheSettings.MaxNumOperators = 64;
+
+		OperatorCache = MakeShared<FOperatorCache>(OperatorCacheSettings);
+	}
+
+	void FMetasoundGeneratorModule::ShutdownModule() 
+	{
+		OperatorCache.Reset();
+	}
+
+	TSharedPtr<FOperatorCache> FMetasoundGeneratorModule::GetOperatorCache()
+	{
+		return OperatorCache;
+	}
 }
 
 IMPLEMENT_MODULE(Metasound::FMetasoundGeneratorModule, MetasoundGenerator);
 
-
diff --git a/Engine/Plugins/Runtime/Metasound/Source/MetasoundGenerator/Private/MetasoundGeneratorModuleImpl.h b/Engine/Plugins/Runtime/Metasound/Source/MetasoundGenerator/Private/MetasoundGeneratorModuleImpl.h
new file mode 100644
--- /dev/null
+++ b/Engine/Plugins/Runtime/Metasound/Source/MetasoundGenerator/Private/MetasoundGeneratorModuleImpl.h
@@ -0,0 +1,28 @@
+// Copyright Epic Games, Inc. All Rights Reserved.
+
+#pragma once
+
+#include "MetasoundGeneratorModule.h"
+
+#include "MetasoundOperatorCache.h"
+
+namespace Metasound
+{
+	class FMetasoundGeneratorModule : public IMetasoundGeneratorModule
+	{
+	public:
+
+		virtual void StartupModule() override;
+		virtual void ShutdownModule() override;
+
+		TSharedPtr<FOperatorCache> GetOperatorCache();
+
+	private:
+
+		TSharedPtr<FOperatorCache> OperatorCache;
+	};
+}
+
+
+
+
diff --git a/Engine/Plugins/Runtime/Metasound/Source/MetasoundGenerator/Private/MetasoundOperatorCache.cpp b/Engine/Plugins/Runtime/Metasound/Source/MetasoundGenerator/Private/MetasoundOperatorCache.cpp
new file mode 100644
--- /dev/null
+++ b/Engine/Plugins/Runtime/Metasound/Source/MetasoundGenerator/Private/MetasoundOperatorCache.cpp
@@ -0,0 +1,140 @@
+// Copyright Epic Games, Inc. All Rights Reserved.
+
+#include "MetasoundOperatorCache.h"
+
+#include "Containers/Array.h"
+#include "Containers/Map.h"
+#include "HAL/CriticalSection.h"
+#include "Misc/Guid.h"
+#include "Misc/ScopeLock.h"
+#include "ProfilingDebugging/CountersTrace.h"
+#include "Templates/UniquePtr.h"
+
+#ifndef METASOUND_OPERATORCACHEPROFILER_ENABLED
+#define METASOUND_OPERATORCACHEPROFILER_ENABLED COUNTERSTRACE_ENABLED
+#endif
+
+namespace Metasound
+{
+#if METASOUND_OPERATORCACHEPROFILER_ENABLED
+	TRACE_DECLARE_FLOAT_COUNTER(MetaSound_OperatorCache_HitRatio, TEXT("MetaSound/OperatorCache/HitRatio"));
+	namespace OperatorCachePrivate
+	{
+		static std::atomic<uint32> CacheHitCount = 0;
+		static std::atomic<uint32> CacheAttemptCount = 0;
+
+		double GetHitRatio()
+		{
+			uint32 NumHits = CacheHitCount;
+			uint32 Total = CacheAttemptCount;
+			if (Total > 0)
+			{
+				return static_cast<double>(NumHits) / static_cast<double>(Total);
+			}
+			else
+			{
+				return 0.f;
+			}
+		}
+	}
+#endif
+
+	FOperatorCache::FOperatorCache(const FOperatorCacheSettings& InSettings)
+	: Settings(InSettings)
+	{
+	}
+
+	FOperatorAndInputs FOperatorCache::ClaimCachedOperator(const FGuid& InOperatorID)
+	{
+		FScopeLock Lock(&CriticalSection);
+
+		FOperatorAndInputs OpAndInputs;
+		if (TArray<FOperatorAndInputs>* OperatorsWithID = Operators.Find(InOperatorID))
+		{
+			if (OperatorsWithID->Num() > 0)
+			{
+				OpAndInputs = OperatorsWithID->Pop();
+				Stack.RemoveAt(Stack.FindLast(InOperatorID));
+#if METASOUND_OPERATORCACHEPROFILER_ENABLED
+				OperatorCachePrivate::CacheHitCount++;
+#endif
+			}
+		}
+
+#if METASOUND_OPERATORCACHEPROFILER_ENABLED
+		OperatorCachePrivate::CacheAttemptCount++;
+
+		TRACE_COUNTER_SET(MetaSound_OperatorCache_HitRatio, OperatorCachePrivate::GetHitRatio());
+#endif
+		return OpAndInputs;
+	}
+
+	void FOperatorCache::AddOperatorToCache(const FGuid& InOperatorID, TUniquePtr<IOperator>&& InOperator, FInputVertexInterfaceData&& InInputData)
+	{
+		if (InOperator.IsValid())
+		{
+			FScopeLock Lock(&CriticalSection);
+
+			Stack.Add(InOperatorID);
+
+			FOperatorAndInputs OpAndInputs;
+			OpAndInputs.Operator = MoveTemp(InOperator);
+			OpAndInputs.Inputs = MoveTemp(InInputData);
+
+			if (TArray<FOperatorAndInputs>* OperatorArray = Operators.Find(InOperatorID))
+			{
+				OperatorArray->Add(MoveTemp(OpAndInputs));
+			}
+			else
+			{
+				TArray<FOperatorAndInputs> NewOperatorArray;
+				NewOperatorArray.Add(MoveTemp(OpAndInputs));
+				Operators.Add(InOperatorID, MoveTemp(NewOperatorArray));
+			}
+
+			TrimCache();
+		}
+	}
+
+	void FOperatorCache::RemoveOperatorsWithID(const FGuid& InOperatorID)
+	{
+		FScopeLock Lock(&CriticalSection);
+		Operators.Remove(InOperatorID);
+		Stack.Remove(InOperatorID);
+	}
+
+	void FOperatorCache::SetMaxNumOperators(uint32 InMaxNumOperators)
+	{
+		FScopeLock Lock(&CriticalSection);
+		Settings.MaxNumOperators = InMaxNumOperators;
+		TrimCache();
+	}
+
+	void FOperatorCache::TrimCache()
+	{
+		int32 NumToTrim = Stack.Num() - Settings.MaxNumOperators;
+		if (NumToTrim > 0)
+		{
+			UE_LOG(LogMetasoundGenerator, Verbose, TEXT("Trimming %d operators"), NumToTrim);
+			for (int32 i = 0; i < NumToTrim; i++)
+			{
+				UE_LOG(LogMetasoundGenerator, VeryVerbose, TEXT("Trimming operator with ID %s"), *LexToString(Stack[i]));
+				// Destructor called by TUniquePtr going out of scope.
+				TArray<FOperatorAndInputs>* OperatorArray = Operators.Find(Stack[i]);
+				if (ensure(OperatorArray))
+				{
+					if (ensure(OperatorArray->Num() > 0))
+					{
+						OperatorArray->Pop();
+						if (OperatorArray->Num() == 0)
+						{
+							Operators.Remove(Stack[i]);
+						}
+					}
+				}
+			}
+			Stack.RemoveAt(0, NumToTrim);
+		}
+	}
+}
+
diff --git a/Engine/Plugins/Runtime/Metasound/Source/MetasoundGenerator/Private/MetasoundOperatorCache.h b/Engine/Plugins/Runtime/Metasound/Source/MetasoundGenerator/Private/MetasoundOperatorCache.h
new file mode 100644
--- /dev/null
+++ b/Engine/Plugins/Runtime/Metasound/Source/MetasoundGenerator/Private/MetasoundOperatorCache.h
@@ -0,0 +1,45 @@
+// Copyright Epic Games, Inc. All Rights Reserved.
+
+#pragma once
+
+#include "Containers/Array.h"
+#include "Containers/Map.h"
+#include "MetasoundOperatorInterface.h"
+#include "Misc/Guid.h"
+#include "HAL/CriticalSection.h"
+#include "Templates/UniquePtr.h"
+
+namespace Metasound
+{
+	struct FOperatorCacheSettings
+	{
+		uint32 MaxNumOperators = 64;
+	};
+
+	struct FOperatorAndInputs
+	{
+		TUniquePtr<IOperator> Operator;
+		FInputVertexInterfaceData Inputs;
+	};
+
+	class FOperatorCache
+	{
+	public:
+
+		FOperatorCache(const FOperatorCacheSettings& InSettings);
+		FOperatorAndInputs ClaimCachedOperator(const FGuid& InOperatorID);
+		void AddOperatorToCache(const FGuid& InOperatorID, TUniquePtr<IOperator>&& InOperator, FInputVertexInterfaceData&& InputData);
+		void RemoveOperatorsWithID(const FGuid& InOperatorID);
+		void SetMaxNumOperators(uint32 InMaxNumOperators);
+
+	private:
+
+		void TrimCache();
+
+		FOperatorCacheSettings Settings;
+		FCriticalSection CriticalSection;
+		TMap<FGuid, TArray<FOperatorAndInputs>> Operators;
+		TArray<FGuid> Stack;
+	};
+}
+
diff --git a/Engine/Plugins/Runtime/Metasound/Source/MetasoundGenerator/Public/MetasoundGenerator.h b/Engine/Plugins/Runtime/Metasound/Source/MetasoundGenerator/Public/MetasoundGenerator.h
--- a/Engine/Plugins/Runtime/Metasound/Source/MetasoundGenerator/Public/MetasoundGenerator.h
+++ b/Engine/Plugins/Runtime/Metasound/Source/MetasoundGenerator/Public/MetasoundGenerator.h
@@ -13,6 +13,7 @@
 #include "Analysis/MetasoundFrontendGraphAnalyzer.h"
 #include "Async/AsyncWork.h"
 #include "Containers/MpscQueue.h"
+#include "Containers/SpscQueue.h"
 #include "Sound/SoundGenerator.h"
 #include "Delegates/Delegate.h"
 
@@ -87,38 +88,6 @@ namespace Metasound
 		void Release();
 	};
 
-	class FMetasoundGenerator;
-
-	class FAsyncMetaSoundBuilder : public FNonAbandonableTask
-	{
-	public:
-		FAsyncMetaSoundBuilder(FMetasoundGenerator* InGenerator, FMetasoundGeneratorInitParams&& InInitParams, bool bInTriggerGenerator);
-
-		~FAsyncMetaSoundBuilder() = default;
-
-		void DoWork();
-
-		FORCEINLINE TStatId GetStatId() const
-		{
-			RETURN_QUICK_DECLARE_CYCLE_STAT(FAsyncMetaSoundBuilder, STATGROUP_ThreadPoolAsyncTasks);
-		}
-
-	private:
-		struct FGraphOperatorAndInputs
-		{
-			TUniquePtr<IOperator> Operator;
-			FInputVertexInterfaceData Inputs;
-		};
-		MetasoundGeneratorPrivate::FMetasoundGeneratorData BuildGeneratorData(const FMetasoundGeneratorInitParams& InInitParams, FGraphOperatorAndInputs&& InOperatorAndInputs, TUniquePtr<Frontend::FGraphAnalyzer> InAnalyzer) const;
-		FGraphOperatorAndInputs BuildGraphOperator(TArray<FAudioParameter>&& InParameters, FBuildResults& OutBuildResults) const;
-		TUniquePtr<Frontend::FGraphAnalyzer> BuildGraphAnalyzer(TMap<FGuid, FDataReferenceCollection>&& InInternalDataReferences) const;
-		void LogBuildErrors(const FBuildResults& InBuildResults) const;
-		TArray<FAudioBufferReadRef> FindOutputAudioBuffers(const FVertexInterfaceData& InVertexData) const;
-
-		FMetasoundGenerator* Generator;
-		FMetasoundGeneratorInitParams InitParams;
-		bool bTriggerGenerator;
-	};
 
 	DECLARE_TS_MULTICAST_DELEGATE(FOnSetGraph);
 
@@ -249,6 +218,9 @@ namespace Metasound
 		FOnSetGraph OnSetGraph;
 
 	private:
+		bool TryUseCachedOperator(FMetasoundGeneratorInitParams& InParams, bool bTriggerGenerator);
+		void ReleaseOperatorToCache();
+
 		friend class FAsyncMetaSoundBuilder;
 
 		/** Update the current graph operator with a new graph operator. The number of channels
@@ -263,7 +235,12 @@ namespace Metasound
 
 		bool UpdateGraphIfPending();
 
-		// Internal set graph after checking compatibility.
+		/** Release the graph operator and remove any references to data owned by
+		 * the graph operator.
+		 */
+		void ClearGraph();
+
+		/** Internal set graph after checking compatibility. */
 		void SetGraph(TUniquePtr<MetasoundGeneratorPrivate::FMetasoundGeneratorData>&& InData, bool bTriggerGraph);
 
 		// Fill OutAudio with data in InBuffer, up to maximum number of samples.
@@ -281,6 +258,8 @@ namespace Metasound
 		bool bIsGraphBuilding;
 		bool bIsFinishTriggered;
 		bool bIsFinished;
+		bool bUseOperatorCache = false;
+		FGuid OperatorID;
 
 		int32 FinishSample = INDEX_NONE;
 		int32 NumChannels;
@@ -299,8 +278,7 @@ namespace Metasound
 
 		Audio::FAlignedFloatBuffer OverflowBuffer;
 
-		typedef FAsyncTask<FAsyncMetaSoundBuilder> FBuilderTask;
-		TUniquePtr<FBuilderTask> BuilderTask;
+		TUniquePtr<FAsyncTaskBase> BuilderTask;
 
 		FCriticalSection PendingGraphMutex;
 		TUniquePtr<MetasoundGeneratorPrivate::FMetasoundGeneratorData> PendingGraphData;
diff --git a/Engine/Plugins/Runtime/Metasound/Source/MetasoundGraphCore/Public/MetasoundDataReference.h b/Engine/Plugins/Runtime/Metasound/Source/MetasoundGraphCore/Public/MetasoundDataReference.h
--- a/Engine/Plugins/Runtime/Metasound/Source/MetasoundGraphCore/Public/MetasoundDataReference.h
+++ b/Engine/Plugins/Runtime/Metasound/Source/MetasoundGraphCore/Public/MetasoundDataReference.h
@@ -657,7 +657,7 @@ namespace Metasound
 		 * does match.  
 		 */
 		template<typename DataType>
-		DataType* GetWritableValue()
+		DataType* GetWritableValue() const
 		{
 			if (DataRefPtr.IsValid())
 			{
diff --git a/Engine/Plugins/Runtime/Metasound/Source/MetasoundGraphCore/Public/MetasoundExecutableOperator.h b/Engine/Plugins/Runtime/Metasound/Source/MetasoundGraphCore/Public/MetasoundExecutableOperator.h
--- a/Engine/Plugins/Runtime/Metasound/Source/MetasoundGraphCore/Public/MetasoundExecutableOperator.h
+++ b/Engine/Plugins/Runtime/Metasound/Source/MetasoundGraphCore/Public/MetasoundExecutableOperator.h
@@ -226,6 +226,14 @@ namespace Metasound
 				}
 			}
 
+			TUniquePtr<IOperator> ReleaseOperator()
+			{
+				ExecuteFunction = nullptr;
+				PostExecuteFunction = nullptr;
+				ResetFunction = nullptr;
+				return MoveTemp(Operator);
+			}
+
 			void Execute()
 			{
 				if (ExecuteFunction)
