CIS metasound non-unity build fix.

diff --git a/Engine/Plugins/Runtime/Metasound/Source/MetasoundGenerator/Private/MetasoundGenerator.cpp b/Engine/Plugins/Runtime/Metasound/Source/MetasoundGenerator/Private/MetasoundGenerator.cpp
--- a/Engine/Plugins/Runtime/Metasound/Source/MetasoundGenerator/Private/MetasoundGenerator.cpp
+++ b/Engine/Plugins/Runtime/Metasound/Source/MetasoundGenerator/Private/MetasoundGenerator.cpp
@@ -18,7 +18,7 @@
 #include "Analysis/MetasoundFrontendVertexAnalyzer.h"
 #include "DSP/FloatArrayMath.h"
 #include "Interfaces/MetasoundFrontendSourceInterface.h"
-#include "ProfilingDebugging/CountersTrace.h"
+#include "Modules/ModuleManager.h"
 
 #ifndef ENABLE_METASOUNDGENERATOR_INVALID_SAMPLE_VALUE_LOGGING 
 #define ENABLE_METASOUNDGENERATOR_INVALID_SAMPLE_VALUE_LOGGING !UE_BUILD_SHIPPING
diff --git a/Engine/Plugins/Runtime/Metasound/Source/MetasoundGenerator/Private/MetasoundGeneratorModule.cpp b/Engine/Plugins/Runtime/Metasound/Source/MetasoundGenerator/Private/MetasoundGeneratorModule.cpp
--- a/Engine/Plugins/Runtime/Metasound/Source/MetasoundGenerator/Private/MetasoundGeneratorModule.cpp
+++ b/Engine/Plugins/Runtime/Metasound/Source/MetasoundGenerator/Private/MetasoundGeneratorModule.cpp
@@ -3,14 +3,16 @@
 #include "MetasoundGeneratorModule.h"
 #include "MetasoundGeneratorModuleImpl.h"
 
+#include "Containers/Array.h"
+#include "HAL/IConsoleManager.h"
 #include "Logging/LogMacros.h"
 #include "MetasoundOperatorCache.h"
+#include "Misc/CString.h"
 #include "Modules/ModuleInterface.h"
 #include "Modules/ModuleManager.h"
 
 DEFINE_LOG_CATEGORY(LogMetasoundGenerator);
 
-
 static FAutoConsoleCommand CommandMetaSoundExperimentalOperatorCacheSetMaxNumOperators(
 	TEXT("au.MetaSound.Experimental.OperatorCache.SetMaxNumOperators"),
 	TEXT("Set the maximum number of operators in the MetaSound operator cache."),
diff --git a/Engine/Plugins/Runtime/Metasound/Source/MetasoundGenerator/Private/MetasoundOperatorCache.cpp b/Engine/Plugins/Runtime/Metasound/Source/MetasoundGenerator/Private/MetasoundOperatorCache.cpp
--- a/Engine/Plugins/Runtime/Metasound/Source/MetasoundGenerator/Private/MetasoundOperatorCache.cpp
+++ b/Engine/Plugins/Runtime/Metasound/Source/MetasoundGenerator/Private/MetasoundOperatorCache.cpp
@@ -5,6 +5,7 @@
 #include "Containers/Array.h"
 #include "Containers/Map.h"
 #include "HAL/CriticalSection.h"
+#include "MetasoundGeneratorModule.h"
 #include "Misc/Guid.h"
 #include "Misc/ScopeLock.h"
 #include "ProfilingDebugging/CountersTrace.h"
