Add more tests for metasound automated node tests

diff --git a/Engine/Plugins/Runtime/Metasound/Source/MetasoundEngineTest/Private/EngineTestAutomatedNodeTests.cpp b/Engine/Plugins/Runtime/Metasound/Source/MetasoundEngineTest/Private/EngineTestAutomatedNodeTests.cpp
--- a/Engine/Plugins/Runtime/Metasound/Source/MetasoundEngineTest/Private/EngineTestAutomatedNodeTests.cpp
+++ b/Engine/Plugins/Runtime/Metasound/Source/MetasoundEngineTest/Private/EngineTestAutomatedNodeTests.cpp
@@ -1,5 +1,6 @@
 // Copyright Epic Games, Inc. All Rights Reserved.
 
+#include "Algo/AllOf.h"
 #include "AudioDevice.h"
 #include "AudioDeviceManager.h"
 #include "AudioMixerDevice.h"
@@ -59,6 +60,30 @@ namespace Metasound::MetasoundEngineTestPrivate {
 		return Environment;
 	}
 
+	FString GetPrettyName(const Frontend::FNodeRegistryKey& InRegistryKey)
+	{
+		Frontend::IMetaSoundAssetManager* AssetManager = Frontend::IMetaSoundAssetManager::Get();
+		if (ensure(AssetManager))
+		{
+			if (const FSoftObjectPath* ObjectPath = AssetManager->FindObjectPathFromKey(InRegistryKey))
+			{
+				return ObjectPath->ToString();
+			}
+		}
+
+		FMetasoundFrontendRegistryContainer* NodeRegistry = FMetasoundFrontendRegistryContainer::Get();
+		if (ensure(NodeRegistry))
+		{
+			FMetasoundFrontendClass NodeClass;
+			if (NodeRegistry->FindFrontendClassFromRegistered(InRegistryKey, NodeClass))
+			{
+				return FString::Printf(TEXT("%s %s"), *NodeClass.Metadata.GetClassName().ToString(), *NodeClass.Metadata.GetVersion().ToString());
+			}
+		}
+
+		return TEXT("");
+	}
+
 	// TTestTypeInfo converts test types to strings
 	template<typename DataType>
 	struct TTestTypeInfo
@@ -72,12 +97,18 @@ namespace Metasound::MetasoundEngineTestPrivate {
 	template<typename ElementType>
 	struct TTestTypeInfo<TArray<ElementType>>
 	{
-		static FString ToString(const TArray<ElementType>& InData)
+		static FString ToString(TArrayView<const ElementType> InData)
 		{
 			return FString::Printf(TEXT("[%s]"), *FString::JoinBy(InData, TEXT(","), &TTestTypeInfo<ElementType>::ToString));
 		}
 	};
 
+	template<>
+	struct TTestTypeInfo<FAudioBuffer> : TTestTypeInfo<TArray<float>>
+	{
+	};
+	
+
 	template<>
 	struct TTestTypeInfo<FTime>
 	{
@@ -87,6 +118,278 @@ namespace Metasound::MetasoundEngineTestPrivate {
 		}
 	};
 
+	// TTestTypeAnalysis performs analysis operator on data types.
+	template<typename DataType>
+	struct TTestTypeAnalysis
+	{
+		static bool IsValid(const DataType& InArray)
+		{
+			return true;
+		}
+
+		static bool IsEqual(const DataType& InLHS, const DataType& InRHS)
+		{
+			return InLHS == InRHS;
+		}
+	};
+
+	// Specialize type analysis for TArrays
+	template<typename ElementType>
+	struct TTestTypeAnalysis<TArray<ElementType>>
+	{
+		static bool IsValid(TArrayView<const ElementType> InArray)
+		{
+			return Algo::AllOf(InArray, &TTestTypeAnalysis<ElementType>::IsValid);
+		}
+
+		static bool IsEqual(TArrayView<const ElementType> InLHS, TArrayView<const ElementType> InRHS)
+		{
+			const int32 Num = InLHS.Num();
+			if (Num == InRHS.Num())
+			{
+				for (int32 i = 0; i < Num; i++)
+				{
+					if (InLHS[i] != InRHS[i])
+					{
+						return false;
+					}
+				}
+				return true;
+			}
+			return false;
+		}
+	};
+
+	// Specialize type analysis for audio buffers.
+	template<>
+	struct TTestTypeAnalysis<FAudioBuffer> : TTestTypeAnalysis<TArray<float>>
+	{
+	};
+
+	// Specialize type analysis for float
+	template<>
+	struct TTestTypeAnalysis<float>
+	{
+		static bool IsValid(float InValue)
+		{
+			return !FMath::IsNaN(InValue);
+		}
+
+		static bool IsEqual(float InLHS, float InRHS)
+		{
+			return InLHS == InRHS;
+		}
+	};
+
+	// Specialize type analysis for time
+	template<>
+	struct TTestTypeAnalysis<FTime>
+	{
+		static bool IsValid(const FTime& InValue)
+		{
+			return FMath::IsFinite(InValue.GetSeconds());
+		}
+
+		static bool IsEqual(const FTime& InLHS, const FTime& InRHS)
+		{
+			return InLHS == InRHS;
+		}
+	};
+
+	// Interface for data reference analyzers.
+	struct IDataReferenceAnalyzer
+	{
+		virtual ~IDataReferenceAnalyzer() = default;
+		virtual FAnyDataReference Copy(const FAnyDataReference& InDataRef) const = 0;
+		virtual bool IsEqual(const FAnyDataReference& InLHS, const FAnyDataReference& InRHS) const = 0;
+		virtual bool IsValid(const FAnyDataReference& InDataRef) const = 0;
+		virtual FString ToString(const FAnyDataReference& InDataRef) const = 0;
+	};
+
+	// TDataReferenceAnalyzer fulfils the IDataReferenceAnalyzer interface and uses
+	// various templates to implement methods.
+	template<typename DataType>
+	struct TDataReferenceAnalyzer : public IDataReferenceAnalyzer
+	{
+		virtual FAnyDataReference Copy(const FAnyDataReference& InDataRef) const override
+		{
+			if (const DataType* InValue = GetDataOrLogError(InDataRef))
+			{
+				return FAnyDataReference{TDataValueReference<DataType>::CreateNew(*InValue)};
+			}
+			return InDataRef;
+		}
+
+		virtual bool IsEqual(const FAnyDataReference& InLHS, const FAnyDataReference& InRHS) const override
+		{
+			if (const DataType* InLHSValue = GetDataOrLogError(InLHS))
+			{
+				if (const DataType* InRHSValue = GetDataOrLogError(InRHS))
+				{
+					return TTestTypeAnalysis<DataType>::IsEqual(*InLHSValue, *InRHSValue);
+				}
+			}
+			return false;
+		}
+
+		virtual bool IsValid(const FAnyDataReference& InDataRef) const override
+		{
+			if (const DataType* InValue = GetDataOrLogError(InDataRef))
+			{
+				return TTestTypeAnalysis<DataType>::IsValid(*InValue);
+			}
+			return false;
+		}
+
+		virtual FString ToString(const FAnyDataReference& InDataRef) const override
+		{
+			if (const DataType* InValue = GetDataOrLogError(InDataRef))
+			{
+				return TTestTypeInfo<DataType>::ToString(*InValue);
+			}
+			return TEXT("");
+		}
+	private:
+		const DataType* GetDataOrLogError(const FAnyDataReference& InDataRef) const
+		{
+			const DataType* Value = InDataRef.GetValue<DataType>();
+			if (nullptr == Value)	
+			{
+				// Data references should never be null
+				UE_LOG(LogMetaSound, Error, TEXT("Failed to get data type value of type %s"), *GetMetasoundDataTypeString<DataType>());
+			}
+			return Value;
+		}
+	};
+
+	// Register a data reference analyzer to support analyzing output of metasound nodes.
+	template<typename DataType>
+	void AddDataReferenceAnalyzerToMap(TMap<FName, TSharedPtr<const IDataReferenceAnalyzer>>& InMap)
+	{
+		InMap.Add(GetMetasoundDataTypeName<DataType>(), MakeShared<const TDataReferenceAnalyzer<DataType>>());
+	}
+
+	const TMap<FName, TSharedPtr<const IDataReferenceAnalyzer>>& GetDataTypeAnalyzerMap()
+	{
+		static TMap<FName, TSharedPtr<const IDataReferenceAnalyzer>> Map;
+
+		AddDataReferenceAnalyzerToMap<bool>(Map);
+		AddDataReferenceAnalyzerToMap<int32>(Map);
+		AddDataReferenceAnalyzerToMap<float>(Map);
+		AddDataReferenceAnalyzerToMap<FString>(Map);
+		AddDataReferenceAnalyzerToMap<FTime>(Map);
+		AddDataReferenceAnalyzerToMap<FAudioBuffer>(Map);
+		AddDataReferenceAnalyzerToMap<FTrigger>(Map);
+		AddDataReferenceAnalyzerToMap<TArray<bool>>(Map);
+		AddDataReferenceAnalyzerToMap<TArray<int32>>(Map);
+		AddDataReferenceAnalyzerToMap<TArray<float>>(Map);
+		AddDataReferenceAnalyzerToMap<TArray<FString>>(Map);
+		AddDataReferenceAnalyzerToMap<TArray<FTime>>(Map);
+
+		return Map;
+	}
+
+	TMap<FVertexName, FAnyDataReference> VertexState;
+
+	// Convenience class for setting node input data reference values to default, min, max or random values. 
+	struct FOutputVertexDataTestController
+	{
+		struct FAnalyzableOutput
+		{
+			FAnyDataReference CapturedValue;
+			FAnyDataReference DataReference;
+			FVertexName VertexName;
+			TSharedPtr<const IDataReferenceAnalyzer> DataReferenceAnalyzer;
+
+			void CaptureValue()
+			{
+				CapturedValue = DataReferenceAnalyzer->Copy(DataReference);
+			}
+
+			bool IsDataReferenceValid() const
+			{
+				return DataReferenceAnalyzer->IsValid(DataReference);
+			}
+
+			bool IsDataReferenceEqualToCapturedValue() const
+			{
+				return DataReferenceAnalyzer->IsEqual(CapturedValue, DataReference);
+			}
+
+			FString DataReferenceToString() const
+			{
+				return DataReferenceAnalyzer->ToString(DataReference);
+			}
+
+			FString CapturedValueToString() const
+			{
+				return DataReferenceAnalyzer->ToString(CapturedValue);
+			}
+		};
+
+		FOutputVertexDataTestController( const FOutputVertexInterface& InOutputInterface, const FOutputVertexInterfaceData& InOutputData)
+		{
+			const TMap<FName, TSharedPtr<const IDataReferenceAnalyzer>>& AnalyzerMap = GetDataTypeAnalyzerMap();
+
+			for (const FOutputDataVertex& Vertex : InOutputInterface)
+			{
+				if (AnalyzerMap.Contains(Vertex.DataTypeName))
+				{
+					if (const FAnyDataReference* Ref = InOutputData.FindDataReference(Vertex.VertexName))
+					{
+						AnalyzableOutputs.Add(FAnalyzableOutput{*Ref, *Ref, Vertex.VertexName, AnalyzerMap[Vertex.DataTypeName]});
+					}
+				}
+			}
+		}
+
+		int32 GetNumAnalyzableOutputs() const
+		{
+			return AnalyzableOutputs.Num();
+		}
+
+		bool AreAllAnalyzableOutputsValid() const
+		{
+			bool bAllAreValid = true;
+			for (const FAnalyzableOutput& AnalyzableOutput : AnalyzableOutputs)
+			{
+				if (!AnalyzableOutput.IsDataReferenceValid())
+				{
+					UE_LOG(LogMetaSound, Warning, TEXT("Invalid output encountered %s %s"), *AnalyzableOutput.VertexName.ToString(), *AnalyzableOutput.DataReferenceToString());
+					bAllAreValid = false;
+				}
+			}
+			return bAllAreValid;
+		}
+
+		void CaptureCurrentOutputValues()
+		{
+			for (FAnalyzableOutput& AnalyzableOutput : AnalyzableOutputs)
+			{
+				AnalyzableOutput.CaptureValue();
+			}
+		}
+
+		bool AreAllOutputValuesEqualToCapturedValues() const
+		{
+			bool bAllAreEqual = true;
+			for (const FAnalyzableOutput& AnalyzableOutput : AnalyzableOutputs)
+			{
+				if (!AnalyzableOutput.IsDataReferenceEqualToCapturedValue())
+				{
+					UE_LOG(LogMetaSound, Warning, TEXT("Unequal output encountered %s found: %s expected: %s"), *AnalyzableOutput.VertexName.ToString(), *AnalyzableOutput.DataReferenceToString(), *AnalyzableOutput.CapturedValueToString());
+					bAllAreEqual = false;
+				}
+			}
+			return bAllAreEqual;
+		}
+
+	private:
+
+		TArray<FAnalyzableOutput> AnalyzableOutputs;
+	};
+
+
 	// TTestTypeValues should return basic bounds for tested input data types.
 	// Similar to TNumericLimits<>
 	template<typename DataType>
@@ -199,6 +502,8 @@ namespace Metasound::MetasoundEngineTestPrivate {
 		virtual void SetMax(const FOperatorSettings& InSettings, const FAnyDataReference& InDataRef) const = 0;
 		virtual void SetMin(const FOperatorSettings& InSettings, const FAnyDataReference& InDataRef) const = 0; 
 		virtual void SetRandom(const FOperatorSettings& InSettings, const FAnyDataReference& InDataRef) const = 0;
+		virtual FAnyDataReference Copy(const FAnyDataReference& InDataRef) const = 0;
+		virtual void SetValue(const FAnyDataReference& InSrcDataRef, const FAnyDataReference& InDstDataRef) const = 0;
 		virtual FString ToString(const FAnyDataReference& InDataRef) const = 0;
 	};
 
@@ -238,6 +543,25 @@ namespace Metasound::MetasoundEngineTestPrivate {
 				return TEXT("");
 			}
 		}
+
+		virtual void SetValue(const FAnyDataReference& InSrcDataRef, const FAnyDataReference& InDstDataRef) const override
+		{
+			*InDstDataRef.GetDataWriteReference<DataType>() = *InSrcDataRef.GetValue<DataType>();
+		}
+
+		virtual FAnyDataReference Copy(const FAnyDataReference& InDataRef) const override
+		{
+			if (const DataType* Data = InDataRef.GetValue<DataType>())
+			{
+				return FAnyDataReference{TDataValueReference<DataType>::CreateNew(*Data)};
+			}
+			else
+			{
+				// Data references should never be null
+				UE_LOG(LogMetaSound, Error, TEXT("Failed to get data type value of type %s"), *GetMetasoundDataTypeString<DataType>());
+				return FAnyDataReference{TDataValueReference<int32>::CreateNew()}; // we are going to crash soon.
+			}
+		}
 	};
 
 
@@ -268,6 +592,8 @@ namespace Metasound::MetasoundEngineTestPrivate {
 		return Map;
 	}
 
+	using FInterfaceState = TMap<FVertexName, FAnyDataReference>;
+
 	// Convenience class for setting node input data reference values to default, min, max or random values. 
 	struct FInputVertexDataTestController
 	{
@@ -303,6 +629,28 @@ namespace Metasound::MetasoundEngineTestPrivate {
 			return MutableInputs.Num();
 		}
 
+		FInterfaceState GetInterfaceState() const
+		{
+			FInterfaceState State;
+			for (const FMutableInput& MutableInput : MutableInputs)
+			{
+				State.Add(MutableInput.VertexName, MutableInput.DataReferenceMutator->Copy(MutableInput.DataReference));
+			}
+			return State;
+		}
+
+		void SetMutableInputsToState(const FInterfaceState& InState) 
+		{
+			for (const FMutableInput& MutableInput : MutableInputs)
+			{
+				if (const FAnyDataReference* Value = InState.Find(MutableInput.VertexName))
+				{
+					MutableInput.DataReferenceMutator->SetValue(*Value, MutableInput.DataReference);
+				}
+			}
+			UE_LOG(LogMetaSound, Verbose, TEXT("Setting operator input values:%s%s"), LINE_TERMINATOR, *FString::Join(GetInputValueStrings(), LINE_TERMINATOR));
+		}
+
 		void SetMutableInputsToMin()
 		{
 			for (const FMutableInput& MutableInput : MutableInputs)
@@ -541,60 +889,117 @@ bool FMetasoundAutomatedNodeTest::RunTest(const FString& InRegistryKey)
 		SourceEnvironment
 	};
 
+	IOperator::FResetParams ResetParams
+	{
+		OperatorSettings,
+		SourceEnvironment
+	};
+
 	// Convenience function for testing entire lifecycle of an individual operator 
 	// with a variety of inputs.
-	auto RunTestIteration = [&]() -> bool
+	auto RunTestIteration = [&]()
 	{
+		// Create the operator from the node factory
 		FBuildResults BuildResults;
 		TUniquePtr<IOperator> Operator = Node->GetDefaultOperatorFactory()->CreateOperator(BuildParams, BuildResults);
 
 		if (!Operator.IsValid())
 		{
-			AddError(FString::Printf(TEXT("Failed to create operator from node %s."), *InRegistryKey));
-			return false;
+			AddError(FString::Printf(TEXT("Failed to create operator from node %s - %s."), *InRegistryKey, *GetPrettyName(InRegistryKey)));
 		}
+
+
+		// Store a copy of the input values data values so the inputs
+		// can be reset to this value later during the test.
+		FInterfaceState InitialInputState = InputTester.GetInterfaceState();
+
+		// Bind to inputs and output data of operator.
+		FVertexInterface VertexInterface = Node->GetVertexInterface();
+		FVertexInterfaceData VertexInterfaceData{VertexInterface};
+		Operator->Bind(VertexInterfaceData);
+
+		// Create output tester which will analyzer outputs of operator.
+		FOutputVertexDataTestController OutputTester{VertexInterface.GetOutputInterface(), VertexInterfaceData.GetOutputs()};
+		// Convenience method for printing errors if the OutputTester finds an error.
+		auto CheckOutputValuesAreValid = [&]()
+		{
+			if (!OutputTester.AreAllAnalyzableOutputsValid())
+			{
+				AddError(FString::Printf(TEXT("Invalid output value encountered from node %s - %s."), *InRegistryKey, *GetPrettyName(InRegistryKey)));
+			}
+		};
+		// Initial values should all be valid.
+		CheckOutputValuesAreValid();
 		
-		IOperator::FExecuteFunction OpExecFunc = Operator->GetExecuteFunction();
+		// Capture current state of outputs so they can be referenced at a later time.
+		// The captured values are held within the OutputTester.
+		OutputTester.CaptureCurrentOutputValues();
 
+		IOperator::FExecuteFunction OpExecFunc = Operator->GetExecuteFunction();
+		IOperator::FResetFunction OpResetFunc = Operator->GetResetFunction();
+		// Test execute function with input variations
 		if (OpExecFunc)
 		{
 			OpExecFunc(Operator.Get());
+			CheckOutputValuesAreValid();
 
 			if (InputTester.GetNumMutableInputs() > 0)
 			{
 				InputTester.SetMutableInputsToDefault();
 				OpExecFunc(Operator.Get());
+				CheckOutputValuesAreValid();
 
 				InputTester.SetMutableInputsToMin();
 				OpExecFunc(Operator.Get());
+				CheckOutputValuesAreValid();
 
 				InputTester.SetMutableInputsToMax();
 				OpExecFunc(Operator.Get());
+				CheckOutputValuesAreValid();
 
 				InputTester.SetMutableInputsToRandom();
 				OpExecFunc(Operator.Get());
+				CheckOutputValuesAreValid();
 			}
 		}
 
-		return true;
+		// Return inputs to initial state. 
+		InputTester.SetMutableInputsToState(InitialInputState);
+
+		if (OpResetFunc)
+		{
+			OpResetFunc(Operator.Get(), ResetParams);
+			CheckOutputValuesAreValid();
+		}
+		else if (OpExecFunc)
+		{
+			AddError(FString::Printf(TEXT("Missing initialize function when execute function exists for node %s - %s"), *InRegistryKey, *GetPrettyName(InRegistryKey)));
+		}
+
+		// Check that after returning all inputs to their original state and calling
+		// reset on the operator, that all output values have returned to their initail state. 
+		if (!OutputTester.AreAllOutputValuesEqualToCapturedValues())
+		{
+			AddError(FString::Printf(TEXT("Reset function resulted in different starting conditions for node %s - %s"), *InRegistryKey, *GetPrettyName(InRegistryKey)));
+		}
 	};
 
-	bool bSuccess = true;
 	// Test entire operator lifecycle with different starting conditions if
 	// any of the inputs are mutable
 	InputTester.SetMutableInputsToDefault();
-	bSuccess &= RunTestIteration();
+	RunTestIteration();
+
 	if (InputTester.GetNumMutableInputs() > 0)
 	{
 		InputTester.SetMutableInputsToMin();
-		bSuccess &= RunTestIteration();
+		RunTestIteration();
 		InputTester.SetMutableInputsToMax();
-		bSuccess &= RunTestIteration();
+		RunTestIteration();
 		InputTester.SetMutableInputsToRandom();
-		bSuccess &= RunTestIteration();
+		RunTestIteration();
 	}
 
-	return bSuccess;
+	return true;
 }
 
 #endif // WITH_EDITORONLY_DATA
diff --git a/Engine/Plugins/Runtime/Metasound/Source/MetasoundFrontend/Public/MetasoundArrayNodes.h b/Engine/Plugins/Runtime/Metasound/Source/MetasoundFrontend/Public/MetasoundArrayNodes.h
--- a/Engine/Plugins/Runtime/Metasound/Source/MetasoundFrontend/Public/MetasoundArrayNodes.h
+++ b/Engine/Plugins/Runtime/Metasound/Source/MetasoundFrontend/Public/MetasoundArrayNodes.h
@@ -286,11 +286,18 @@ namespace Metasound
 		: Trigger(InParams.Trigger)
 		, Array(InParams.Array)
 		, Index(InParams.Index)
-		, Value(TDataWriteReferenceFactory<ElementType>::CreateAny(InSettings))
+		, Value(TDataWriteReferenceFactory<ElementType>::CreateExplicitArgs(InSettings))
 #if WITH_METASOUND_DEBUG_ENVIRONMENT
 		, GraphName(InParams.GraphName)
 #endif // WITH_METASOUND_DEBUG_ENVIRONMENT
 		{
+			const int32 IndexValue = *Index;
+			const ArrayType& ArrayRef = *Array;
+
+			if ((IndexValue >= 0) && (IndexValue < ArrayRef.Num()))
+			{
+				*Value = ArrayRef[IndexValue];
+			}
 		}
 
 		virtual ~TArrayGetOperator() = default;
@@ -351,7 +358,7 @@ namespace Metasound
 			}
 			else
 			{
-				*Value = TDataTypeFactory<ElementType>::CreateAny(InParams.OperatorSettings);
+				*Value = TDataTypeFactory<ElementType>::CreateExplicitArgs(InParams.OperatorSettings);
 			}
 		}
 
