[Metasound] Reset functions for nodes

diff --git a/Engine/Plugins/Runtime/Metasound/Source/MetasoundStandardNodes/Private/MetasoundDebugTriggerOnNanNode.cpp b/Engine/Plugins/Runtime/Metasound/Source/MetasoundStandardNodes/Private/MetasoundDebugTriggerOnNanNode.cpp
--- a/Engine/Plugins/Runtime/Metasound/Source/MetasoundStandardNodes/Private/MetasoundDebugTriggerOnNanNode.cpp
+++ b/Engine/Plugins/Runtime/Metasound/Source/MetasoundStandardNodes/Private/MetasoundDebugTriggerOnNanNode.cpp
@@ -37,6 +37,8 @@ namespace Metasound
 
 		FTriggerOnNanOperator(const FOperatorSettings& InSettings, const FAudioBufferReadRef& InAudioInput, const FBoolReadRef& InTriggerOnce);
 
+		void Reset(const IOperator::FResetParams&);
+
 		virtual FDataReferenceCollection GetInputs() const override;
 
 		virtual FDataReferenceCollection GetOutputs() const override;
@@ -64,6 +66,12 @@ namespace Metasound
 	{
 	}
 
+	void FTriggerOnNanOperator::Reset(const IOperator::FResetParams&)
+	{
+		TriggerOutput->Reset();
+		bHasTriggered = false;
+	}
+
 	FDataReferenceCollection FTriggerOnNanOperator::GetInputs() const
 	{
 		using namespace TriggerOnNanVertexNames;
diff --git a/Engine/Plugins/Runtime/Metasound/Source/MetasoundStandardNodes/Private/MetasoundLowFrequencyNoise.cpp b/Engine/Plugins/Runtime/Metasound/Source/MetasoundStandardNodes/Private/MetasoundLowFrequencyNoise.cpp
--- a/Engine/Plugins/Runtime/Metasound/Source/MetasoundStandardNodes/Private/MetasoundLowFrequencyNoise.cpp
+++ b/Engine/Plugins/Runtime/Metasound/Source/MetasoundStandardNodes/Private/MetasoundLowFrequencyNoise.cpp
@@ -177,16 +177,24 @@ namespace Metasound
 				, InputCol.GetDataReadReferenceOrConstructWithVertexDefault<float>(InputInterface, METASOUND_GET_PARAM_NAME(MaxOutputValuePin), Settings)
 			};
 
-			return MakeUnique<FLowFrequencyNoiseOperator>(Settings, MoveTemp(Pins));
+			return MakeUnique<FLowFrequencyNoiseOperator>(InParams, MoveTemp(Pins));
 		}
 
-		FLowFrequencyNoiseOperator(const FOperatorSettings& InSettings, FPinReadRefs&& InPins)
-			: BlockRate{ InSettings.GetActualBlockRate() }
-			, Phase{ 0.0f }
+		FLowFrequencyNoiseOperator(const FCreateOperatorParams& InParams, FPinReadRefs&& InPins)
+			: BlockRate{ InParams.OperatorSettings.GetActualBlockRate() }
 			, Pins{ MoveTemp(InPins) }
 			, ScaledOutput{ FFloatWriteRef::CreateNew(0.f) }
 			, NormalizedOutput{ FFloatWriteRef::CreateNew(0.f) }
 		{	
+			Reset(InParams);
+		}
+
+		void Reset(const IOperator::FResetParams& InParams)
+		{
+			Phase = 0.f;
+			*ScaledOutput = 0.f;
+			*NormalizedOutput = 0.f;
+
 			SetSeedAndFillValues();
 		}
 
diff --git a/Engine/Plugins/Runtime/Metasound/Source/MetasoundStandardNodes/Private/MetasoundPerlinNoiseNode.cpp b/Engine/Plugins/Runtime/Metasound/Source/MetasoundStandardNodes/Private/MetasoundPerlinNoiseNode.cpp
--- a/Engine/Plugins/Runtime/Metasound/Source/MetasoundStandardNodes/Private/MetasoundPerlinNoiseNode.cpp
+++ b/Engine/Plugins/Runtime/Metasound/Source/MetasoundStandardNodes/Private/MetasoundPerlinNoiseNode.cpp
@@ -104,15 +104,6 @@ namespace Metasound
 			const FInputVertexInterfaceData& InputData = InParams.InputData;
 
 			using namespace PerlinNoiseVertexNames;
-
-			// Constructor pin.
-			// If Offset is -1, use system rand as a "seed"
-			float Offset = InputData.GetOrCreateDefaultValue<float>(METASOUND_GET_PARAM_NAME(OffsetPin), InParams.OperatorSettings); 
-			if (FMath::IsNearlyEqual(Offset, -1.f))
-			{
-				Offset = FMath::FRand();
-			}
-
 			FPinReadRefs Pins
 			{
 				  InputData.GetOrCreateDefaultDataReadReference<TDataClass>(METASOUND_GET_PARAM_NAME(XPin), Settings)
@@ -120,21 +111,36 @@ namespace Metasound
 				, InputData.GetOrCreateDefaultDataReadReference<float>(METASOUND_GET_PARAM_NAME(MinValuePin), Settings)
 				, InputData.GetOrCreateDefaultDataReadReference<float>(METASOUND_GET_PARAM_NAME(MaxValuePin), Settings)
 				, InputData.IsVertexBound(METASOUND_GET_PARAM_NAME(XPin))
-				, Offset
+				, InputData.GetOrCreateDefaultValue<float>(METASOUND_GET_PARAM_NAME(OffsetPin), Settings)
 			};
 
-			return MakeUnique<TThisType>(Settings, MoveTemp(Pins));
+			return MakeUnique<TThisType>(InParams, MoveTemp(Pins));
 		}
 
-		TBasePerlinNoiseOperator(const FOperatorSettings& InSettings, FPinReadRefs&& InPins)
-			: SampleRate{ InSettings.GetSampleRate() }
-			, NumFramesPerBlock{ InSettings.GetNumFramesPerBlock() }
-			, Pins{ MoveTemp(InPins) }
-			, OutputScaled{ TThisType::CreateOutputWriteRef(InSettings) }
-			, OutputNormalized{ TThisType::CreateOutputWriteRef(InSettings) }
+		TBasePerlinNoiseOperator(const FBuildOperatorParams& InParams, FPinReadRefs&& InPins)
+			: Pins{ MoveTemp(InPins) }
+			, OutputScaled{ TThisType::CreateOutputWriteRef(InParams.OperatorSettings) }
+			, OutputNormalized{ TThisType::CreateOutputWriteRef(InParams.OperatorSettings) }
 		{
+			Reset(InParams);
+		}
+
+		void Reset(const IOperator::FResetParams& InParams)
+		{
+			// If Offset is -1, use system rand as a "seed"
+			if (FMath::IsNearlyEqual(Pins.Offset, -1.f))
+			{
+				Pins.Offset = FMath::FRand();
+			}
+
+			SampleRate = InParams.OperatorSettings.GetSampleRate();
+			NumFramesPerBlock = InParams.OperatorSettings.GetNumFramesPerBlock();
+
 			check(SampleRate > 0);
 			check(NumFramesPerBlock > 0);
+
+			TThisType::ResetOutputWriteRef(OutputScaled);
+			TThisType::ResetOutputWriteRef(OutputNormalized);
 		}
 
 		FDataReferenceCollection GetInputs() const override { return {}; }
@@ -179,6 +185,8 @@ namespace Metasound
 		using BaseClass::CreateOperator;
 
 		static FFloatWriteRef CreateOutputWriteRef(const FOperatorSettings& InSettings) { return FFloatWriteRef::CreateNew(0.f); }
+		static void ResetOutputWriteRef(FFloatWriteRef& WriteRef) { *WriteRef = 0; };
+
 		static FNodeClassMetadata CreateNodeInfo()
 		{
 			FNodeClassMetadata Info;
@@ -188,6 +196,11 @@ namespace Metasound
 			return Info;
 		}
 
+		void Reset(const IOperator::FResetParams& InParams)
+		{
+			BaseClass::Reset(InParams);
+		}
+
 		void Execute()
 		{
 			// Use input pin if its connected, otherwise use elapsed time in seconds.
@@ -220,6 +233,7 @@ namespace Metasound
 		using BaseClass::CreateOperator;
 
 		static FAudioBufferWriteRef CreateOutputWriteRef(const FOperatorSettings& InSettings) { return FAudioBufferWriteRef::CreateNew(InSettings); }
+		static void ResetOutputWriteRef(FAudioBufferWriteRef& WriteRef) { WriteRef->Zero(); };
 
 		static FNodeClassMetadata CreateNodeInfo()
 		{
@@ -230,9 +244,9 @@ namespace Metasound
 			return Info;
 		}
 
-		FPerlinAudioRate(const FOperatorSettings& InSettings, FPinReadRefs&& InPins)
-			: BaseClass{ InSettings, MoveTemp(InPins) }
-			, TimeBuffer{ InPins.bXConnected ? 0 : InSettings.GetNumFramesPerBlock() }	// Create a scratch buffer if X is not connected.
+		FPerlinAudioRate(const FBuildOperatorParams& InParams, FPinReadRefs&& InPins)
+			: BaseClass{ InParams, MoveTemp(InPins) }
+			, TimeBuffer{ InPins.bXConnected ? 0 : InParams.OperatorSettings.GetNumFramesPerBlock() }	// Create a scratch buffer if X is not connected.
 		{
 		}
 
@@ -262,6 +276,11 @@ namespace Metasound
 			}
 		}
 
+		void Reset(const IOperator::FResetParams& InParams)
+		{
+			BaseClass::Reset(InParams);
+		}
+
 		void Execute()
 		{
 			// Use input pin if its connected, otherwise generate an array of increasing time
