[Metasounds] removing Reset() function from0 audio bus writer node to see if it is the cause of automated test crashes.

diff --git a/Engine/Plugins/Runtime/Metasound/Source/MetasoundEngine/Private/MetasoundAudioBusWriterNode.cpp b/Engine/Plugins/Runtime/Metasound/Source/MetasoundEngine/Private/MetasoundAudioBusWriterNode.cpp
--- a/Engine/Plugins/Runtime/Metasound/Source/MetasoundEngine/Private/MetasoundAudioBusWriterNode.cpp
+++ b/Engine/Plugins/Runtime/Metasound/Source/MetasoundEngine/Private/MetasoundAudioBusWriterNode.cpp
@@ -78,15 +78,17 @@ namespace Metasound
 		static TUniquePtr<IOperator> CreateOperator(const FCreateOperatorParams& InParams, FBuildErrorArray& OutErrors)
 		{
 			using namespace Frontend;
-			using namespace AudioBusWriterNode;
 
+			using namespace AudioBusWriterNode;
 			const FDataReferenceCollection& InputCollection = InParams.InputDataReferences;
 
-			bool bHasEnvironmentVars = InParams.Environment.Contains<Audio::FDeviceId>(SourceInterface::Environment::DeviceID);
-			bHasEnvironmentVars &= InParams.Environment.Contains<int32>(SourceInterface::Environment::AudioMixerNumOutputFrames);
-
+			bool bHasEnvironmentVars = InParams.Environment.Contains<Audio::FDeviceId>(SourceInterface::Environment::DeviceID) && InParams.Environment.Contains<int32>(SourceInterface::Environment::AudioMixerNumOutputFrames);
+			
 			if (bHasEnvironmentVars)
 			{
+				Audio::FDeviceId AudioDeviceId = InParams.Environment.GetValue<Audio::FDeviceId>(SourceInterface::Environment::DeviceID);
+				int32 AudioMixerOutputFrames = InParams.Environment.GetValue<int32>(SourceInterface::Environment::AudioMixerNumOutputFrames);
+
 				FAudioBusAssetReadRef AudioBusIn = InputCollection.GetDataReadReferenceOrConstruct<FAudioBusAsset>(METASOUND_GET_PARAM_NAME(InParamAudioBusOutput));
 
 				TArray<FAudioBufferReadRef> AudioInputs;
@@ -95,43 +97,22 @@ namespace Metasound
 					AudioInputs.Add(InputCollection.GetDataReadReferenceOrConstruct<FAudioBuffer>(METASOUND_GET_PARAM_NAME_WITH_INDEX(InParamAudio, ChannelIndex), InParams.OperatorSettings));
 				}
 
-				return MakeUnique<TAudioBusWriterOperator<NumChannels>>(InParams, MoveTemp(AudioBusIn), MoveTemp(AudioInputs));
+				return MakeUnique<TAudioBusWriterOperator<NumChannels>>(InParams, AudioMixerOutputFrames, AudioDeviceId, MoveTemp(AudioBusIn), MoveTemp(AudioInputs));
 			}
 			else
 			{
-				UE_LOG(LogMetaSound, Warning, TEXT("Audio bus writer node requires audio device ID '%s' and audio mixer num output frames '%s' environment variables")
-					, *SourceInterface::Environment::DeviceID.ToString(), *SourceInterface::Environment::AudioMixerNumOutputFrames.ToString());
+				UE_LOG(LogMetaSound, Warning, TEXT("Audio bus writer node requires audio device ID '%s' and audio mixer num output frames '%s' environment variables"), *SourceInterface::Environment::DeviceID.ToString(), *SourceInterface::Environment::AudioMixerNumOutputFrames.ToString());
 				return nullptr;
 			}
 		}
 
-		TAudioBusWriterOperator(const FCreateOperatorParams& InParams, FAudioBusAssetReadRef InAudioBusAsset, TArray<FAudioBufferReadRef> InAudioInputs)
-			: AudioBusAsset(MoveTemp(InAudioBusAsset))
-			, AudioInputs(MoveTemp(InAudioInputs))
+		TAudioBusWriterOperator(const FCreateOperatorParams& InParams, int32 InAudioMixerOutputFrames, Audio::FDeviceId InAudioDeviceId, FAudioBusAssetReadRef InAudioBusAsset, TArray<FAudioBufferReadRef> InAudioInputs) :
+			AudioBusAsset(MoveTemp(InAudioBusAsset)),
+			AudioInputs(MoveTemp(InAudioInputs)),
+			AudioMixerOutputFrames(InAudioMixerOutputFrames),
+			AudioDeviceId(InAudioDeviceId),
+			SampleRate(InParams.OperatorSettings.GetSampleRate())
 		{
-			Reset(InParams);
-		}
-
-		void Reset(const IOperator::FResetParams& InParams)
-		{
-			using namespace Frontend;
-			using namespace AudioBusWriterNode;
-
-			bool bHasEnvironmentVars = InParams.Environment.Contains<Audio::FDeviceId>(SourceInterface::Environment::DeviceID);
-			bHasEnvironmentVars &= InParams.Environment.Contains<int32>(SourceInterface::Environment::AudioMixerNumOutputFrames);
-
-			if (bHasEnvironmentVars)
-			{
-				SampleRate = InParams.OperatorSettings.GetSampleRate();
-				AudioDeviceId = InParams.Environment.GetValue<Audio::FDeviceId>(SourceInterface::Environment::DeviceID);
-				AudioMixerOutputFrames = InParams.Environment.GetValue<int32>(SourceInterface::Environment::AudioMixerNumOutputFrames);
-			}
-			else
-			{
-				UE_LOG(LogMetaSound, Warning, TEXT("Audio bus writer node requires audio device ID '%s' and audio mixer num output frames '%s' environment variables")
-					, *SourceInterface::Environment::DeviceID.ToString(), *SourceInterface::Environment::AudioMixerNumOutputFrames.ToString());
-			}
-			
 			const FAudioBusProxyPtr& AudioBusProxy = AudioBusAsset->GetAudioBusProxy();
 			if (AudioBusProxy.IsValid())
 			{
@@ -151,12 +132,9 @@ namespace Metasound
 
 						// Create a bus patch input with enough room for the number of samples we expect and some buffering
 						AudioBusPatchInput = AudioBusSubsystem->AddPatchInputForAudioBus(AudioBusKey, BlockSizeFrames, AudioBusChannels);
-
 					}
 				}
 			}
-
-			bFirstBlock = true;
 		}
 
 		virtual FDataReferenceCollection GetInputs() const override
