[Metasound] Pass audio component id to OnGeneratorInstanceCreated instead of the individual voice id

diff --git a/Engine/Plugins/Runtime/Metasound/Source/MetasoundEngine/Private/MetasoundGeneratorHandle.cpp b/Engine/Plugins/Runtime/Metasound/Source/MetasoundEngine/Private/MetasoundGeneratorHandle.cpp
--- a/Engine/Plugins/Runtime/Metasound/Source/MetasoundEngine/Private/MetasoundGeneratorHandle.cpp
+++ b/Engine/Plugins/Runtime/Metasound/Source/MetasoundEngine/Private/MetasoundGeneratorHandle.cpp
@@ -21,6 +21,14 @@ UMetasoundGeneratorHandle* UMetasoundGeneratorHandle::CreateMetaSoundGeneratorHa
 		return nullptr;
 	}
 
+	if (OnComponent->bCanPlayMultipleInstances)
+	{
+		UE_LOG(
+			LogMetaSound,
+			Warning,
+			TEXT("Created a UMetaSoundGeneratorHandle for a UAudioComponent that is allowed to play multiple instances. This may not work as expected."))
+	}
+
 	UMetasoundGeneratorHandle* Result = NewObject<UMetasoundGeneratorHandle>();
 	Result->SetAudioComponent(OnComponent);
 	return Result;
diff --git a/Engine/Source/Runtime/AudioMixer/Private/AudioMixerSource.cpp b/Engine/Source/Runtime/AudioMixer/Private/AudioMixerSource.cpp
--- a/Engine/Source/Runtime/AudioMixer/Private/AudioMixerSource.cpp
+++ b/Engine/Source/Runtime/AudioMixer/Private/AudioMixerSource.cpp
@@ -1134,6 +1134,7 @@ namespace Audio
 
 		FMixerSourceBufferInitArgs BufferInitArgs;
 		BufferInitArgs.AudioDeviceID = AudioDevice->DeviceID;
+		BufferInitArgs.AudioComponentID = InstanceID;
 		BufferInitArgs.InstanceID = GetTransmitterID(InstanceID, WaveInstance->WaveInstanceHash, PlayOrder);
 		BufferInitArgs.SampleRate = AudioDevice->GetSampleRate();
 		BufferInitArgs.AudioMixerNumOutputFrames = MixerDevice->GetNumOutputFrames();
diff --git a/Engine/Source/Runtime/AudioMixer/Private/AudioMixerSourceBuffer.cpp b/Engine/Source/Runtime/AudioMixer/Private/AudioMixerSourceBuffer.cpp
--- a/Engine/Source/Runtime/AudioMixer/Private/AudioMixerSourceBuffer.cpp
+++ b/Engine/Source/Runtime/AudioMixer/Private/AudioMixerSourceBuffer.cpp
@@ -108,7 +108,7 @@ namespace Audio
 			InitParams.AudioMixerNumOutputFrames = InArgs.AudioMixerNumOutputFrames;
 			InitParams.NumChannels = NumChannels;
 			InitParams.NumFramesPerCallback = MONO_PCM_BUFFER_SAMPLES;
-			InitParams.InstanceID = InArgs.InstanceID;
+			InitParams.InstanceID = InArgs.AudioComponentID;
 			InitParams.bIsPreviewSound = InArgs.bIsPreviewSound;
 
 			SoundGenerator = InArgs.SoundWave->CreateSoundGenerator(InitParams, MoveTemp(InDefaultParams));
diff --git a/Engine/Source/Runtime/AudioMixer/Private/AudioMixerSourceBuffer.h b/Engine/Source/Runtime/AudioMixer/Private/AudioMixerSourceBuffer.h
--- a/Engine/Source/Runtime/AudioMixer/Private/AudioMixerSourceBuffer.h
+++ b/Engine/Source/Runtime/AudioMixer/Private/AudioMixerSourceBuffer.h
@@ -49,6 +49,7 @@ namespace Audio
 	struct FMixerSourceBufferInitArgs
 	{
 		FDeviceId AudioDeviceID = 0;
+		uint64 AudioComponentID = 0;
 		uint32 InstanceID = 0;
 		int32 SampleRate = 0;
 		int32 AudioMixerNumOutputFrames = 0;
