Adding diagnostcis for metasounds which load before interfaces are registered.

diff --git a/Engine/Plugins/Runtime/Metasound/Source/MetasoundFrontend/Private/MetasoundAssetBase.cpp b/Engine/Plugins/Runtime/Metasound/Source/MetasoundFrontend/Private/MetasoundAssetBase.cpp
--- a/Engine/Plugins/Runtime/Metasound/Source/MetasoundFrontend/Private/MetasoundAssetBase.cpp
+++ b/Engine/Plugins/Runtime/Metasound/Source/MetasoundFrontend/Private/MetasoundAssetBase.cpp
@@ -373,7 +373,7 @@ bool FMetasoundAssetBase::VersionAsset()
 			const TArray<FMetasoundFrontendVersion> Versions = Doc->Interfaces.Array();
 			for (const FMetasoundFrontendVersion& Version : Versions)
 			{
-				bPassUpdated |= FUpdateRootGraphInterface(Version).Transform(DocHandle);
+				bPassUpdated |= FUpdateRootGraphInterface(Version, GetOwningAssetName()).Transform(DocHandle);
 			}
 
 			bInterfaceUpdated |= bPassUpdated;
diff --git a/Engine/Plugins/Runtime/Metasound/Source/MetasoundFrontend/Private/MetasoundFrontendTransform.cpp b/Engine/Plugins/Runtime/Metasound/Source/MetasoundFrontend/Private/MetasoundFrontendTransform.cpp
--- a/Engine/Plugins/Runtime/Metasound/Source/MetasoundFrontend/Private/MetasoundFrontendTransform.cpp
+++ b/Engine/Plugins/Runtime/Metasound/Source/MetasoundFrontend/Private/MetasoundFrontendTransform.cpp
@@ -474,7 +474,7 @@ namespace Metasound
 			bool bFoundTargetInterface = ISearchEngine::Get().FindInterfaceWithHighestVersion(InterfaceVersion.Name, TargetInterface);
 			if (!bFoundTargetInterface)
 			{
-				UE_LOG(LogMetaSound, Error, TEXT("Target interface is not registered [InterfaceVersion:%s]"), *InterfaceVersion.ToString());
+				UE_LOG(LogMetaSound, Warning, TEXT("Could not check for interface updates. Target interface is not registered [InterfaceVersion:%s] when attempting to update root graph of asset (%s). Ensure that the module which registers the interface has been loaded before the asset is loaded."), *InterfaceVersion.ToString(), *OwningAssetName);
 				return false;
 			}
 
diff --git a/Engine/Plugins/Runtime/Metasound/Source/MetasoundFrontend/Public/MetasoundFrontendTransform.h b/Engine/Plugins/Runtime/Metasound/Source/MetasoundFrontend/Public/MetasoundFrontendTransform.h
--- a/Engine/Plugins/Runtime/Metasound/Source/MetasoundFrontend/Public/MetasoundFrontendTransform.h
+++ b/Engine/Plugins/Runtime/Metasound/Source/MetasoundFrontend/Public/MetasoundFrontendTransform.h
@@ -130,8 +130,9 @@ namespace Metasound
 		class METASOUNDFRONTEND_API FUpdateRootGraphInterface : public IDocumentTransform
 		{
 		public:
-			FUpdateRootGraphInterface(const FMetasoundFrontendVersion& InInterfaceVersion)
+			FUpdateRootGraphInterface(const FMetasoundFrontendVersion& InInterfaceVersion, const FString& InOwningAssetName=FString(TEXT("Unknown")))
 				: InterfaceVersion(InInterfaceVersion)
+				, OwningAssetName(InOwningAssetName)
 			{
 			}
 
@@ -142,6 +143,7 @@ namespace Metasound
 			bool UpdateDocumentInterface(const TArray<const IInterfaceRegistryEntry*>& InUpgradePath, FDocumentHandle InDocument) const;
 
 			FMetasoundFrontendVersion InterfaceVersion;
+			FString OwningAssetName;
 		};
 
 		/** Completely rebuilds the graph connecting a preset's inputs to the reference
