Sanitize input values to prevent NaNs in the case of extremely large inputs. i.e. FLT_Max and FLT_Min.
#okforhub public

diff --git a/Engine/Plugins/Runtime/Metasound/Source/MetasoundStandardNodes/Private/MetasoundPerlinNoiseNode.cpp b/Engine/Plugins/Runtime/Metasound/Source/MetasoundStandardNodes/Private/MetasoundPerlinNoiseNode.cpp
--- a/Engine/Plugins/Runtime/Metasound/Source/MetasoundStandardNodes/Private/MetasoundPerlinNoiseNode.cpp
+++ b/Engine/Plugins/Runtime/Metasound/Source/MetasoundStandardNodes/Private/MetasoundPerlinNoiseNode.cpp
@@ -35,6 +35,12 @@ namespace Metasound
 		METASOUND_PARAM(OutputNormalizedPin, "Normalized", "Normalized output from Perlin function (-1...1)")
 	}
 
+	// For sanitizing inputs against so we can still sum them in range.
+	static float SanitizeFloatInput(const float In, const float InMaxMultiple)
+	{
+		return FMath::Clamp(In, FLT_MIN / InMaxMultiple, FLT_MAX / InMaxMultiple);
+	}
+
 	template<typename TDataClass, typename TThisType>
 	class TBasePerlinNoiseOperator : public TExecutableOperator<TThisType>
 	{
@@ -208,10 +214,20 @@ namespace Metasound
 				*Pins.X :
 				GetElapsedTimeInSeconds();
 
+			// Sanitize input. 
 			int32 OctavesClamped = FMath::Clamp(*Pins.Octaves, 1, MaxOctaves);
-
+			
+			// Make sure we can sum safely the X+Offset without overflow. So limit to FLT_Max/2
+			float OffsetClamped = SanitizeFloatInput(Pins.Offset, 2.f);
+			float XClamped = SanitizeFloatInput(X, 2.f);
+			
+			// Value noise will scale up this value to 2^Octaves frequency, so clamp to that range.
+			float MaxFreq = (float)(1UL << OctavesClamped);
+			float Sum = OffsetClamped + XClamped;
+			float SumClamped = SanitizeFloatInput(Sum, MaxFreq);
+			
 			// Generate noise (this is multiple layers of Perlin noise at different frequencies).
-			const float ValueNoise = Audio::PerlinValueNoise1D(X + Pins.Offset, OctavesClamped);
+			const float ValueNoise = Audio::PerlinValueNoise1D(SumClamped, OctavesClamped);
 
 			// Output the scaled version.
 			const float ScaledValue = FMath::GetMappedRangeValueClamped(
@@ -223,7 +239,7 @@ namespace Metasound
 			*OutputNormalized = ValueNoise;
 
 			AccumulateTime();
-		}		
+		}
 	};
 
 	class FPerlinAudioRate : public TBasePerlinNoiseOperator<FAudioBuffer, FPerlinAudioRate>
