Clear object literals when connecting MetaSound nodes to prevent hidden asset references

diff --git a/Engine/Plugins/Runtime/Metasound/Source/MetasoundFrontend/Private/Interfaces/MetasoundFrontendInterface.cpp b/Engine/Plugins/Runtime/Metasound/Source/MetasoundFrontend/Private/Interfaces/MetasoundFrontendInterface.cpp
--- a/Engine/Plugins/Runtime/Metasound/Source/MetasoundFrontend/Private/Interfaces/MetasoundFrontendInterface.cpp
+++ b/Engine/Plugins/Runtime/Metasound/Source/MetasoundFrontend/Private/Interfaces/MetasoundFrontendInterface.cpp
@@ -233,11 +233,6 @@ namespace Metasound
 			Algo::ForEach(InSubgraphs, GatherRequiredEnvironmentVariablesFromClass);
 		}
 
-		static FMetasoundFrontendVersionNumber GetMaxVersion()
-		{
-			return FMetasoundFrontendVersionNumber{ 1, 10 };
-		}
-
 		const FMetasoundFrontendInterface* FindMostSimilarInterfaceSupportingEnvironment(const FMetasoundFrontendGraphClass& InRootGraph, const TArray<FMetasoundFrontendClass>& InDependencies, const TArray<FMetasoundFrontendGraphClass>& InSubgraphs, const TArray<FMetasoundFrontendInterface>& InCandidateInterfaces)
 		{
 			using namespace InterfacePrivate;
diff --git a/Engine/Plugins/Runtime/Metasound/Source/MetasoundFrontend/Private/MetasoundFrontendDocumentVersioning.cpp b/Engine/Plugins/Runtime/Metasound/Source/MetasoundFrontend/Private/MetasoundFrontendDocumentVersioning.cpp
--- a/Engine/Plugins/Runtime/Metasound/Source/MetasoundFrontend/Private/MetasoundFrontendDocumentVersioning.cpp
+++ b/Engine/Plugins/Runtime/Metasound/Source/MetasoundFrontend/Private/MetasoundFrontendDocumentVersioning.cpp
@@ -485,6 +485,32 @@ namespace Metasound
 			}
 		};
 
+		/** Versions document from 1.10 to 1.11. */
+		class FVersionDocument_1_11 : public FVersionDocumentTransform
+		{
+		public:
+			FVersionDocument_1_11() = default;
+
+			FMetasoundFrontendVersionNumber GetTargetVersion() const override
+			{
+				return { 1, 11 };
+			}
+
+			void TransformInternal(FDocumentHandle InDocument) const override
+			{
+				// Clear object literals on inputs that are connected 
+				// to prevent referencing assets that are not used in the graph
+				InDocument->GetRootGraph()->IterateNodes([](FNodeHandle NodeHandle)
+				{
+					TArray<FInputHandle> NodeInputs = NodeHandle->GetInputs();
+					for (FInputHandle NodeInput : NodeInputs)
+					{
+						NodeInput->ClearConnectedObjectLiterals();
+					}
+				});
+			}
+		};
+
 		FVersionDocument::FVersionDocument(FName InName, const FString& InPath)
 			: Name(InName)
 			, Path(InPath)
@@ -515,6 +541,7 @@ namespace Metasound
 			bWasUpdated |= FVersionDocument_1_8(Name, Path).Transform(InDocument);
 			bWasUpdated |= FVersionDocument_1_9(Name, Path).Transform(InDocument);
 			bWasUpdated |= FVersionDocument_1_10().Transform(InDocument);
+			bWasUpdated |= FVersionDocument_1_11().Transform(InDocument);
 
 			if (bWasUpdated)
 			{
diff --git a/Engine/Plugins/Runtime/Metasound/Source/MetasoundFrontend/Private/MetasoundFrontendDocumentVersioning.h b/Engine/Plugins/Runtime/Metasound/Source/MetasoundFrontend/Private/MetasoundFrontendDocumentVersioning.h
--- a/Engine/Plugins/Runtime/Metasound/Source/MetasoundFrontend/Private/MetasoundFrontendDocumentVersioning.h
+++ b/Engine/Plugins/Runtime/Metasound/Source/MetasoundFrontend/Private/MetasoundFrontendDocumentVersioning.h
@@ -21,7 +21,7 @@ namespace Metasound
 		public:
 			static FMetasoundFrontendVersionNumber GetMaxVersion()
 			{
-				return FMetasoundFrontendVersionNumber { 1, 10 };
+				return FMetasoundFrontendVersionNumber { 1, 11 };
 			}
 
 			FVersionDocument(FName InName, const FString& InPath);
diff --git a/Engine/Plugins/Runtime/Metasound/Source/MetasoundFrontend/Private/MetasoundFrontendInputController.cpp b/Engine/Plugins/Runtime/Metasound/Source/MetasoundFrontend/Private/MetasoundFrontendInputController.cpp
--- a/Engine/Plugins/Runtime/Metasound/Source/MetasoundFrontend/Private/MetasoundFrontendInputController.cpp
+++ b/Engine/Plugins/Runtime/Metasound/Source/MetasoundFrontend/Private/MetasoundFrontendInputController.cpp
@@ -294,6 +294,8 @@ namespace Metasound
 						Edge->FromNodeID = InController.GetOwningNodeID();
 						Edge->FromVertexID = InController.GetID();
 
+						ClearConnectedObjectLiterals();
+
 						return true;
 					}
 					else
@@ -389,6 +391,21 @@ namespace Metasound
 
 			return false;
 		}
+		
+		void FBaseInputController::ClearConnectedObjectLiterals()
+		{
+			FDataTypeRegistryInfo DataTypeInfo;
+			if (IsConnected())
+			{
+				if (IDataTypeRegistry::Get().GetDataTypeInfo(GetDataType(), DataTypeInfo))
+				{
+					if (DataTypeInfo.IsDataTypeProxyParsable())
+					{
+						ClearLiteral();
+					}
+				}
+			}
+		}
 
 		const FMetasoundFrontendEdge* FBaseInputController::FindEdge() const
 		{
diff --git a/Engine/Plugins/Runtime/Metasound/Source/MetasoundFrontend/Private/MetasoundFrontendInputController.h b/Engine/Plugins/Runtime/Metasound/Source/MetasoundFrontend/Private/MetasoundFrontendInputController.h
--- a/Engine/Plugins/Runtime/Metasound/Source/MetasoundFrontend/Private/MetasoundFrontendInputController.h
+++ b/Engine/Plugins/Runtime/Metasound/Source/MetasoundFrontend/Private/MetasoundFrontendInputController.h
@@ -78,6 +78,8 @@ namespace Metasound
 			virtual bool Disconnect(IOutputController& InController) override;
 			virtual bool Disconnect() override;
 
+			virtual void ClearConnectedObjectLiterals() override;
+
 		protected:
 
 
diff --git a/Engine/Plugins/Runtime/Metasound/Source/MetasoundFrontend/Private/MetasoundFrontendInvalidController.h b/Engine/Plugins/Runtime/Metasound/Source/MetasoundFrontend/Private/MetasoundFrontendInvalidController.h
--- a/Engine/Plugins/Runtime/Metasound/Source/MetasoundFrontend/Private/MetasoundFrontendInvalidController.h
+++ b/Engine/Plugins/Runtime/Metasound/Source/MetasoundFrontend/Private/MetasoundFrontendInvalidController.h
@@ -131,6 +131,8 @@ namespace Metasound
 			virtual bool Connect(IOutputController& InController) override { return false; }
 			virtual bool ConnectWithConverterNode(IOutputController& InController, const FConverterNodeInfo& InNodeClassName) override { return false; }
 			virtual bool Disconnect(IOutputController& InController) override { return false; }
+
+			virtual void ClearConnectedObjectLiterals() override { }
 		protected:
 			virtual FDocumentAccess ShareAccess() override { return FDocumentAccess(); }
 			virtual FConstDocumentAccess ShareAccess() const override { return FConstDocumentAccess(); }
diff --git a/Engine/Plugins/Runtime/Metasound/Source/MetasoundFrontend/Public/MetasoundFrontendController.h b/Engine/Plugins/Runtime/Metasound/Source/MetasoundFrontend/Public/MetasoundFrontendController.h
--- a/Engine/Plugins/Runtime/Metasound/Source/MetasoundFrontend/Public/MetasoundFrontendController.h
+++ b/Engine/Plugins/Runtime/Metasound/Source/MetasoundFrontend/Public/MetasoundFrontendController.h
@@ -407,6 +407,9 @@ namespace Metasound
 			 */
 			virtual bool Disconnect() = 0;
 
+			/** Clear object literals to prevent referencing assets
+			 * on connected inputs that are not used in the graph. */
+			virtual void ClearConnectedObjectLiterals() = 0;
 		};
 
 		/* An INodeController provides methods for querying and manipulating a Metasound node. */
diff --git a/Engine/Plugins/Runtime/Metasound/Source/MetasoundFrontend/Public/MetasoundFrontendDataTypeRegistry.h b/Engine/Plugins/Runtime/Metasound/Source/MetasoundFrontend/Public/MetasoundFrontendDataTypeRegistry.h
--- a/Engine/Plugins/Runtime/Metasound/Source/MetasoundFrontend/Public/MetasoundFrontendDataTypeRegistry.h
+++ b/Engine/Plugins/Runtime/Metasound/Source/MetasoundFrontend/Public/MetasoundFrontendDataTypeRegistry.h
@@ -23,7 +23,6 @@ namespace Metasound
 		// Function signature for assigning a FLiteral to a FAnyDataReference 
 		using FLiteralAssignmentFunction = void(*)(const FOperatorSettings& InOperatorSettings, const FLiteral& InLiteral, const FAnyDataReference& OutDataRef);
 
-
 		/** FDataTypeRegsitryInfo contains runtime inspectable behavior of a registered
 		 * MetaSound data type.
 		 */
@@ -74,6 +73,12 @@ namespace Metasound
 
 			// If provided in registration call, UClass this datatype was registered with.
 			UClass* ProxyGeneratorClass = nullptr;
+
+			// Whether the data type can be created from a UObject proxy. 
+			const bool IsDataTypeProxyParsable() const
+			{
+				return bIsProxyParsable || bIsProxyArrayParsable;
+			}
 		};
 
 		/** Interface for metadata of a registered MetaSound enum type. */
