[Metasound] Support constructor pins in test graph builder

diff --git a/Engine/Plugins/Runtime/Metasound/Source/MetasoundEngineTest/Private/NodeTestGraphBuilder.cpp b/Engine/Plugins/Runtime/Metasound/Source/MetasoundEngineTest/Private/NodeTestGraphBuilder.cpp
--- a/Engine/Plugins/Runtime/Metasound/Source/MetasoundEngineTest/Private/NodeTestGraphBuilder.cpp
+++ b/Engine/Plugins/Runtime/Metasound/Source/MetasoundEngineTest/Private/NodeTestGraphBuilder.cpp
@@ -34,14 +34,19 @@ namespace Metasound::Test
 		return Node;
 	}
 
-	FNodeHandle FNodeTestGraphBuilder::AddInput(const FName& InputName, const FName& TypeName) const
+	FNodeHandle FNodeTestGraphBuilder::AddInput(
+		const FName& InputName,
+		const FName& TypeName,
+		EMetasoundFrontendVertexAccessType AccessType) const
 	{
 		check(RootGraph->IsValid());
 
 		FMetasoundFrontendClassInput Input;
 		Input.Name = InputName;
 		Input.TypeName = TypeName;
+		Input.AccessType = AccessType;
 		Input.VertexID = FGuid::NewGuid();
+		
 		return RootGraph->AddInputVertex(Input);
 	}
 
@@ -122,7 +127,7 @@ namespace Metasound::Test
 		// add the inputs and connect them
 		for (FInputHandle Input : NodeHandle->GetInputs())
 		{
-			FNodeHandle InputNode = Builder.AddInput(Input->GetName(), Input->GetDataType());
+			FNodeHandle InputNode = Builder.AddInput(Input->GetName(), Input->GetDataType(), Input->GetVertexAccessType());
 
 			if (!InputNode->IsValid())
 			{
diff --git a/Engine/Plugins/Runtime/Metasound/Source/MetasoundEngineTest/Public/NodeTestGraphBuilder.h b/Engine/Plugins/Runtime/Metasound/Source/MetasoundEngineTest/Public/NodeTestGraphBuilder.h
--- a/Engine/Plugins/Runtime/Metasound/Source/MetasoundEngineTest/Public/NodeTestGraphBuilder.h
+++ b/Engine/Plugins/Runtime/Metasound/Source/MetasoundEngineTest/Public/NodeTestGraphBuilder.h
@@ -19,7 +19,25 @@ namespace Metasound::Test
 		Frontend::FNodeHandle AddNode(const FNodeClassName& ClassName, int32 MajorVersion) const;
 
 		/** Add an input node to the graph */
-		Frontend::FNodeHandle AddInput(const FName& InputName, const FName& TypeName) const;
+		Frontend::FNodeHandle AddInput(
+			const FName& InputName,
+			const FName& TypeName,
+			EMetasoundFrontendVertexAccessType AccessType = EMetasoundFrontendVertexAccessType::Reference) const;
+
+		/** Add a constructor input to the graph */
+		template<typename DataType>
+		Frontend::FNodeHandle AddConstructorInput(const FName& InputName, const DataType& Value) const
+		{
+			check(RootGraph->IsValid());
+
+			FMetasoundFrontendClassInput Input;
+			Input.Name = InputName;
+			Input.TypeName = GetMetasoundDataTypeName<DataType>();
+			Input.VertexID = FGuid::NewGuid();
+			Input.AccessType = EMetasoundFrontendVertexAccessType::Value;
+			Input.DefaultLiteral.Set(Value);
+			return RootGraph->AddInputVertex(Input);
+		}
 
 		/** Add an output node to the graph */
 		Frontend::FNodeHandle AddOutput(const FName& OutputName, const FName& TypeName);
