[Metasound] Get history of output changes instead of last value

diff --git a/Engine/Plugins/Runtime/Metasound/Source/MetasoundEngine/Private/MetasoundEngineModule.cpp b/Engine/Plugins/Runtime/Metasound/Source/MetasoundEngine/Private/MetasoundEngineModule.cpp
--- a/Engine/Plugins/Runtime/Metasound/Source/MetasoundEngine/Private/MetasoundEngineModule.cpp
+++ b/Engine/Plugins/Runtime/Metasound/Source/MetasoundEngine/Private/MetasoundEngineModule.cpp
@@ -1,24 +1,25 @@
 // Copyright Epic Games, Inc. All Rights Reserved.
 #include "MetasoundEngineModule.h"
 
-#include "Analysis/MetasoundFrontendAnalyzerRegistry.h"
-#include "Analysis/MetasoundFrontendVertexAnalyzerForwardValue.h"
-#include "Interfaces/MetasoundDeprecatedInterfaces.h"
-#include "Interfaces/MetasoundInterface.h"
-#include "Interfaces/MetasoundInterfaceBindingsPrivate.h"
 #include "Metasound.h"
 #include "MetasoundAudioBus.h"
 #include "MetasoundDataReference.h"
 #include "MetasoundDataTypeRegistrationMacro.h"
-#include "MetasoundFrontendDataTypeTraits.h"
 #include "MetasoundFrontendDocumentBuilder.h"
 #include "MetasoundOutputSubsystem.h"
-#include "MetasoundOutputWatcher.h"
 #include "MetasoundSource.h"
 #include "MetasoundTrace.h"
 #include "MetasoundUObjectRegistry.h"
 #include "MetasoundWave.h"
 #include "MetasoundWaveTable.h"
+#include "Analysis/MetasoundFrontendAnalyzerRegistry.h"
+#include "Analysis/MetasoundFrontendVertexAnalyzerEnvelopeFollower.h"
+#include "Analysis/MetasoundFrontendVertexAnalyzerForwardValue.h"
+#include "Analysis/MetasoundFrontendVertexAnalyzerTriggerDensity.h"
+#include "Analysis/MetasoundFrontendVertexAnalyzerTriggerToTime.h"
+#include "Interfaces/MetasoundDeprecatedInterfaces.h"
+#include "Interfaces/MetasoundInterface.h"
+#include "Interfaces/MetasoundInterfaceBindingsPrivate.h"
 #include "Modules/ModuleManager.h"
 #include "Sound/AudioSettings.h"
 
@@ -64,27 +65,35 @@ class FMetasoundEngineModule : public IMetasoundEngineModule
 		FMetasoundFrontendRegistryContainer::Get()->RegisterPendingNodes();
 
 		// Register Analyzers
-		Metasound::Frontend::IVertexAnalyzerRegistry::Get().RegisterAnalyzerFactories();
-
-		// Register output watcher ops
-		METASOUND_PRIVATE_REGISTER_GENERATOR_OUTPUT_WATCHER_TYPE_OPERATIONS(float, 0.0f)
-		METASOUND_PRIVATE_REGISTER_GENERATOR_OUTPUT_WATCHER_TYPE_OPERATIONS(int32, 0)
-		METASOUND_PRIVATE_REGISTER_GENERATOR_OUTPUT_WATCHER_TYPE_OPERATIONS(bool, false)
-		METASOUND_PRIVATE_REGISTER_GENERATOR_OUTPUT_WATCHER_TYPE_OPERATIONS(FString, "")
+		METASOUND_REGISTER_VERTEX_ANALYZER_FACTORY(Frontend::FVertexAnalyzerEnvelopeFollower)
+		METASOUND_REGISTER_VERTEX_ANALYZER_FACTORY(Frontend::FVertexAnalyzerTriggerDensity)
+		METASOUND_REGISTER_VERTEX_ANALYZER_FACTORY(Frontend::FVertexAnalyzerForwardBool)
+		METASOUND_REGISTER_VERTEX_ANALYZER_FACTORY(Frontend::FVertexAnalyzerForwardFloat)
+		METASOUND_REGISTER_VERTEX_ANALYZER_FACTORY(Frontend::FVertexAnalyzerForwardInt)
+		METASOUND_REGISTER_VERTEX_ANALYZER_FACTORY(Frontend::FVertexAnalyzerForwardString)
+		METASOUND_REGISTER_VERTEX_ANALYZER_FACTORY(Frontend::FVertexAnalyzerTriggerToTime)
 
 		// Register passthrough output analyzers
 		UMetasoundGeneratorHandle::RegisterPassthroughAnalyzerForType(
 			GetMetasoundDataTypeName<float>(),
-			Frontend::FVertexAnalyzerForwardFloat::GetAnalyzerName());
+			Frontend::FVertexAnalyzerForwardFloat::GetAnalyzerName(),
+			Frontend::FVertexAnalyzerForwardFloat::FOutputs::GetValue().Name);
 		UMetasoundGeneratorHandle::RegisterPassthroughAnalyzerForType(
 			GetMetasoundDataTypeName<int32>(),
-			Frontend::FVertexAnalyzerForwardInt::GetAnalyzerName());
+			Frontend::FVertexAnalyzerForwardInt::GetAnalyzerName(),
+			Frontend::FVertexAnalyzerForwardInt::FOutputs::GetValue().Name);
 		UMetasoundGeneratorHandle::RegisterPassthroughAnalyzerForType(
 			GetMetasoundDataTypeName<bool>(),
-			Frontend::FVertexAnalyzerForwardBool::GetAnalyzerName());
+			Frontend::FVertexAnalyzerForwardBool::GetAnalyzerName(),
+			Frontend::FVertexAnalyzerForwardBool::FOutputs::GetValue().Name);
 		UMetasoundGeneratorHandle::RegisterPassthroughAnalyzerForType(
 			GetMetasoundDataTypeName<FString>(),
-			Frontend::FVertexAnalyzerForwardString::GetAnalyzerName());
+			Frontend::FVertexAnalyzerForwardString::GetAnalyzerName(),
+			Frontend::FVertexAnalyzerForwardString::FOutputs::GetValue().Name);
+		UMetasoundGeneratorHandle::RegisterPassthroughAnalyzerForType(
+			GetMetasoundDataTypeName<FTrigger>(),
+			Frontend::FVertexAnalyzerTriggerToTime::GetAnalyzerName(),
+			Frontend::FVertexAnalyzerTriggerToTime::FOutputs::GetValue().Name);
 
 		UE_LOG(LogMetasoundEngine, Log, TEXT("MetaSound Engine Initialized"));
 	}
diff --git a/Engine/Plugins/Runtime/Metasound/Source/MetasoundEngine/Private/MetasoundGeneratorHandle.cpp b/Engine/Plugins/Runtime/Metasound/Source/MetasoundEngine/Private/MetasoundGeneratorHandle.cpp
--- a/Engine/Plugins/Runtime/Metasound/Source/MetasoundEngine/Private/MetasoundGeneratorHandle.cpp
+++ b/Engine/Plugins/Runtime/Metasound/Source/MetasoundEngine/Private/MetasoundGeneratorHandle.cpp
@@ -3,13 +3,12 @@
 #include "MetasoundGeneratorHandle.h"
 
 #include "Components/AudioComponent.h"
-#include "AudioDeviceManager.h"
 #include "MetasoundGenerator.h"
 #include "MetasoundSource.h"
 #include "MetasoundTrace.h"
 #include "MetasoundParameterPack.h"
 
-TMap<FName, FName> UMetasoundGeneratorHandle::PassthroughAnalyzers{};
+TMap<FName, UMetasoundGeneratorHandle::FPassthroughAnalyzerInfo> UMetasoundGeneratorHandle::PassthroughAnalyzers{};
 
 UMetasoundGeneratorHandle* UMetasoundGeneratorHandle::CreateMetaSoundGeneratorHandle(UAudioComponent* OnComponent)
 {
@@ -38,6 +37,11 @@ void UMetasoundGeneratorHandle::BeginDestroy()
 {
 	Super::BeginDestroy();
 	DetachGeneratorDelegates();
+
+	if (TSharedPtr<Metasound::FMetasoundGenerator> PinnedGenerator = CachedGeneratorPtr.Pin())
+	{
+		PinnedGenerator->OnOutputChanged.RemoveAll(this);
+	}
 }
 
 bool UMetasoundGeneratorHandle::IsValid() const
@@ -282,6 +286,11 @@ bool UMetasoundGeneratorHandle::RemoveGraphSetCallback(const FDelegateHandle& Ha
 	return OnGeneratorsGraphChanged.Remove(Handle);
 }
 
+bool AnalyzerAddressesReferToSameGeneratorOutput(const Metasound::Frontend::FAnalyzerAddress& Lhs, const Metasound::Frontend::FAnalyzerAddress& Rhs)
+{
+	return Lhs.OutputName == Rhs.OutputName && Lhs.AnalyzerName == Rhs.AnalyzerName && Lhs.AnalyzerMemberName == Rhs.AnalyzerMemberName;
+}
+
 bool UMetasoundGeneratorHandle::WatchOutput(
 	FName OutputName,
 	const FOnMetasoundOutputValueChanged& OnOutputValueChanged,
@@ -335,35 +344,31 @@ bool UMetasoundGeneratorHandle::WatchOutput(
 			return false;
 		}
 
-		AnalyzerName = PassthroughAnalyzers[TypeName];
-		AnalyzerOutputName = "Value";
+		AnalyzerName = PassthroughAnalyzers[TypeName].AnalyzerName;
+		AnalyzerOutputName = PassthroughAnalyzers[TypeName].OutputName;
 	}
 	AnalyzerAddress.AnalyzerName = AnalyzerName;
+	AnalyzerAddress.AnalyzerMemberName = AnalyzerOutputName;
 	
 	AnalyzerAddress.AnalyzerInstanceID = FGuid::NewGuid();
 	AnalyzerAddress.NodeID = NodeId;
 
-	// if we already have a generator, go ahead and make the analyzer and watcher
+	// if we already have a generator, go ahead and make the analyzer
 	if (const TSharedPtr<Metasound::FMetasoundGenerator> PinnedGenerator = PinGenerator())
 	{
-		CreateAnalyzerAndWatcher(PinnedGenerator, MoveTemp(AnalyzerAddress));
+		CreateAnalyzer(AnalyzerAddress, PinnedGenerator);
 	}
-	// otherwise enqueue it for later
-	else
-	{
-		OutputAnalyzersToAdd.Enqueue(MoveTemp(AnalyzerAddress));
-	}
-	
-	// either way, add the delegate to the map
-	OutputListenerMap.FindOrAdd(OutputName).FindOrAdd(AnalyzerOutputName).AddUnique(OnOutputValueChanged);
+
+	// Create the listener
+	CreateListener(AnalyzerAddress, OnOutputValueChanged);
 
 	return true;
 }
 
-void UMetasoundGeneratorHandle::RegisterPassthroughAnalyzerForType(FName TypeName, FName AnalyzerName)
+void UMetasoundGeneratorHandle::RegisterPassthroughAnalyzerForType(FName TypeName, FName AnalyzerName, FName OutputName)
 {
 	check(!PassthroughAnalyzers.Contains(TypeName));
-	PassthroughAnalyzers.Add(TypeName, AnalyzerName);
+	PassthroughAnalyzers.Add(TypeName, { AnalyzerName, OutputName });
 }
 
 void UMetasoundGeneratorHandle::UpdateWatchers()
@@ -371,18 +376,18 @@ void UMetasoundGeneratorHandle::UpdateWatchers()
 	METASOUND_LLM_SCOPE;
 	METASOUND_TRACE_CPUPROFILER_EVENT_SCOPE(UMetasoundGeneratorHandle::UpdateWatchers);
 
-	for (auto& Watcher : OutputWatchers)
+	while (TOptional<FOutputPayload> ChangedOutput = ChangedOutputs.Dequeue())
 	{
-		Watcher.Update([this](FName OutputName, const FMetaSoundOutput& Output)
-		{
-			if (TMap<FName, FOnOutputValueChangedMulticast>* Map = OutputListenerMap.Find(OutputName))
+		if (const FOutputListener* OutputListener = OutputListeners.FindByPredicate(
+			[&ChangedOutput](const FOutputListener& ExistingListener)
 			{
-				if (const FOnOutputValueChangedMulticast* Delegate = Map->Find(Output.Name))
-				{
-					Delegate->Broadcast(OutputName, Output);
-				}
-			}
-		});
+				return (*ChangedOutput).AnalyzerName == ExistingListener.AnalyzerAddress.AnalyzerName
+				&& (*ChangedOutput).OutputName == ExistingListener.AnalyzerAddress.OutputName
+				&& (*ChangedOutput).OutputValue.Name == ExistingListener.AnalyzerAddress.AnalyzerMemberName;
+			}))
+		{
+			OutputListener->OnOutputValueChanged.Broadcast((*ChangedOutput).OutputName, (*ChangedOutput).OutputValue);
+		}
 	}
 }
 
@@ -406,12 +411,11 @@ void UMetasoundGeneratorHandle::OnSourceCreatedAGenerator(uint64 InAudioComponen
 				InGenerator->QueueParameterPack(CachedParameterPack);
 			}
 
-			// If there are analyzers to create, create them
+			// If we have listeners, add the analyzers for them
 			{
-				Metasound::Frontend::FAnalyzerAddress Address;
-				while (OutputAnalyzersToAdd.Dequeue(Address))
+				for (const FOutputListener& Listener : OutputListeners)
 				{
-					CreateAnalyzerAndWatcher(InGenerator, MoveTemp(Address));
+					CreateAnalyzer(Listener.AnalyzerAddress, InGenerator);
 				}
 			}
 			
@@ -439,19 +443,12 @@ void UMetasoundGeneratorHandle::OnSourceDestroyedAGenerator(uint64 InAudioCompon
 			OnGeneratorHandleDetached.Broadcast();
 		}
 		CachedGeneratorPtr = nullptr;
-
-		// Clear out the watchers, but keep track of what we had in case the sound comes back
-		for (const Metasound::Private::FMetasoundOutputWatcher& Watcher : OutputWatchers)
-		{
-			OutputAnalyzersToAdd.Enqueue(Watcher.GetAddress());
-		}
-		OutputWatchers.Empty();
 	}
 }
 
-void UMetasoundGeneratorHandle::CreateAnalyzerAndWatcher(
-	const TSharedPtr<Metasound::FMetasoundGenerator> Generator,
-	Metasound::Frontend::FAnalyzerAddress&& AnalyzerAddress)
+void UMetasoundGeneratorHandle::CreateAnalyzer(
+	const Metasound::Frontend::FAnalyzerAddress& AnalyzerAddress,
+	const TSharedPtr<Metasound::FMetasoundGenerator> Generator)
 {
 	if (!IsValid())
 	{
@@ -460,16 +457,46 @@ void UMetasoundGeneratorHandle::CreateAnalyzerAndWatcher(
 
 	// Create the analyzer (will skip if there's already one)
 	Generator->AddOutputVertexAnalyzer(AnalyzerAddress);
+
+	// Subscribe for output updates if we haven't already
+	if (!Generator->OnOutputChanged.IsBoundToObject(this))
+	{
+		Generator->OnOutputChanged.AddUObject(this, &UMetasoundGeneratorHandle::HandleOutputChanged);
+	}
+}
+
+void UMetasoundGeneratorHandle::CreateListener(
+	const Metasound::Frontend::FAnalyzerAddress& AnalyzerAddress,
+	const FOnMetasoundOutputValueChanged& OnOutputValueChanged)
+{
+	if (!IsValid())
+	{
+		return;
+	}
 	
-	// Create the watcher
-	if (!OutputWatchers.ContainsByPredicate(
-		[&AnalyzerAddress](const Metasound::Private::FMetasoundOutputWatcher& Watcher)
+	// If we already have a listener for this output, just add the delegate to that one
+	if (FOutputListener* Listener = OutputListeners.FindByPredicate(
+		[&AnalyzerAddress](const FOutputListener& ExistingListener)
 		{
-			return AnalyzerAddress.OutputName == Watcher.Name;
+			return AnalyzerAddressesReferToSameGeneratorOutput(AnalyzerAddress, ExistingListener.AnalyzerAddress);
 		}))
 	{
-		OutputWatchers.Emplace(MoveTemp(AnalyzerAddress), Generator->OperatorSettings);
+		Listener->OnOutputValueChanged.AddUnique(OnOutputValueChanged);
+	}
+	// Otherwise add a new listener
+	else
+	{
+		OutputListeners.Emplace(AnalyzerAddress, OnOutputValueChanged);
 	}
 }
 
-
+// NOTE: We're in an audio render thread here, so get out fast
+// and be mindful of touching things that get touched on the game thread
+void UMetasoundGeneratorHandle::HandleOutputChanged(
+	FName AnalyzerName,
+	FName OutputName,
+	FName AnalyzerOutputName,
+	TSharedPtr<Metasound::IOutputStorage> OutputData)
+{
+	ChangedOutputs.Enqueue(AnalyzerName, OutputName, AnalyzerOutputName, OutputData);
+}
diff --git a/Engine/Plugins/Runtime/Metasound/Source/MetasoundEngine/Private/MetasoundOutput.cpp b/Engine/Plugins/Runtime/Metasound/Source/MetasoundEngine/Private/MetasoundOutput.cpp
--- a/Engine/Plugins/Runtime/Metasound/Source/MetasoundEngine/Private/MetasoundOutput.cpp
+++ b/Engine/Plugins/Runtime/Metasound/Source/MetasoundEngine/Private/MetasoundOutput.cpp
@@ -3,11 +3,16 @@
 #include "MetasoundOutput.h"
 
 #include "MetasoundPrimitives.h"
+#include "MetasoundTime.h"
 
 FMetaSoundOutput::FMetaSoundOutput(const FMetaSoundOutput& Other)
 {
 	Name = Other.Name;
-	DataReference = MakeUnique<Metasound::FAnyDataReference>(*Other.DataReference);
+
+	if (Other.IsValid())
+	{
+		Data = Other.Data->Clone();
+	}
 }
 
 FMetaSoundOutput& FMetaSoundOutput::operator=(const FMetaSoundOutput& Other)
@@ -15,20 +20,50 @@ FMetaSoundOutput& FMetaSoundOutput::operator=(const FMetaSoundOutput& Other)
 	if (this != &Other)
 	{
 		this->Name = Other.Name;
-		this->DataReference = MakeUnique<Metasound::FAnyDataReference>(*Other.DataReference);
+		
+		if (Other.IsValid())
+		{
+			Data = Other.Data->Clone();
+		}
 	}
 
 	return *this;
 }
 
+FMetaSoundOutput::FMetaSoundOutput(FMetaSoundOutput&& Other) noexcept
+{
+	Name = MoveTemp(Other.Name);
+	Data = MoveTemp(Other.Data);
+}
+
+FMetaSoundOutput& FMetaSoundOutput::operator=(FMetaSoundOutput&& Other) noexcept
+{
+	if (this != &Other)
+	{
+		Name = MoveTemp(Other.Name);
+		Data = MoveTemp(Other.Data);
+	}
+
+	return *this;
+}
+
+FMetaSoundOutput::FMetaSoundOutput(FName InName, const TSharedPtr<Metasound::IOutputStorage>& InData)
+{
+	Name = MoveTemp(InName);
+	if (InData.IsValid())
+	{
+		Data = InData->Clone();
+	}
+}
+
 bool FMetaSoundOutput::IsValid() const
 {
-	return nullptr != DataReference;
+	return nullptr != Data;
 }
 
-FName FMetaSoundOutput::GetTypeName() const
+FName FMetaSoundOutput::GetDataTypeName() const
 {
-	return IsValid() ? DataReference->GetDataTypeName() : FName();
+	return IsValid() ? Data->GetDataTypeName() : FName();
 }
 
 bool UMetasoundOutputBlueprintAccess::IsFloat(const FMetaSoundOutput& Output)
@@ -78,3 +113,15 @@ FString UMetasoundOutputBlueprintAccess::GetString(const FMetaSoundOutput& Outpu
 	Success = Output.Get(Value);
 	return Value;
 }
+
+bool UMetasoundOutputBlueprintAccess::IsTime(const FMetaSoundOutput& Output)
+{
+	return Output.IsType<Metasound::FTime>();
+}
+
+float UMetasoundOutputBlueprintAccess::GetTimeSeconds(const FMetaSoundOutput& Output, bool& Success)
+{
+	Metasound::FTime TimeValue;
+	Success = Output.Get(TimeValue);
+	return TimeValue.GetSeconds();
+}
diff --git a/Engine/Plugins/Runtime/Metasound/Source/MetasoundEngine/Private/MetasoundOutputWatcher.cpp b/Engine/Plugins/Runtime/Metasound/Source/MetasoundEngine/Private/MetasoundOutputWatcher.cpp
deleted file mode 100644
diff --git a/Engine/Plugins/Runtime/Metasound/Source/MetasoundEngine/Public/MetasoundGeneratorHandle.h b/Engine/Plugins/Runtime/Metasound/Source/MetasoundEngine/Public/MetasoundGeneratorHandle.h
--- a/Engine/Plugins/Runtime/Metasound/Source/MetasoundEngine/Public/MetasoundGeneratorHandle.h
+++ b/Engine/Plugins/Runtime/Metasound/Source/MetasoundEngine/Public/MetasoundGeneratorHandle.h
@@ -2,8 +2,9 @@
 
 #pragma once
 
-#include "MetasoundOutputWatcher.h"
+#include "MetasoundOutput.h"
 #include "MetasoundParameterPack.h"
+#include "Analysis/MetasoundFrontendAnalyzerAddress.h"
 #include "Containers/SpscQueue.h"
 #include "Templates/SharedPointer.h"
 
@@ -92,8 +93,9 @@ public:
 	 *
 	 * @param TypeName - The type name returned from GetMetasoundDataTypeName()
 	 * @param AnalyzerName - The name of the analyzer to use
+	 * @param OutputName - The name of the output in the analyzer
 	 */
-	static void RegisterPassthroughAnalyzerForType(FName TypeName, FName AnalyzerName);
+	static void RegisterPassthroughAnalyzerForType(FName TypeName, FName AnalyzerName, FName OutputName);
 
 	/**
 	 * Update any watched outputs
@@ -142,13 +144,65 @@ private:
 	FDelegateHandle GeneratorDestroyedDelegateHandle;
 	FDelegateHandle GeneratorGraphChangedDelegateHandle;
 
+	void CreateAnalyzer(
+		const Metasound::Frontend::FAnalyzerAddress& AnalyzerAddress,
+		const TSharedPtr<Metasound::FMetasoundGenerator> Generator);
+	void CreateListener(
+		const Metasound::Frontend::FAnalyzerAddress& AnalyzerAddress,
+		const FOnMetasoundOutputValueChanged& OnOutputValueChanged);
+	void HandleOutputChanged(
+		FName AnalyzerName,
+		FName OutputName,
+		FName AnalyzerOutputName,
+		TSharedPtr<Metasound::IOutputStorage> OutputData);
+
+	struct FPassthroughAnalyzerInfo
+	{
+		FName AnalyzerName;
+		FName OutputName;
+	};
+	static TMap<FName, FPassthroughAnalyzerInfo> PassthroughAnalyzers;
+
+	/**
+	 * Multicast delegate to broadcast to users calling WatchOutput
+	 */
 	DECLARE_DYNAMIC_MULTICAST_DELEGATE_TwoParams(FOnOutputValueChangedMulticast, FName, Name, const FMetaSoundOutput&, Output);
-	static TMap<FName, FName> PassthroughAnalyzers;
-	// This is a map of output names on the generator to a map of output names on an analyzer
-	// to delegates listening to that analyzer output. We do this because analyzers can have more
-	// than one output.
-	TMap<FName, TMap<FName, FOnOutputValueChangedMulticast>> OutputListenerMap;
-	TSpscQueue<Metasound::Frontend::FAnalyzerAddress> OutputAnalyzersToAdd;
-	void CreateAnalyzerAndWatcher(const TSharedPtr<Metasound::FMetasoundGenerator> Generator, Metasound::Frontend::FAnalyzerAddress&& AnalyzerAddress);
-	TArray<Metasound::Private::FMetasoundOutputWatcher> OutputWatchers;
+
+	/**
+	 * Info about an output being watched by one or more listeners
+	 */
+	struct FOutputListener
+	{
+		Metasound::Frontend::FAnalyzerAddress AnalyzerAddress;
+		FOnOutputValueChangedMulticast OnOutputValueChanged;
+
+		FOutputListener(
+			const Metasound::Frontend::FAnalyzerAddress& InAnalyzerAddress,
+			const FOnMetasoundOutputValueChanged& InOnOutputValueChanged)
+				: AnalyzerAddress(InAnalyzerAddress)
+		{
+			OnOutputValueChanged.AddUnique(InOnOutputValueChanged);
+		}
+	};
+
+	TArray<FOutputListener> OutputListeners;
+
+	struct FOutputPayload
+	{
+		FName AnalyzerName;
+		FName OutputName;
+		FMetaSoundOutput OutputValue;
+
+		FOutputPayload(
+			const FName InAnalyzerName,
+			const FName InOutputName,
+			const FName AnalyzerOutputName,
+			TSharedPtr<Metasound::IOutputStorage> OutputData)
+				: AnalyzerName(InAnalyzerName)
+				, OutputName(InOutputName)
+				, OutputValue(AnalyzerOutputName, OutputData)
+		{}
+	};
+
+	TSpscQueue<FOutputPayload> ChangedOutputs;
 };
\ No newline at end of file
diff --git a/Engine/Plugins/Runtime/Metasound/Source/MetasoundEngine/Public/MetasoundOutput.h b/Engine/Plugins/Runtime/Metasound/Source/MetasoundEngine/Public/MetasoundOutput.h
--- a/Engine/Plugins/Runtime/Metasound/Source/MetasoundEngine/Public/MetasoundOutput.h
+++ b/Engine/Plugins/Runtime/Metasound/Source/MetasoundEngine/Public/MetasoundOutput.h
@@ -2,7 +2,7 @@
 
 #pragma once
 
-#include "MetasoundDataReference.h"
+#include "MetasoundOutputStorage.h"
 #include "SoundGeneratorOutput.h"
 #include "Kismet/BlueprintFunctionLibrary.h"
 
@@ -22,6 +22,12 @@ struct METASOUNDENGINE_API FMetaSoundOutput : public FSoundGeneratorOutput
 
 	FMetaSoundOutput& operator=(const FMetaSoundOutput& Other);
 
+	FMetaSoundOutput(FMetaSoundOutput&& Other) noexcept;
+
+	FMetaSoundOutput& operator=(FMetaSoundOutput&& Other) noexcept;
+
+	FMetaSoundOutput(FName InName, const TSharedPtr<Metasound::IOutputStorage>& InData);
+
 	/**
 	 * Has this output been initialized?
 	 *
@@ -34,7 +40,7 @@ struct METASOUNDENGINE_API FMetaSoundOutput : public FSoundGeneratorOutput
 	 *
 	 * @returns The type name, or none if invalid
 	 */
-	FName GetTypeName() const;
+	FName GetDataTypeName() const;
 	
 	/**
 	 * Initialize the output with some data.
@@ -45,8 +51,8 @@ struct METASOUNDENGINE_API FMetaSoundOutput : public FSoundGeneratorOutput
 	template<typename DataType>
 	void Init(const DataType& InitialValue)
 	{
-		DataReference = MakeUnique<Metasound::FAnyDataReference>(Metasound::TDataWriteReference<DataType>::CreateNew());
-		*DataReference->GetDataWriteReference<DataType>() = InitialValue;
+		using FOutputData = Metasound::TOutputStorage<DataType>;
+		Data = MakeUnique<FOutputData>(InitialValue);
 	}
 
 	/**
@@ -58,7 +64,8 @@ struct METASOUNDENGINE_API FMetaSoundOutput : public FSoundGeneratorOutput
 	template<typename DataType>
 	bool IsType() const
 	{
-		return IsValid() && Metasound::IsDataReferenceOfType<DataType>(*DataReference);
+		static const FName TypeName = Metasound::GetMetasoundDataTypeName<DataType>();
+		return IsValid() && TypeName == Data->GetDataTypeName();
 	}
 	
 	/**
@@ -76,7 +83,8 @@ struct METASOUNDENGINE_API FMetaSoundOutput : public FSoundGeneratorOutput
 			return false;
 		}
 
-		Value = *DataReference->GetDataReadReference<DataType>();
+		using FOutputData = Metasound::TOutputStorage<DataType>;
+		Value = static_cast<FOutputData*>(Data.Get())->Get();
 		return true;
 	}
 	
@@ -95,12 +103,33 @@ struct METASOUNDENGINE_API FMetaSoundOutput : public FSoundGeneratorOutput
 			return false;
 		}
 
-		*DataReference->GetDataWriteReference<DataType>() = Value;
+		using FOutputData = Metasound::TOutputStorage<DataType>;
+		static_cast<FOutputData*>(Data.Get())->Set(Value);
 		return true;
 	}
 
+	/**
+	 * Set the value, for moveable registered Metasound data types
+	 *
+	 * @tparam DataType - The expected data type of the output
+	 * @param Value - The value to use
+	 * @returns true if the value was set, false otherwise
+	 */
+	template<typename DataType>
+	bool Set(DataType&& Value)
+	{
+		if (!IsType<DataType>())
+		{
+			return false;
+		}
+
+		using FOutputData = Metasound::TOutputStorage<DataType>;
+		static_cast<FOutputData*>(Data.Get())->Set(MoveTemp(Value));
+		return true;
+	}
+	
 private:
-	TUniquePtr<Metasound::FAnyDataReference> DataReference;
+	TUniquePtr<Metasound::IOutputStorage> Data;
 };
 
 /**
@@ -116,16 +145,24 @@ class UMetasoundOutputBlueprintAccess final : public UBlueprintFunctionLibrary
 	static bool IsFloat(const FMetaSoundOutput& Output);
 	UFUNCTION(BlueprintCallable, Category="MetaSoundOutput")
 	static float GetFloat(const FMetaSoundOutput& Output, bool& Success);
+	
 	UFUNCTION(BlueprintCallable, Category="MetaSoundOutput")
 	static bool IsInt32(const FMetaSoundOutput& Output);
 	UFUNCTION(BlueprintCallable, Category="MetaSoundOutput")
 	static int32 GetInt32(const FMetaSoundOutput& Output, bool& Success);
+
 	UFUNCTION(BlueprintCallable, Category="MetaSoundOutput")
 	static bool IsBool(const FMetaSoundOutput& Output);
 	UFUNCTION(BlueprintCallable, Category="MetaSoundOutput")
 	static bool GetBool(const FMetaSoundOutput& Output, bool& Success);
+
 	UFUNCTION(BlueprintCallable, Category="MetaSoundOutput")
 	static bool IsString(const FMetaSoundOutput& Output);
 	UFUNCTION(BlueprintCallable, Category="MetaSoundOutput")
 	static FString GetString(const FMetaSoundOutput& Output, bool& Success);
+
+	UFUNCTION(BlueprintCallable, Category="MetaSoundOutput")
+	static bool IsTime(const FMetaSoundOutput& Output);
+	UFUNCTION(BlueprintCallable, Category="MetaSoundOutput")
+	static float GetTimeSeconds(const FMetaSoundOutput& Output, bool& Success);
 };
diff --git a/Engine/Plugins/Runtime/Metasound/Source/MetasoundEngine/Public/MetasoundOutputWatcher.h b/Engine/Plugins/Runtime/Metasound/Source/MetasoundEngine/Public/MetasoundOutputWatcher.h
deleted file mode 100644
diff --git a/Engine/Plugins/Runtime/Metasound/Source/MetasoundEngineTest/Private/MetasoundGeneratorTests.cpp b/Engine/Plugins/Runtime/Metasound/Source/MetasoundEngineTest/Private/MetasoundGeneratorTests.cpp
--- a/Engine/Plugins/Runtime/Metasound/Source/MetasoundEngineTest/Private/MetasoundGeneratorTests.cpp
+++ b/Engine/Plugins/Runtime/Metasound/Source/MetasoundEngineTest/Private/MetasoundGeneratorTests.cpp
@@ -2,6 +2,8 @@
 
 #include "NodeTestGraphBuilder.h"
 #include "Analysis/MetasoundFrontendAnalyzerView.h"
+#include "Analysis/MetasoundFrontendVertexAnalyzerTriggerToTime.h"
+
 #include "Misc/AutomationTest.h"
 
 #if WITH_DEV_AUTOMATION_TESTS
@@ -386,6 +388,101 @@ namespace Metasound::Test::Generator
 		const FDataType ExpectedValue { "unexpected value" };
 		return RunOutputAnalyzerForwardValueTest<FDataType>(*this, ExpectedValue, "UE.Forward.String");
 	}
+
+	IMPLEMENT_SIMPLE_AUTOMATION_TEST(
+	FMetasoundGeneratorOutputAnalyzerTriggerToTimeTest,
+	"Audio.Metasound.Generator.OutputAnalyzer.TriggerToTime",
+	EAutomationTestFlags::EditorContext | EAutomationTestFlags::EngineFilter)
+	bool FMetasoundGeneratorOutputAnalyzerTriggerToTimeTest::RunTest(const FString&)
+	{
+		const FName InputName = "MyInput";
+		const FName OutputName = "MyOutput";
+		constexpr FSampleRate SampleRate = 48000;
+		constexpr int32 NumSamplesPerBlock = 128;
+
+		// Build a passthrough graph
+		FGuid OutputNodeId;
+		const TUniquePtr<FMetasoundGenerator> Generator =
+			BuildPassthroughGraph<FTrigger>(*this, InputName, OutputName, SampleRate, NumSamplesPerBlock, &OutputNodeId);
+		UTEST_NOT_NULL("Generator build", Generator.Get());
+
+		// Add an analyzer to the output
+		FAnalyzerAddress AnalyzerAddress;
+		AnalyzerAddress.DataType = GetMetasoundDataTypeName<FTrigger>();
+		AnalyzerAddress.InstanceID = 1234;
+		AnalyzerAddress.OutputName = OutputName;
+		AnalyzerAddress.AnalyzerName = FVertexAnalyzerTriggerToTime::GetAnalyzerName();
+		AnalyzerAddress.AnalyzerInstanceID = FGuid::NewGuid();
+		AnalyzerAddress.AnalyzerMemberName = FVertexAnalyzerTriggerToTime::FOutputs::GetValue().Name;
+		AnalyzerAddress.NodeID = OutputNodeId;
+		Generator->AddOutputVertexAnalyzer(AnalyzerAddress);
+
+		// make the interval such that the triggers don't line up exactly with the beginning of the block after the first block
+		constexpr int32 TriggerIntervalSamples = NumSamplesPerBlock / 3 + NumSamplesPerBlock / 2; 
+		constexpr double TriggerIntervalSeconds = static_cast<double>(TriggerIntervalSamples) / SampleRate;
+		double NextTriggerReceivedSeconds = 0;
+		bool CallbackSuccess = true;
+		int32 NumTriggersReceived = 0;
+		
+		// Subscribe for updates
+		Generator->OnOutputChanged.AddLambda(
+			[this, &NextTriggerReceivedSeconds, &AnalyzerAddress, &CallbackSuccess, &NumTriggersReceived]
+			(const FName AnalyzerName, const FName OutputName, const FName AnalyzerOutputName, TSharedPtr<IOutputStorage> OutputData)
+			{
+				CallbackSuccess = TestEqual("Data types match", OutputData->GetDataTypeName(), GetMetasoundDataTypeName<FTime>())
+				&& TestEqual("Analyzer names match", AnalyzerName, AnalyzerAddress.AnalyzerName)
+				&& TestEqual("Output names match", OutputName, AnalyzerAddress.OutputName)
+				&& TestEqual("Analyzer output names match", AnalyzerOutputName, AnalyzerAddress.AnalyzerMemberName);
+
+				if (!CallbackSuccess)
+				{
+					return;
+				}
+
+				const FTime Time = static_cast<TOutputStorage<FTime>*>(OutputData.Get())->Get();
+				CallbackSuccess = TestEqual("Time is as expected", Time.GetSeconds(), NextTriggerReceivedSeconds);
+				NextTriggerReceivedSeconds += TriggerIntervalSeconds;
+				++NumTriggersReceived;
+			});
+		
+		// Get the input
+		const TOptional<FTriggerWriteRef> TriggerIn = Generator->GetInputWriteReference<FTrigger>(InputName);
+		UTEST_TRUE("Trigger input is valid", TriggerIn.IsSet());
+
+		int32 NextTriggerBlockOffset = 0;
+		constexpr int32 NumBlocksToTest = 20;
+
+		for (int32 i = 0; i < NumBlocksToTest; ++i)
+		{
+			int32 NumTriggersThisBlock = 0;
+			NumTriggersReceived = 0;
+			
+			// Set the triggers for this block
+			{
+				while (NextTriggerBlockOffset < NumSamplesPerBlock)
+				{
+					(*TriggerIn)->TriggerFrame(NextTriggerBlockOffset);
+					NextTriggerBlockOffset += TriggerIntervalSamples;
+					++NumTriggersThisBlock;
+				}
+				NextTriggerBlockOffset -= NumSamplesPerBlock;
+			}
+		
+			// Render a block
+			{
+				TArray<float> Buffer;
+				Buffer.Reserve(NumSamplesPerBlock);
+				UTEST_EQUAL("Generated the right number of samples.",
+					Generator->OnGenerateAudio(Buffer.GetData(), NumSamplesPerBlock),
+					NumSamplesPerBlock);
+			}
+
+			UTEST_TRUE(FString::Printf(TEXT("Callback success on iteration %i"), i), CallbackSuccess);
+			UTEST_EQUAL(FString::Printf(TEXT("Correct number of triggers on iteration %i"), i), NumTriggersReceived, NumTriggersThisBlock);
+		}
+
+		return true;
+	}
 }
 
 #endif
\ No newline at end of file
diff --git a/Engine/Plugins/Runtime/Metasound/Source/MetasoundEngineTest/Private/MetasoundOutputTests.cpp b/Engine/Plugins/Runtime/Metasound/Source/MetasoundEngineTest/Private/MetasoundOutputTests.cpp
--- a/Engine/Plugins/Runtime/Metasound/Source/MetasoundEngineTest/Private/MetasoundOutputTests.cpp
+++ b/Engine/Plugins/Runtime/Metasound/Source/MetasoundEngineTest/Private/MetasoundOutputTests.cpp
@@ -2,6 +2,8 @@
 
 #include "MetasoundOutput.h"
 #include "MetasoundPrimitives.h"
+#include "MetasoundTime.h"
+
 #include "Misc/AutomationTest.h"
 
 #if WITH_DEV_AUTOMATION_TESTS
@@ -21,24 +23,36 @@ namespace Metasound::Test::GeneratorOutput
 		UTEST_FALSE("float is not int32", Output.IsType<int32>());
 		UTEST_FALSE("float is not bool", Output.IsType<bool>());
 		UTEST_FALSE("float is not FString", Output.IsType<FString>());
+		UTEST_FALSE("float is not FTime", Output.IsType<FTime>());
 
 		Output.Init(int32{ 0 });
 		UTEST_FALSE("int32 is not float", Output.IsType<float>());
 		UTEST_TRUE("int32 is int32", Output.IsType<int32>());
 		UTEST_FALSE("int32 is not bool", Output.IsType<bool>());
 		UTEST_FALSE("int32 is not FString", Output.IsType<FString>());
+		UTEST_FALSE("int32 is not FTime", Output.IsType<FTime>());
 		
 		Output.Init(bool{ false });
 		UTEST_FALSE("bool is not float", Output.IsType<float>());
 		UTEST_FALSE("bool is not int32", Output.IsType<int32>());
 		UTEST_TRUE("bool is bool", Output.IsType<bool>());
 		UTEST_FALSE("bool is not FString", Output.IsType<FString>());
+		UTEST_FALSE("bool is not FTime", Output.IsType<FTime>());
 		
 		Output.Init(FString{ "hi!" });
 		UTEST_FALSE("FString is not float", Output.IsType<float>());
 		UTEST_FALSE("FString is not int32", Output.IsType<int32>());
 		UTEST_FALSE("FString is not bool", Output.IsType<bool>());
 		UTEST_TRUE("FString is FString", Output.IsType<FString>());
+		UTEST_FALSE("FString is not FTime", Output.IsType<FTime>());
+
+		constexpr float TimeSeconds = 123.456f;
+		Output.Init(FTime::FromSeconds(TimeSeconds));
+		UTEST_FALSE("FTime is not float", Output.IsType<float>());
+		UTEST_FALSE("FTime is not int32", Output.IsType<int32>());
+		UTEST_FALSE("FTime is not bool", Output.IsType<bool>());
+		UTEST_FALSE("FTime is not FString", Output.IsType<FString>());
+		UTEST_TRUE("FTime is FTime", Output.IsType<FTime>());
 		
 		return true;
 	}
@@ -116,6 +130,18 @@ namespace Metasound::Test::GeneratorOutput
 		const FDataType ExpectedValue{ "goodbye" };
 		return RunOutputGetSetTest<FDataType>(*this, InitialValue, ExpectedValue);
 	}
+
+	IMPLEMENT_SIMPLE_AUTOMATION_TEST(
+	FMetasoundOutputGetSetFTimeTest,
+	"Audio.Metasound.Output.GetSet.FTime",
+	EAutomationTestFlags::EditorContext | EAutomationTestFlags::EngineFilter)
+	bool FMetasoundOutputGetSetFTimeTest::RunTest(const FString&)
+	{
+		using FDataType = FTime;
+		const FDataType InitialValue{ FTime::FromSeconds(123.456f) };
+		const FDataType ExpectedValue{ FTime::FromSeconds(654.321f) };
+		return RunOutputGetSetTest<FDataType>(*this, InitialValue, ExpectedValue);
+	}
 }
 
 #endif
\ No newline at end of file
diff --git a/Engine/Plugins/Runtime/Metasound/Source/MetasoundEngineTest/Private/MetasoundOutputWatcherTests.cpp b/Engine/Plugins/Runtime/Metasound/Source/MetasoundEngineTest/Private/MetasoundOutputWatcherTests.cpp
deleted file mode 100644
diff --git a/Engine/Plugins/Runtime/Metasound/Source/MetasoundFrontend/Private/Analysis/MetasoundFrontendAnalyzerRegistry.cpp b/Engine/Plugins/Runtime/Metasound/Source/MetasoundFrontend/Private/Analysis/MetasoundFrontendAnalyzerRegistry.cpp
--- a/Engine/Plugins/Runtime/Metasound/Source/MetasoundFrontend/Private/Analysis/MetasoundFrontendAnalyzerRegistry.cpp
+++ b/Engine/Plugins/Runtime/Metasound/Source/MetasoundFrontend/Private/Analysis/MetasoundFrontendAnalyzerRegistry.cpp
@@ -16,16 +16,16 @@ namespace Metasound
 		{
 			TMap<FName, TUniquePtr<IVertexAnalyzerFactory>> AnalyzerFactoryRegistry;
 
+		public:
 			template<typename TAnalyzerFactoryClass>
 			void RegisterAnalyzerFactory()
 			{
 				TUniquePtr<IVertexAnalyzerFactory> Factory(new TAnalyzerFactoryClass());
 				AnalyzerFactoryRegistry.Emplace(TAnalyzerFactoryClass::GetAnalyzerName(), MoveTemp(Factory));
 			}
-
-		public:
+			
 			FVertexAnalyzerRegistry() = default;
-			virtual ~FVertexAnalyzerRegistry() = default;
+			virtual ~FVertexAnalyzerRegistry() override = default;
 
 			virtual const IVertexAnalyzerFactory* FindAnalyzerFactory(FName InAnalyzerName) const override
 			{
@@ -39,16 +39,10 @@ namespace Metasound
 				return nullptr;
 			}
 
-			virtual void RegisterAnalyzerFactories() override
+			virtual void RegisterAnalyzerFactory(FName AnalyzerName, TUniquePtr<IVertexAnalyzerFactory>&& Factory) override
 			{
-				RegisterAnalyzerFactory<FVertexAnalyzerEnvelopeFollower::FFactory>();
-				RegisterAnalyzerFactory<FVertexAnalyzerTriggerDensity::FFactory>();
-
-				// Primitives
-				RegisterAnalyzerFactory<FVertexAnalyzerForwardBool::FFactory>();
-				RegisterAnalyzerFactory<FVertexAnalyzerForwardFloat::FFactory>();
-				RegisterAnalyzerFactory<FVertexAnalyzerForwardInt::FFactory>();
-				RegisterAnalyzerFactory<FVertexAnalyzerForwardString::FFactory>();
+				check(!AnalyzerFactoryRegistry.Contains(AnalyzerName));
+				AnalyzerFactoryRegistry.Emplace(AnalyzerName, MoveTemp(Factory));
 			}
 		};
 
diff --git a/Engine/Plugins/Runtime/Metasound/Source/MetasoundFrontend/Private/Analysis/MetasoundFrontendVertexAnalyzerTriggerToTime.cpp b/Engine/Plugins/Runtime/Metasound/Source/MetasoundFrontend/Private/Analysis/MetasoundFrontendVertexAnalyzerTriggerToTime.cpp
new file mode 100644
--- /dev/null
+++ b/Engine/Plugins/Runtime/Metasound/Source/MetasoundFrontend/Private/Analysis/MetasoundFrontendVertexAnalyzerTriggerToTime.cpp
@@ -0,0 +1,58 @@
+// Copyright Epic Games, Inc. All Rights Reserved.
+
+#include "Analysis/MetasoundFrontendVertexAnalyzerTriggerToTime.h"
+
+namespace Metasound
+{
+	namespace Frontend
+	{
+		const FName& FVertexAnalyzerTriggerToTime::GetAnalyzerName()
+		{
+			static const FName AnalyzerName = "UE.Trigger.ToTime";
+			return AnalyzerName;
+		}
+
+		const FName& FVertexAnalyzerTriggerToTime::GetDataType()
+		{
+			return GetMetasoundDataTypeName<FTrigger>();
+		}
+
+		const FAnalyzerOutput& FVertexAnalyzerTriggerToTime::FOutputs::GetValue()
+		{
+			static FAnalyzerOutput Value = { "LastTriggerTime", GetMetasoundDataTypeName<FTime>() };
+			return Value;
+		}
+
+		const TArray<FAnalyzerOutput>& FVertexAnalyzerTriggerToTime::FFactory::GetAnalyzerOutputs() const
+		{
+			static const TArray<FAnalyzerOutput> Outputs { FOutputs::GetValue() };
+			return Outputs;
+		}
+
+		FVertexAnalyzerTriggerToTime::FVertexAnalyzerTriggerToTime(const FCreateAnalyzerParams& InParams)
+		: FVertexAnalyzerBase(InParams.AnalyzerAddress, InParams.VertexDataReference)
+		, SampleRate(InParams.OperatorSettings.GetSampleRate())
+		, FramesPerBlock(InParams.OperatorSettings.GetNumFramesPerBlock())
+		, SecondsPerBlock(static_cast<double>(FramesPerBlock) / SampleRate)
+		, LastTriggerTime(TDataWriteReference<FTime>::CreateNew())
+		{
+			BindOutputData(FOutputs::GetValue().Name, InParams.OperatorSettings, TDataReadReference{ LastTriggerTime });
+		}
+
+		void FVertexAnalyzerTriggerToTime::Execute()
+		{
+			const FTrigger& Triggers = GetVertexData<FTrigger>();
+
+			Triggers.ExecuteBlock(
+				[](int32, int32) {},
+				[this](const int32 TriggerFrame, int32)
+				{
+					const double OffsetSeconds = static_cast<double>(TriggerFrame) / SampleRate;
+					*LastTriggerTime = FTime::FromSeconds(OffsetSeconds + BlockStartTime);
+					MarkOutputDirty();
+				});
+
+			BlockStartTime += SecondsPerBlock;
+		}
+	}
+}
diff --git a/Engine/Plugins/Runtime/Metasound/Source/MetasoundFrontend/Public/Analysis/MetasoundFrontendAnalyzerAddress.h b/Engine/Plugins/Runtime/Metasound/Source/MetasoundFrontend/Public/Analysis/MetasoundFrontendAnalyzerAddress.h
--- a/Engine/Plugins/Runtime/Metasound/Source/MetasoundFrontend/Public/Analysis/MetasoundFrontendAnalyzerAddress.h
+++ b/Engine/Plugins/Runtime/Metasound/Source/MetasoundFrontend/Public/Analysis/MetasoundFrontendAnalyzerAddress.h
@@ -64,6 +64,11 @@ namespace Metasound
 			// the same name. Useful if the analyzer requires outputting multiple analysis values.
 			// Can potentially be used as an input as well to modify analyzer settings.
 			FName AnalyzerMemberName;
+
+			/**
+			 * Specifies whether to use the data transmission center for passing values to views.
+			 */
+			bool UseDataTransmissionCenter = true;
 		};
 	} // namespace Frontend
 
diff --git a/Engine/Plugins/Runtime/Metasound/Source/MetasoundFrontend/Public/Analysis/MetasoundFrontendAnalyzerRegistry.h b/Engine/Plugins/Runtime/Metasound/Source/MetasoundFrontend/Public/Analysis/MetasoundFrontendAnalyzerRegistry.h
--- a/Engine/Plugins/Runtime/Metasound/Source/MetasoundFrontend/Public/Analysis/MetasoundFrontendAnalyzerRegistry.h
+++ b/Engine/Plugins/Runtime/Metasound/Source/MetasoundFrontend/Public/Analysis/MetasoundFrontendAnalyzerRegistry.h
@@ -17,10 +17,20 @@ namespace Metasound
 		class METASOUNDFRONTEND_API IVertexAnalyzerRegistry
 		{
 		public:
+			virtual ~IVertexAnalyzerRegistry() = default;
+			
 			static IVertexAnalyzerRegistry& Get();
 
 			virtual const IVertexAnalyzerFactory* FindAnalyzerFactory(FName InAnalyzerName) const = 0;
-			virtual void RegisterAnalyzerFactories() = 0;
+			virtual void RegisterAnalyzerFactory(FName AnalyzerName, TUniquePtr<IVertexAnalyzerFactory>&& Factory) = 0;
 		};
 	} // namespace Frontend
 } // namespace Metasound
+
+#define METASOUND_REGISTER_VERTEX_ANALYZER_FACTORY(ANALYZER_CLASS) \
+	{ \
+	Frontend::IVertexAnalyzerRegistry& AnalyzerRegistry = Frontend::IVertexAnalyzerRegistry::Get(); \
+	AnalyzerRegistry.RegisterAnalyzerFactory( \
+		ANALYZER_CLASS::GetAnalyzerName(), \
+		MakeUnique<ANALYZER_CLASS::FFactory>()); \
+	}
diff --git a/Engine/Plugins/Runtime/Metasound/Source/MetasoundFrontend/Public/Analysis/MetasoundFrontendVertexAnalyzer.h b/Engine/Plugins/Runtime/Metasound/Source/MetasoundFrontend/Public/Analysis/MetasoundFrontendVertexAnalyzer.h
--- a/Engine/Plugins/Runtime/Metasound/Source/MetasoundFrontend/Public/Analysis/MetasoundFrontendVertexAnalyzer.h
+++ b/Engine/Plugins/Runtime/Metasound/Source/MetasoundFrontend/Public/Analysis/MetasoundFrontendVertexAnalyzer.h
@@ -5,10 +5,10 @@
 #include "Containers/Map.h"
 #include "HAL/Platform.h"
 #include "MetasoundDataReference.h"
-#include "MetasoundDataReferenceCollection.h"
 #include "MetasoundOperatorSettings.h"
-#include "Misc/Guid.h"
+#include "MetasoundOutputStorage.h"
 
+#include "Misc/Guid.h"
 
 namespace Metasound
 {
@@ -50,6 +50,9 @@ namespace Metasound
 
 			// Executes analysis
 			virtual void Execute() = 0;
+
+			DECLARE_DELEGATE_TwoParams(FOnOutputDataChanged, FName, TSharedPtr<IOutputStorage>);
+			FOnOutputDataChanged OnOutputDataChanged;
 		};
 
 		// Bound output from an analyzer (not to be confused with an output vertex)
@@ -62,9 +65,12 @@ namespace Metasound
 
 			// Returns the data reference associated with the analyzer's output
 			virtual FAnyDataReference GetDataReference() const = 0;
-
+			
 			// Pushes data to the sender to be forwarded to all actively associated analyser views.
 			virtual void PushData() = 0;
+
+			// Gets the latest value to pass to IVertexAnalyzer::OnOutputDataChanged
+			virtual TSharedPtr<IOutputStorage> CreateOutputData() const = 0;
 		};
 
 		// Templatized implementation of a bound analyzer output (see IBoundAnalyzerOutput)
@@ -72,15 +78,15 @@ namespace Metasound
 		class TBoundAnalyzerOutput final : public IBoundAnalyzerOutput
 		{
 		public:
-			TBoundAnalyzerOutput(const FAnalyzerAddress& InAnalyzerOutputAddress, const FOperatorSettings& InOperatorSettings, TDataReadReference<DataType>&& InData)
+			TBoundAnalyzerOutput(const FAnalyzerAddress& InAnalyzerOutputAddress, const FOperatorSettings& InOperatorSettings, TDataReadReference<DataType> InData)
 				: Address(InAnalyzerOutputAddress)
-				, DataRef(MoveTemp(InData))
+				, DataRef(InData)
 			{
 				Sender = FDataTransmissionCenter::Get().RegisterNewSender<DataType>(InAnalyzerOutputAddress, FSenderInitParams { InOperatorSettings, 0 });
 				ensure(Sender);
 			}
 
-			virtual ~TBoundAnalyzerOutput()
+			virtual ~TBoundAnalyzerOutput() override
 			{
 				// Only unregister the data channel if we had a sender using that 
 				// data channel. This protects against removing the data channel 
@@ -94,12 +100,12 @@ namespace Metasound
 				}
 			}
 
-			FAnyDataReference GetDataReference() const override
+			virtual FAnyDataReference GetDataReference() const override
 			{
 				return DataRef;
 			}
 
-			void PushData() override
+			virtual void PushData() override
 			{
 				if (Sender.IsValid())
 				{
@@ -107,6 +113,11 @@ namespace Metasound
 				}
 			}
 
+			virtual TSharedPtr<IOutputStorage> CreateOutputData() const override
+			{
+				return MakeShared<TOutputStorage<DataType>>(*DataRef);
+			}
+
 		private:
 			FAnalyzerAddress Address;
 			TDataReadReference<DataType> DataRef;
@@ -141,6 +152,10 @@ namespace Metasound
 				for (TPair<FName, FBoundOutputDataPtr>& Pair : BoundOutputData)
 				{
 					Pair.Value->PushData();
+					if (OnOutputDataChanged.IsBound())
+					{
+						OnOutputDataChanged.Execute(Pair.Key, Pair.Value->CreateOutputData());
+					}
 				}
 			}
 
@@ -152,12 +167,12 @@ namespace Metasound
 
 			// Binds a particular named analyzer output to be updated when MarkOutputDirty is called.
 			template<typename DataType>
-			void BindOutputData(FName InAnalyzerOutputName, const FOperatorSettings& InOperatorSettings, TDataReadReference<DataType>&& InData)
+			void BindOutputData(FName InAnalyzerOutputName, const FOperatorSettings& InOperatorSettings, TDataReadReference<DataType> InData)
 			{
 				FAnalyzerAddress OutputAddress = AnalyzerAddress;
 				OutputAddress.AnalyzerMemberName = InAnalyzerOutputName;
 				OutputAddress.DataType = GetMetasoundDataTypeName<DataType>();
-				FBoundOutputDataPtr BoundOutputDataPtr = FBoundOutputDataPtr(new TBoundAnalyzerOutput<DataType>(OutputAddress, InOperatorSettings, MoveTemp(InData)));
+				FBoundOutputDataPtr BoundOutputDataPtr = FBoundOutputDataPtr(new TBoundAnalyzerOutput<DataType>(OutputAddress, InOperatorSettings, InData));
 				BoundOutputData.Emplace(InAnalyzerOutputName, MoveTemp(BoundOutputDataPtr));
 			}
 
diff --git a/Engine/Plugins/Runtime/Metasound/Source/MetasoundFrontend/Public/Analysis/MetasoundFrontendVertexAnalyzerTriggerToTime.h b/Engine/Plugins/Runtime/Metasound/Source/MetasoundFrontend/Public/Analysis/MetasoundFrontendVertexAnalyzerTriggerToTime.h
new file mode 100644
--- /dev/null
+++ b/Engine/Plugins/Runtime/Metasound/Source/MetasoundFrontend/Public/Analysis/MetasoundFrontendVertexAnalyzerTriggerToTime.h
@@ -0,0 +1,42 @@
+// Copyright Epic Games, Inc. All Rights Reserved.
+#pragma once
+
+#include "MetasoundDataReference.h"
+#include "Analysis/MetasoundFrontendAnalyzerFactory.h"
+#include "Analysis/MetasoundFrontendVertexAnalyzer.h"
+
+namespace Metasound
+{
+	namespace Frontend
+	{
+		class METASOUNDFRONTEND_API FVertexAnalyzerTriggerToTime final : public FVertexAnalyzerBase
+		{
+		public:
+			static const FName& GetAnalyzerName();
+			static const FName& GetDataType();
+
+			struct METASOUNDFRONTEND_API FOutputs
+			{
+				static const FAnalyzerOutput& GetValue();
+			};
+
+			class METASOUNDFRONTEND_API FFactory final : public TVertexAnalyzerFactory<FVertexAnalyzerTriggerToTime>
+			{
+			public:
+				virtual const TArray<FAnalyzerOutput>& GetAnalyzerOutputs() const override;
+			};
+
+			explicit FVertexAnalyzerTriggerToTime(const FCreateAnalyzerParams& InParams);
+			virtual ~FVertexAnalyzerTriggerToTime() override = default;
+
+			virtual void Execute() override;
+
+		private:
+			const FSampleRate SampleRate;
+			const int32 FramesPerBlock;
+			const double SecondsPerBlock;
+			FTime BlockStartTime{ 0.0f };
+			TDataWriteReference<FTime> LastTriggerTime;
+		};
+	} // namespace Frontend
+} // namespace Metasound
diff --git a/Engine/Plugins/Runtime/Metasound/Source/MetasoundFrontend/Public/MetasoundOutputStorage.h b/Engine/Plugins/Runtime/Metasound/Source/MetasoundFrontend/Public/MetasoundOutputStorage.h
new file mode 100644
--- /dev/null
+++ b/Engine/Plugins/Runtime/Metasound/Source/MetasoundFrontend/Public/MetasoundOutputStorage.h
@@ -0,0 +1,57 @@
+// Copyright Epic Games, Inc. All Rights Reserved.
+
+#pragma once
+
+#include "MetasoundDataReference.h"
+#include "Templates/UniquePtr.h"
+
+namespace Metasound
+{
+	class IOutputStorage
+	{
+	public:
+		virtual ~IOutputStorage() = default;
+
+		virtual FName GetDataTypeName() const = 0;
+
+		virtual TUniquePtr<IOutputStorage> Clone() const = 0;
+	};
+	
+	template<typename DataType>
+	class TOutputStorage final : public IOutputStorage
+	{
+	public:
+		explicit TOutputStorage(const DataType& InData)
+			: Data(InData)
+		{}
+		
+		virtual FName GetDataTypeName() const override
+		{
+			static FName TypeName = GetMetasoundDataTypeName<DataType>();
+			return TypeName;
+		}
+
+		virtual TUniquePtr<IOutputStorage> Clone() const override
+		{
+			return MakeUnique<TOutputStorage<DataType>>(Data);
+		}
+		
+		void Set(const DataType& Value)
+		{
+			Data = Value;
+		}
+
+		void Set(DataType&& Value)
+		{
+			Data = MoveTemp(Value);
+		}
+
+		const DataType& Get() const
+		{
+			return Data;
+		}
+
+	private:
+		DataType Data;
+	};
+}
diff --git a/Engine/Plugins/Runtime/Metasound/Source/MetasoundGenerator/Private/MetasoundGenerator.cpp b/Engine/Plugins/Runtime/Metasound/Source/MetasoundGenerator/Private/MetasoundGenerator.cpp
--- a/Engine/Plugins/Runtime/Metasound/Source/MetasoundGenerator/Private/MetasoundGenerator.cpp
+++ b/Engine/Plugins/Runtime/Metasound/Source/MetasoundGenerator/Private/MetasoundGenerator.cpp
@@ -481,11 +481,19 @@ PRAGMA_ENABLE_DEPRECATION_WARNINGS
 		return NumChannels;
 	}
 
+	bool AnalyzerAddressesReferToSameGeneratorOutput(const Frontend::FAnalyzerAddress& Lhs, const Frontend::FAnalyzerAddress& Rhs)
+	{
+		return Lhs.OutputName == Rhs.OutputName && Lhs.AnalyzerName == Rhs.AnalyzerName;
+	}
+	
 	void FMetasoundGenerator::AddOutputVertexAnalyzer(const Frontend::FAnalyzerAddress& AnalyzerAddress)
 	{
 		OutputAnalyzerModificationQueue.Enqueue([this, AnalyzerAddress]
 		{
-			if (OutputAnalyzers.Contains(AnalyzerAddress.OutputName))
+			if (OutputAnalyzers.ContainsByPredicate([AnalyzerAddress](const TUniquePtr<Frontend::IVertexAnalyzer>& Analyzer)
+			{
+				return AnalyzerAddressesReferToSameGeneratorOutput(AnalyzerAddress, Analyzer->GetAnalyzerAddress());
+			}))
 			{
 				return;
 			}
@@ -510,16 +518,32 @@ PRAGMA_ENABLE_DEPRECATION_WARNINGS
 			{
 				return;
 			}
+
+			if (!Analyzer->OnOutputDataChanged.IsBound())
+			{
+				Analyzer->OnOutputDataChanged.BindLambda(
+					[this, AnalyzerAddress](const FName AnalyzerOutputName, TSharedPtr<IOutputStorage> OutputData)
+					{
+						OnOutputChanged.Broadcast(
+							AnalyzerAddress.AnalyzerName,
+							AnalyzerAddress.OutputName,
+							AnalyzerOutputName,
+							OutputData);
+					});
+			}
 			
-			OutputAnalyzers.Emplace(AnalyzerAddress.OutputName, MoveTemp(Analyzer));
+			OutputAnalyzers.Emplace(MoveTemp(Analyzer));
 		});
 	}
-	
-	void FMetasoundGenerator::RemoveOutputVertexAnalyzer(const FName& OutputName)
+
+	void FMetasoundGenerator::RemoveOutputVertexAnalyzer(const Frontend::FAnalyzerAddress& AnalyzerAddress)
 	{
-		OutputAnalyzerModificationQueue.Enqueue([this, OutputName]
+		OutputAnalyzerModificationQueue.Enqueue([this, AnalyzerAddress]
 		{
-			OutputAnalyzers.Remove(OutputName);
+			OutputAnalyzers.RemoveAll([AnalyzerAddress](const TUniquePtr<Frontend::IVertexAnalyzer>& Analyzer)
+			{
+				return AnalyzerAddressesReferToSameGeneratorOutput(AnalyzerAddress, Analyzer->GetAnalyzerAddress());
+			});
 		});
 	}
 
@@ -596,9 +620,9 @@ PRAGMA_ENABLE_DEPRECATION_WARNINGS
 			}
 
 			// Execute the output analyzers
-			for (const auto& Pair : OutputAnalyzers)
+			for (const TUniquePtr<Frontend::IVertexAnalyzer>& Analyzer : OutputAnalyzers)
 			{
-				Pair.Value->Execute();
+				Analyzer->Execute();
 			}
 
 			// Check if generated finished during this execute call
diff --git a/Engine/Plugins/Runtime/Metasound/Source/MetasoundGenerator/Public/MetasoundGenerator.h b/Engine/Plugins/Runtime/Metasound/Source/MetasoundGenerator/Public/MetasoundGenerator.h
--- a/Engine/Plugins/Runtime/Metasound/Source/MetasoundGenerator/Public/MetasoundGenerator.h
+++ b/Engine/Plugins/Runtime/Metasound/Source/MetasoundGenerator/Public/MetasoundGenerator.h
@@ -187,16 +187,19 @@ namespace Metasound
 		/**
 		 * Add a vertex analyzer for a named output with the given address info.
 		 *
-		 * @param AnalyzerAddress - Address information for the analyzer to use when making views
+		 * @param AnalyzerAddress - Address information for the analyzer
 		 */
 		void AddOutputVertexAnalyzer(const Frontend::FAnalyzerAddress& AnalyzerAddress);
 		
 		/**
 		 * Remove a vertex analyzer for a named output
 		 *
-		 * @param OutputName - The name of the output in the MetaSound graph
+		 * @param AnalyzerAddress - Address information for the analyzer
 		 */
-		void RemoveOutputVertexAnalyzer(const FName& OutputName);
+		void RemoveOutputVertexAnalyzer(const Frontend::FAnalyzerAddress& AnalyzerAddress);
+		
+		DECLARE_TS_MULTICAST_DELEGATE_FourParams(FOnOutputChanged, FName, FName, FName, TSharedPtr<IOutputStorage>);
+		FOnOutputChanged OnOutputChanged;
 		
 		/** Return the number of audio channels. */
 		int32 GetNumChannels() const;
@@ -307,7 +310,7 @@ namespace Metasound
 		TMpscQueue<TSharedPtr<FMetasoundParameterPackStorage>> ParameterPackQueue;
 
 		TMpscQueue<TUniqueFunction<void()>> OutputAnalyzerModificationQueue;
-		TMap<FName, TUniquePtr<Frontend::IVertexAnalyzer>> OutputAnalyzers;
+		TArray<TUniquePtr<Frontend::IVertexAnalyzer>> OutputAnalyzers;
 #if ENABLE_METASOUND_GENERATOR_RENDER_TIMING
 		TUniquePtr<MetasoundGeneratorPrivate::FRenderTimer> RenderTimer;
 #endif // if ENABLE_METASOUND_GENERATOR_RENDER_TIMING
