[Metasound] Fix crash when attempting to connect vertices with mismatched access types

#rnx

diff --git a/Engine/Plugins/Runtime/Metasound/Source/MetasoundFrontend/Private/MetasoundFrontendInputController.cpp b/Engine/Plugins/Runtime/Metasound/Source/MetasoundFrontend/Private/MetasoundFrontendInputController.cpp
--- a/Engine/Plugins/Runtime/Metasound/Source/MetasoundFrontend/Private/MetasoundFrontendInputController.cpp
+++ b/Engine/Plugins/Runtime/Metasound/Source/MetasoundFrontend/Private/MetasoundFrontendInputController.cpp
@@ -300,7 +300,7 @@ namespace Metasound
 					}
 					else
 					{
-						UE_LOG(LogMetaSound, Error, TEXT("Cannot connect incompatible vertex access types (Input)%s and (Output)%s."), *LexToString(GetVertexAccessType()), *LexToString(InController.GetVertexAccessType()));
+						UE_LOG(LogMetaSound, Error, TEXT("Cannot connect incompatible vertex access types (Input)%s and (Output)%s."), LexToString(GetVertexAccessType()), LexToString(InController.GetVertexAccessType()));
 					}
 				}
 				else
