[Metasound] Build fix - missing variable capture in a lambda definition

diff --git a/Engine/Plugins/Runtime/Metasound/Source/MetasoundEngineTest/Private/MetasoundGeneratorTests.cpp b/Engine/Plugins/Runtime/Metasound/Source/MetasoundEngineTest/Private/MetasoundGeneratorTests.cpp
--- a/Engine/Plugins/Runtime/Metasound/Source/MetasoundEngineTest/Private/MetasoundGeneratorTests.cpp
+++ b/Engine/Plugins/Runtime/Metasound/Source/MetasoundEngineTest/Private/MetasoundGeneratorTests.cpp
@@ -426,7 +426,7 @@ namespace Metasound::Test::Generator
 		
 		// Subscribe for updates
 		Generator->OnOutputChanged.AddLambda(
-			[this, &NextTriggerReceivedSeconds, &AnalyzerAddress, &CallbackSuccess, &NumTriggersReceived]
+			[this, &NextTriggerReceivedSeconds, &AnalyzerAddress, &CallbackSuccess, &NumTriggersReceived, TriggerIntervalSeconds]
 			(const FName AnalyzerName, const FName OutputName, const FName AnalyzerOutputName, TSharedPtr<IOutputStorage> OutputData)
 			{
 				CallbackSuccess = TestEqual("Data types match", OutputData->GetDataTypeName(), GetMetasoundDataTypeName<FTime>())
