Fix for ensure in the MS BP builder subsystem when attempting to add a MetaSound Node Class using a literal BP node value
#rnx

diff --git a/Engine/Plugins/Runtime/Metasound/Source/MetasoundEngine/Private/MetasoundBuilderSubsystem.cpp b/Engine/Plugins/Runtime/Metasound/Source/MetasoundEngine/Private/MetasoundBuilderSubsystem.cpp
--- a/Engine/Plugins/Runtime/Metasound/Source/MetasoundEngine/Private/MetasoundBuilderSubsystem.cpp
+++ b/Engine/Plugins/Runtime/Metasound/Source/MetasoundEngine/Private/MetasoundBuilderSubsystem.cpp
@@ -128,7 +128,7 @@ void UMetaSoundBuilderBase::AddInterface(FName InterfaceName, EMetaSoundBuilderR
 	OutResult = bInterfaceAdded ? EMetaSoundBuilderResult::Succeeded : EMetaSoundBuilderResult::Failed;
 }
 
-FMetaSoundNodeHandle UMetaSoundBuilderBase::AddNode(TScriptInterface<IMetaSoundDocumentInterface> NodeClass, EMetaSoundBuilderResult& OutResult)
+FMetaSoundNodeHandle UMetaSoundBuilderBase::AddNode(const TScriptInterface<IMetaSoundDocumentInterface>& NodeClass, EMetaSoundBuilderResult& OutResult)
 {
 	using namespace Metasound;
 
diff --git a/Engine/Plugins/Runtime/Metasound/Source/MetasoundEngine/Public/MetasoundBuilderSubsystem.h b/Engine/Plugins/Runtime/Metasound/Source/MetasoundEngine/Public/MetasoundBuilderSubsystem.h
--- a/Engine/Plugins/Runtime/Metasound/Source/MetasoundEngine/Public/MetasoundBuilderSubsystem.h
+++ b/Engine/Plugins/Runtime/Metasound/Source/MetasoundEngine/Public/MetasoundBuilderSubsystem.h
@@ -132,7 +132,7 @@ public:
 	// Adds a node to the graph using the provided MetaSound asset as its defining NodeClass.
 	// Returns a node handle to the created node if successful, or an invliad handle if it failed.
 	UFUNCTION(BlueprintCallable, Category = "Audio|MetaSound|Builder", DisplayName = "Add MetaSound Node From Asset Class", meta = (ExpandEnumAsExecs = "OutResult"))
-	UPARAM(DisplayName = "Node Handle") FMetaSoundNodeHandle AddNode(TScriptInterface<IMetaSoundDocumentInterface> NodeClass, EMetaSoundBuilderResult& OutResult);
+	UPARAM(DisplayName = "Node Handle") FMetaSoundNodeHandle AddNode(const TScriptInterface<IMetaSoundDocumentInterface>& NodeClass, EMetaSoundBuilderResult& OutResult);
 
 	// Connects node output to a node input. Does *NOT* provide loop detection for performance reasons.  Loop detection is checked on class registration when built or played.
 	// Returns succeeded if connection made, failed if connection already exists with input, the data types do not match, or the connection is not supported due to access type
