Fix for shadowing of function when registering interface introduced in refactor
#rnx

diff --git a/Engine/Plugins/Runtime/Metasound/Source/MetasoundEngine/Private/Interfaces/MetasoundOutputFormatInterfaces.cpp b/Engine/Plugins/Runtime/Metasound/Source/MetasoundEngine/Private/Interfaces/MetasoundOutputFormatInterfaces.cpp
--- a/Engine/Plugins/Runtime/Metasound/Source/MetasoundEngine/Private/Interfaces/MetasoundOutputFormatInterfaces.cpp
+++ b/Engine/Plugins/Runtime/Metasound/Source/MetasoundEngine/Private/Interfaces/MetasoundOutputFormatInterfaces.cpp
@@ -122,7 +122,7 @@ namespace Metasound::Engine
 			};
 		}
 
-		TArray<Audio::FParameterInterface::FClassOptions> GetUClassOptions(bool bIsSourceDefault = false)
+		TArray<Audio::FParameterInterface::FClassOptions> CreateUClassOptions(bool bIsSourceDefault = false)
 		{
 			return
 			{
@@ -158,8 +158,10 @@ namespace Metasound::Engine
 				FInterface()
 					: FParameterInterface(OutputFormatMonoInterface::GetVersion().Name, OutputFormatMonoInterface::GetVersion().Number.ToInterfaceVersion())
 				{
+					using namespace OutputFormatPrivate;
+
 					constexpr bool bIsSourceDefault = true;
-					UClassOptions = OutputFormatPrivate::GetUClassOptions(bIsSourceDefault);
+					UClassOptions = CreateUClassOptions(bIsSourceDefault);
 
 					Outputs =
 					{
@@ -208,7 +210,10 @@ namespace Metasound::Engine
 				FInterface()
 					: FParameterInterface(OutputFormatStereoInterface::GetVersion().Name, OutputFormatStereoInterface::GetVersion().Number.ToInterfaceVersion())
 				{
-					UClassOptions = OutputFormatPrivate::GetUClassOptions();
+					using namespace OutputFormatPrivate;
+					
+					UClassOptions = CreateUClassOptions();
+
 					Outputs =
 					{
 						{
@@ -269,7 +274,8 @@ namespace Metasound::Engine
 				{
 					using namespace OutputFormatPrivate;
 
-					UClassOptions = GetUClassOptions();
+					UClassOptions = CreateUClassOptions();
+
 					Outputs =
 					{
 						GetFrontLeftOutput(Outputs::FrontLeftOut),
@@ -318,7 +324,8 @@ namespace Metasound::Engine
 				{
 					using namespace OutputFormatPrivate;
 
-					UClassOptions = GetUClassOptions();
+					UClassOptions = CreateUClassOptions();
+
 					Outputs =
 					{
 						GetFrontLeftOutput(Outputs::FrontLeftOut),
@@ -371,7 +378,8 @@ namespace Metasound::Engine
 				{
 					using namespace OutputFormatPrivate;
 
-					UClassOptions = GetUClassOptions();
+					UClassOptions = CreateUClassOptions();
+
 					Outputs =
 					{
 						GetFrontLeftOutput(Outputs::FrontLeftOut),
