Refactor interface modification placement algorithm to position newly added nodes using sort order of adjacent graph members
- Clean-up & remove dead code in FModifyRootGraphInterfaces transform
#rnx

diff --git a/Engine/Plugins/Runtime/Metasound/Source/MetasoundFrontend/Private/MetasoundFrontendTransform.cpp b/Engine/Plugins/Runtime/Metasound/Source/MetasoundFrontend/Private/MetasoundFrontendTransform.cpp
--- a/Engine/Plugins/Runtime/Metasound/Source/MetasoundFrontend/Private/MetasoundFrontendTransform.cpp
+++ b/Engine/Plugins/Runtime/Metasound/Source/MetasoundFrontend/Private/MetasoundFrontendTransform.cpp
@@ -89,6 +89,23 @@ namespace Metasound
 			Init(&InNamePairingFunction);
 		}
 
+		bool FModifyRootGraphInterfaces::AddMissingVertices(FGraphHandle GraphHandle) const
+		{
+			for (const FInputData& InputData : InputsToAdd)
+			{
+				const FMetasoundFrontendClassInput& InputToAdd = InputData.Input;
+				GraphHandle->AddInputVertex(InputToAdd);
+			}
+
+			for (const FOutputData& OutputData : OutputsToAdd)
+			{
+				const FMetasoundFrontendClassOutput& OutputToAdd = OutputData.Output;
+				GraphHandle->AddOutputVertex(OutputToAdd);
+			}
+
+			return !InputsToAdd.IsEmpty() || !OutputsToAdd.IsEmpty();
+		}
+
 		void FModifyRootGraphInterfaces::Init(const TFunction<bool(FName, FName)>* InNamePairingFunction)
 		{
 			InputsToRemove.Reset();
@@ -193,33 +210,15 @@ namespace Metasound
 
 				if (INDEX_NONE != RemoveIndex)
 				{
-					PairedOutputs.Add(FVertexPair{OutputsToRemove[RemoveIndex], OutputsToAdd[AddIndex].Output});
+					PairedOutputs.Add(FVertexPair {OutputsToRemove[RemoveIndex], OutputsToAdd[AddIndex].Output });
 					OutputsToRemove.RemoveAtSwap(RemoveIndex);
 					OutputsToAdd.RemoveAtSwap(AddIndex);
 				}
-			}			
+			}
 		}
 
-		bool FModifyRootGraphInterfaces::Transform(FDocumentHandle InDocument) const
+		bool FModifyRootGraphInterfaces::RemoveUnsupportedVertices(FGraphHandle GraphHandle) const
 		{
-			bool bDidEdit = false;
-
-			FGraphHandle GraphHandle = InDocument->GetRootGraph();
-			if (!ensure(GraphHandle->IsValid()))
-			{
-				return false;
-			}
-
-			for (const FMetasoundFrontendInterface& Interface : InterfacesToRemove)
-			{
-				InDocument->RemoveInterfaceVersion(Interface.Version);
-			}
-
-			for (const FMetasoundFrontendInterface& Interface : InterfacesToAdd)
-			{
-				InDocument->AddInterfaceVersion(Interface.Version);
-			}
-
 			// Remove unsupported inputs
 			for (const FMetasoundFrontendClassVertex& InputToRemove : InputsToRemove)
 			{
@@ -227,128 +226,30 @@ namespace Metasound
 				{
 					if (FMetasoundFrontendClassInput::IsFunctionalEquivalent(*ClassInput, InputToRemove))
 					{
-						bDidEdit = true;
 						GraphHandle->RemoveInputVertex(InputToRemove.Name);
 					}
 				}
 			}
 
-			// Remove unrequired outputs
+			// Remove unsupported outputs
 			for (const FMetasoundFrontendClassVertex& OutputToRemove : OutputsToRemove)
 			{
 				if (const FMetasoundFrontendClassOutput* ClassOutput = GraphHandle->FindClassOutputWithName(OutputToRemove.Name).Get())
 				{
 					if (FMetasoundFrontendClassOutput::IsFunctionalEquivalent(*ClassOutput, OutputToRemove))
 					{
-						bDidEdit = true;
 						GraphHandle->RemoveOutputVertex(OutputToRemove.Name);
 					}
 				}
 			}
 
-			auto InputDataTypeCompareFilter = [](FConstNodeHandle NodeHandle, FName DataType)
-			{
-				bool bMatchesDataType = false;
-				NodeHandle->IterateConstOutputs([&](FConstOutputHandle OutputHandle)
-				{
-					if (OutputHandle->GetDataType() == DataType)
-					{
-						bMatchesDataType = true;
-					}
-				});
-
-				return bMatchesDataType;
-			};
-
-			auto OutputDataTypeCompareFilter = [](FConstNodeHandle NodeHandle, FName DataType)
-			{
-				bool bMatchesDataType = false;
-				NodeHandle->IterateConstInputs([&](FConstInputHandle InputHandle)
-				{
-					if (InputHandle->GetDataType() == DataType)
-					{
-						bMatchesDataType = true;
-					}
-				});
-
-				return bMatchesDataType;
-			};
-
-#if WITH_EDITOR
-			auto FindLowestNodeLocationOfClassType = [](EMetasoundFrontendClassType ClassType, FGraphHandle Graph, FName DataType, TFunctionRef<bool(FConstNodeHandle, FName)> NodeDataTypeFilter)
-			{
-				FVector2D LowestLocation = { TNumericLimits<float>::Min(), TNumericLimits<float>::Min() };
-				bool bFoundLocation = false;
-
-				Graph->IterateConstNodes([&](FConstNodeHandle NodeHandle)
-				{
-					const TMap<FGuid, FVector2D>& Locations = NodeHandle->GetNodeStyle().Display.Locations;
-					for (const TPair<FGuid, FVector2D>& Pair : Locations)
-					{
-						if (Pair.Value.Y > LowestLocation.Y)
-						{
-							LowestLocation = Pair.Value;
-							bFoundLocation = true;
-						}
-					}
-				}, ClassType);
-
-				if (!bFoundLocation)
-				{
-					LowestLocation = { DisplayStyle::NodeLayout::DefaultOffsetX.X * 2.0f, 0.0f };
-				}
-
-				return LowestLocation;
-			};
-#endif // WITH_EDITOR
-
-			// Add missing inputs
-			for (const FInputData& InputData : InputsToAdd)
-			{
-				bDidEdit = true;
-				const FMetasoundFrontendClassInput& InputToAdd = InputData.Input;
-				FNodeHandle NewInputNode = GraphHandle->AddInputVertex(InputToAdd);
-
-#if WITH_EDITOR
-				FName InputName = InputToAdd.Name;
-				bool bRequiredInput = InputData.InputInterface->GetInputStyle().RequiredMembers.Contains(InputName);
-
-				if (bSetDefaultNodeLocations || bRequiredInput)
-				{
-					FMetasoundFrontendNodeStyle Style = NewInputNode->GetNodeStyle();
-					const FVector2D LastOutputLocation = FindLowestNodeLocationOfClassType(EMetasoundFrontendClassType::Input, GraphHandle, InputToAdd.TypeName, InputDataTypeCompareFilter);
-					Style.Display.Locations.Add(FGuid(), LastOutputLocation + DisplayStyle::NodeLayout::DefaultOffsetY);
-					NewInputNode->SetNodeStyle(Style);
-				}
-#endif // WITH_EDITOR
-			}
-
-			// Add missing outputs
-			for (const FOutputData& OutputData : OutputsToAdd)
-			{
-				bDidEdit = true;
-				const FMetasoundFrontendClassOutput& OutputToAdd = OutputData.Output;
-				FNodeHandle NewOutputNode = GraphHandle->AddOutputVertex(OutputToAdd);
-
-#if WITH_EDITOR
-				FName OutputName = OutputToAdd.Name;
-				bool bRequiredOutput = OutputData.OutputInterface->GetOutputStyle().RequiredMembers.Contains(OutputName);
-
-				if (bSetDefaultNodeLocations || bRequiredOutput)
-				{
-					FMetasoundFrontendNodeStyle Style = NewOutputNode->GetNodeStyle();
-					const FVector2D LastOutputLocation = FindLowestNodeLocationOfClassType(EMetasoundFrontendClassType::Output, GraphHandle, OutputToAdd.TypeName, OutputDataTypeCompareFilter);
-					Style.Display.Locations.Add(FGuid(), LastOutputLocation + DisplayStyle::NodeLayout::DefaultOffsetY);
-					NewOutputNode->SetNodeStyle(Style);
-				}
-#endif // WITH_EDITOR
-			}
+			return !InputsToRemove.IsEmpty() || !OutputsToRemove.IsEmpty();
+		}
 
-			// Swap paired inputs.
+		bool FModifyRootGraphInterfaces::SwapPairedVertices(FGraphHandle GraphHandle) const
+		{
 			for (const FVertexPair& InputPair : PairedInputs)
 			{
-				bDidEdit = true;
-
 				const FMetasoundFrontendClassVertex& OriginalVertex = InputPair.Get<0>();
 				FMetasoundFrontendClassInput NewVertex = InputPair.Get<1>();
 
@@ -379,7 +280,6 @@ namespace Metasound
 				// Copy prior node locations
 				if (!Locations.IsEmpty())
 				{
-					// TODO: copy entire style.
 					FMetasoundFrontendNodeStyle Style = NewInputNode->GetNodeStyle();
 					Style.Display.Locations = Locations;
 					NewInputNode->SetNodeStyle(Style);
@@ -397,8 +297,6 @@ namespace Metasound
 			// Swap paired outputs.
 			for (const FVertexPair& OutputPair : PairedOutputs)
 			{
-				bDidEdit = true;
-
 				const FMetasoundFrontendClassVertex& OriginalVertex = OutputPair.Get<0>();
 				FMetasoundFrontendClassVertex NewVertex = OutputPair.Get<1>();
 
@@ -406,7 +304,7 @@ namespace Metasound
 				// Cache off node locations to push to new node
 				// Default add output node to origin.
 				TMap<FGuid, FVector2D> Locations;
-				Locations.Add(FGuid(), FVector2D{0.f, 0.f});
+				Locations.Add(FGuid(), FVector2D { 0.f, 0.f });
 #endif // WITH_EDITOR
 
 				FOutputHandle ConnectedOutput = IOutputController::GetInvalidHandle();
@@ -450,6 +348,32 @@ namespace Metasound
 				ConnectedOutput->Connect(*InputHandle);
 			}
 
+			return !PairedInputs.IsEmpty() || !PairedOutputs.IsEmpty();
+		}
+
+		bool FModifyRootGraphInterfaces::Transform(FDocumentHandle InDocument) const
+		{
+			bool bDidEdit = false;
+
+			FGraphHandle GraphHandle = InDocument->GetRootGraph();
+			if (ensure(GraphHandle->IsValid()))
+			{
+				bDidEdit |= UpdateInterfacesInternal(InDocument);
+
+				const bool bAddedVertices = AddMissingVertices(GraphHandle);
+				bDidEdit |= bAddedVertices;
+
+				bDidEdit |= SwapPairedVertices(GraphHandle);
+				bDidEdit |= RemoveUnsupportedVertices(GraphHandle);
+
+#if WITH_EDITOR
+				if (bAddedVertices)
+				{
+					UpdateAddedVertexNodePositions(GraphHandle);
+				}
+#endif // WITH_EDITOR
+			}
+
 			return bDidEdit;
 		}
 
@@ -460,6 +384,103 @@ namespace Metasound
 			return Transform(FDocumentController::CreateDocumentHandle(DocAccessPtr));
 		}
 
+		bool FModifyRootGraphInterfaces::UpdateInterfacesInternal(FDocumentHandle DocumentHandle) const
+		{
+			for (const FMetasoundFrontendInterface& Interface : InterfacesToRemove)
+			{
+				DocumentHandle->RemoveInterfaceVersion(Interface.Version);
+			}
+
+			for (const FMetasoundFrontendInterface& Interface : InterfacesToAdd)
+			{
+				DocumentHandle->AddInterfaceVersion(Interface.Version);
+			}
+
+			return !InterfacesToRemove.IsEmpty() || !InterfacesToAdd.IsEmpty();
+		}
+
+#if WITH_EDITOR
+		void FModifyRootGraphInterfaces::UpdateAddedVertexNodePositions(FGraphHandle GraphHandle) const
+		{
+			auto SortAndPlaceMemberNodes = [&GraphHandle](EMetasoundFrontendClassType ClassType, TSet<FName>& AddedNames, TFunctionRef<int32(const FVertexName&)> InGetSortOrder)
+			{
+				// Add graph member nodes by sort order
+				TSortedMap<int32, FNodeHandle> SortOrderToName;
+				GraphHandle->IterateNodes([&GraphHandle, &SortOrderToName, &InGetSortOrder](FNodeHandle NodeHandle)
+				{
+					const int32 Index = InGetSortOrder(NodeHandle->GetNodeName());
+					SortOrderToName.Add(Index, NodeHandle);
+				}, ClassType);
+
+				// Prime the first location as an offset prior to an existing location (as provided by a swapped member)
+				//  to avoid placing away from user's active area if possible.
+				FVector2D NextLocation = { 0.0f, 0.0f };
+				{
+					int32 NumBeforeDefined = 1;
+					for (const TPair<int32, FNodeHandle>& Pair : SortOrderToName)
+					{
+						const FConstNodeHandle& NodeHandle = Pair.Value;
+						const FName NodeName = NodeHandle->GetNodeName();
+						if (AddedNames.Contains(NodeName))
+						{
+							NumBeforeDefined++;
+						}
+						else
+						{
+							const TMap<FGuid, FVector2D>& Locations = NodeHandle->GetNodeStyle().Display.Locations;
+							if (!Locations.IsEmpty())
+							{
+								for (const TPair<FGuid, FVector2D>& Location : Locations)
+								{
+									NextLocation = Location.Value - (NumBeforeDefined * DisplayStyle::NodeLayout::DefaultOffsetY);
+									break;
+								}
+								break;
+							}
+						}
+					}
+				}
+
+				// Iterate through sorted map in sequence, slotting in new locations after existing swapped nodes with predefined locations.
+				for (TPair<int32, FNodeHandle>& Pair : SortOrderToName)
+				{
+					FNodeHandle& NodeHandle = Pair.Value;
+					const FName NodeName = NodeHandle->GetNodeName();
+					if (AddedNames.Contains(NodeName))
+					{
+						FMetasoundFrontendNodeStyle NewStyle = NodeHandle->GetNodeStyle();
+						NewStyle.Display.Locations.Add(FGuid(), NextLocation);
+						NextLocation += DisplayStyle::NodeLayout::DefaultOffsetY;
+						NodeHandle->SetNodeStyle(NewStyle);
+					}
+					else
+					{
+						for (const TPair<FGuid, FVector2D>& Location : NodeHandle->GetNodeStyle().Display.Locations)
+						{
+							NextLocation = Location.Value + DisplayStyle::NodeLayout::DefaultOffsetY;
+						}
+					}
+				}
+			};
+
+			// Sort/Place Inputs
+			{
+				TSet<FName> AddedNames;
+				Algo::Transform(InputsToAdd, AddedNames, [](const FInputData& InputData) { return InputData.Input.Name; });
+				auto GetInputSortOrder = [&GraphHandle](const FVertexName& InVertexName) { return GraphHandle->GetSortOrderIndexForInput(InVertexName); };
+				SortAndPlaceMemberNodes(EMetasoundFrontendClassType::Input, AddedNames, GetInputSortOrder);
+			}
+
+			// Sort/Place Outputs
+			{
+				TSet<FName> AddedNames;
+				Algo::Transform(OutputsToAdd, AddedNames, [](const FOutputData& OutputData) { return OutputData.Output.Name; });
+				auto GetOutputSortOrder = [&GraphHandle](const FVertexName& InVertexName) { return GraphHandle->GetSortOrderIndexForOutput(InVertexName); };
+				SortAndPlaceMemberNodes(EMetasoundFrontendClassType::Output, AddedNames, GetOutputSortOrder);
+			}
+		}
+#endif // WITH_EDITOR
+
 		bool FUpdateRootGraphInterface::Transform(FDocumentHandle InDocument) const
 		{
 			bool bDidEdit = false;
diff --git a/Engine/Plugins/Runtime/Metasound/Source/MetasoundFrontend/Public/MetasoundFrontendTransform.h b/Engine/Plugins/Runtime/Metasound/Source/MetasoundFrontend/Public/MetasoundFrontendTransform.h
--- a/Engine/Plugins/Runtime/Metasound/Source/MetasoundFrontend/Public/MetasoundFrontendTransform.h
+++ b/Engine/Plugins/Runtime/Metasound/Source/MetasoundFrontend/Public/MetasoundFrontendTransform.h
@@ -94,9 +94,15 @@ namespace Metasound
 			virtual bool Transform(FMetasoundFrontendDocument& InOutDocument) const override;
 
 		private:
+			bool AddMissingVertices(FGraphHandle GraphHandle) const;
 			void Init(const TFunction<bool(FName, FName)>* InNamePairingFunction = nullptr);
+			bool SwapPairedVertices(FGraphHandle GraphHandle) const;
+			bool RemoveUnsupportedVertices(FGraphHandle GraphHandle) const;
+			bool UpdateInterfacesInternal(FDocumentHandle DocumentHandle) const;
 
 #if WITH_EDITOR
+			void UpdateAddedVertexNodePositions(FGraphHandle GraphHandle) const;
+
 			bool bSetDefaultNodeLocations = true;
 #endif // WITH_EDITOR
 
