Allow metasound C++ nodes to swap pointers on vertex data binding.

diff --git a/Engine/Plugins/Runtime/Metasound/Source/MetasoundFrontend/Private/MetasoundFrontendDataTypeRegistry.cpp b/Engine/Plugins/Runtime/Metasound/Source/MetasoundFrontend/Private/MetasoundFrontendDataTypeRegistry.cpp
--- a/Engine/Plugins/Runtime/Metasound/Source/MetasoundFrontend/Private/MetasoundFrontendDataTypeRegistry.cpp
+++ b/Engine/Plugins/Runtime/Metasound/Source/MetasoundFrontend/Private/MetasoundFrontendDataTypeRegistry.cpp
@@ -1141,7 +1141,7 @@ namespace Metasound
 			return Registry;
 		}
 
-		void CreateAndBindDefaults(const FOperatorSettings& InOperatorSettings, FInputVertexInterfaceData& OutVertexData)
+		void CreateDefaults(const FOperatorSettings& InOperatorSettings, FInputVertexInterfaceData& OutVertexData)
 		{
 			using namespace MetasoundVertexDataPrivate;
 
@@ -1160,7 +1160,7 @@ namespace Metasound
 
 			IDataTypeRegistry& DataTypeRegistry = IDataTypeRegistry::Get();
 
-			for (const TBinding<FInputDataVertex>& Binding : OutVertexData)
+			for (const FInputBinding& Binding : OutVertexData)
 			{
 				// Attempt to create default data reference from the literal stored
 				// on the input vertex.
@@ -1171,7 +1171,7 @@ namespace Metasound
 				if (DataRef.IsSet())
 				{
 					// Set as vertex data reference.
-					OutVertexData.BindVertex(InputVertex.VertexName, *DataRef);
+					OutVertexData.SetVertex(InputVertex.VertexName, *DataRef);
 				}
 				else
 				{
diff --git a/Engine/Plugins/Runtime/Metasound/Source/MetasoundFrontend/Private/MetasoundInputNode.cpp b/Engine/Plugins/Runtime/Metasound/Source/MetasoundFrontend/Private/MetasoundInputNode.cpp
--- a/Engine/Plugins/Runtime/Metasound/Source/MetasoundFrontend/Private/MetasoundInputNode.cpp
+++ b/Engine/Plugins/Runtime/Metasound/Source/MetasoundFrontend/Private/MetasoundInputNode.cpp
@@ -37,16 +37,16 @@ namespace Metasound
 
 		void FNonExecutableInputOperatorBase::Bind(FVertexInterfaceData& InOutVertexData) const
 		{
-			BindInputs(InOutVertexData.GetInputs());
-			BindOutputs(InOutVertexData.GetOutputs());
+			BindInputsInternal(InOutVertexData.GetInputs());
+			BindOutputsInternal(InOutVertexData.GetOutputs());
 		}
 
-		void FNonExecutableInputOperatorBase::BindInputs(FInputVertexInterfaceData& InOutVertexData) const
+		void FNonExecutableInputOperatorBase::BindInputsInternal(FInputVertexInterfaceData& InOutVertexData) const
 		{
 			InOutVertexData.BindVertex(VertexName, DataRef);
 		}
 
-		void FNonExecutableInputOperatorBase::BindOutputs(FOutputVertexInterfaceData& InOutVertexData) const
+		void FNonExecutableInputOperatorBase::BindOutputsInternal(FOutputVertexInterfaceData& InOutVertexData) const
 		{
 			InOutVertexData.BindVertex(VertexName, DataRef);
 		}
diff --git a/Engine/Plugins/Runtime/Metasound/Source/MetasoundFrontend/Public/MetasoundFrontendDataTypeRegistry.h b/Engine/Plugins/Runtime/Metasound/Source/MetasoundFrontend/Public/MetasoundFrontendDataTypeRegistry.h
--- a/Engine/Plugins/Runtime/Metasound/Source/MetasoundFrontend/Public/MetasoundFrontendDataTypeRegistry.h
+++ b/Engine/Plugins/Runtime/Metasound/Source/MetasoundFrontend/Public/MetasoundFrontendDataTypeRegistry.h
@@ -327,8 +327,8 @@ namespace Metasound
 		 * using the default literals which live on the FInputDataVertexes.
 		 *
 		 * If the vertex has access type EVertexAccessType::Reference, a writable
-		 * data reference will be bound, otherwise a value data reference will be bound.
+		 * data reference will be set, otherwise a value data reference will be set.
 		 */
-		METASOUNDFRONTEND_API void CreateAndBindDefaults(const FOperatorSettings& InOperatorSettings, FInputVertexInterfaceData& OutVertexData);
+		METASOUNDFRONTEND_API void CreateDefaults(const FOperatorSettings& InOperatorSettings, FInputVertexInterfaceData& OutVertexData);
 	}
 }
diff --git a/Engine/Plugins/Runtime/Metasound/Source/MetasoundFrontend/Public/MetasoundInputNode.h b/Engine/Plugins/Runtime/Metasound/Source/MetasoundFrontend/Public/MetasoundInputNode.h
--- a/Engine/Plugins/Runtime/Metasound/Source/MetasoundFrontend/Public/MetasoundInputNode.h
+++ b/Engine/Plugins/Runtime/Metasound/Source/MetasoundFrontend/Public/MetasoundInputNode.h
@@ -38,9 +38,9 @@ namespace Metasound
 			FNonExecutableInputOperatorBase(const FVertexName& InVertexName, FAnyDataReference&& InDataRef);
 
 		private:
-			void BindInputs(FInputVertexInterfaceData& InOutVertexData) const;
+			void BindInputsInternal(FInputVertexInterfaceData& InOutVertexData) const;
 		protected:
-			void BindOutputs(FOutputVertexInterfaceData& InOutVertexData) const;
+			void BindOutputsInternal(FOutputVertexInterfaceData& InOutVertexData) const;
 
 
 			FVertexName VertexName;
@@ -101,8 +101,8 @@ namespace Metasound
 
 			virtual void Bind(FVertexInterfaceData& InOutVertexData) const override
 			{
-				BindInputs(InOutVertexData.GetInputs());
-				BindOutputs(InOutVertexData.GetOutputs());
+				BindInputsInternal(InOutVertexData.GetInputs());
+				BindOutputsInternal(InOutVertexData.GetOutputs());
 			}
 
 			virtual FExecuteFunction GetExecuteFunction() override
@@ -116,14 +116,14 @@ namespace Metasound
 			}
 
 		private:
-			void BindInputs(FInputVertexInterfaceData& InOutVertexData) const
+			void BindInputsInternal(FInputVertexInterfaceData& InOutVertexData) const
 			{
 				InOutVertexData.BindWriteVertex(DataReferenceName, InputValue);
 			}
 
 		protected:
 
-			void BindOutputs(FOutputVertexInterfaceData& InOutVertexData) const
+			void BindOutputsInternal(FOutputVertexInterfaceData& InOutVertexData) const
 			{
 				InOutVertexData.BindReadVertex(DataReferenceName, OutputValue);
 			}
@@ -197,8 +197,8 @@ namespace Metasound
 
 			virtual void Bind(FVertexInterfaceData& InOutVertexData) const override
 			{
-				BindInputs(InOutVertexData.GetInputs());
-				BindOutputs(InOutVertexData.GetOutputs());
+				BindInputsInternal(InOutVertexData.GetInputs());
+				BindOutputsInternal(InOutVertexData.GetOutputs());
 			}
 
 			virtual FExecuteFunction GetExecuteFunction() override
@@ -217,14 +217,14 @@ namespace Metasound
 			}
 
 		private:
-			void BindInputs(FInputVertexInterfaceData& InOutVertexData) const
+			void BindInputsInternal(FInputVertexInterfaceData& InOutVertexData) const
 			{
 				InOutVertexData.BindWriteVertex(DataReferenceName, Value);
 			}
 
 		protected:
 
-			void BindOutputs(FOutputVertexInterfaceData& InOutVertexData) const
+			void BindOutputsInternal(FOutputVertexInterfaceData& InOutVertexData) const
 			{
 				InOutVertexData.BindReadVertex(DataReferenceName, Value);
 			}
diff --git a/Engine/Plugins/Runtime/Metasound/Source/MetasoundFrontend/Public/MetasoundLiteralNode.h b/Engine/Plugins/Runtime/Metasound/Source/MetasoundFrontend/Public/MetasoundLiteralNode.h
--- a/Engine/Plugins/Runtime/Metasound/Source/MetasoundFrontend/Public/MetasoundLiteralNode.h
+++ b/Engine/Plugins/Runtime/Metasound/Source/MetasoundFrontend/Public/MetasoundLiteralNode.h
@@ -128,7 +128,7 @@ namespace Metasound
 			virtual void Bind(FVertexInterfaceData& InVertexData) const override
 			{
 				// Do not bind inputs as there are no inputs exposed from the literal node. 
-				TInputOperator<DataType>::BindOutputs(InVertexData.GetOutputs());
+				TInputOperator<DataType>::BindOutputsInternal(InVertexData.GetOutputs());
 			}
 		};
 
diff --git a/Engine/Plugins/Runtime/Metasound/Source/MetasoundGenerator/Private/MetasoundGeneratorBuilder.cpp b/Engine/Plugins/Runtime/Metasound/Source/MetasoundGenerator/Private/MetasoundGeneratorBuilder.cpp
--- a/Engine/Plugins/Runtime/Metasound/Source/MetasoundGenerator/Private/MetasoundGeneratorBuilder.cpp
+++ b/Engine/Plugins/Runtime/Metasound/Source/MetasoundGenerator/Private/MetasoundGeneratorBuilder.cpp
@@ -155,7 +155,7 @@ namespace Metasound
 			}
 
 			// Set any remaining inputs to their default values.
-			for (const MetasoundVertexDataPrivate::TBinding<FInputDataVertex>& Binding : InputData)
+			for (const MetasoundVertexDataPrivate::FInputBinding& Binding : InputData)
 			{
 				// Only create data reference if something does not already exist. 
 				if (!Binding.IsBound())
@@ -218,7 +218,7 @@ namespace Metasound
 			FParameterSetterSortedMap ParameterSetters;
 			TMap<FName, FParameterPackSetter> ParameterPackSetters;
 			FInputVertexInterfaceData& GraphInputs = VertexData.GetInputs();
-			for (const MetasoundVertexDataPrivate::TBinding<FInputDataVertex>& Binding : GraphInputs)
+			for (const MetasoundVertexDataPrivate::FInputBinding& Binding : GraphInputs)
 			{
 				// Only assign inputs that are writable. 
 				if (EDataReferenceAccessType::Write == Binding.GetAccessType())
diff --git a/Engine/Plugins/Runtime/Metasound/Source/MetasoundGraphCore/Private/MetasoundGraphOperator.cpp b/Engine/Plugins/Runtime/Metasound/Source/MetasoundGraphCore/Private/MetasoundGraphOperator.cpp
--- a/Engine/Plugins/Runtime/Metasound/Source/MetasoundGraphCore/Private/MetasoundGraphOperator.cpp
+++ b/Engine/Plugins/Runtime/Metasound/Source/MetasoundGraphCore/Private/MetasoundGraphOperator.cpp
@@ -2,14 +2,81 @@
 
 #include "MetasoundGraphOperator.h"
 
-#include "CoreMinimal.h"
+#include "Containers/Array.h"
+#include "Containers/SortedMap.h"
 #include "MetasoundDataReference.h"
 #include "MetasoundDataReferenceCollection.h"
-#include "MetasoundOperatorInterface.h"
 #include "MetasoundExecutableOperator.h"
+#include "MetasoundOperatorInterface.h"
 
 namespace Metasound
 {
+	namespace MetasoundGraphPrivate
+	{
+		// FGraphOperator does not support rebinding with new inputs or outputs. This checks
+		// that underlying data pointers were not updated when bind is called on the graph
+		// operator.
+		//
+		// In order for FGraphOperator to support rebinding with new inputs, it would need
+		// to maintain an internal map of all connections in the graph in order to update
+		// internal operators appropriately. It does not hold onto this data for 
+		// performance reasons. 
+		template<typename InterfaceDataType>
+		bool IsSupportedVertexData(const InterfaceDataType& InCurrentData, const InterfaceDataType& InNewData)
+		{
+#if DO_CHECK
+			TArray<FVertexDataState> CurrentState;
+			GetVertexInterfaceDataState(InCurrentData, CurrentState);
+
+			TArray<FVertexDataState> NewState;
+			GetVertexInterfaceDataState(InNewData, NewState);
+
+			CurrentState.Sort();
+			NewState.Sort();
+
+			TArray<FVertexDataState>::TConstIterator CurrentIter = CurrentState.CreateConstIterator();
+			TArray<FVertexDataState>::TConstIterator NewIter = NewState.CreateConstIterator();
+
+			while (CurrentIter && NewIter)
+			{
+				if (NewIter->VertexName == CurrentIter->VertexName)
+				{
+					if ((NewIter->ID != nullptr) && (NewIter->ID != CurrentIter->ID))
+					{
+						UE_LOG(LogMetaSound, Warning, TEXT("Cannot bind to FGraphOperator because vertex %s has mismatched data"), *(NewIter->VertexName.ToString()));
+						return false;
+					}
+					else
+					{
+						NewIter++;
+						CurrentIter++;
+					}
+				}
+				else if (*NewIter < *CurrentIter)
+				{
+					UE_LOG(LogMetaSound, Warning, TEXT("Cannot bind to FGraphOperator because vertex %s does not exist in current vertex data"), *(NewIter->VertexName.ToString()));
+					return false;
+				}
+				else 
+				{
+					// It's ok if we have an entry in the current vertex data that does not exist in the new vertex data. 
+					CurrentIter++;
+				}
+			}
+
+			if (!NewIter)
+			{
+				UE_LOG(LogMetaSound, Warning, TEXT("Cannot bind to FGraphOperator because vertex %s does not exist in current vertex data"), *(NewIter->VertexName.ToString()));
+				return false;
+			}
+
+			return true;
+#else
+			return true;
+#endif // DO_CHECK
+		}
+	}
+
 	void FGraphOperator::AppendOperator(FOperatorPtr InOperator)
 	{
 		if (InOperator.IsValid())
@@ -41,16 +108,6 @@ namespace Metasound
 		}
 	}
 
-	void FGraphOperator::SetInputs(const FDataReferenceCollection& InCollection)
-	{
-		VertexData.GetInputs().Bind(InCollection);
-	}
-
-	void FGraphOperator::SetOutputs(const FDataReferenceCollection& InCollection)
-	{
-		VertexData.GetOutputs().Bind(InCollection);
-	}
-
 	void FGraphOperator::SetVertexInterfaceData(FVertexInterfaceData&& InVertexData)
 	{
 		VertexData = InVertexData;
@@ -66,9 +123,16 @@ namespace Metasound
 		return VertexData.GetOutputs().ToDataReferenceCollection();
 	}
 
-	void FGraphOperator::Bind(FVertexInterfaceData& InVertexData) const
+	void FGraphOperator::BindInputs(FInputVertexInterfaceData& InInputVertexData)
+	{
+		checkf(MetasoundGraphPrivate::IsSupportedVertexData(VertexData.GetInputs(), InInputVertexData), TEXT("FGraphOperator does not support rebinding with new data"));
+
+		InInputVertexData = VertexData.GetInputs();
+	}
+
+	void FGraphOperator::BindOutputs(FOutputVertexInterfaceData& InOutputVertexData)
 	{
-		InVertexData.Bind(VertexData);
+		InOutputVertexData = VertexData.GetOutputs();
 	}
 
 	IOperator::FPostExecuteFunction FGraphOperator::GetPostExecuteFunction()
@@ -150,6 +214,5 @@ namespace Metasound
 		Function(Operator, InParams);
 	}
 
-			static void StaticPostExecute(IOperator* Operator);
 
 }
diff --git a/Engine/Plugins/Runtime/Metasound/Source/MetasoundGraphCore/Private/MetasoundOperatorBuilder.cpp b/Engine/Plugins/Runtime/Metasound/Source/MetasoundGraphCore/Private/MetasoundOperatorBuilder.cpp
--- a/Engine/Plugins/Runtime/Metasound/Source/MetasoundGraphCore/Private/MetasoundOperatorBuilder.cpp
+++ b/Engine/Plugins/Runtime/Metasound/Source/MetasoundGraphCore/Private/MetasoundOperatorBuilder.cpp
@@ -257,7 +257,7 @@ namespace Metasound
 			{
 				check(DataReference->GetDataTypeName() == Edge->From.Vertex.DataTypeName);
 				check(DataReference->GetDataTypeName() == Edge->To.Vertex.DataTypeName);
-				OutVertexData.BindVertex(Edge->To.Vertex.VertexName, FAnyDataReference{*DataReference});
+				OutVertexData.SetVertex(Edge->To.Vertex.VertexName, FAnyDataReference{*DataReference});
 			}
 			else 
 			{
@@ -290,7 +290,7 @@ namespace Metasound
 				{
 					if (DataReference->GetDataTypeName() == Destination->Vertex.DataTypeName)
 					{
-						OutVertexData.BindVertex(Destination->Vertex.VertexName, FAnyDataReference{*DataReference});
+						OutVertexData.SetVertex(Destination->Vertex.VertexName, FAnyDataReference{*DataReference});
 					}
 					else
 					{
@@ -329,7 +329,7 @@ namespace Metasound
 
 	FOperatorBuilder::FBuildStatus FOperatorBuilder::ValidateOperatorOutputsAreBound(const INode& InNode, const FOutputVertexInterfaceData& InVertexData) const
 	{
-		using FOutputVertexBinding = MetasoundVertexDataPrivate::TBinding<FOutputDataVertex>;
+		using FOutputVertexBinding = MetasoundVertexDataPrivate::FOutputBinding;
 
 		FBuildStatus BuildStatus = FBuildStatus::NoError;
 		bool bFoundUnboundVertex = false;
@@ -448,7 +448,7 @@ namespace Metasound
 					if (DataReference->GetDataTypeName() == InputDestination.Vertex.DataTypeName)
 					{
 						bFoundDataReference = true;
-						OutVertexData.GetInputs().BindVertex(InputDestination.Vertex.VertexName, *DataReference);
+						OutVertexData.GetInputs().SetVertex(InputDestination.Vertex.VertexName, *DataReference);
 					}
 				}
 			}
@@ -474,7 +474,7 @@ namespace Metasound
 					if (DataReference->GetDataTypeName() == OutputSource.Vertex.DataTypeName)
 					{
 						bFoundDataReference = true;
-						OutVertexData.GetOutputs().BindVertex(OutputSource.Vertex.VertexName, FAnyDataReference{*DataReference});
+						OutVertexData.GetOutputs().SetVertex(OutputSource.Vertex.VertexName, FAnyDataReference{*DataReference});
 					}
 				}
 			}
diff --git a/Engine/Plugins/Runtime/Metasound/Source/MetasoundGraphCore/Private/MetasoundVertexData.cpp b/Engine/Plugins/Runtime/Metasound/Source/MetasoundGraphCore/Private/MetasoundVertexData.cpp
--- a/Engine/Plugins/Runtime/Metasound/Source/MetasoundGraphCore/Private/MetasoundVertexData.cpp
+++ b/Engine/Plugins/Runtime/Metasound/Source/MetasoundGraphCore/Private/MetasoundVertexData.cpp
@@ -10,6 +10,9 @@ namespace Metasound
 {
 	namespace MetasoundVertexDataPrivate
 	{
+		template<typename VertexType>
+		using TBindingType = std::conditional_t<std::is_same_v<VertexType, FInputDataVertex>, FInputBinding, FOutputBinding>;
+
 #if ENABLE_METASOUND_ACCESS_TYPE_COMPATIBILITY_TEST
 		void CheckAccessTypeCompatibility(const FDataVertex& InDataVertex, const FAnyDataReference& InDataReference)
 		{
@@ -37,9 +40,194 @@ namespace Metasound
 			checkf(bIsCompatible, TEXT("Vertex access type \"%s\" is incompatible with data access type \"%s\" on vertex \"%s\" on node \"%s\""), *LexToString(InDataVertex.AccessType), *LexToString(ReferenceAccessType), *InDataVertex.VertexName.ToString(), ThreadLocalDebug::GetActiveNodeClassNameAndVersion());
 		}
 #endif // #if ENABLE_METASOUND_ACCESS_TYPE_COMPATIBILITY_TEST
+		EVertexAccessType DataReferenceAccessTypeToVertexAccessType(EDataReferenceAccessType InReferenceAccessType)
+		{
+			switch (InReferenceAccessType)
+			{
+				case EDataReferenceAccessType::Read:
+				case EDataReferenceAccessType::Write:
+					return EVertexAccessType::Reference;
+
+				case EDataReferenceAccessType::Value:
+					return EVertexAccessType::Value;
+
+				default:
+					return EVertexAccessType::Reference;
+			}
+		}
+
+		FInputBinding::FInputBinding(FInputDataVertex&& InVertex)
+		: Vertex(MoveTemp(InVertex))
+		{
+		}
+
+		FInputBinding::FInputBinding(const FInputDataVertex& InVertex)
+		: Vertex(InVertex)
+		{
+		}
+
+		FInputBinding::FInputBinding(const FVertexName& InVertexName, FAnyDataReference&& InReference)
+		: Vertex(InVertexName, InReference.GetDataTypeName(), FDataVertexMetadata{}, DataReferenceAccessTypeToVertexAccessType(InReference.GetAccessType()))
+		{
+			Set(MoveTemp(InReference));
+		}
+
+
+
+		void FInputBinding::Set(FAnyDataReference&& InAnyDataReference)
+		{
+			check(Vertex.DataTypeName == InAnyDataReference.GetDataTypeName());
+			Data.Emplace(MoveTemp(InAnyDataReference));
+			CheckAccessTypeCompatibility(Vertex, *Data);
+		}
+		
+		const FInputDataVertex& FInputBinding::GetVertex() const
+		{
+			return Vertex;
+		}
+
+		bool FInputBinding::IsBound() const
+		{
+			return Data.IsSet();
+		}
+
+		EDataReferenceAccessType FInputBinding::GetAccessType() const
+		{
+			if (Data.IsSet())
+			{
+				return Data->GetAccessType();
+			}
+			return EDataReferenceAccessType::None;
+		}
+
+		// Get data reference 
+		const FAnyDataReference* FInputBinding::GetDataReference() const
+		{
+			return Data.GetPtrOrNull();
+		}
+
+		FDataReferenceID FInputBinding::GetDataReferenceID() const
+		{
+			if (Data.IsSet())
+			{
+				return Metasound::GetDataReferenceID(*Data);
+			}
+			return nullptr;
+		}
+
+
+		void FInputBinding::Bind(FAnyDataReference& InOutDataReference)
+		{
+			check(Vertex.DataTypeName == InOutDataReference.GetDataTypeName());
+			if (Data.IsSet())
+			{
+				InOutDataReference = *Data;
+			}
+			else
+			{
+				Data.Emplace(InOutDataReference);
+				CheckAccessTypeCompatibility(Vertex, *Data);
+			}
+		}
+
+		void FInputBinding::Bind(FInputBinding& InBinding)
+		{
+			check(Vertex.DataTypeName == InBinding.GetVertex().DataTypeName);
+
+			if (Data.IsSet()) 
+			{ 
+				if (InBinding.Data.IsSet())
+				{
+					*InBinding.Data = *Data;
+				}
+				else
+				{
+					InBinding.Data = Data;
+				}
+			}
+			else if (InBinding.Data.IsSet())
+			{
+				Data = InBinding.Data;
+				CheckAccessTypeCompatibility(Vertex, *Data);
+			}
+		}
+
+		FOutputBinding::FOutputBinding(FOutputDataVertex&& InVertex)
+		: Vertex(MoveTemp(InVertex))
+		{
+		}
+
+		FOutputBinding::FOutputBinding(const FOutputDataVertex& InVertex)
+		: Vertex(InVertex)
+		{
+		}
+
+		FOutputBinding::FOutputBinding(const FVertexName& InVertexName, FAnyDataReference&& InReference)
+		: Vertex(InVertexName, InReference.GetDataTypeName(), FDataVertexMetadata{}, DataReferenceAccessTypeToVertexAccessType(InReference.GetAccessType()))
+		{
+			Set(MoveTemp(InReference));
+		}
+
+		void FOutputBinding::Set(FAnyDataReference&& InAnyDataReference)
+		{
+			check(Vertex.DataTypeName == InAnyDataReference.GetDataTypeName());
+			Data.Emplace(MoveTemp(InAnyDataReference));
+			CheckAccessTypeCompatibility(Vertex, *Data);
+		}
+		
+		const FOutputDataVertex& FOutputBinding::GetVertex() const
+		{
+			return Vertex;
+		}
+
+		bool FOutputBinding::IsBound() const
+		{
+			return Data.IsSet();
+		}
+
+		EDataReferenceAccessType FOutputBinding::GetAccessType() const
+		{
+			if (Data.IsSet())
+			{
+				return Data->GetAccessType();
+			}
+			return EDataReferenceAccessType::None;
+		}
 
+		// Get data reference 
+		const FAnyDataReference* FOutputBinding::GetDataReference() const
+		{
+			return Data.GetPtrOrNull();
+		}
+
+		FDataReferenceID FOutputBinding::GetDataReferenceID() const
+		{
+			if (Data.IsSet())
+			{
+				return Metasound::GetDataReferenceID(*Data);
+			}
+			return nullptr;
+		}
+
+		void FOutputBinding::Bind(FAnyDataReference& InOutDataReference)
+		{
+			check(Vertex.DataTypeName == InOutDataReference.GetDataTypeName());
+			Data.Emplace(InOutDataReference);
+			CheckAccessTypeCompatibility(Vertex, *Data);
+		}
+
+		void FOutputBinding::Bind(FOutputBinding& InBinding)
+		{
+			check(Vertex.DataTypeName == InBinding.GetVertex().DataTypeName);
+
+			if (InBinding.Data.IsSet())
+			{
+				Data = InBinding.Data;
+				CheckAccessTypeCompatibility(Vertex, *Data);
+			}
+		}
 		template<typename VertexType>
-		void EmplaceBindings(TArray<TBinding<VertexType>>& InArray, const TVertexInterfaceGroup<VertexType>& InVertexInterface)
+		void EmplaceBindings(TArray<TBindingType<VertexType>>& InArray, const TVertexInterfaceGroup<VertexType>& InVertexInterface)
 		{
 			for (const VertexType& DataVertex : InVertexInterface)
 			{
@@ -68,19 +256,24 @@ namespace Metasound
 		}
 
 		template<typename BindingType>
-		void BindVertex(TArray<BindingType>& InBindings, const FVertexName& InVertexName, FAnyDataReference&& InDataReference)
+		void SetVertex(bool bIsVertexInterfaceFrozen, TArray<BindingType>& InBindings, const FVertexName& InVertexName, FAnyDataReference&& InDataReference)
 		{
 			if (BindingType* Binding = Find(InBindings, InVertexName))
 			{
 				if (Binding->GetVertex().DataTypeName == InDataReference.GetDataTypeName())
 				{
-					Binding->Bind(MoveTemp(InDataReference));
+					Binding->Set(MoveTemp(InDataReference));
 				}
 				else
 				{
 					UE_LOG(LogMetaSound, Warning, TEXT("Failed bind vertex with name '%s'. Supplied data type (%s) does not match vertex data type (%s)"), *InVertexName.ToString(), *InDataReference.GetDataTypeName().ToString(), *Binding->GetVertex().DataTypeName.ToString());
 				}
 			}
+			else if (!bIsVertexInterfaceFrozen)
+			{
+				BindingType NewBinding(InVertexName, MoveTemp(InDataReference));
+				InBindings.Add(MoveTemp(NewBinding));
+			}
 			else
 			{
 				UE_LOG(LogMetaSound, Warning, TEXT("Failed find vertex with name '%s'. Failed to bind vertex data"), *InVertexName.ToString());
@@ -88,31 +281,42 @@ namespace Metasound
 		}
 
 		template<typename BindingType>
-		void BindVertex(TArray<BindingType>& InBindings, const FVertexName& InVertexName, const FAnyDataReference& InDataReference)
+		void SetVertex(bool bIsVertexInterfaceFrozen, TArray<BindingType>& InBindings, const FVertexName& InVertexName, const FAnyDataReference& InDataReference)
 		{
-			BindVertex<BindingType>(InBindings, InVertexName, FAnyDataReference{InDataReference});
+			SetVertex<BindingType>(bIsVertexInterfaceFrozen, InBindings, InVertexName, FAnyDataReference{InDataReference});
 		}
 
 		template<typename BindingType>
-		void Bind(TArray<BindingType>& InTargetBindings, const TArray<BindingType>& InSourceBindings)
+		void Bind(bool bIsVertexInterfaceFrozen, TArray<BindingType>& InThisBindings, TArray<BindingType>& InOtherBindings)
 		{
-			for (const BindingType& SourceBinding : InSourceBindings)
+			for (BindingType& OtherBinding : InOtherBindings)
 			{
-				if (const FAnyDataReference* DataReference = SourceBinding.GetDataReference())
+				const FVertexName& OtherVertexName = OtherBinding.GetVertex().VertexName;
+				if (BindingType* ThisBinding = Find(InThisBindings, OtherVertexName))
 				{
-					if (BindingType* TargetBinding = Find(InTargetBindings, SourceBinding.GetVertex().VertexName))
+					const FName& OtherDataTypeName= OtherBinding.GetVertex().DataTypeName;
+					if (OtherDataTypeName == ThisBinding->GetVertex().DataTypeName)
 					{
-						if (DataReference->GetDataTypeName() == TargetBinding->GetVertex().DataTypeName)
-						{
-							TargetBinding->Bind(*DataReference);
-						}
+						ThisBinding->Bind(OtherBinding);
+					}
+					else
+					{
+						UE_LOG(LogMetaSound, Warning, TEXT("Failed bind vertex with name '%s'. Supplied data type (%s) does not match vertex data type (%s)"), *OtherVertexName.ToString(), *OtherDataTypeName.ToString(), *ThisBinding->GetVertex().DataTypeName.ToString());
 					}
 				}
+				else if (!bIsVertexInterfaceFrozen)
+				{
+					InThisBindings.Add(OtherBinding);
+				}
+				else
+				{
+					UE_LOG(LogMetaSound, Warning, TEXT("Failed find vertex with name '%s'. Failed to bind vertex data"), *OtherVertexName.ToString());
+				}
 			}
 		}
 
 		template<typename BindingType>
-		void Bind(TArray<BindingType>& InBindings, const FDataReferenceCollection& InCollection)
+		void Set(TArray<BindingType>& InBindings, const FDataReferenceCollection& InCollection)
 		{
 			for (BindingType& Binding : InBindings)
 			{
@@ -120,7 +324,7 @@ namespace Metasound
 
 				if (const FAnyDataReference* DataRef = InCollection.FindDataReference(Vertex.VertexName))
 				{
-					Binding.Bind(*DataRef);
+					Binding.Set(*DataRef);
 				}
 			}
 		}
@@ -141,7 +345,6 @@ namespace Metasound
 			return Algo::AllOf(InBindings, [](const BindingType& Binding) { return Binding.IsBound(); });
 		}
 
-
 		template<typename BindingType>
 		EDataReferenceAccessType GetVertexDataAccessType(const TArray<BindingType>& InBindings, const FVertexName& InVertexName)
 		{
@@ -175,32 +378,113 @@ namespace Metasound
 			}
 			return nullptr;
 		}
+
+		template<typename BindingType>
+		void GetVertexInterfaceDataState(const TArray<BindingType>& InBindings, TArray<FVertexDataState>& OutState)
+		{
+			OutState.Reset();
+
+			for (const BindingType& InBinding : InBindings)
+			{
+				OutState.Add(FVertexDataState{InBinding.GetVertex().VertexName, InBinding.GetDataReferenceID()});
+			}
+		}
+
+		template<typename BindingType>
+		void CompareVertexInterfaceDataToPriorState(const TArray<BindingType>& InBindings, const TArray<FVertexDataState>& InPriorState, TSortedMap<FDataReferenceID, FAnyDataReference>& OutUpdates)
+		{
+			for (const BindingType& Binding : InBindings)
+			{
+				if (const FAnyDataReference* CurrentReference = Binding.GetDataReference())
+				{
+					const FVertexDataState* OtherState = nullptr;
+					const FName& VertexName = Binding.GetVertex().VertexName;
+
+					for (const FVertexDataState& State : InPriorState)
+					{
+						if (State.VertexName == VertexName)
+						{
+							OtherState = &State;
+							break;
+						}
+					}
+
+					if (OtherState)
+					{
+						if (OtherState->ID != Binding.GetDataReferenceID())
+						{
+							OutUpdates.Add(OtherState->ID, *CurrentReference);
+						}
+					}
+				}
+			}
+		}
+	}
+
+	FInputVertexInterfaceData::FInputVertexInterfaceData()
+	: bIsVertexInterfaceFrozen(false)
+	{
 	}
 
 	FInputVertexInterfaceData::FInputVertexInterfaceData(const FInputVertexInterface& InVertexInterface)
+	: bIsVertexInterfaceFrozen(true)
 	{
 		MetasoundVertexDataPrivate::EmplaceBindings(Bindings, InVertexInterface);
 	}
 
-	
-	void FInputVertexInterfaceData::BindVertex(const FVertexName& InVertexName, FAnyDataReference&& InDataReference)
+	bool FInputVertexInterfaceData::IsVertexInterfaceFrozen() const
 	{
-		MetasoundVertexDataPrivate::BindVertex(Bindings, InVertexName, MoveTemp(InDataReference));
+		return bIsVertexInterfaceFrozen;
 	}
 
-	void FInputVertexInterfaceData::BindVertex(const FVertexName& InVertexName, const FAnyDataReference& InDataReference)
+	void FInputVertexInterfaceData::SetIsVertexInterfaceFrozen(bool bInFreezeVertices)
 	{
-		MetasoundVertexDataPrivate::BindVertex(Bindings, InVertexName, InDataReference);
+		bIsVertexInterfaceFrozen = bInFreezeVertices;
+	}
+
+	void FInputVertexInterfaceData::BindVertex(const FVertexName& InVertexName, const FAnyDataReference& InOutDataReference)
+	{
+		BindVertex(InVertexName, const_cast<FAnyDataReference&>(InOutDataReference));
+	}
+
+	void FInputVertexInterfaceData::BindVertex(const FVertexName& InVertexName, FAnyDataReference& InOutDataReference)
+	{
+		using namespace MetasoundVertexDataPrivate;
+		auto CreateBinding = [&]()
+		{
+			FInputDataVertex Vertex(InVertexName, InOutDataReference.GetDataTypeName(), FDataVertexMetadata{}, DataReferenceAccessTypeToVertexAccessType(InOutDataReference.GetAccessType()));
+
+			return FInputBinding(Vertex);
+		};
+
+		auto BindData = [&](FInputBinding& Binding) { Binding.Bind(InOutDataReference); };
+		
+		Apply(InVertexName, CreateBinding, BindData);
 	}
 
 	void FInputVertexInterfaceData::Bind(const FInputVertexInterfaceData& InVertexData)
 	{
-		MetasoundVertexDataPrivate::Bind(Bindings, InVertexData.Bindings);
+		Bind(const_cast<FInputVertexInterfaceData&>(InVertexData));
+	}
+
+	void FInputVertexInterfaceData::Bind(FInputVertexInterfaceData& InVertexData)
+	{
+		MetasoundVertexDataPrivate::Bind(IsVertexInterfaceFrozen(), Bindings, InVertexData.Bindings);
+	}
+
+	void FInputVertexInterfaceData::SetVertex(const FVertexName& InVertexName, FAnyDataReference&& InDataReference)
+	{
+		MetasoundVertexDataPrivate::SetVertex(IsVertexInterfaceFrozen(), Bindings, InVertexName, MoveTemp(InDataReference));
 	}
 
-	void FInputVertexInterfaceData::Bind(const FDataReferenceCollection& InCollection)
+	void FInputVertexInterfaceData::SetVertex(const FVertexName& InVertexName, const FAnyDataReference& InDataReference)
 	{
-		MetasoundVertexDataPrivate::Bind(Bindings, InCollection);
+		MetasoundVertexDataPrivate::SetVertex(IsVertexInterfaceFrozen(), Bindings, InVertexName, InDataReference);
+	}
+
+	void FInputVertexInterfaceData::Set(const FDataReferenceCollection& InCollection)
+	{
+		MetasoundVertexDataPrivate::Set(Bindings, InCollection);
 	}
 
 	FDataReferenceCollection FInputVertexInterfaceData::ToDataReferenceCollection() const
@@ -233,55 +517,112 @@ namespace Metasound
 		return MetasoundVertexDataPrivate::FindDataReference(Bindings, InVertexName);
 	}
 
-	FInputVertexInterfaceData::FBinding* FInputVertexInterfaceData::Find(const FVertexName& InVertexName)
+	void FInputVertexInterfaceData::Apply(const FVertexName& InVertexName, TFunctionRef<FInputBinding ()> InCreateFunc, TFunctionRef<void (FInputBinding&)> InBindFunc)
+	{
+		if (FInputBinding* Binding = Find(InVertexName))
+		{
+			InBindFunc(*Binding);
+		}
+		else if (IsVertexInterfaceFrozen())
+		{
+			UE_LOG(LogMetaSound, Warning, TEXT("Failed to find input vertex with name '%s'."), *InVertexName.ToString());
+		}
+		else
+		{
+			FInputBinding NewBinding = InCreateFunc();
+			InBindFunc(NewBinding);
+			Bindings.Add(MoveTemp(NewBinding));
+		}
+	}
+
+	FInputVertexInterfaceData::FInputBinding* FInputVertexInterfaceData::Find(const FVertexName& InVertexName)
 	{
 		return MetasoundVertexDataPrivate::Find(Bindings, InVertexName);
 	}
 
-	const FInputVertexInterfaceData::FBinding* FInputVertexInterfaceData::Find(const FVertexName& InVertexName) const
+	const FInputVertexInterfaceData::FInputBinding* FInputVertexInterfaceData::Find(const FVertexName& InVertexName) const
 	{
 		return MetasoundVertexDataPrivate::Find(Bindings, InVertexName);
 	}
 
-	FInputVertexInterfaceData::FBinding* FInputVertexInterfaceData::FindChecked(const FVertexName& InVertexName)
+	FInputVertexInterfaceData::FInputBinding* FInputVertexInterfaceData::FindChecked(const FVertexName& InVertexName)
 	{
-		FBinding* Binding = Find(InVertexName);
+		FInputBinding* Binding = Find(InVertexName);
 		checkf(nullptr != Binding, TEXT("Attempt to access vertex \"%s\" which does not exist on interface."), *InVertexName.ToString());
 		return Binding;
 	}
 
-	const FInputVertexInterfaceData::FBinding* FInputVertexInterfaceData::FindChecked(const FVertexName& InVertexName) const
+	const FInputVertexInterfaceData::FInputBinding* FInputVertexInterfaceData::FindChecked(const FVertexName& InVertexName) const
 	{
-		const FBinding* Binding = Find(InVertexName);
+		const FInputBinding* Binding = Find(InVertexName);
 		checkf(nullptr != Binding, TEXT("Attempt to access vertex \"%s\" which does not exist on interface."), *InVertexName.ToString());
 		return Binding;
 	}
 
+	FOutputVertexInterfaceData::FOutputVertexInterfaceData()
+	: bIsVertexInterfaceFrozen(false)
+	{
+	}
 
 	FOutputVertexInterfaceData::FOutputVertexInterfaceData(const FOutputVertexInterface& InVertexInterface)
+	: bIsVertexInterfaceFrozen(true)
 	{
 		MetasoundVertexDataPrivate::EmplaceBindings(Bindings, InVertexInterface);
 	}
 
-	
-	void FOutputVertexInterfaceData::BindVertex(const FVertexName& InVertexName, FAnyDataReference&& InDataReference)
+	bool FOutputVertexInterfaceData::IsVertexInterfaceFrozen() const
 	{
-		MetasoundVertexDataPrivate::BindVertex(Bindings, InVertexName, MoveTemp(InDataReference));
+		return bIsVertexInterfaceFrozen;
 	}
 
-	void FOutputVertexInterfaceData::BindVertex(const FVertexName& InVertexName, const FAnyDataReference& InDataReference)
+	void FOutputVertexInterfaceData::SetIsVertexInterfaceFrozen(bool bInFreezeVertices)
 	{
-		MetasoundVertexDataPrivate::BindVertex(Bindings, InVertexName, InDataReference);
+		bIsVertexInterfaceFrozen = bInFreezeVertices;
+	}
+
+	void FOutputVertexInterfaceData::BindVertex(const FVertexName& InVertexName, const FAnyDataReference& InOutDataReference)
+	{
+		BindVertex(InVertexName, const_cast<FAnyDataReference&>(InOutDataReference));
+	}
+
+	void FOutputVertexInterfaceData::BindVertex(const FVertexName& InVertexName, FAnyDataReference& InOutDataReference)
+	{
+		using namespace MetasoundVertexDataPrivate;
+		auto CreateBinding = [&]()
+		{
+			FOutputDataVertex Vertex(InVertexName, InOutDataReference.GetDataTypeName(), FDataVertexMetadata{}, DataReferenceAccessTypeToVertexAccessType(InOutDataReference.GetAccessType()));
+
+			return FOutputBinding(Vertex);
+		};
+
+		auto BindData = [&](FOutputBinding& Binding) { Binding.Bind(InOutDataReference); };
+		
+		Apply(InVertexName, CreateBinding, BindData);
 	}
 
 	void FOutputVertexInterfaceData::Bind(const FOutputVertexInterfaceData& InVertexData)
 	{
-		MetasoundVertexDataPrivate::Bind(Bindings, InVertexData.Bindings);
+		Bind(const_cast<FOutputVertexInterfaceData&>(InVertexData));
+	}
+
+	void FOutputVertexInterfaceData::Bind(FOutputVertexInterfaceData& InVertexData)
+	{
+		MetasoundVertexDataPrivate::Bind(IsVertexInterfaceFrozen(), Bindings, InVertexData.Bindings);
+	}
+
+	void FOutputVertexInterfaceData::SetVertex(const FVertexName& InVertexName, FAnyDataReference&& InDataReference)
+	{
+		MetasoundVertexDataPrivate::SetVertex(IsVertexInterfaceFrozen(), Bindings, InVertexName, MoveTemp(InDataReference));
 	}
 
-	void FOutputVertexInterfaceData::Bind(const FDataReferenceCollection& InCollection)
+	void FOutputVertexInterfaceData::SetVertex(const FVertexName& InVertexName, const FAnyDataReference& InDataReference)
 	{
-		MetasoundVertexDataPrivate::Bind(Bindings, InCollection);
+		MetasoundVertexDataPrivate::SetVertex(IsVertexInterfaceFrozen(), Bindings, InVertexName, InDataReference);
+	}
+
+	void FOutputVertexInterfaceData::Set(const FDataReferenceCollection& InCollection)
+	{
+		MetasoundVertexDataPrivate::Set(Bindings, InCollection);
 	}
 
 	FDataReferenceCollection FOutputVertexInterfaceData::ToDataReferenceCollection() const
@@ -314,34 +655,103 @@ namespace Metasound
 		return MetasoundVertexDataPrivate::FindDataReference(Bindings, InVertexName);
 	}
 
+	void FOutputVertexInterfaceData::Apply(const FVertexName& InVertexName, TFunctionRef<FOutputBinding ()> InCreateFunc, TFunctionRef<void (FOutputBinding&)> InBindFunc)
+	{
+		if (FOutputBinding* Binding = Find(InVertexName))
+		{
+			InBindFunc(*Binding);
+		}
+		else if (!IsVertexInterfaceFrozen())
+		{
+			FOutputBinding NewBinding = InCreateFunc();
+			InBindFunc(NewBinding);
+			Bindings.Add(MoveTemp(NewBinding));
+		}
+		else
+		{
+			UE_LOG(LogMetaSound, Warning, TEXT("Failed to find output vertex with name '%s'."), *InVertexName.ToString());
+		}
+	}
 
-	FOutputVertexInterfaceData::FBinding* FOutputVertexInterfaceData::Find(const FVertexName& InVertexName)
+	FOutputVertexInterfaceData::FOutputBinding* FOutputVertexInterfaceData::Find(const FVertexName& InVertexName)
 	{
 		return MetasoundVertexDataPrivate::Find(Bindings, InVertexName);
 	}
 
-	const FOutputVertexInterfaceData::FBinding* FOutputVertexInterfaceData::Find(const FVertexName& InVertexName) const
+	const FOutputVertexInterfaceData::FOutputBinding* FOutputVertexInterfaceData::Find(const FVertexName& InVertexName) const
 	{
 		return MetasoundVertexDataPrivate::Find(Bindings, InVertexName);
 	}
 
-	FOutputVertexInterfaceData::FBinding* FOutputVertexInterfaceData::FindChecked(const FVertexName& InVertexName)
+	FOutputVertexInterfaceData::FOutputBinding* FOutputVertexInterfaceData::FindChecked(const FVertexName& InVertexName)
 	{
-		FBinding* Binding = Find(InVertexName);
+		FOutputBinding* Binding = Find(InVertexName);
 		checkf(nullptr != Binding, TEXT("Attempt to access vertex \"%s\" which does not exist on interface."), *InVertexName.ToString());
 		return Binding;
 	}
 
-	const FOutputVertexInterfaceData::FBinding* FOutputVertexInterfaceData::FindChecked(const FVertexName& InVertexName) const
+	const FOutputVertexInterfaceData::FOutputBinding* FOutputVertexInterfaceData::FindChecked(const FVertexName& InVertexName) const
 	{
-		const FBinding* Binding = Find(InVertexName);
+		const FOutputBinding* Binding = Find(InVertexName);
 		checkf(nullptr != Binding, TEXT("Attempt to access vertex \"%s\" which does not exist on interface."), *InVertexName.ToString());
 		return Binding;
 	}
 
+	FVertexInterfaceData::FVertexInterfaceData(const FVertexInterface& InVertexInterface)
+	: InputVertexInterfaceData(InVertexInterface.GetInputInterface())
+	, OutputVertexInterfaceData(InVertexInterface.GetOutputInterface())
+	{
+	}
+
 	void FVertexInterfaceData::Bind(const FVertexInterfaceData& InVertexData)
+	{
+		Bind(const_cast<FVertexInterfaceData&>(InVertexData));
+	}
+
+	void FVertexInterfaceData::Bind(FVertexInterfaceData& InVertexData)
 	{
 		InputVertexInterfaceData.Bind(InVertexData.GetInputs());
 		OutputVertexInterfaceData.Bind(InVertexData.GetOutputs());
 	}
+
+	bool operator<(const FVertexDataState& InLHS, const FVertexDataState& InRHS)
+	{
+		if (InLHS.VertexName.FastLess(InRHS.VertexName))
+		{
+			return true;
+		}
+		else if (InRHS.VertexName.FastLess(InLHS.VertexName))
+		{
+			return false;
+		}
+		else
+		{
+			return InLHS.ID < InRHS.ID;
+		}
+	}
+
+	bool operator==(const FVertexDataState& InLHS, const FVertexDataState& InRHS)
+	{
+		return (InLHS.ID == InRHS.ID) && (InLHS.VertexName == InRHS.VertexName);
+	}
+
+	void GetVertexInterfaceDataState(const FInputVertexInterfaceData& InVertexInterface, TArray<FVertexDataState>& OutState)
+	{
+		MetasoundVertexDataPrivate::GetVertexInterfaceDataState(InVertexInterface.Bindings, OutState);
+	}
+
+	void GetVertexInterfaceDataState(const FOutputVertexInterfaceData& InVertexInterface, TArray<FVertexDataState>& OutState)
+	{
+		MetasoundVertexDataPrivate::GetVertexInterfaceDataState(InVertexInterface.Bindings, OutState);
+	}
+
+	void CompareVertexInterfaceDataToPriorState(const FInputVertexInterfaceData& InCurrentInterface, const TArray<FVertexDataState>& InPriorState, TSortedMap<FDataReferenceID, FAnyDataReference>& OutUpdates)
+	{
+		MetasoundVertexDataPrivate::CompareVertexInterfaceDataToPriorState(InCurrentInterface.Bindings, InPriorState, OutUpdates);
+	}
+
+	void CompareVertexInterfaceDataToPriorState(const FOutputVertexInterfaceData& InCurrentInterface, const TArray<FVertexDataState>& InPriorState, TSortedMap<FDataReferenceID, FAnyDataReference>& OutUpdates)
+	{
+		MetasoundVertexDataPrivate::CompareVertexInterfaceDataToPriorState(InCurrentInterface.Bindings, InPriorState, OutUpdates);
+	}
 }
diff --git a/Engine/Plugins/Runtime/Metasound/Source/MetasoundGraphCore/Public/MetasoundBuilderInterface.h b/Engine/Plugins/Runtime/Metasound/Source/MetasoundGraphCore/Public/MetasoundBuilderInterface.h
--- a/Engine/Plugins/Runtime/Metasound/Source/MetasoundGraphCore/Public/MetasoundBuilderInterface.h
+++ b/Engine/Plugins/Runtime/Metasound/Source/MetasoundGraphCore/Public/MetasoundBuilderInterface.h
@@ -176,25 +176,6 @@ namespace Metasound
 			virtual ~IOperatorFactory() = default;
 
 			virtual TUniquePtr<IOperator> CreateOperator(const FBuildOperatorParams& InParams, FBuildResults& OutResults) = 0;
-
-			UE_DEPRECATED(5.1, "Use CreateOperator overload providing the FBuildResults struct.")
-			virtual TUniquePtr<IOperator> CreateOperator(const FCreateOperatorParams& InParams, FBuildErrorArray& OutErrors)
-			{
-				// Map inputs to new structures
-				FInputVertexInterfaceData InputData(InParams.Node.GetVertexInterface().GetInputInterface());
-				InputData.Bind(InParams.InputDataReferences);
-
-				FBuildOperatorParams Params2 { InParams.Node, InParams.OperatorSettings, InputData, InParams.Environment, InParams.Builder };
-
-				// Create operator
-				FBuildResults Results;
-				TUniquePtr<IOperator> Operator = CreateOperator(Params2, Results);
-
-				// Map new output to old output
-				OutErrors = MoveTemp(Results.Errors);
-
-				return MoveTemp(Operator);
-			}
 	};
 
 
@@ -210,28 +191,6 @@ namespace Metasound
 
 			virtual ~IOperatorBuilder() = default;
 
-			/** Build a graph operator from a graph.
-			 *
-			 * @params InParams - Input parameters for building a graph.
-			 * @param OutErrors - An array of errors which occurred during while building the operator.
-			 *
-			 * @return A unique pointer to the built IOperator. Null if build failed.
-			 */
-			UE_DEPRECATED(5.1, "Use BuildGraphOperator overload providing the FBuildResults struct.")
-			virtual TUniquePtr<IOperator> BuildGraphOperator(const FBuildGraphParams& InParams, FBuildErrorArray& OutErrors) const 
-			{ 
-				FInputVertexInterfaceData InputData(InParams.Graph.GetVertexInterface().GetInputInterface());
-				InputData.Bind(InParams.InputDataReferences);
-				FBuildGraphOperatorParams NewParams{InParams.Graph, InParams.OperatorSettings, InputData, InParams.Environment};
-				FBuildResults Results;
-
-				TUniquePtr<IOperator> GraphOperator = BuildGraphOperator(NewParams, Results);
-
-				OutErrors.Append(MoveTemp(Results.Errors));
-				
-				return MoveTemp(GraphOperator);
-			}
-
 			/** Build a graph operator from a graph.
 			 *
 			 * @params InParams - Input parameters for building a graph.
diff --git a/Engine/Plugins/Runtime/Metasound/Source/MetasoundGraphCore/Public/MetasoundDataReference.h b/Engine/Plugins/Runtime/Metasound/Source/MetasoundGraphCore/Public/MetasoundDataReference.h
--- a/Engine/Plugins/Runtime/Metasound/Source/MetasoundGraphCore/Public/MetasoundDataReference.h
+++ b/Engine/Plugins/Runtime/Metasound/Source/MetasoundGraphCore/Public/MetasoundDataReference.h
@@ -31,6 +31,9 @@ using FMetasoundDataTypeId = void const*;
 
 namespace Metasound
 {
+	// Unique ID type which corresponds to the underlying object referred to by a data reference.
+	using FDataReferenceID = const void*;
+
 	/** Helper class to enforce specialization of TDataReferenceTypeInfo */
 	template<typename DataType>
 	struct TSpecializationHelper 
@@ -127,6 +130,12 @@ namespace Metasound
 		/** provides a raw pointer to the storage where the data actually resides. */
 		virtual void* GetRaw() const = 0;
 	};
+
+	/** Return the ID of the data reference. */
+	FORCEINLINE FDataReferenceID GetDataReferenceID(const IDataReference& InDataReference)
+	{
+		return static_cast<FDataReferenceID>(InDataReference.GetRaw());
+	}
 	
 	/** Test if an IDataReference contains the same data type as the template
 	 * parameter.
diff --git a/Engine/Plugins/Runtime/Metasound/Source/MetasoundGraphCore/Public/MetasoundGraphOperator.h b/Engine/Plugins/Runtime/Metasound/Source/MetasoundGraphCore/Public/MetasoundGraphOperator.h
--- a/Engine/Plugins/Runtime/Metasound/Source/MetasoundGraphCore/Public/MetasoundGraphOperator.h
+++ b/Engine/Plugins/Runtime/Metasound/Source/MetasoundGraphCore/Public/MetasoundGraphOperator.h
@@ -25,12 +25,6 @@ namespace Metasound
 			// Add an operator to the end of the executation stack.
 			void AppendOperator(FOperatorPtr InOperator);
 
-			UE_DEPRECATED(5.1, "Use FGraphOperator::SetVertexInterfaceData instead.")
-			void SetInputs(const FDataReferenceCollection& InCollection);
-
-			UE_DEPRECATED(5.1, "Use FGraphOperator::SetVertexInterfaceData instead.")
-			void SetOutputs(const FDataReferenceCollection& InCollection);
-
 			// Set the vertex interface data. This data will be copied to output 
 			// during calls to Bind(InOutVertexData).
 			void SetVertexInterfaceData(FVertexInterfaceData&& InVertexData);
@@ -40,7 +34,8 @@ namespace Metasound
 			virtual FDataReferenceCollection GetOutputs() const override;
 
 			// Bind the graph's interface data references to FVertexInterfaceData.
-			virtual void Bind(FVertexInterfaceData& InOutVertexData) const override;
+			virtual void BindInputs(FInputVertexInterfaceData& InOutVertexData) override;
+			virtual void BindOutputs(FOutputVertexInterfaceData& InOutVertexData) override;
 
 			virtual IOperator::FPostExecuteFunction GetPostExecuteFunction() override;
 
diff --git a/Engine/Plugins/Runtime/Metasound/Source/MetasoundGraphCore/Public/MetasoundOperatorInterface.h b/Engine/Plugins/Runtime/Metasound/Source/MetasoundGraphCore/Public/MetasoundOperatorInterface.h
--- a/Engine/Plugins/Runtime/Metasound/Source/MetasoundGraphCore/Public/MetasoundOperatorInterface.h
+++ b/Engine/Plugins/Runtime/Metasound/Source/MetasoundGraphCore/Public/MetasoundOperatorInterface.h
@@ -55,11 +55,25 @@ namespace Metasound
 		 */
 		virtual void Bind(FVertexInterfaceData& InVertexData) const
 		{
-			InVertexData.GetInputs().Bind(GetInputs());
-			InVertexData.GetOutputs().Bind(GetOutputs());
+			InVertexData.GetInputs().Set(GetInputs());
+			InVertexData.GetOutputs().Set(GetOutputs());
 		}
 
-		/** Pointer to intialize function for an operator.
+		virtual void BindInputs(FInputVertexInterfaceData& InVertexData)
+		{
+			FVertexInterfaceData Data;
+			Data.GetInputs() = InVertexData;
+			Bind(Data);
+		}
+
+		virtual void BindOutputs(FOutputVertexInterfaceData& InVertexData) 
+		{
+			FVertexInterfaceData Data;
+			Data.GetOutputs() = InVertexData;
+			Bind(Data);
+		}
+
+		/** Pointer to initialize function for an operator.
 		 *
 		 * @param IOperator* - The operator associated with the function pointer.
 		 */
diff --git a/Engine/Plugins/Runtime/Metasound/Source/MetasoundGraphCore/Public/MetasoundVertexData.h b/Engine/Plugins/Runtime/Metasound/Source/MetasoundGraphCore/Public/MetasoundVertexData.h
--- a/Engine/Plugins/Runtime/Metasound/Source/MetasoundGraphCore/Public/MetasoundVertexData.h
+++ b/Engine/Plugins/Runtime/Metasound/Source/MetasoundGraphCore/Public/MetasoundVertexData.h
@@ -4,6 +4,7 @@
 
 #include "Algo/AllOf.h"
 #include "Containers/Array.h"
+#include "Containers/SortedMap.h"
 #include "MetasoundDataFactory.h"
 #include "MetasoundDataReference.h"
 #include "MetasoundDataReferenceCollection.h"
@@ -21,9 +22,6 @@ namespace Metasound
 {
 	namespace MetasoundVertexDataPrivate
 	{
-		template<typename V>
-		using EnableIfIsInputDataVertex = std::enable_if_t<std::is_same_v<FInputDataVertex, V>>;
-
 		// Tests to see that the access type of the data reference is compatible 
 		// with the access type of the vertex.
 #if ENABLE_METASOUND_ACCESS_TYPE_COMPATIBILITY_TEST
@@ -32,62 +30,106 @@ namespace Metasound
 		FORCEINLINE void CheckAccessTypeCompatibility(const FDataVertex& InDataVertex, const FAnyDataReference& InDataReference) {}
 #endif // #if DO_CHECK
 
-		// Binds a vertex to a data reference.
-		template<typename VertexType>
-		class TBinding
+		METASOUNDGRAPHCORE_API EVertexAccessType DataReferenceAccessTypeToVertexAccessType(EDataReferenceAccessType InReferenceAccessType);
+
+		// Helper for getting a EVertexAccessType from a TData*Reference<DataType> object
+		template<typename ...>
+		struct TGetVertexAccess
 		{
+		};
 
-		public:
-			TBinding(VertexType&& InVertex)
-			: Vertex(MoveTemp(InVertex))
-			{
-			}
+		template<typename DataType>
+		struct TGetVertexAccess<TDataReadReference<DataType>>
+		{
+			static constexpr EVertexAccessType VertexAccess = EVertexAccessType::Reference;
+		};
 
-			TBinding(const VertexType& InVertex)
-			: Vertex(InVertex)
-			{
-			}
+		template<typename DataType>
+		struct TGetVertexAccess<TDataWriteReference<DataType>>
+		{
+			static constexpr EVertexAccessType VertexAccess = EVertexAccessType::Reference;
+		};
 
-			template<typename DataReferenceType>
-			void Bind(const DataReferenceType& InDataReference)
-			{
-				check(Vertex.DataTypeName == InDataReference.GetDataTypeName());
-				Data.Emplace(InDataReference);
-				CheckAccessTypeCompatibility(Vertex, *Data);
-			}
+		template<typename DataType>
+		struct TGetVertexAccess<TDataValueReference<DataType>>
+		{
+			static constexpr EVertexAccessType VertexAccess = EVertexAccessType::Value;
+		};
+		
 
-			void Bind(FAnyDataReference&& InAnyDataReference)
-			{
-				check(Vertex.DataTypeName == InAnyDataReference.GetDataTypeName());
-				Data.Emplace(MoveTemp(InAnyDataReference));
-				CheckAccessTypeCompatibility(Vertex, *Data);
-			}
-			
-			const VertexType& GetVertex() const
+		// An input binding which connects an FInputVertex to a IDataReference
+		class METASOUNDGRAPHCORE_API FInputBinding
+		{
+		public:
+
+			FInputBinding(FInputDataVertex&& InVertex);
+			FInputBinding(const FInputDataVertex& InVertex);
+			FInputBinding(const FVertexName& InVertexName, FAnyDataReference&& InReference);
+
+			template<typename DataType>
+			void BindRead(TDataReadReference<DataType>& InOutDataReference)
 			{
-				return Vertex;
+				if (Data.IsSet())
+				{
+					InOutDataReference = Data->GetDataReadReference<DataType>();
+				}
+				else
+				{
+					Set(InOutDataReference);
+				}
 			}
 
-			bool IsBound() const
+			template<typename DataType>
+			void BindRead(TDataWriteReference<DataType>& InOutDataReference)
 			{
-				return Data.IsSet();
+				if (Data.IsSet())
+				{
+					InOutDataReference = Data->GetDataWriteReference<DataType>();
+				}
+				else
+				{
+					Set(TDataReadReference<DataType>(InOutDataReference));
+				}
 			}
 
-			EDataReferenceAccessType GetAccessType() const
+			template<typename DataType>
+			void BindWrite(TDataWriteReference<DataType>& InOutDataReference)
 			{
 				if (Data.IsSet())
 				{
-					return Data->GetAccessType();
+					InOutDataReference = Data->GetDataWriteReference<DataType>();
+				}
+				else
+				{
+					Set(InOutDataReference);
 				}
-				return EDataReferenceAccessType::None;
 			}
 
-			// Get data reference 
-			const FAnyDataReference* GetDataReference() const
+			void Bind(FAnyDataReference& InOutDataReference);
+			void Bind(FInputBinding& InBinding);
+
+			// Set the data reference, overwriting any existing bound data references
+			template<typename DataReferenceType>
+			void Set(const DataReferenceType& InDataReference)
 			{
-				return Data.GetPtrOrNull();
+				check(Vertex.DataTypeName == InDataReference.GetDataTypeName());
+				Data.Emplace(InDataReference);
+				CheckAccessTypeCompatibility(Vertex, *Data);
 			}
 
+			// Set the data reference, overwriting any existing bound data references
+			void Set(FAnyDataReference&& InAnyDataReference);
+
+			const FInputDataVertex& GetVertex() const;
+
+			bool IsBound() const;
+
+			EDataReferenceAccessType GetAccessType() const;
+
+			const FAnyDataReference* GetDataReference() const;
+
+			FDataReferenceID GetDataReferenceID() const;
+
 			// Get data read reference assuming data is bound and read or write accessible.
 			template<typename DataType>
 			TDataReadReference<DataType> GetDataReadReference() const
@@ -99,11 +141,7 @@ namespace Metasound
 			// Get the bound data read reference if it exists. Otherwise create and 
 			// return a data read reference by constructing one using the Vertex's 
 			// default literal.
-			template<
-				typename DataType,
-				typename V = VertexType,
-				typename = EnableIfIsInputDataVertex<V>
-			> 
+			template< typename DataType > 
 			TDataReadReference<DataType> GetOrCreateDefaultDataReadReference(const FOperatorSettings& InSettings) const
 			{
 				if (Data.IsSet())
@@ -141,11 +179,7 @@ namespace Metasound
 			// Get the bound data write reference if it exists. Otherwise create and 
 			// return a data write reference by constructing one using the Vertex's 
 			// default literal.
-			template<
-				typename DataType,
-				typename V = VertexType,
-				typename = EnableIfIsInputDataVertex<V>
-			> 
+			template< typename DataType > 
 			TDataWriteReference<DataType> GetOrCreateDefaultDataWriteReference(const FOperatorSettings& InSettings) const
 			{
 				if (Data.IsSet())
@@ -180,6 +214,7 @@ namespace Metasound
 				return nullptr;
 			}
 
+
 			// Gets the value of the bound data reference if it exists. Otherwise
 			// create and return a value by constructing one using the Vertex's 
 			// default literal.
@@ -196,6 +231,153 @@ namespace Metasound
 				return TDataTypeLiteralFactory<DataType>::CreateExplicitArgs(InSettings, Vertex.GetDefaultLiteral());
 			}
 
+			// Set the value with a constant value reference.
+			template<typename DataType>
+			void SetValue(const DataType& InValue)
+			{
+				check(Vertex.DataTypeName == GetMetasoundDataTypeName<DataType>());
+				Data.Emplace(TDataValueReference<DataType>::CreateNew(InValue));
+				CheckAccessTypeCompatibility(Vertex, *Data);
+			}
+
+		private:
+
+			FInputDataVertex Vertex;
+			TOptional<FAnyDataReference> Data;
+		};
+
+
+		// Factory for creating a new FInputBinding from a IDataReference derived class and a vertex name.
+		struct FInputBindingFactory
+		{
+			template<typename DataReferenceType>
+			static FInputBinding CreateBinding(const FVertexName& InVertexName, const DataReferenceType& InRef)
+			{
+				FInputDataVertex Vertex(InVertexName, InRef.GetDataTypeName(), FDataVertexMetadata{}, TGetVertexAccess<DataReferenceType>::VertexAccess);
+				return FInputBinding(Vertex);
+			}
+		};
+
+
+		// Binds an IDataReference to a FOutputDataVertex
+		class METASOUNDGRAPHCORE_API FOutputBinding
+		{
+		public:
+			FOutputBinding(FOutputDataVertex&& InVertex);
+			FOutputBinding(const FOutputDataVertex& InVertex);
+			FOutputBinding(const FVertexName& InVertexName, FAnyDataReference&& InReference);
+
+			template<typename DataType>
+			void BindValue(const TDataValueReference<DataType>& InOutDataReference)
+			{
+				Set(InOutDataReference);
+			}
+
+			template<typename DataType>
+			void BindRead(TDataReadReference<DataType>& InOutDataReference)
+			{
+				Set(InOutDataReference);
+			}
+
+			template<typename DataType>
+			void BindRead(TDataValueReference<DataType>& InOutDataReference)
+			{
+				Set(TDataReadReference<DataType>(InOutDataReference));
+			}
+
+			template<typename DataType>
+			void BindRead(TDataWriteReference<DataType>& InOutDataReference)
+			{
+				Set(TDataReadReference<DataType>(InOutDataReference));
+			}
+
+			template<typename DataType>
+			void BindWrite(TDataWriteReference<DataType>& InOutDataReference)
+			{
+				Set(InOutDataReference);
+			}
+
+			void Bind(FAnyDataReference& InOutDataReference);
+			void Bind(FOutputBinding& InBinding);
+
+			// Set the data reference, overwriting any existing bound data references
+			template<typename DataReferenceType>
+			void Set(const DataReferenceType& InDataReference)
+			{
+				check(Vertex.DataTypeName == InDataReference.GetDataTypeName());
+				Data.Emplace(InDataReference);
+				CheckAccessTypeCompatibility(Vertex, *Data);
+			}
+
+			// Set the data reference, overwriting any existing bound data references
+			void Set(FAnyDataReference&& InAnyDataReference);
+			
+			const FOutputDataVertex& GetVertex() const;
+
+			bool IsBound() const;
+
+			EDataReferenceAccessType GetAccessType() const;
+
+			// Get data reference 
+			const FAnyDataReference* GetDataReference() const;
+
+			FDataReferenceID GetDataReferenceID() const;
+
+			// Get data read reference assuming data is bound and read or write accessible.
+			template<typename DataType>
+			TDataReadReference<DataType> GetDataReadReference() const
+			{
+				check(Data.IsSet());
+				return Data->GetDataReadReference<DataType>();
+			}
+
+			// Get the bound data read reference if it exists. Otherwise create and 
+			// return a data read reference by constructing one the supplied constructor 
+			// arguments.
+			template<typename DataType, typename ... ConstructorArgTypes>
+			TDataReadReference<DataType> GetOrConstructDataReadReference(ConstructorArgTypes&&... ConstructorArgs) const
+			{
+				if (Data.IsSet())
+				{
+					return Data->GetDataReadReference<DataType>();
+				}
+
+				return TDataReadReference<DataType>::CreateNew(Forward<ConstructorArgTypes>(ConstructorArgs)...);
+			}
+
+			// Get data write reference assuming data is bound and write accessible.
+			template<typename DataType> 
+			TDataWriteReference<DataType> GetDataWriteReference() const
+			{
+				check(Data.IsSet());
+				return Data->GetDataWriteReference<DataType>();
+			}
+
+			// Get the bound data write reference if it exists. Otherwise create and 
+			// return a data write reference by constructing one the supplied constructor 
+			// arguments.
+			template<typename DataType, typename ... ConstructorArgTypes>
+			TDataWriteReference<DataType> GetOrConstructDataWriteReference(ConstructorArgTypes&&... ConstructorArgs) const
+			{
+				if (Data.IsSet())
+				{
+					return Data->GetDataWriteReference<DataType>();
+				}
+				return TDataWriteReference<DataType>::CreateNew(Forward<ConstructorArgTypes>(ConstructorArgs)...);
+			}
+
+			// Gets the value of the bound data reference if it exists. Returns
+			// nullptr if not bound.
+			template<typename DataType>
+			const DataType* GetValue() const
+			{
+				if (Data.IsSet())
+				{
+					return Data->GetValue<DataType>();
+				}
+				return nullptr;
+			}
+
 			// Set the value with a constant value reference.
 			template<typename DataType>
 			void SetValue(const DataType& InValue)
@@ -206,72 +388,192 @@ namespace Metasound
 
 		private:
 
-			VertexType Vertex;
+			FOutputDataVertex Vertex;
 			TOptional<FAnyDataReference> Data;
 		};
+
+		// Factory for creating FOutputBinding from an IDataReference derived object.
+		struct FOutputBindingFactory
+		{
+			template<typename DataReferenceType>
+			static FOutputBinding CreateBinding(const FVertexName& InVertexName, const DataReferenceType& InRef)
+			{
+				FOutputDataVertex Vertex(InVertexName, InRef.GetDataTypeName(), FDataVertexMetadata{}, TGetVertexAccess<DataReferenceType>::VertexAccess);
+				return FOutputBinding(Vertex);
+			}
+		};
 	}
 
+
+	/** MetaSound Binding 
+	 *
+	 * MetaSound IOperators read and write to shared data. For example: a TDataWriteReference<float> 
+	 * written to by one operator may be read by several other operators holding
+	 * onto TDataReadReference<float>. Binding empowers sharing of the underlying 
+	 * objects (a `float` in the example) between IOperators and minimizing copying
+	 * of data. In general, it is assumed that any individual node does not need 
+	 * to worry about the specific implications of a binding, but can rather assume 
+	 * that data references are managed by external systems. 
+	 *
+	 * There are 2 known scenarios where an IOperator must be aware of what binding
+	 * can do.
+	 *
+	 * 1. If the IOperator internally caches raw pointers to the underlying data 
+	 * of a data reference.
+	 *
+	 * 2. If the IOperator manages multiple internal IOperators with their own 
+	 * shared connections (e.g. FGraphOperator)
+	 *
+	 * In the case that a IOperator does need to know manage bound connections, it's
+	 * important that the IOperator follows the binding rules.
+	 *
+	 * General Binding Rules:
+	 * - Binding cannot change the access type of an existing data reference.
+	 * - IOperators will ignore new TDataValueReferences because Value references
+	 *   cannot be updated after the operator has been constructed. 
+	 *
+	 * Input Binding Rules:
+	 * - Binding an input may update the underlying data object to point to
+	 * a new object.
+	 * - When updating input pointers of underlying data, a data reference from
+	 * and outer scope should replace that of an inner scope. For example, a graph 
+	 * representing multiple IOperators can override the input pointers of individual
+	 * IOperators, but an individual IOperator cannot override the input pointers of 
+	 * the containing graph.
+	 *
+	 *
+	 * Output Binding Rules:
+	 * - Binding an output may NOT update the location of the underlying data object.
+	 * - When determining output pointers of underlying data, a data reference from
+	 *  the inner scope can override the data reference of an outerscope. For example, 
+	 *  a graph representing multiple IOperators cannot override the output pointers 
+	 *  of individual IOperators, but an individual IOperator can override the output 
+	 *  pointers of the containing graph.
+	 *
+	 *
+	 *  These rules apply to any method on the FInputVertexInterfaData or 
+	 *  FOutptuVertexInterfaceData which begins with `Bind`
+	 */
+
+	// Forward declare
+	struct FVertexDataState;
+
 	/** An input vertex interface with optionally bound data references. */
 	class METASOUNDGRAPHCORE_API FInputVertexInterfaceData
 	{
-		using FBinding = MetasoundVertexDataPrivate::TBinding<FInputDataVertex>;
+		using FInputBinding = MetasoundVertexDataPrivate::FInputBinding;
 
 	public:
 
-		using FRangedForIteratorType = typename TArray<FBinding>::RangedForIteratorType;
-		using FRangedForConstIteratorType = typename TArray<FBinding>::RangedForConstIteratorType;
+		using FRangedForIteratorType = typename TArray<FInputBinding>::RangedForIteratorType;
+		using FRangedForConstIteratorType = typename TArray<FInputBinding>::RangedForConstIteratorType;
 
-		FInputVertexInterfaceData() = default;
+		/** Construct with an FInputVertexInterface. This will default to a unfrozen vertex interface. */
+		FInputVertexInterfaceData();
 
-		/** Construct with an FInputVertexInterface. */
+		/** Construct with an FInputVertexInterface. This will default to a frozen vertex interface. */
 		FInputVertexInterfaceData(const FInputVertexInterface& InVertexInterface);
 
-		/** Bind a value vertex from a value reference. */
+		/** Returns true if the vertex interface is frozen. */
+		bool IsVertexInterfaceFrozen() const;
+
+		/** Set whether the vertex interface is frozen or not. 
+		 *
+		 * If frozen, attempts to access vertices which do not already exist will result in an error. 
+		 *
+		 * If not frozen, attempts to bind to a missing vertex will automatically
+		 * add the missing vertex.
+		 */
+		void SetIsVertexInterfaceFrozen(bool bInIsVertexInterfaceFrozen);
+
+		/** Set the value of a vertex. */
 		template<typename DataType>
-		void BindValueVertex(const FVertexName& InVertexName, const TDataValueReference<DataType>& InDataReference)
+		void SetValue(const FVertexName& InVertexName, const DataType& InValue)
 		{
-			BindVertex(InVertexName, FAnyDataReference{InDataReference});
+			auto CreateBinding = [&]()
+			{
+				FInputDataVertex Vertex(InVertexName, GetMetasoundDataTypeName<DataType>(), FDataVertexMetadata{}, EVertexAccessType::Value);
+				return FInputBinding(Vertex);
+			};
+
+			auto BindData = [&](FInputBinding& Binding) { Binding.SetValue<DataType>(InValue); };
+			
+			Apply(InVertexName, CreateBinding, BindData);
 		}
 
-		/** Bind a read vertex from a value reference. */
+		/** Bind a read vertex from a read reference. Slated for deprecation. */
 		template<typename DataType>
-		void BindReadVertex(const FVertexName& InVertexName, const TDataValueReference<DataType>& InDataReference)
+		void BindReadVertex(const FVertexName& InVertexName, const TDataReadReference<DataType>& InOutDataReference)
 		{
-			BindVertex(InVertexName, FAnyDataReference{static_cast<TDataReadReference<DataType>>(InDataReference)});
+			BindReadVertex(InVertexName, const_cast<TDataReadReference<DataType>&>(InOutDataReference));
 		}
 
 		/** Bind a read vertex from a read reference. */
 		template<typename DataType>
-		void BindReadVertex(const FVertexName& InVertexName, const TDataReadReference<DataType>& InDataReference)
+		void BindReadVertex(const FVertexName& InVertexName, TDataReadReference<DataType>& InOutDataReference)
 		{
-			BindVertex(InVertexName, FAnyDataReference{InDataReference});
-		}
+			auto CreateBinding = [&]()
+			{
+				return MetasoundVertexDataPrivate::FInputBindingFactory::CreateBinding(InVertexName, InOutDataReference);
+			};
 
+			auto BindData = [&](FInputBinding& Binding) { Binding.BindRead(InOutDataReference); };
+			
+			Apply(InVertexName, CreateBinding, BindData);
+		}
+		
 		/** Bind a read vertex from a write reference. */
 		template<typename DataType>
-		void BindReadVertex(const FVertexName& InVertexName, const TDataWriteReference<DataType>& InDataReference)
+		void BindReadVertex(const FVertexName& InVertexName, TDataWriteReference<DataType>& InOutDataReference)
 		{
-			BindVertex(InVertexName, FAnyDataReference{static_cast<TDataReadReference<DataType>>(InDataReference)});
+			auto CreateBinding = [&]()
+			{
+				return MetasoundVertexDataPrivate::FInputBindingFactory::CreateBinding(InVertexName, InOutDataReference);
+			};
+
+			auto BindData = [&](FInputBinding& Binding) { Binding.BindRead(InOutDataReference); };
+			
+			Apply(InVertexName, CreateBinding, BindData);
 		}
 
-		/** Bind a write vertex from a write reference. */
+		/** Bind a write vertex from a write reference. Slated for deprecation. */
 		template<typename DataType>
-		void BindWriteVertex(const FVertexName& InVertexName, const TDataWriteReference<DataType>& InDataReference)
+		void BindWriteVertex(const FVertexName& InVertexName, const TDataWriteReference<DataType>& InOutDataReference)
 		{
-			BindVertex(InVertexName, FAnyDataReference{InDataReference});
+			BindWriteVertex(InVertexName, const_cast<TDataWriteReference<DataType>&>(InOutDataReference));
 		}
 
-		/** Bind a data reference to a vertex. */
-		void BindVertex(const FVertexName& InVertexName, FAnyDataReference&& InDataReference);
+		/** Bind a write vertex from a write reference. */
+		template<typename DataType>
+		void BindWriteVertex(const FVertexName& InVertexName, TDataWriteReference<DataType>& InOutDataReference)
+		{
+			auto CreateBinding = [&]()
+			{
+				return MetasoundVertexDataPrivate::FInputBindingFactory::CreateBinding(InVertexName, InOutDataReference);
+			};
 
-		/** Bind a data reference to a vertex. */
-		void BindVertex(const FVertexName& InVertexName, const FAnyDataReference& InDataReference);
+			auto BindData = [&](FInputBinding& Binding) { Binding.BindWrite(InOutDataReference); };
+			
+			Apply(InVertexName, CreateBinding, BindData);
+		}
 
+		/** Bind a vertex with a any data reference. Slated for deprecation. */
+		void BindVertex(const FVertexName& InVertexName, const FAnyDataReference& InOutDataReference);
+		/** Bind a vertex with a any data reference. */
+		void BindVertex(const FVertexName& InVertexName, FAnyDataReference& InOutDataReference);
+
+		/** Bind vertex data using other vertex data. Slated for deprecation. */
+		void Bind(const FInputVertexInterfaceData& InOutInputVertexData);
 		/** Bind vertex data using other vertex data. */
-		void Bind(const FInputVertexInterfaceData& InInputVertexData);
+		void Bind(FInputVertexInterfaceData& InOutInputVertexData);
 
-		/** Bind a data references to vertices with matching vertex names. */
-		void Bind(const FDataReferenceCollection& InCollection);
+		/** Sets a vertex to use a data reference, ignoring any existing data bound to the vertex. */
+		void SetVertex(const FVertexName& InVertexName, const FAnyDataReference& InDataReference);
+		/** Sets a vertex to use a data reference, ignoring any existing data bound to the vertex. */
+		void SetVertex(const FVertexName& InVertexName, FAnyDataReference&& InDataReference);
+
+		/** Set a data references to vertices with matching vertex names. Ignores any existing data bound to the vertex. */
+		void Set(const FDataReferenceCollection& InCollection);
 
 		/** Convert vertex data to a data reference collection. */
 		FDataReferenceCollection ToDataReferenceCollection() const;
@@ -319,7 +621,7 @@ namespace Metasound
 		template<typename DataType>
 		const DataType* GetValue(const FVertexName& InVertexName) const
 		{
-			const FBinding* Binding = FindChecked(InVertexName);
+			const FInputBinding* Binding = FindChecked(InVertexName);
 			return Binding->GetValue<DataType>();
 		}
 
@@ -329,23 +631,15 @@ namespace Metasound
 		template<typename DataType>
 		DataType GetOrCreateDefaultValue(const FVertexName& InVertexName, const FOperatorSettings& InSettings) const
 		{
-			const FBinding* Binding = FindChecked(InVertexName);
+			const FInputBinding* Binding = FindChecked(InVertexName);
 			return Binding->GetOrCreateDefaultValue<DataType>(InSettings);
 		}
 
-		/** Set the value of a vertex. */
-		template<typename DataType>
-		void SetValue(const FVertexName& InVertexName, const DataType& InValue)
-		{
-			FBinding* Binding = FindChecked(InVertexName);
-			return Binding->SetValue<DataType>(InValue);
-		}
-
 		/** Get data read reference assuming data is bound and read or write accessible. */
 		template<typename DataType> 
 		TDataReadReference<DataType> GetDataReadReference(const FVertexName& InVertexName) const
 		{
-			const FBinding* Binding = FindChecked(InVertexName);
+			const FInputBinding* Binding = FindChecked(InVertexName);
 			return Binding->GetDataReadReference<DataType>();
 		}
 
@@ -356,7 +650,7 @@ namespace Metasound
 		template<typename DataType> 
 		TDataReadReference<DataType> GetOrCreateDefaultDataReadReference(const FVertexName& InVertexName, const FOperatorSettings& InSettings) const
 		{
-			if (const FBinding* Binding = Find(InVertexName))
+			if (const FInputBinding* Binding = Find(InVertexName))
 			{
 				return Binding->GetOrCreateDefaultDataReadReference<DataType>(InSettings);
 			}
@@ -375,7 +669,7 @@ namespace Metasound
 		template<typename DataType, typename ... ConstructorArgTypes>
 		TDataReadReference<DataType> GetOrConstructDataReadReference(const FVertexName& InVertexName, ConstructorArgTypes&&... ConstructorArgs) const
 		{
-			if (const FBinding* Binding = Find(InVertexName))
+			if (const FInputBinding* Binding = Find(InVertexName))
 			{
 				return Binding->GetOrConstructDataReadReference<DataType>(Forward<ConstructorArgTypes>(ConstructorArgs)...);
 			}
@@ -392,7 +686,7 @@ namespace Metasound
 		template<typename DataType> 
 		TDataWriteReference<DataType> GetDataWriteReference(const FVertexName& InVertexName) const
 		{
-			const FBinding* Binding = FindChecked(InVertexName);
+			const FInputBinding* Binding = FindChecked(InVertexName);
 			return Binding->GetDataWriteReference<DataType>();
 		}
 
@@ -403,7 +697,7 @@ namespace Metasound
 		template<typename DataType> 
 		TDataWriteReference<DataType> GetOrCreateDefaultDataWriteReference(const FVertexName& InVertexName, const FOperatorSettings& InSettings) const
 		{
-			if (const FBinding* Binding = Find(InVertexName))
+			if (const FInputBinding* Binding = Find(InVertexName))
 			{
 				return Binding->GetOrCreateDefaultDataWriteReference<DataType>(InSettings);
 			}
@@ -422,7 +716,7 @@ namespace Metasound
 		template<typename DataType, typename ... ConstructorArgTypes>
 		TDataWriteReference<DataType> GetOrConstructDataWriteReference(const FVertexName& InVertexName, ConstructorArgTypes&&... ConstructorArgs) const
 		{
-			if (const FBinding* Binding = Find(InVertexName))
+			if (const FInputBinding* Binding = Find(InVertexName))
 			{
 				return Binding->GetOrConstructDataWriteReference<DataType>(Forward<ConstructorArgTypes>(ConstructorArgs)...);
 			}
@@ -435,77 +729,158 @@ namespace Metasound
 		}
 	
 	private:
+		friend METASOUNDGRAPHCORE_API void GetVertexInterfaceDataState(const FInputVertexInterfaceData& InVertexInterface, TArray<FVertexDataState>& OutState);
+		friend METASOUNDGRAPHCORE_API void CompareVertexInterfaceDataToPriorState(const FInputVertexInterfaceData& InCurrentInterface, const TArray<FVertexDataState>& InPriorState, TSortedMap<FDataReferenceID, FAnyDataReference>& OutUpdates);
+
+		void Apply(const FVertexName& InVertexName, TFunctionRef<FInputBinding ()> InCreateFunc, TFunctionRef<void (FInputBinding&)> InApplyFunc);
 
-		FBinding* Find(const FVertexName& InVertexName);
-		const FBinding* Find(const FVertexName& InVertexName) const;
+		FInputBinding* Find(const FVertexName& InVertexName);
+		const FInputBinding* Find(const FVertexName& InVertexName) const;
 
-		FBinding* FindChecked(const FVertexName& InVertexName);
-		const FBinding* FindChecked(const FVertexName& InVertexName) const;
+		FInputBinding* FindChecked(const FVertexName& InVertexName);
+		const FInputBinding* FindChecked(const FVertexName& InVertexName) const;
 
-		TArray<FBinding> Bindings;
+		bool bIsVertexInterfaceFrozen = false;
+		TArray<FInputBinding> Bindings;
 	};
 
+
 	/** An output vertex interface with optionally bound data references. */
 	class METASOUNDGRAPHCORE_API FOutputVertexInterfaceData
 	{
-		using FBinding = MetasoundVertexDataPrivate::TBinding<FOutputDataVertex>;
+		using FOutputBinding = MetasoundVertexDataPrivate::FOutputBinding;
 
 	public:
 
-		using FRangedForIteratorType = typename TArray<FBinding>::RangedForIteratorType;
-		using FRangedForConstIteratorType = typename TArray<FBinding>::RangedForConstIteratorType;
+		using FRangedForIteratorType = typename TArray<FOutputBinding>::RangedForIteratorType;
+		using FRangedForConstIteratorType = typename TArray<FOutputBinding>::RangedForConstIteratorType;
 
-		FOutputVertexInterfaceData() = default;
+		/** Construct with an FOutputVertexInterface. This will default to a unfrozen vertex interface. */
+		FOutputVertexInterfaceData();
 
-		/* Construct using an FOutputVertexInterface. */
+		/** Construct with an FOutputVertexInterface. This will default to a frozen vertex interface. */
 		FOutputVertexInterfaceData(const FOutputVertexInterface& InVertexInterface);
 
+		/** Returns true if the vertex interface is frozen. */
+		bool IsVertexInterfaceFrozen() const;
+
+		/** Set whether the vertex interface is frozen or not. 
+		 *
+		 * If frozen, attempts to access vertices which do not already exist will result in an error. 
+		 *
+		 * If not frozen, attempts to bind to a missing vertex will automatically
+		 * add the missing vertex.
+		 */
+		void SetIsVertexInterfaceFrozen(bool bInIsVertexInterfaceFrozen);
+
+		/** Set the value of a vertex. */
+		template<typename DataType>
+		void SetValue(const FVertexName& InVertexName, const DataType& InValue)
+		{
+			BindValueVertex<DataType>(InVertexName, TDataValueReference<DataType>::CreateNew(InValue));
+		}
+
 		/** Bind a value vertex from a value reference. */
 		template<typename DataType>
-		void BindValueVertex(const FVertexName& InVertexName, const TDataValueReference<DataType>& InDataReference)
+		void BindValueVertex(const FVertexName& InVertexName, const TDataValueReference<DataType>& InOutDataReference)
 		{
-			BindVertex(InVertexName, FAnyDataReference{InDataReference});
+			auto CreateBinding = [&]()
+			{
+				return MetasoundVertexDataPrivate::FOutputBindingFactory::CreateBinding(InVertexName, InOutDataReference);
+			};
+
+			auto BindData = [&](FOutputBinding& Binding) { Binding.BindValue(InOutDataReference); };
+
+			Apply(InVertexName, CreateBinding, BindData);
 		}
 
 		/** Bind a read vertex from a value reference. */
 		template<typename DataType>
-		void BindReadVertex(const FVertexName& InVertexName, const TDataValueReference<DataType>& InDataReference)
+		void BindReadVertex(const FVertexName& InVertexName, const TDataValueReference<DataType>& InOutDataReference)
+		{
+			BindReadVertex(InVertexName, static_cast<TDataReadReference<DataType>>(InOutDataReference));
+		}
+
+		/** Bind a read vertex from a read reference. Slated for deprecation. */
+		template<typename DataType>
+		void BindReadVertex(const FVertexName& InVertexName, const TDataReadReference<DataType>& InOutDataReference)
 		{
-			BindVertex(InVertexName, FAnyDataReference{static_cast<TDataReadReference<DataType>>(InDataReference)});
+			BindReadVertex(InVertexName, const_cast<TDataReadReference<DataType>&>(InOutDataReference));
 		}
 
 		/** Bind a read vertex from a read reference. */
 		template<typename DataType>
-		void BindReadVertex(const FVertexName& InVertexName, const TDataReadReference<DataType>& InDataReference)
+		void BindReadVertex(const FVertexName& InVertexName, TDataReadReference<DataType>& InOutDataReference)
 		{
-			BindVertex(InVertexName, FAnyDataReference{InDataReference});
+			auto CreateBinding = [&]()
+			{
+				return MetasoundVertexDataPrivate::FOutputBindingFactory::CreateBinding(InVertexName, InOutDataReference);
+			};
+
+			auto BindData = [&](FOutputBinding& Binding) { Binding.BindRead(InOutDataReference); };
+
+			Apply(InVertexName, CreateBinding, BindData);
+		}
+
+		/** Bind a read vertex from a write reference. Slated for deprecation. */
+		template<typename DataType>
+		void BindReadVertex(const FVertexName& InVertexName, const TDataWriteReference<DataType>& InOutDataReference)
+		{
+			BindReadVertex(InVertexName, const_cast<TDataWriteReference<DataType>&>(InOutDataReference));
 		}
 
 		/** Bind a read vertex from a write reference. */
 		template<typename DataType>
-		void BindReadVertex(const FVertexName& InVertexName, const TDataWriteReference<DataType>& InDataReference)
+		void BindReadVertex(const FVertexName& InVertexName, TDataWriteReference<DataType>& InOutDataReference)
 		{
-			BindVertex(InVertexName, FAnyDataReference{static_cast<TDataReadReference<DataType>>(InDataReference)});
+			auto CreateBinding = [&]()
+			{
+				return MetasoundVertexDataPrivate::FOutputBindingFactory::CreateBinding(InVertexName, InOutDataReference);
+			};
+
+			auto BindData = [&](FOutputBinding& Binding) { Binding.BindRead(InOutDataReference); };
+
+			Apply(InVertexName, CreateBinding, BindData);
 		}
 
-		/** Bind a write vertex from a write reference. */
+		/** Bind a write vertex from a write reference. Slated for deprecation. */
 		template<typename DataType>
-		void BindWriteVertex(const FVertexName& InVertexName, const TDataWriteReference<DataType>& InDataReference)
+		void BindWriteVertex(const FVertexName& InVertexName, const TDataWriteReference<DataType>& InOutDataReference)
 		{
-			BindVertex(InVertexName, FAnyDataReference{InDataReference});
+			BindWriteVertex(InVertexName, const_cast<TDataWriteReference<DataType>&>(InOutDataReference));
 		}
 
-		/** Bind a data reference to a vertex. */
-		void BindVertex(const FVertexName& InVertexName, FAnyDataReference&& InDataReference);
+		/** Bind a write vertex from a write reference. */
+		template<typename DataType>
+		void BindWriteVertex(const FVertexName& InVertexName, TDataWriteReference<DataType>& InOutDataReference)
+		{
+			auto CreateBinding = [&]()
+			{
+				return MetasoundVertexDataPrivate::FOutputBindingFactory::CreateBinding(InVertexName, InOutDataReference);
+			};
 
-		/** Bind a data reference to a vertex. */
-		void BindVertex(const FVertexName& InVertexName, const FAnyDataReference& InDataReference);
+			auto BindData = [&](FOutputBinding& Binding) { Binding.BindWrite(InOutDataReference); };
 
+			Apply(InVertexName, CreateBinding, BindData);
+		}
+
+		/** Bind a vertex with a any data reference. Slated for deprecation. */
+		void BindVertex(const FVertexName& InVertexName, const FAnyDataReference& InOutDataReference);
+		/** Bind a vertex with a any data reference. */
+		void BindVertex(const FVertexName& InVertexName, FAnyDataReference& InOutDataReference);
+
+		/** Bind vertex data using other vertex data. Slated for deprecation. */
+		void Bind(const FOutputVertexInterfaceData& InOutOutputVertexData);
 		/** Bind vertex data using other vertex data. */
-		void Bind(const FOutputVertexInterfaceData& InOutputVertexData);
+		void Bind(FOutputVertexInterfaceData& InOutOutputVertexData);
+
+		/** Set a data reference to a vertex, ignoring any existing data bound to the vertex*/
+		void SetVertex(const FVertexName& InVertexName, const FAnyDataReference& InDataReference);
+		/** Set a data reference to a vertex, ignoring any existing data bound to the vertex*/
+		void SetVertex(const FVertexName& InVertexName, FAnyDataReference&& InDataReference);
 
 		/** Bind a data references to vertices with matching vertex names. */
-		void Bind(const FDataReferenceCollection& InCollection);
+		void Set(const FDataReferenceCollection& InCollection);
 
 		/** Converts the vertex data to a data reference collection. */
 		FDataReferenceCollection ToDataReferenceCollection() const;
@@ -553,23 +928,16 @@ namespace Metasound
 		template<typename DataType>
 		const DataType* GetValue(const FVertexName& InVertexName) const
 		{
-			const FBinding* Binding = FindChecked(InVertexName);
+			const FOutputBinding* Binding = FindChecked(InVertexName);
 			return Binding->GetValue<DataType>();
 		}
 
-		/** Set the value of a vertex. */
-		template<typename DataType>
-		void SetValue(const FVertexName& InVertexName, const DataType& InValue)
-		{
-			FBinding* Binding = FindChecked(InVertexName);
-			return Binding->SetValue<DataType>(InValue);
-		}
 
 		/** Get data read reference assuming data is bound and read or write accessible. */
 		template<typename DataType> 
 		TDataReadReference<DataType> GetDataReadReference(const FVertexName& InVertexName) const
 		{
-			const FBinding* Binding = FindChecked(InVertexName);
+			const FOutputBinding* Binding = FindChecked(InVertexName);
 			return Binding->GetDataReadReference<DataType>();
 		}
 
@@ -580,7 +948,7 @@ namespace Metasound
 		template<typename DataType, typename ... ConstructorArgTypes>
 		TDataReadReference<DataType> GetOrConstructDataReadReference(const FVertexName& InVertexName, ConstructorArgTypes&&... ConstructorArgs) const
 		{
-			if (const FBinding* Binding = Find(InVertexName))
+			if (const FOutputBinding* Binding = Find(InVertexName))
 			{
 				return Binding->GetOrConstructDataReadReference<DataType>(Forward<ConstructorArgTypes>(ConstructorArgs)...);
 			}
@@ -597,7 +965,7 @@ namespace Metasound
 		template<typename DataType> 
 		TDataWriteReference<DataType> GetDataWriteReference(const FVertexName& InVertexName) const
 		{
-			const FBinding* Binding = FindChecked(InVertexName);
+			const FOutputBinding* Binding = FindChecked(InVertexName);
 			return Binding->GetDataWriteReference<DataType>();
 		}
 
@@ -608,7 +976,7 @@ namespace Metasound
 		template<typename DataType, typename ... ConstructorArgTypes>
 		TDataWriteReference<DataType> GetOrConstructDataWriteReference(const FVertexName& InVertexName, ConstructorArgTypes&&... ConstructorArgs) const
 		{
-			if (const FBinding* Binding = Find(InVertexName))
+			if (const FOutputBinding* Binding = Find(InVertexName))
 			{
 				return Binding->GetOrConstructDataWriteReference<DataType>(Forward<ConstructorArgTypes>(ConstructorArgs)...);
 			}
@@ -621,16 +989,20 @@ namespace Metasound
 		}
 	
 	private:
+		friend METASOUNDGRAPHCORE_API void GetVertexInterfaceDataState(const FOutputVertexInterfaceData& InVertexInterface, TArray<FVertexDataState>& OutState);
+		friend METASOUNDGRAPHCORE_API void CompareVertexInterfaceDataToPriorState(const FOutputVertexInterfaceData& InCurrentInterface, const TArray<FVertexDataState>& InPriorState, TSortedMap<FDataReferenceID, FAnyDataReference>& OutUpdates);
 
-		FBinding* Find(const FVertexName& InVertexName);
-		const FBinding* Find(const FVertexName& InVertexName) const;
+		void Apply(const FVertexName& InVertexName, TFunctionRef<FOutputBinding ()> InCreateFunc, TFunctionRef<void (FOutputBinding&)> InBindFunc);
 
-		FBinding* FindChecked(const FVertexName& InVertexName);
-		const FBinding* FindChecked(const FVertexName& InVertexName) const;
+		FOutputBinding* Find(const FVertexName& InVertexName);
+		const FOutputBinding* Find(const FVertexName& InVertexName) const;
 
-		TArray<FBinding> Bindings;
-	};
+		FOutputBinding* FindChecked(const FVertexName& InVertexName);
+		const FOutputBinding* FindChecked(const FVertexName& InVertexName) const;
 
+		bool bIsVertexInterfaceFrozen = false;
+		TArray<FOutputBinding> Bindings;
+	};
 
 	/** A vertex interface with optionally bound data. */
 	class METASOUNDGRAPHCORE_API FVertexInterfaceData
@@ -640,14 +1012,11 @@ namespace Metasound
 		FVertexInterfaceData() = default;
 
 		/** Construct using an FVertexInterface. */
-		FVertexInterfaceData(const FVertexInterface& InVertexInterface)
-		: InputVertexInterfaceData(InVertexInterface.GetInputInterface())
-		, OutputVertexInterfaceData(InVertexInterface.GetOutputInterface())
-		{
-		}
+		FVertexInterfaceData(const FVertexInterface& InVertexInterface);
 
-		/** Bind vertex data using other vertex data. */
+		/** Set vertex data using other vertex data. */
 		void Bind(const FVertexInterfaceData& InVertexData);
+		void Bind(FVertexInterfaceData& InVertexData);
 
 		/** Get input vertex interface data. */
 		FInputVertexInterfaceData& GetInputs()
@@ -678,4 +1047,26 @@ namespace Metasound
 		FInputVertexInterfaceData InputVertexInterfaceData;
 		FOutputVertexInterfaceData OutputVertexInterfaceData;
 	};
+
+	/** FVertexDataState encapsulates which data reference a vertex is associated
+	 * with. The ID refers to the underlying object associated with the IDataReference.
+	 */
+	struct FVertexDataState
+	{
+		FVertexName VertexName;
+		FDataReferenceID ID;
+
+		METASOUNDGRAPHCORE_API friend bool operator<(const FVertexDataState& InLHS, const FVertexDataState& InRHS);
+		METASOUNDGRAPHCORE_API friend bool operator==(const FVertexDataState& InLHS, const FVertexDataState& InRHS);
+	};
+
+	/** Caches a representation of the current data references bound to the vertex interface */
+	METASOUNDGRAPHCORE_API void GetVertexInterfaceDataState(const FInputVertexInterfaceData& InVertexInterface, TArray<FVertexDataState>& OutState);
+	/** Caches a representation of the current data references bound to the vertex interface */
+	METASOUNDGRAPHCORE_API void GetVertexInterfaceDataState(const FOutputVertexInterfaceData& InVertexInterface, TArray<FVertexDataState>& OutState);
+
+	/** Compares the current data bound to the vertex interface with a prior cached state. */
+	METASOUNDGRAPHCORE_API void CompareVertexInterfaceDataToPriorState(const FInputVertexInterfaceData& InCurrentInterface, const TArray<FVertexDataState>& InPriorState, TSortedMap<FDataReferenceID, FAnyDataReference>& OutUpdates);
+	/** Compares the current data bound to the vertex interface with a prior cached state. */
+	METASOUNDGRAPHCORE_API void CompareVertexInterfaceDataToPriorState(const FOutputVertexInterfaceData& InCurrentInterface, const TArray<FVertexDataState>& InPriorState, TSortedMap<FDataReferenceID, FAnyDataReference>& OutUpdates);
 }
