- Fix to avoid not adding graph vertices to graph on creation
- Fix to MetaSound doc builder ClassIdToNodeIndices not initializing properly
#rnx

diff --git a/Engine/Plugins/Runtime/Metasound/Source/MetasoundFrontend/Private/MetasoundFrontendDocumentCache.cpp b/Engine/Plugins/Runtime/Metasound/Source/MetasoundFrontend/Private/MetasoundFrontendDocumentCache.cpp
--- a/Engine/Plugins/Runtime/Metasound/Source/MetasoundFrontend/Private/MetasoundFrontendDocumentCache.cpp
+++ b/Engine/Plugins/Runtime/Metasound/Source/MetasoundFrontend/Private/MetasoundFrontendDocumentCache.cpp
@@ -102,7 +102,9 @@ namespace Metasound::Frontend
 			const FMetasoundFrontendClass* Class = InParent.FindDependency(Node.ClassID);
 			check(Class);
 
-			IDToIndex.Add(Node.GetID(), Index);
+			const FGuid& NodeID = Node.GetID();
+			IDToIndex.Add(NodeID, Index);
+			ClassIDToNodeIndices.FindOrAdd(Node.ClassID).Add(Index);
 
 			switch (Class->Metadata.GetType())
 			{
diff --git a/Engine/Plugins/Runtime/Metasound/Source/MetasoundFrontend/Private/MetasoundFrontendTransform.cpp b/Engine/Plugins/Runtime/Metasound/Source/MetasoundFrontend/Private/MetasoundFrontendTransform.cpp
--- a/Engine/Plugins/Runtime/Metasound/Source/MetasoundFrontend/Private/MetasoundFrontendTransform.cpp
+++ b/Engine/Plugins/Runtime/Metasound/Source/MetasoundFrontend/Private/MetasoundFrontendTransform.cpp
@@ -210,7 +210,7 @@ namespace Metasound
 
 				if (INDEX_NONE != RemoveIndex)
 				{
-					PairedOutputs.Add(FVertexPair {OutputsToRemove[RemoveIndex], OutputsToAdd[AddIndex].Output });
+					PairedOutputs.Add(FVertexPair { OutputsToRemove[RemoveIndex], OutputsToAdd[AddIndex].Output });
 					OutputsToRemove.RemoveAtSwap(RemoveIndex);
 					OutputsToAdd.RemoveAtSwap(AddIndex);
 				}
@@ -367,7 +367,7 @@ namespace Metasound
 				bDidEdit |= RemoveUnsupportedVertices(GraphHandle);
 
 #if WITH_EDITOR
-				if (bAddedVertices)
+				if (bAddedVertices && bSetDefaultNodeLocations)
 				{
 					UpdateAddedVertexNodePositions(GraphHandle);
 				}
