- Minor bug w/setting vertex ID instead of node ID when member modify context is updated
- Misc clean-up
#rnx

diff --git a/Engine/Plugins/Runtime/Metasound/Source/MetasoundFrontend/Private/MetasoundFrontendDocumentController.cpp b/Engine/Plugins/Runtime/Metasound/Source/MetasoundFrontend/Private/MetasoundFrontendDocumentController.cpp
--- a/Engine/Plugins/Runtime/Metasound/Source/MetasoundFrontend/Private/MetasoundFrontendDocumentController.cpp
+++ b/Engine/Plugins/Runtime/Metasound/Source/MetasoundFrontend/Private/MetasoundFrontendDocumentController.cpp
@@ -157,7 +157,7 @@ namespace Metasound
 			if (FMetasoundFrontendDocument* Doc = DocumentPtr.Get())
 			{
 #if WITH_EDITOR
-				Doc->Metadata.ModifyContext.AddInterfaceModified({ InVersion.Name });
+				Doc->Metadata.ModifyContext.AddInterfaceModified(InVersion.Name);
 #endif // WITH_EDITOR
 				Doc->Interfaces.Add(InVersion);
 			}
@@ -168,7 +168,7 @@ namespace Metasound
 			if (FMetasoundFrontendDocument* Doc = DocumentPtr.Get())
 			{
 #if WITH_EDITOR
-				Doc->Metadata.ModifyContext.AddInterfaceModified({ InVersion.Name });
+				Doc->Metadata.ModifyContext.AddInterfaceModified(InVersion.Name);
 #endif // WITH_EDITOR
 				Doc->Interfaces.Remove(InVersion);
 			}
diff --git a/Engine/Plugins/Runtime/Metasound/Source/MetasoundFrontend/Private/MetasoundFrontendGraphController.cpp b/Engine/Plugins/Runtime/Metasound/Source/MetasoundFrontend/Private/MetasoundFrontendGraphController.cpp
--- a/Engine/Plugins/Runtime/Metasound/Source/MetasoundFrontend/Private/MetasoundFrontendGraphController.cpp
+++ b/Engine/Plugins/Runtime/Metasound/Source/MetasoundFrontend/Private/MetasoundFrontendGraphController.cpp
@@ -1158,9 +1158,9 @@ namespace Metasound
 #if WITH_EDITOR
 							if (FMetasoundFrontendDocumentMetadata* DocMetadata = OwningDocument->GetMetadata())
 							{
-								DocMetadata->ModifyContext.AddMemberIDModified(NewOutput.VertexID);
+								DocMetadata->ModifyContext.AddMemberIDModified(NewOutput.NodeID);
 							}
-#endif // WITH_EDITO
+#endif // WITH_EDITOR
 
 							// Mark interface as changed.
 							GraphClass->Interface.UpdateChangeID();
diff --git a/Engine/Plugins/Runtime/Metasound/Source/MetasoundFrontend/Public/MetasoundDocumentInterface.h b/Engine/Plugins/Runtime/Metasound/Source/MetasoundFrontend/Public/MetasoundDocumentInterface.h
--- a/Engine/Plugins/Runtime/Metasound/Source/MetasoundFrontend/Public/MetasoundDocumentInterface.h
+++ b/Engine/Plugins/Runtime/Metasound/Source/MetasoundFrontend/Public/MetasoundDocumentInterface.h
@@ -12,7 +12,8 @@ struct FMetaSoundFrontendDocumentBuilder;
 
 
 // UInterface for all MetaSound UClasses that implement a MetaSound document
-// as a means for accessing data via scripting, execution, or node class generation.
+// as a means for accessing data via code, scripting, execution, or node
+// class generation.
 UINTERFACE(BlueprintType, meta = (CannotImplementInterfaceInBlueprint))
 class METASOUNDFRONTEND_API UMetaSoundDocumentInterface : public UInterface
 {
@@ -24,9 +25,13 @@ class METASOUNDFRONTEND_API IMetaSoundDocumentInterface : public IInterface
 	GENERATED_BODY()
 
 public:
+	// Returns read-only reference to the the MetaSoundFrontendDocument
+	// containing all MetaSound runtime & editor data.
 	virtual const FMetasoundFrontendDocument& GetDocument() const = 0;
 	virtual const UClass& GetBaseMetaSoundUClass() const = 0;
 
+	// Returns the parent class registered with the MetaSound UObject registry.
+
 private:
 	virtual FMetasoundFrontendDocument& GetDocument() = 0;
 
