Make PlayOrder of most recently played sound accessible from Audio Component, fix editor graph analyzers

diff --git a/Engine/Plugins/Runtime/Metasound/Source/MetasoundEditor/Private/MetasoundEditor.cpp b/Engine/Plugins/Runtime/Metasound/Source/MetasoundEditor/Private/MetasoundEditor.cpp
--- a/Engine/Plugins/Runtime/Metasound/Source/MetasoundEditor/Private/MetasoundEditor.cpp
+++ b/Engine/Plugins/Runtime/Metasound/Source/MetasoundEditor/Private/MetasoundEditor.cpp
@@ -1554,7 +1554,10 @@ namespace Metasound
 
 					const FName& AudioBufferTypeName = GetMetasoundDataTypeName<FAudioBuffer>();
 					const FSampleRate SampleRate = static_cast<FSampleRate>(AudioDevice->GetSampleRate());
-					GraphConnectionManager = MakeUnique<FGraphConnectionManager>(*MetasoundAsset, *PreviewComp, SampleRate);
+					const uint32 PlayOrder = PreviewComp->GetLastPlayOrder();
+					const uint64 TransmitterID = Audio::GetTransmitterID(PreviewComp->GetAudioComponentID(), 0, PlayOrder);
+					
+					GraphConnectionManager = MakeUnique<FGraphConnectionManager>(*MetasoundAsset, *PreviewComp, TransmitterID, SampleRate);
 				}
 
 				MetasoundGraphEditor->RegisterActiveTimer(0.0f,
diff --git a/Engine/Plugins/Runtime/Metasound/Source/MetasoundEditor/Private/MetasoundEditorGraphConnectionManager.cpp b/Engine/Plugins/Runtime/Metasound/Source/MetasoundEditor/Private/MetasoundEditorGraphConnectionManager.cpp
--- a/Engine/Plugins/Runtime/Metasound/Source/MetasoundEditor/Private/MetasoundEditorGraphConnectionManager.cpp
+++ b/Engine/Plugins/Runtime/Metasound/Source/MetasoundEditor/Private/MetasoundEditorGraphConnectionManager.cpp
@@ -12,9 +12,9 @@ namespace Metasound
 {
 	namespace Editor
 	{
-		FGraphConnectionManager::FGraphConnectionManager(const FMetasoundAssetBase& InAssetBase, const UAudioComponent& InAudioComponent, FSampleRate InSampleRate)
+		FGraphConnectionManager::FGraphConnectionManager(const FMetasoundAssetBase& InAssetBase, const UAudioComponent& InAudioComponent, uint64 InTransmitterID, FSampleRate InSampleRate)
 			: AudioComponent(&InAudioComponent)
-			, GraphAnalyzerView(MakeUnique<Frontend::FMetasoundGraphAnalyzerView>(InAssetBase, InAudioComponent.GetInstanceOwnerID(), InSampleRate))
+			, GraphAnalyzerView(MakeUnique<Frontend::FMetasoundGraphAnalyzerView>(InAssetBase, InTransmitterID, InSampleRate))
 		{
 			using namespace Frontend;
 
diff --git a/Engine/Plugins/Runtime/Metasound/Source/MetasoundEditor/Private/MetasoundEditorGraphConnectionManager.h b/Engine/Plugins/Runtime/Metasound/Source/MetasoundEditor/Private/MetasoundEditorGraphConnectionManager.h
--- a/Engine/Plugins/Runtime/Metasound/Source/MetasoundEditor/Private/MetasoundEditorGraphConnectionManager.h
+++ b/Engine/Plugins/Runtime/Metasound/Source/MetasoundEditor/Private/MetasoundEditorGraphConnectionManager.h
@@ -83,7 +83,7 @@ namespace Metasound
 
 		public:
 			FGraphConnectionManager() = default;
-			FGraphConnectionManager(const FMetasoundAssetBase& InAssetBase, const UAudioComponent& InAudioComponent, FSampleRate InSampleRate);
+			FGraphConnectionManager(const FMetasoundAssetBase& InAssetBase, const UAudioComponent& InAudioComponent, uint64 InTransmitterID, FSampleRate InSampleRate);
 			~FGraphConnectionManager() = default;
 
 			bool GetValue(const FGuid& InNodeID, FVertexName InOutputName, float& OutValue) const;
diff --git a/Engine/Source/Runtime/Engine/Classes/Components/AudioComponent.h b/Engine/Source/Runtime/Engine/Classes/Components/AudioComponent.h
--- a/Engine/Source/Runtime/Engine/Classes/Components/AudioComponent.h
+++ b/Engine/Source/Runtime/Engine/Classes/Components/AudioComponent.h
@@ -835,6 +835,7 @@ public:
 	FAudioDevice* GetAudioDevice() const override;
 	TArray<FAudioParameter>& GetInstanceParameters() override { return InstanceParameters; }
 	uint64 GetInstanceOwnerID() const override { return AudioComponentID; }
+	uint32 GetLastPlayOrder() const { return LastSoundPlayOrder; }
 	USoundBase* GetSound() override { return Sound; }
 
 	virtual FName GetFNameForStatID() const override;
@@ -873,6 +874,7 @@ public:
 private:
 
 	uint64 AudioComponentID;
+	uint32 LastSoundPlayOrder = 0;
 
 	float RetriggerTimeSinceLastUpdate;
 	float RetriggerUpdateInterval;
diff --git a/Engine/Source/Runtime/Engine/Private/Components/AudioComponent.cpp b/Engine/Source/Runtime/Engine/Private/Components/AudioComponent.cpp
--- a/Engine/Source/Runtime/Engine/Private/Components/AudioComponent.cpp
+++ b/Engine/Source/Runtime/Engine/Private/Components/AudioComponent.cpp
@@ -853,6 +853,8 @@ void UAudioComponent::PlayInternal(const PlayInternalRequestData& InPlayRequestD
 
 	AudioDevice->AddNewActiveSound(NewActiveSound, MoveTemp(SoundParams));
 
+	LastSoundPlayOrder = NewActiveSound.GetPlayOrder();
+
 	// In editor, the audio thread is not run separate from the game thread, and can result in calling PlaybackComplete prior
 	// to bIsActive being set. Therefore, we assign to the current state of ActiveCount as opposed to just setting to true.
 	SetActiveFlag(ActiveCount > 0);
