[Metasound] Check whether analyzer exists before attempting to add it to a generator.

Also, remove ensure when trying to find an analyzer factory, since we can fail gracefully.

#rnx

diff --git a/Engine/Plugins/Runtime/Metasound/Source/MetasoundEngine/Private/MetasoundGeneratorHandle.cpp b/Engine/Plugins/Runtime/Metasound/Source/MetasoundEngine/Private/MetasoundGeneratorHandle.cpp
--- a/Engine/Plugins/Runtime/Metasound/Source/MetasoundEngine/Private/MetasoundGeneratorHandle.cpp
+++ b/Engine/Plugins/Runtime/Metasound/Source/MetasoundEngine/Private/MetasoundGeneratorHandle.cpp
@@ -2,11 +2,12 @@
 
 #include "MetasoundGeneratorHandle.h"
 
-#include "Components/AudioComponent.h"
 #include "MetasoundGenerator.h"
+#include "MetasoundParameterPack.h"
 #include "MetasoundSource.h"
 #include "MetasoundTrace.h"
-#include "MetasoundParameterPack.h"
+#include "Analysis/MetasoundFrontendAnalyzerRegistry.h"
+#include "Components/AudioComponent.h"
 
 TMap<FName, UMetasoundGeneratorHandle::FPassthroughAnalyzerInfo> UMetasoundGeneratorHandle::PassthroughAnalyzers{};
 
@@ -353,6 +354,17 @@ bool UMetasoundGeneratorHandle::WatchOutput(
 	AnalyzerAddress.AnalyzerInstanceID = FGuid::NewGuid();
 	AnalyzerAddress.NodeID = NodeId;
 
+	// Check to see if the analyzer exists
+	{
+		using namespace Metasound::Frontend;
+		const IVertexAnalyzerFactory* Factory =
+			IVertexAnalyzerRegistry::Get().FindAnalyzerFactory(AnalyzerAddress.AnalyzerName);
+		if (nullptr == Factory)
+		{
+			return false;
+		}
+	}
+
 	// if we already have a generator, go ahead and make the analyzer
 	if (const TSharedPtr<Metasound::FMetasoundGenerator> PinnedGenerator = PinGenerator())
 	{
diff --git a/Engine/Plugins/Runtime/Metasound/Source/MetasoundFrontend/Private/Analysis/MetasoundFrontendAnalyzerRegistry.cpp b/Engine/Plugins/Runtime/Metasound/Source/MetasoundFrontend/Private/Analysis/MetasoundFrontendAnalyzerRegistry.cpp
--- a/Engine/Plugins/Runtime/Metasound/Source/MetasoundFrontend/Private/Analysis/MetasoundFrontendAnalyzerRegistry.cpp
+++ b/Engine/Plugins/Runtime/Metasound/Source/MetasoundFrontend/Private/Analysis/MetasoundFrontendAnalyzerRegistry.cpp
@@ -30,13 +30,14 @@ namespace Metasound
 			virtual const IVertexAnalyzerFactory* FindAnalyzerFactory(FName InAnalyzerName) const override
 			{
 				const TUniquePtr<IVertexAnalyzerFactory>* Factory = AnalyzerFactoryRegistry.Find(InAnalyzerName);
-				if (ensureMsgf(Factory, TEXT("Failed to find registered MetaSound Analyzer Factory with name '%s'"), *InAnalyzerName.ToString()))
+				if (nullptr == Factory)
 				{
-					check(Factory->IsValid());
-					return Factory->Get();
+					UE_LOG(LogMetaSound, Warning, TEXT("Failed to find registered MetaSound Analyzer Factory with name '%s'"), *InAnalyzerName.ToString());
+					return nullptr;
 				}
 
-				return nullptr;
+				check(Factory->IsValid());
+				return Factory->Get();
 			}
 
 			virtual void RegisterAnalyzerFactory(FName AnalyzerName, TUniquePtr<IVertexAnalyzerFactory>&& Factory) override
