Read sample rate and num frames from file header in MetaSound wave proxy reader to account for incorrect sample rate when using Resample for Device platform cook override

diff --git a/Engine/Plugins/Runtime/Metasound/Source/MetasoundEngine/Private/MetasoundWaveProxyReader.cpp b/Engine/Plugins/Runtime/Metasound/Source/MetasoundEngine/Private/MetasoundWaveProxyReader.cpp
--- a/Engine/Plugins/Runtime/Metasound/Source/MetasoundEngine/Private/MetasoundWaveProxyReader.cpp
+++ b/Engine/Plugins/Runtime/Metasound/Source/MetasoundEngine/Private/MetasoundWaveProxyReader.cpp
@@ -431,7 +431,6 @@ namespace Metasound
 					InStartTimeInSeconds);
 			}
 		}
-		CurrentFrameIndex = FMath::Clamp(static_cast<int32>(InStartTimeInSeconds * GetSampleRate()), 0, GetNumFramesInWave());
 
 		// Get codec ptr by reading the header info from the decoder input.
 		ICodecRegistry::FCodecPtr Codec = ICodecRegistry::Get().FindCodecByParsingInput(DecoderInput.Get());
@@ -441,6 +440,18 @@ namespace Metasound
 			return false;
 		}
 
+		// Read the sample rate and number of frames from the header 
+		// Similar to refreshing the wave data in FMixerBuffer::CreateStreamingBuffer
+		// This is a runtime hack to address incorrect sample rate on soundwaves 
+		// on platforms with Resample for Device enabled (UE-183237)
+		Audio::FFormatDescriptorSection FormatDesc;
+		DecoderInput->FindSection(FormatDesc);
+		SampleRate = FormatDesc.NumFramesPerSec;
+		NumFramesInWave = FormatDesc.NumFrames;
+		// end hack
+
+		CurrentFrameIndex = FMath::Clamp(static_cast<int32>(InStartTimeInSeconds * GetSampleRate()), 0, GetNumFramesInWave());
+
 		// Create the decoder
 		Decoder = Codec->CreateDecoder(DecoderInput.Get(), &DecoderOutput);
 		if (!Decoder.IsValid())
