Fix MetaSound Variable update timing when connected to an input
Avoid unnecessary copying of metasound variables.

diff --git a/Engine/Plugins/Runtime/Metasound/Source/MetasoundGraphCore/Private/MetasoundNode.cpp b/Engine/Plugins/Runtime/Metasound/Source/MetasoundGraphCore/Private/MetasoundNode.cpp
--- a/Engine/Plugins/Runtime/Metasound/Source/MetasoundGraphCore/Private/MetasoundNode.cpp
+++ b/Engine/Plugins/Runtime/Metasound/Source/MetasoundGraphCore/Private/MetasoundNode.cpp
@@ -30,4 +30,31 @@ namespace Metasound
 	{
 		return Info;
 	}
+
+	const FVertexInterface& FNode::GetVertexInterface() const
+	{
+		return Info.DefaultInterface;
+	}
+
+	bool FNode::SetVertexInterface(const FVertexInterface& InInterface)
+	{
+		// This code looks a bit counter-intuitive. Here's the explanation: 
+		// 1. This node doesn't support manipulating the interface. The default 
+		//    interface is the only interface that works with this node.
+		// 2. The only valid argument to this function is so set the interface 
+		//    to it's default state. 
+		// 3. If the input vertex interface is exactly the same as the default 
+		//    interface, then we can say that this function executed 
+		//    successfully.
+		// 4. Return "true" if the supplied interface is the same as the 
+		//    default interface, false otherwise. 
+		return GetVertexInterface() == InInterface;
+	}
+
+	bool FNode::IsVertexInterfaceSupported(const FVertexInterface& InInterface) const
+	{
+		// This node doesn't support manipulating the interface. The default 
+		// interface is the only interface that works with this node.
+		return GetVertexInterface() == InInterface;
+	}
 }
diff --git a/Engine/Plugins/Runtime/Metasound/Source/MetasoundGraphCore/Private/MetasoundVariableNodes.cpp b/Engine/Plugins/Runtime/Metasound/Source/MetasoundGraphCore/Private/MetasoundVariableNodes.cpp
--- a/Engine/Plugins/Runtime/Metasound/Source/MetasoundGraphCore/Private/MetasoundVariableNodes.cpp
+++ b/Engine/Plugins/Runtime/Metasound/Source/MetasoundGraphCore/Private/MetasoundVariableNodes.cpp
@@ -33,4 +33,31 @@ namespace Metasound
 			return FNodeClassName{NodeNamespace, InDataTypeName, ""};
 		}
 	}
+
+	
+	namespace MetasoundVariableNodesPrivate
+	{
+		bool IsSupportedVertexData(const IDataReference& InCurrentVariableRef, const FInputVertexInterfaceData& InNew)
+		{
+			using namespace VariableNames;
+#if DO_CHECK
+			// Variable nodes are internal to the graph and their underlying TVariable<> Object should 
+			// never change. Changing the underlying variable would break the state of the TVariable<>
+			// object which is initialized during CreateOperator(...) calls of various TVariable*Operators.
+			FDataReferenceID CurrentID = GetDataReferenceID(InCurrentVariableRef);
+			FDataReferenceID NewID = nullptr;
+
+			if (const FAnyDataReference* InNewVar = InNew.FindDataReference(METASOUND_GET_PARAM_NAME(InputVariable)))
+			{
+				NewID = GetDataReferenceID(*InNewVar);
+			}
+
+			// The new data reference must be missing or equal to the current data reference.
+			return  (NewID == nullptr) || (NewID == CurrentID);
+#else
+			return true;
+#endif // #if DO_CHECK
+
+		}
+	}
 }
diff --git a/Engine/Plugins/Runtime/Metasound/Source/MetasoundGraphCore/Public/MetasoundNode.h b/Engine/Plugins/Runtime/Metasound/Source/MetasoundGraphCore/Public/MetasoundNode.h
--- a/Engine/Plugins/Runtime/Metasound/Source/MetasoundGraphCore/Public/MetasoundNode.h
+++ b/Engine/Plugins/Runtime/Metasound/Source/MetasoundGraphCore/Public/MetasoundNode.h
@@ -16,13 +16,19 @@ namespace Metasound
 			virtual ~FNode() = default;
 
 			/** Return the name of this specific instance of the node class. */
-			const FVertexName& GetInstanceName() const override;
+			virtual const FVertexName& GetInstanceName() const override;
 
 			/** Return the ID of this specific instance of the node class. */
-			const FGuid& GetInstanceID() const override;
+			virtual const FGuid& GetInstanceID() const override;
 
 			/** Return metadata associated with this node. */
-			const FNodeClassMetadata& GetMetadata() const override;
+			virtual const FNodeClassMetadata& GetMetadata() const override;
+
+			virtual const FVertexInterface& GetVertexInterface() const override;
+
+			virtual bool SetVertexInterface(const FVertexInterface& InInterface) override;
+
+			virtual bool IsVertexInterfaceSupported(const FVertexInterface& InInterface) const override;
 
 		private:
 
diff --git a/Engine/Plugins/Runtime/Metasound/Source/MetasoundGraphCore/Public/MetasoundVariable.h b/Engine/Plugins/Runtime/Metasound/Source/MetasoundGraphCore/Public/MetasoundVariable.h
--- a/Engine/Plugins/Runtime/Metasound/Source/MetasoundGraphCore/Public/MetasoundVariable.h
+++ b/Engine/Plugins/Runtime/Metasound/Source/MetasoundGraphCore/Public/MetasoundVariable.h
@@ -2,10 +2,21 @@
 
 #pragma once
 
+#include "MetasoundDataFactory.h"
 #include "MetasoundDataReference.h"
+#include "MetasoundLiteral.h"
+#include "MetasoundOperatorSettings.h"
 
 namespace Metasound
 {
+	namespace MetasoundVariablePrivate
+	{
+		// This tag resolves issues around the behavior of the TLiteralTraits<>::bIsParsable.
+		// Use of the tag dispatch in constructing TVariable<> disables implicit conversions
+		// of FLiteral::FVariantType underlying types to FLiteral. 
+		struct FConstructWithLiteral {};
+	}
+
 	/** A MetaSound Variable contains a data reference's prior and current value.
 	 *
 	 * @tparam DataType - Underlying data type of data reference.
@@ -23,46 +34,74 @@ namespace Metasound
 
 		/** Create a delayed variable.
 		 *
-		 * @param InDelayedReference - A writable reference which will hold the delayed
-		 *                           version of the data reference.
+		 * @param InInitialValue - A literal used to construct the initial value.
 		 */
-		TVariable(FWriteReference InDelayedReference)
-		: DelayedDataReference(InDelayedReference)
-		, DataReference(InDelayedReference)
+		TVariable(const FLiteral& InInitialValue, MetasoundVariablePrivate::FConstructWithLiteral)
+		: Literal(InInitialValue)
+		{
+		}
+
+		bool RequiresDelayedDataCopy() const
 		{
+			return DelayedDataReference.IsSet() && DataReference.IsSet();
 		}
 
 		void CopyReferencedData()
 		{
-			*DelayedDataReference = *DataReference;
+			check(RequiresDelayedDataCopy());
+			*(*DelayedDataReference) = *(*DataReference);
 		}
 
 		void SetDataReference(FReadReference InDataReference)
 		{
-			DataReference = InDataReference;
+			DataReference.Emplace(InDataReference);
 		}
-
-		void SetDelayedDataValue(const DataType& InDataType)
+		
+		void InitDataReference(const FOperatorSettings& InOperatorSettings)
 		{
-			*DelayedDataReference = InDataType;
+			if (!DataReference.IsSet())
+			{
+				InitDelayedDataReference(InOperatorSettings);
+				DataReference = GetDelayedDataReference();
+			}
 		}
 
 		/** Get the current data reference. */
 		FReadReference GetDataReference() const
 		{
-			return DataReference;
+			checkf(DataReference.IsSet(), TEXT("InitDataReference must be called before accessing"));
+			return *DataReference;
+		}
+
+		void InitDelayedDataReference(const FOperatorSettings& InOperatorSettings)
+		{
+			if (!DelayedDataReference.IsSet())
+			{
+				DelayedDataReference.Emplace(TDataWriteReferenceLiteralFactory<DataType>::CreateExplicitArgs(InOperatorSettings, Literal));
+			}
 		}
 
 		/** Get the delayed data reference */
 		FReadReference GetDelayedDataReference() const
 		{
-			return DelayedDataReference;
+			checkf(DelayedDataReference.IsSet(), TEXT("InitDelayedDataReference must be called before accessing"));
+			return *DelayedDataReference;
+		}
+
+		void Reset(const FOperatorSettings& InOperatorSettings)
+		{
+			if (DelayedDataReference.IsSet())
+			{
+				FWriteReference DelayedWritable = *DelayedDataReference;
+				*DelayedWritable = TDataTypeLiteralFactory<DataType>::CreateExplicitArgs(InOperatorSettings, Literal);
+			}
 		}
 
 	private:
 
-		FWriteReference DelayedDataReference;
-		FReadReference DataReference;
+		FLiteral Literal;
+		TOptional<FWriteReference> DelayedDataReference;
+		TOptional<FReadReference> DataReference;
 	};
 
 	/** Template to determine if data type is a variable. */
diff --git a/Engine/Plugins/Runtime/Metasound/Source/MetasoundGraphCore/Public/MetasoundVariableNodes.h b/Engine/Plugins/Runtime/Metasound/Source/MetasoundGraphCore/Public/MetasoundVariableNodes.h
--- a/Engine/Plugins/Runtime/Metasound/Source/MetasoundGraphCore/Public/MetasoundVariableNodes.h
+++ b/Engine/Plugins/Runtime/Metasound/Source/MetasoundGraphCore/Public/MetasoundVariableNodes.h
@@ -96,17 +96,16 @@ namespace Metasound
 
 	namespace MetasoundVariableNodesPrivate
 	{
+		METASOUNDGRAPHCORE_API bool IsSupportedVertexData(const IDataReference& InCurrentVariableRef, const FInputVertexInterfaceData& InNew);
+
 		template<typename DataType>
-		class TVariableOperator : public TExecutableOperator<TVariableOperator<DataType>>
+		class TVariableOperator : public FNoOpOperator
 		{
-			using Super = TExecutableOperator<TVariableOperator<DataType>>;
 			using FVariable = TVariable<DataType>;
 
 		public:
 			TVariableOperator(const FOperatorSettings& InOperatorSettings, const FLiteral& InLiteral)
-			: Literal(InLiteral)
-			, Variable(TDataWriteReference<FVariable>::CreateNew(TDataWriteReferenceLiteralFactory<DataType>::CreateExplicitArgs(InOperatorSettings, Literal)))
-			, bCopyReferenceDataOnExecute(false)
+			: Variable(TDataWriteReference<FVariable>::CreateNew(InLiteral, MetasoundVariablePrivate::FConstructWithLiteral()))
 			{
 			}
 
@@ -114,57 +113,442 @@ namespace Metasound
 
 			virtual FDataReferenceCollection GetInputs() const override
 			{
-				FDataReferenceCollection Collection;
-				return Collection;
+				checkNoEntry();
+				return FDataReferenceCollection{};
 			}
 
 			virtual FDataReferenceCollection GetOutputs() const override
+			{
+				checkNoEntry();
+				return FDataReferenceCollection{};
+			}
+
+			virtual void Bind(FVertexInterfaceData& InVertexData) const override
+			{
+				const_cast<TVariableOperator<DataType>*>(this)->BindInputs(InVertexData.GetInputs());
+				const_cast<TVariableOperator<DataType>*>(this)->BindOutputs(InVertexData.GetOutputs());
+			}
+
+			virtual void BindInputs(FInputVertexInterfaceData& InVertexData) override
+			{
+			}
+
+			virtual void BindOutputs(FOutputVertexInterfaceData& OutVertexData) override
 			{
 				using namespace VariableNames; 
 
-				FDataReferenceCollection Collection;
-				Collection.AddDataWriteReference(METASOUND_GET_PARAM_NAME(OutputVariable), Variable);
-				return Collection;
+				OutVertexData.BindWriteVertex(METASOUND_GET_PARAM_NAME(OutputVariable), Variable);
 			}
 
-			void Execute()
+		private:
+
+			TDataWriteReference<FVariable> Variable;
+		};
+
+		/** Operator for providing setting a variable */
+		template<typename DataType>
+		class TVariableMutatorOperator : public IOperator
+		{
+		public:
+			using FVariable = TVariable<DataType>;
+
+			static FVertexInterface DeclareVertexInterface()
+			{
+				using namespace VariableNames;
+
+				return FVertexInterface(
+					FInputVertexInterface(
+						TInputDataVertex<DataType>(METASOUND_GET_PARAM_NAME(InputData), InputDataMetadata),
+						TInputDataVertex<FVariable>(METASOUND_GET_PARAM_NAME(InputVariable), InputVariableMetadata)
+					),
+					FOutputVertexInterface(
+						TOutputDataVertex<FVariable>(METASOUND_GET_PARAM_NAME(OutputVariable), OutputVariableMetadata)
+					)
+				);
+			}
+
+			static const FNodeClassMetadata& GetNodeInfo()
 			{
-				if (bCopyReferenceDataOnExecute)
+				auto CreateNodeClassMetadata = []()
 				{
-					Variable->CopyReferencedData();
+					FNodeClassMetadata Info;
+
+					Info.ClassName = VariableNames::GetVariableMutatorNodeClassName<DataType>();
+					Info.MajorVersion = 1;
+					Info.MinorVersion = 0;
+
+#if WITH_EDITOR
+					Info.DisplayStyle.bShowName = false;
+					Info.DisplayName = LOCTEXT("Metasound_VariableMutatorNodeDisplayName", "Set");
+					Info.Description = LOCTEXT("Metasound_VariableMutatorNodeDescription", "Set variable on MetaSound graph.");
+					Info.CategoryHierarchy = { LOCTEXT("Metasound_VariableCategory", "Variable") };
+#endif // WITH_EDITOR
+					Info.Author = PluginAuthor;
+					Info.PromptIfMissing = PluginNodeMissingPrompt;
+					Info.DefaultInterface = DeclareVertexInterface();
+
+					return Info;
+				};
+
+				static const FNodeClassMetadata Metadata = CreateNodeClassMetadata();
+				return Metadata;
+			}
+
+			TVariableMutatorOperator(TDataWriteReference<FVariable> InVariable, TDataReadReference<DataType> InData)
+			: Variable(InVariable)
+			, Data(InData)
+			{
+				Variable->SetDataReference(Data);
+			}
+
+			static TUniquePtr<IOperator> CreateOperator(const FBuildOperatorParams& InParams, FBuildResults& OutResults) 
+			{
+				using namespace VariableNames;
+
+				if (const FAnyDataReference* VariableRef = InParams.InputData.FindDataReference(METASOUND_GET_PARAM_NAME(InputVariable)))
+				{
+					TDataWriteReference<FVariable> Variable = VariableRef->GetDataWriteReference<FVariable>();
+					TDataReadReference<DataType> Data = InParams.InputData.GetOrCreateDefaultDataReadReference<DataType>(METASOUND_GET_PARAM_NAME(InputData), InParams.OperatorSettings);
+
+					return MakeUnique<TVariableMutatorOperator<DataType>>(Variable, Data);
 				}
-				else
+				else 
 				{
-					// The first time a variable node is run, it should not copy
-					// reference data in the variable, but instead use the original
-					// initial value of the variable. 
-					//
-					// The TVariableNode is currently executed before the TVariableDeferredAccessor
-					// nodes. This execution order is managed in the Metasound::Frontend::FGraphController.
-					// Because the TVariableNode is executed before the TVariableDeferredAccessor node
-					// it can undesireably override the "init" value of the variable with the "init" 
-					// value of the data reference set in the TVariableMutatorNode. 
-					// This would mean that the first call to execute on TVariableDeferredAccessor
-					// would result in reading the "init" value of the data reference
-					// set in the TVariableMutatorNode as opposed to the "init" value
-					// of the data reference set in the TVariableNode. This boolean
-					// protects against that situation so that on first call to execute, the 
-					// TVariableDeferredAccessor always reads the "init" value provided
-					// by the TVariableNode. 
-					bCopyReferenceDataOnExecute = true;
+					// Nothing to do if there's no input data.
+					UE_LOG(LogMetaSound, Verbose, TEXT("Missing internal variable connection. Failed to create valid \"SetVariable\" operator"));
+					return MakeUnique<FNoOpOperator>();
 				}
 			}
 
+			virtual ~TVariableMutatorOperator() = default;
+
+			virtual FDataReferenceCollection GetInputs() const override
+			{
+				checkNoEntry();
+				return FDataReferenceCollection{};
+			}
+
+			virtual FDataReferenceCollection GetOutputs() const override
+			{
+				checkNoEntry();
+				return FDataReferenceCollection{};
+			}
+
+			virtual void Bind(FVertexInterfaceData& InVertexData) const override
+			{
+				const_cast<TVariableMutatorOperator<DataType>*>(this)->BindInputs(InVertexData.GetInputs());
+				const_cast<TVariableMutatorOperator<DataType>*>(this)->BindOutputs(InVertexData.GetOutputs());
+			}
+			
+			virtual void BindInputs(FInputVertexInterfaceData& InVertexData) override
+			{
+				using namespace VariableNames; 
+
+				checkf(IsSupportedVertexData(Variable, InVertexData), TEXT("Input variables cannot be bound to new underlying objects"));
+
+				InVertexData.BindWriteVertex(METASOUND_GET_PARAM_NAME(InputVariable), Variable);
+				InVertexData.BindReadVertex(METASOUND_GET_PARAM_NAME(InputData), Data);
+
+				// Update input variable with new input data in case it has been updated. 
+				Variable->SetDataReference(Data);
+			}
+
+			virtual void BindOutputs(FOutputVertexInterfaceData& InVertexData) override
+			{
+				using namespace VariableNames; 
+				InVertexData.BindWriteVertex(METASOUND_GET_PARAM_NAME(OutputVariable), Variable);
+			}
+
+			void PostExecute()
+			{
+				Variable->CopyReferencedData();
+			}
+
 			void Reset(const IOperator::FResetParams& InParams)
 			{
-				Variable->SetDelayedDataValue(TDataTypeLiteralFactory<DataType>::CreateExplicitArgs(InParams.OperatorSettings, Literal));
+				Variable->Reset(InParams.OperatorSettings);
+			}
+
+			virtual IOperator::FExecuteFunction GetExecuteFunction() override
+			{
+				// Do not do anything during execution
+				return nullptr;
+			}
+
+			virtual IOperator::FPostExecuteFunction GetPostExecuteFunction() override
+			{
+				if (Variable->RequiresDelayedDataCopy())
+				{
+					// Only copy over data if it's required. Otherwise, avoid the copy
+					// to improve performance.
+					return &StaticPostExecute;
+				}
+				else
+				{
+					return nullptr;
+				}
+			}
+
+			virtual IOperator::FResetFunction GetResetFunction() override
+			{
+				if (Variable->RequiresDelayedDataCopy())
+				{
+					// Reset is only applied to the delayed data reference. If we
+					// are not modifying the delayed data reference, do not perform
+					// any reset.
+					return &StaticReset;
+				}
+				else
+				{
+					return nullptr;
+				}
+			}
+
+		private:
+			static void StaticPostExecute(IOperator* InOperator)
+			{
+				static_cast<TVariableMutatorOperator<DataType>*>(InOperator)->PostExecute();
+			}
+
+			static void StaticReset(IOperator* InOperator, const IOperator::FResetParams& InParams)
+			{
+				static_cast<TVariableMutatorOperator<DataType>*>(InOperator)->Reset(InParams);
+			}
+
+			TDataWriteReference<FVariable> Variable;
+			TDataReadReference<DataType> Data;
+		};
+
+		/** Operator for providing delayed access to a variable */
+		template<typename DataType>
+		class TVariableDeferredAccessorOperator : public FNoOpOperator 
+		{
+		public:
+
+			using FVariable = TVariable<DataType>;
+
+			static FVertexInterface DeclareVertexInterface()
+			{
+				using namespace VariableNames; 
+
+				return FVertexInterface(
+					FInputVertexInterface(
+						TInputDataVertex<FVariable>(METASOUND_GET_PARAM_NAME(InputVariable), InputVariableMetadata)
+					),
+					FOutputVertexInterface(
+						TOutputDataVertex<FVariable>(METASOUND_GET_PARAM_NAME(OutputVariable), OutputVariableMetadata),
+						TOutputDataVertex<DataType>(METASOUND_GET_PARAM_NAME(OutputData), OutputDataMetadata)
+					)
+				);
+			}
+
+			static const FNodeClassMetadata& GetNodeInfo()
+			{
+				auto CreateNodeClassMetadata = []()
+				{
+					FNodeClassMetadata Info;
+
+					Info.ClassName = VariableNames::GetVariableDeferredAccessorNodeClassName<DataType>();
+					Info.MajorVersion = 1;
+					Info.MinorVersion = 0;
+#if WITH_EDITOR
+					Info.DisplayStyle.bShowName = false;
+					Info.DisplayName = LOCTEXT("Metasound_VariableDeferredAccessorNodeDisplayName", "Get Delayed");
+					Info.Description = LOCTEXT("Metasound_VariableDeferredAccessorNodeDescription", "Get a delayed variable on MetaSound graph.");
+					Info.CategoryHierarchy = { LOCTEXT("Metasound_VariableCategory", "Variable") };
+#endif // WITH_EDITOR
+					Info.Author = PluginAuthor;
+					Info.PromptIfMissing = PluginNodeMissingPrompt;
+					Info.DefaultInterface = DeclareVertexInterface();
+
+					return Info;
+				};
+
+				static const FNodeClassMetadata Metadata = CreateNodeClassMetadata();
+				return Metadata;
+			}
+
+			TVariableDeferredAccessorOperator(TDataWriteReference<FVariable> InVariable)
+			: Variable(InVariable)
+			, DelayedData(Variable->GetDelayedDataReference())
+			{
+			}
+
+			static TUniquePtr<IOperator> CreateOperator(const FBuildOperatorParams& InParams, FBuildResults& OutResults) 
+			{
+				using namespace VariableNames;
+
+				if (const FAnyDataReference* VariableRef = InParams.InputData.FindDataReference(METASOUND_GET_PARAM_NAME(InputVariable)))
+				{
+					TDataWriteReference<FVariable> DelayedVariable = VariableRef->GetDataWriteReference<FVariable>();
+					DelayedVariable->InitDelayedDataReference(InParams.OperatorSettings);
+
+					return MakeUnique<TVariableDeferredAccessorOperator>(DelayedVariable);
+				}
+				else 
+				{
+					// Nothing to do if there's no input data.
+					UE_LOG(LogMetaSound, Verbose, TEXT("Missing internal variable connection. Failed to create valid \"GetDelayedVariable\" operator"));
+					return MakeUnique<FNoOpOperator>();
+				}
+			}
+
+			virtual ~TVariableDeferredAccessorOperator() = default;
+
+			virtual FDataReferenceCollection GetInputs() const override
+			{
+				checkNoEntry();
+				return FDataReferenceCollection{};
+			}
+
+			virtual FDataReferenceCollection GetOutputs() const override
+			{
+				checkNoEntry();
+				return FDataReferenceCollection{};
+			}
+
+			virtual void Bind(FVertexInterfaceData& InVertexData) const override
+			{
+				const_cast<TVariableDeferredAccessorOperator<DataType>*>(this)->BindInputs(InVertexData.GetInputs());
+				const_cast<TVariableDeferredAccessorOperator<DataType>*>(this)->BindOutputs(InVertexData.GetOutputs());
+			}
+
+			virtual void BindInputs(FInputVertexInterfaceData& InVertexData) override
+			{
+				using namespace VariableNames; 
+
+				checkf(IsSupportedVertexData(Variable, InVertexData), TEXT("Input variables cannot be bound to new underlying objects"));
+
+				InVertexData.BindWriteVertex<FVariable>(METASOUND_GET_PARAM_NAME(OutputVariable), Variable);
+			}
+
+			virtual void BindOutputs(FOutputVertexInterfaceData& InVertexData) override
+			{
+				using namespace VariableNames; 
+
+				InVertexData.BindWriteVertex<FVariable>(METASOUND_GET_PARAM_NAME(OutputVariable), Variable);
+				InVertexData.BindReadVertex<DataType>(METASOUND_GET_PARAM_NAME(OutputData), Variable->GetDelayedDataReference());
+			}
+
+		private:
+
+			TDataWriteReference<FVariable> Variable;
+			TDataReadReference<DataType> DelayedData;
+		};
+
+		/** Operator for providing inline access to a variable */
+		template<typename DataType>
+		class TVariableAccessorOperator : public FNoOpOperator 
+		{
+		public:
+
+			using FVariable = TVariable<DataType>;
+
+			static FVertexInterface DeclareVertexInterface()
+			{
+				using namespace VariableNames; 
+
+				return FVertexInterface(
+					FInputVertexInterface(
+						TInputDataVertex<FVariable>(METASOUND_GET_PARAM_NAME(InputVariable), InputDataMetadata)
+					),
+					FOutputVertexInterface(
+						TOutputDataVertex<DataType>(METASOUND_GET_PARAM_NAME(OutputData), OutputDataMetadata),
+						TOutputDataVertex<FVariable>(METASOUND_GET_PARAM_NAME(OutputVariable), OutputVariableMetadata)
+					)
+				);
+			}
+
+			static const FNodeClassMetadata& GetNodeInfo()
+			{
+				auto CreateNodeClassMetadata = []()
+				{
+					FNodeClassMetadata Info;
+
+					Info.ClassName = VariableNames::GetVariableAccessorNodeClassName<DataType>();
+					Info.MajorVersion = 1;
+					Info.MinorVersion = 0;
+#if WITH_EDITOR
+					Info.DisplayStyle.bShowName = false;
+					Info.DisplayName = LOCTEXT("Metasound_VariableAccessorNodeDisplayName", "Get");
+					Info.Description = LOCTEXT("Metasound_VariableAccessorNodeDescription", "Get variable on MetaSound graph.");
+#endif // WITH_EDITOR
+
+					Info.Author = PluginAuthor;
+					Info.PromptIfMissing = PluginNodeMissingPrompt;
+					Info.DefaultInterface = DeclareVertexInterface();
+
+					return Info;
+				};
+
+				static const FNodeClassMetadata Metadata = CreateNodeClassMetadata();
+				return Metadata;
+			}
+
+			TVariableAccessorOperator(TDataWriteReference<FVariable> InVariable)
+			: Variable(InVariable)
+			{
+			}
+
+			static TUniquePtr<IOperator> CreateOperator(const FBuildOperatorParams& InParams, FBuildResults& OutResults) 
+			{
+				using namespace VariableNames;
+
+				if (const FAnyDataReference* VariableRef = InParams.InputData.FindDataReference(METASOUND_GET_PARAM_NAME(InputVariable)))
+				{
+					TDataWriteReference<FVariable> Variable = VariableRef->GetDataWriteReference<FVariable>();
+					Variable->InitDataReference(InParams.OperatorSettings);
+
+					return MakeUnique<TVariableAccessorOperator>(Variable);
+				}
+				else 
+				{
+					// Nothing to do if there's no input data.
+					UE_LOG(LogMetaSound, Verbose, TEXT("Missing internal variable connection. Failed to create valid \"GetVariable\" operator"));
+					return MakeUnique<FNoOpOperator>();
+				}
+			}
+
+			virtual ~TVariableAccessorOperator() = default;
+
+			virtual FDataReferenceCollection GetInputs() const override
+			{
+				checkNoEntry();
+				return FDataReferenceCollection{};
+			}
+
+			virtual FDataReferenceCollection GetOutputs() const override
+			{
+				checkNoEntry();
+				return FDataReferenceCollection{};
+			}
+
+			virtual void Bind(FVertexInterfaceData& InVertexData) const override
+			{
+				const_cast<TVariableAccessorOperator<DataType>*>(this)->BindInputs(InVertexData.GetInputs());
+				const_cast<TVariableAccessorOperator<DataType>*>(this)->BindOutputs(InVertexData.GetOutputs());
+			}
+
+			virtual void BindInputs(FInputVertexInterfaceData& InVertexData) override
+			{
+				using namespace VariableNames; 
+
+				checkf(IsSupportedVertexData(Variable, InVertexData), TEXT("Input variables cannot be bound to new underlying objects"));
+
+				InVertexData.BindWriteVertex<FVariable>(METASOUND_GET_PARAM_NAME(InputVariable), Variable);
+			}
+
+			virtual void BindOutputs(FOutputVertexInterfaceData& InVertexData) override
+			{
+				using namespace VariableNames; 
+
+				InVertexData.BindReadVertex<DataType>(METASOUND_GET_PARAM_NAME(OutputData), Variable->GetDataReference());
+				InVertexData.BindWriteVertex<FVariable>(METASOUND_GET_PARAM_NAME(OutputVariable), Variable);
 			}
 
 		private:
 
-			FLiteral Literal;
 			TDataWriteReference<FVariable> Variable;
-			bool bCopyReferenceDataOnExecute;
 		};
 	}
 
@@ -237,21 +621,6 @@ namespace Metasound
 		{
 		}
 
-		virtual const FVertexInterface& GetVertexInterface() const override
-		{
-			return Interface;
-		}
-
-		virtual bool SetVertexInterface(const FVertexInterface& InInterface) override
-		{
-			return Interface == InInterface;
-		}
-
-		virtual bool IsVertexInterfaceSupported(const FVertexInterface& InInterface) const override
-		{
-			return Interface == InInterface;
-		}
-
 		virtual TSharedRef<IOperatorFactory, ESPMode::ThreadSafe> GetDefaultOperatorFactory() const override
 		{
 			return Factory;
@@ -267,38 +636,10 @@ namespace Metasound
 	template<typename DataType>
 	class TVariableMutatorNode : public FNodeFacade
 	{
-		using FVariable = TVariable<DataType>;
-
-		class FOperator : public FNoOpOperator 
-		{
-		public:
-			static TUniquePtr<IOperator> CreateOperator(const FCreateOperatorParams& InParams, FBuildErrorArray& OutErrors);
-			static FVertexInterface DeclareVertexInterface();
-			static const FNodeClassMetadata& GetNodeInfo();
-
-			FOperator(TDataReadReference<FVariable> InVariable)
-			: Variable(InVariable)
-			{
-			}
-
-			virtual ~FOperator() = default;
-
-			virtual FDataReferenceCollection GetOutputs() const override
-			{
-				using namespace VariableNames; 
-
-				FDataReferenceCollection Collection;
-				Collection.AddDataReadReference<FVariable>(METASOUND_GET_PARAM_NAME(OutputVariable), Variable);
-				return Collection;
-			}
-
-		private:
-			TDataReadReference<FVariable> Variable;
-		};
 
 	public:
 		TVariableMutatorNode(const FVertexName& InNodeName, const FGuid& InInstanceID)
-		: FNodeFacade(InNodeName, InInstanceID, TFacadeOperatorClass<FOperator>())
+		: FNodeFacade(InNodeName, InInstanceID, TFacadeOperatorClass<MetasoundVariableNodesPrivate::TVariableMutatorOperator<DataType>>())
 		{
 		}
 
@@ -307,86 +648,6 @@ namespace Metasound
 		{}
 	};
 
-	template<typename DataType>
-	TUniquePtr<IOperator> TVariableMutatorNode<DataType>::FOperator::CreateOperator(const FCreateOperatorParams& InParams, FBuildErrorArray& OutErrors) 
-	{
-		using namespace VariableNames;
-
-		const FDataReferenceCollection& Inputs = InParams.InputDataReferences;
-
-		if (Inputs.ContainsDataWriteReference<FVariable>(METASOUND_GET_PARAM_NAME(InputVariable)))
-		{
-			// If a variable is provided, set the reference to read.
-			TDataWriteReference<FVariable> Variable = Inputs.GetDataWriteReference<FVariable>(METASOUND_GET_PARAM_NAME(InputVariable));
-			if (Inputs.ContainsDataReadReference<DataType>(METASOUND_GET_PARAM_NAME(InputData)))
-			{
-				// Update the input variable with the data reference to copy from.
-				TDataReadReference InputData = Inputs.GetDataReadReference<DataType>(METASOUND_GET_PARAM_NAME(InputData));
-				Variable->SetDataReference(InputData);
-			}
-
-			return MakeUnique<FOperator>(Variable);
-		}
-		else if (Inputs.ContainsDataReadReference<DataType>(METASOUND_GET_PARAM_NAME(InputData)))
-		{
-			// If no input variable is provided create Variable with input variable
-			TDataReadReference<DataType> InputData = Inputs.GetDataReadReference<DataType>(METASOUND_GET_PARAM_NAME(InputData));
-			TDataWriteReference<DataType> InitData = TDataWriteReference<DataType>::CreateNew(*InputData);
-			TDataWriteReference<FVariable> Variable = TDataWriteReference<FVariable>::CreateNew(InitData);
-			Variable->SetDataReference(InputData);
-			return MakeUnique<FOperator>(Variable);
-		}
-		else 
-		{
-			// Nothing to do if there's no input data.
-			UE_LOG(LogMetaSound, Verbose, TEXT("Missing internal variable connection. Failed to create valid \"SetVariable\" operator"));
-			return MakeUnique<FNoOpOperator>();
-		}
-	}
-
-	template<typename DataType>
-	FVertexInterface TVariableMutatorNode<DataType>::FOperator::DeclareVertexInterface()
-	{
-		using namespace VariableNames;
-
-		return FVertexInterface(
-			FInputVertexInterface(
-				TInputDataVertex<DataType>(METASOUND_GET_PARAM_NAME(InputData), InputDataMetadata),
-				TInputDataVertex<FVariable>(METASOUND_GET_PARAM_NAME(InputVariable), InputVariableMetadata)
-			),
-			FOutputVertexInterface(
-				TOutputDataVertex<FVariable>(METASOUND_GET_PARAM_NAME(OutputVariable), OutputVariableMetadata)
-			)
-		);
-	}
-
-	template<typename DataType>
-	const FNodeClassMetadata& TVariableMutatorNode<DataType>::FOperator::GetNodeInfo()
-	{
-		auto CreateNodeClassMetadata = []()
-		{
-			FNodeClassMetadata Info;
-
-			Info.ClassName = VariableNames::GetVariableMutatorNodeClassName<DataType>();
-			Info.MajorVersion = 1;
-			Info.MinorVersion = 0;
-
-#if WITH_EDITOR
-			Info.DisplayStyle.bShowName = false;
-			Info.DisplayName = LOCTEXT("Metasound_VariableMutatorNodeDisplayName", "Set");
-			Info.Description = LOCTEXT("Metasound_VariableMutatorNodeDescription", "Set variable on MetaSound graph.");
-			Info.CategoryHierarchy = { LOCTEXT("Metasound_VariableCategory", "Variable") };
-#endif // WITH_EDITOR
-			Info.Author = PluginAuthor;
-			Info.PromptIfMissing = PluginNodeMissingPrompt;
-			Info.DefaultInterface = DeclareVertexInterface();
-
-			return Info;
-		};
-
-		static const FNodeClassMetadata Metadata = CreateNodeClassMetadata();
-		return Metadata;
-	}
 
 	/** TVariableDeferredAccessorNode provides access to the prior executions variable value.
 	 * TVariableDeferredAccessorNodes must always be before TVariableMutatorNodes in the dependency
@@ -395,42 +656,10 @@ namespace Metasound
 	template<typename DataType>
 	class TVariableDeferredAccessorNode : public FNodeFacade
 	{
-		using FVariable = TVariable<DataType>;
-	
-		class FOperator : public FNoOpOperator 
-		{
-		public:
-			static TUniquePtr<IOperator> CreateOperator(const FCreateOperatorParams& InParams, FBuildErrorArray& OutErrors);
-			static FVertexInterface DeclareVertexInterface();
-			static const FNodeClassMetadata& GetNodeInfo();
-
-			FOperator(TDataWriteReference<FVariable> InDelayedVariable)
-			: DelayedVariable(InDelayedVariable)
-			, DelayedData(DelayedVariable->GetDelayedDataReference())
-			{
-			}
-
-			virtual ~FOperator() = default;
-
-			virtual FDataReferenceCollection GetOutputs() const override
-			{
-				using namespace VariableNames; 
-
-				FDataReferenceCollection Collection;
-				Collection.AddDataWriteReference<FVariable>(METASOUND_GET_PARAM_NAME(OutputVariable), DelayedVariable);
-				Collection.AddDataReadReference<DataType>(METASOUND_GET_PARAM_NAME(OutputData), DelayedVariable->GetDelayedDataReference());
-				return Collection;
-			}
-
-		private:
-			TDataWriteReference<FVariable> DelayedVariable;
-			TDataReadReference<DataType> DelayedData;
-		};
-
 	public:
 
 		TVariableDeferredAccessorNode(const FVertexName& InNodeName, const FGuid& InInstanceID)
-		: FNodeFacade(InNodeName, InInstanceID, TFacadeOperatorClass<FOperator>())
+		: FNodeFacade(InNodeName, InInstanceID, TFacadeOperatorClass<MetasoundVariableNodesPrivate::TVariableDeferredAccessorOperator<DataType>>())
 		{
 		}
 
@@ -439,107 +668,14 @@ namespace Metasound
 		{}
 	};
 
-	template<typename DataType>
-	TUniquePtr<IOperator> TVariableDeferredAccessorNode<DataType>::FOperator::CreateOperator(const FCreateOperatorParams& InParams, FBuildErrorArray& OutErrors) 
-	{
-		using namespace VariableNames;
-
-		const FDataReferenceCollection& Inputs = InParams.InputDataReferences;
-
-		if (Inputs.ContainsDataWriteReference<FVariable>(METASOUND_GET_PARAM_NAME(InputVariable)))
-		{
-			TDataWriteReference<FVariable> DelayedVariable = Inputs.GetDataWriteReference<FVariable>(METASOUND_GET_PARAM_NAME(InputVariable));
-
-			return MakeUnique<FOperator>(DelayedVariable);
-		}
-		else 
-		{
-			// Nothing to do if there's no input data.
-			UE_LOG(LogMetaSound, Verbose, TEXT("Missing internal variable connection. Failed to create valid \"GetDelayedVariable\" operator"));
-			return MakeUnique<FNoOpOperator>();
-		}
-	}
-
-	template<typename DataType>
-	FVertexInterface TVariableDeferredAccessorNode<DataType>::FOperator::DeclareVertexInterface()
-	{
-		using namespace VariableNames; 
-
-		return FVertexInterface(
-			FInputVertexInterface(
-				TInputDataVertex<FVariable>(METASOUND_GET_PARAM_NAME(InputVariable), InputVariableMetadata)
-			),
-			FOutputVertexInterface(
-				TOutputDataVertex<FVariable>(METASOUND_GET_PARAM_NAME(OutputVariable), OutputVariableMetadata),
-				TOutputDataVertex<DataType>(METASOUND_GET_PARAM_NAME(OutputData), OutputDataMetadata)
-			)
-		);
-	}
-
-	template<typename DataType>
-	const FNodeClassMetadata& TVariableDeferredAccessorNode<DataType>::FOperator::GetNodeInfo()
-	{
-		auto CreateNodeClassMetadata = []()
-		{
-			FNodeClassMetadata Info;
-
-			Info.ClassName = VariableNames::GetVariableDeferredAccessorNodeClassName<DataType>();
-			Info.MajorVersion = 1;
-			Info.MinorVersion = 0;
-#if WITH_EDITOR
-			Info.DisplayStyle.bShowName = false;
-			Info.DisplayName = LOCTEXT("Metasound_VariableDeferredAccessorNodeDisplayName", "Get Delayed");
-			Info.Description = LOCTEXT("Metasound_VariableDeferredAccessorNodeDescription", "Get a delayed variable on MetaSound graph.");
-			Info.CategoryHierarchy = { LOCTEXT("Metasound_VariableCategory", "Variable") };
-#endif // WITH_EDITOR
-			Info.Author = PluginAuthor;
-			Info.PromptIfMissing = PluginNodeMissingPrompt;
-			Info.DefaultInterface = DeclareVertexInterface();
-
-			return Info;
-		};
-
-		static const FNodeClassMetadata Metadata = CreateNodeClassMetadata();
-		return Metadata;
-	}
-
 	/** FGetVariable node provides delay free, cpu free access to a set variable. */
 	template<typename DataType>
 	class TVariableAccessorNode : public FNodeFacade
 	{
-		using FVariable = TVariable<DataType>;
-
-		class FOperator : public FNoOpOperator 
-		{
-		public:
-			static TUniquePtr<IOperator> CreateOperator(const FCreateOperatorParams& InParams, FBuildErrorArray& OutErrors);
-			static FVertexInterface DeclareVertexInterface();
-			static const FNodeClassMetadata& GetNodeInfo();
-
-			FOperator(TDataReadReference<FVariable> InVariable)
-			: Variable(InVariable)
-			{
-			}
-
-			virtual ~FOperator() = default;
-
-			virtual FDataReferenceCollection GetOutputs() const override
-			{
-				using namespace VariableNames; 
-
-				FDataReferenceCollection Collection;
-				Collection.AddDataReadReference<DataType>(METASOUND_GET_PARAM_NAME(OutputData), Variable->GetDataReference());
-				Collection.AddDataReadReference<FVariable>(METASOUND_GET_PARAM_NAME(OutputVariable), Variable);
-				return Collection;
-			}
-
-		private:
-			TDataReadReference<FVariable> Variable;
-		};
 	public:
 
 		TVariableAccessorNode(const FVertexName& InNodeName, const FGuid& InInstanceID)
-		: FNodeFacade(InNodeName, InInstanceID, TFacadeOperatorClass<FOperator>())
+		: FNodeFacade(InNodeName, InInstanceID, TFacadeOperatorClass<MetasoundVariableNodesPrivate::TVariableAccessorOperator<DataType>>())
 		{
 		}
 
@@ -547,75 +683,6 @@ namespace Metasound
 		: TVariableAccessorNode(InInitData.InstanceName, InInitData.InstanceID)
 		{}
 	};
-
-	template<typename DataType>
-	TUniquePtr<IOperator> TVariableAccessorNode<DataType>::FOperator::CreateOperator(const FCreateOperatorParams& InParams, FBuildErrorArray& OutErrors) 
-	{
-		using namespace VariableNames;
-
-		const FDataReferenceCollection& Inputs = InParams.InputDataReferences;
-
-		// Update delayed variable.
-		if (Inputs.ContainsDataReadReference<FVariable>(METASOUND_GET_PARAM_NAME(InputVariable)))
-		{
-			TDataReadReference<FVariable> Variable = Inputs.GetDataReadReference<FVariable>(METASOUND_GET_PARAM_NAME(InputVariable));
-			return MakeUnique<FOperator>(Variable);
-		}
-		else if (Inputs.ContainsDataWriteReference<FVariable>(METASOUND_GET_PARAM_NAME(InputVariable)))
-		{
-			TDataReadReference<FVariable> Variable = Inputs.GetDataWriteReference<FVariable>(METASOUND_GET_PARAM_NAME(InputVariable));
-			return MakeUnique<FOperator>(Variable);
-		}
-		else 
-		{
-			// Nothing to do if there's no input data.
-			UE_LOG(LogMetaSound, Verbose, TEXT("Missing internal variable connection. Failed to create valid \"GetVariable\" operator"));
-			return MakeUnique<FNoOpOperator>();
-		}
-	}
-
-	template<typename DataType>
-	FVertexInterface TVariableAccessorNode<DataType>::FOperator::DeclareVertexInterface()
-	{
-		using namespace VariableNames; 
-
-		return FVertexInterface(
-			FInputVertexInterface(
-				TInputDataVertex<FVariable>(METASOUND_GET_PARAM_NAME(InputVariable), InputDataMetadata)
-			),
-			FOutputVertexInterface(
-				TOutputDataVertex<DataType>(METASOUND_GET_PARAM_NAME(OutputData), OutputDataMetadata),
-				TOutputDataVertex<FVariable>(METASOUND_GET_PARAM_NAME(OutputVariable), OutputVariableMetadata)
-			)
-		);
-	}
-
-	template<typename DataType>
-	const FNodeClassMetadata& TVariableAccessorNode<DataType>::FOperator::GetNodeInfo()
-	{
-		auto CreateNodeClassMetadata = []()
-		{
-			FNodeClassMetadata Info;
-
-			Info.ClassName = VariableNames::GetVariableAccessorNodeClassName<DataType>();
-			Info.MajorVersion = 1;
-			Info.MinorVersion = 0;
-#if WITH_EDITOR
-			Info.DisplayStyle.bShowName = false;
-			Info.DisplayName = LOCTEXT("Metasound_VariableAccessorNodeDisplayName", "Get");
-			Info.Description = LOCTEXT("Metasound_VariableAccessorNodeDescription", "Get variable on MetaSound graph.");
-#endif // WITH_EDITOR
-
-			Info.Author = PluginAuthor;
-			Info.PromptIfMissing = PluginNodeMissingPrompt;
-			Info.DefaultInterface = DeclareVertexInterface();
-
-			return Info;
-		};
-
-		static const FNodeClassMetadata Metadata = CreateNodeClassMetadata();
-		return Metadata;
-	}
 }
 
 #undef LOCTEXT_NAMESPACE // MetasoundGraphCore
