Merge conflict reolution (by hand)

diff --git a/Engine/Plugins/Runtime/Metasound/Source/MetasoundEngine/Private/MetasoundGeneratorHandle.cpp b/Engine/Plugins/Runtime/Metasound/Source/MetasoundEngine/Private/MetasoundGeneratorHandle.cpp
--- a/Engine/Plugins/Runtime/Metasound/Source/MetasoundEngine/Private/MetasoundGeneratorHandle.cpp
+++ b/Engine/Plugins/Runtime/Metasound/Source/MetasoundEngine/Private/MetasoundGeneratorHandle.cpp
@@ -164,9 +164,8 @@ void UMetasoundGeneratorHandle::AttachGraphChangedDelegate()
 		// the audio render thread so we need to take steps to assure this UMetasoundGeneratorHandle
 		// hasn't been garbage collected before trying to dereference it later when the delegate "fires".
 		// That is why we capture a TWeakObjectPtr and try to dereference that later!
-		GeneratorGraphChangedDelegateHandle = Generator->OnSetGraph.AddLambda(
-			[WeakGeneratorHandlePtr = TWeakObjectPtr<UMetasoundGeneratorHandle>(this),
-			StatId = this->GetStatID(true)]
+		Metasound::FOnSetGraph::FDelegate NewSetGraphDelegate;
+		NewSetGraphDelegate.BindLambda([WeakGeneratorHandlePtr = TWeakObjectPtr<UMetasoundGeneratorHandle>(this), StatId = this->GetStatID(true)]
 			()
 			{
 				// We are in the audio render thread here, so create a "dispatch task" to be
@@ -185,8 +184,10 @@ void UMetasoundGeneratorHandle::AttachGraphChangedDelegate()
 						}
 					},
 					StatId, nullptr, ENamedThreads::GameThread);
-		});
-}
+			});
+
+		GeneratorGraphChangedDelegateHandle = Generator->AddGraphSetCallback(MoveTemp(NewSetGraphDelegate));
+	}
 }
 
 void UMetasoundGeneratorHandle::DetachGeneratorDelegates()
@@ -205,7 +206,7 @@ void UMetasoundGeneratorHandle::DetachGeneratorDelegates()
 	TSharedPtr<Metasound::FMetasoundGenerator> PinnedGenerator = CachedGeneratorPtr.Pin();
 	if (PinnedGenerator)
 	{
-		PinnedGenerator->OnSetGraph.Remove(GeneratorGraphChangedDelegateHandle);
+		PinnedGenerator->RemoveGraphSetCallback(GeneratorGraphChangedDelegateHandle);
 		GeneratorGraphChangedDelegateHandle.Reset();
 	}
 }
@@ -273,7 +274,7 @@ TSharedPtr<Metasound::FMetasoundGenerator> UMetasoundGeneratorHandle::GetGenerat
 	return PinGenerator();
 }
 
-FDelegateHandle UMetasoundGeneratorHandle::AddGraphSetCallback(const UMetasoundGeneratorHandle::FOnSetGraph& Delegate)
+FDelegateHandle UMetasoundGeneratorHandle::AddGraphSetCallback(UMetasoundGeneratorHandle::FOnSetGraph::FDelegate&& Delegate)
 {
 	
 	FDelegateHandle Handle = OnGeneratorsGraphChanged.Add(Delegate);
diff --git a/Engine/Plugins/Runtime/Metasound/Source/MetasoundEngine/Public/MetasoundGeneratorHandle.h b/Engine/Plugins/Runtime/Metasound/Source/MetasoundEngine/Public/MetasoundGeneratorHandle.h
--- a/Engine/Plugins/Runtime/Metasound/Source/MetasoundEngine/Public/MetasoundGeneratorHandle.h
+++ b/Engine/Plugins/Runtime/Metasound/Source/MetasoundEngine/Public/MetasoundGeneratorHandle.h
@@ -51,7 +51,7 @@ public:
 	// render thread. 
 	DECLARE_MULTICAST_DELEGATE(FOnAttached);
 	DECLARE_MULTICAST_DELEGATE(FOnDetached);
-	DECLARE_DELEGATE(FOnSetGraph);
+	DECLARE_MULTICAST_DELEGATE(FOnSetGraph);
 
 	/**
 	 * Makes a copy of the supplied parameter pack and passes it to the MetaSoundGenerator
@@ -69,7 +69,7 @@ public:
 	// Note: We don't allow direct assignment to the OnGeneratorsGraphChanged delegate
 	// because we need to know that someone actually wants this message so we can
 	// start actively listening for the corresponding audio render thread callback...
-	FDelegateHandle AddGraphSetCallback(const UMetasoundGeneratorHandle::FOnSetGraph& Delegate);
+	FDelegateHandle AddGraphSetCallback(UMetasoundGeneratorHandle::FOnSetGraph::FDelegate&& Delegate);
 	bool RemoveGraphSetCallback(const FDelegateHandle& Handle);
 
 	/**
diff --git a/Engine/Plugins/Runtime/Metasound/Source/MetasoundGenerator/Private/MetasoundGenerator.cpp b/Engine/Plugins/Runtime/Metasound/Source/MetasoundGenerator/Private/MetasoundGenerator.cpp
--- a/Engine/Plugins/Runtime/Metasound/Source/MetasoundGenerator/Private/MetasoundGenerator.cpp
+++ b/Engine/Plugins/Runtime/Metasound/Source/MetasoundGenerator/Private/MetasoundGenerator.cpp
@@ -354,6 +354,22 @@ PRAGMA_ENABLE_DEPRECATION_WARNINGS
 		ClearGraph();
 	}
 
+	FDelegateHandle FMetasoundGenerator::AddGraphSetCallback(FOnSetGraph::FDelegate&& Delegate)
+	{
+		FScopeLock SetPendingGraphLock(&PendingGraphMutex);
+		// If we already have a graph give the delegate an initial call
+		if (!bIsNewGraphPending && !bIsWaitingForFirstGraph)
+		{
+			Delegate.ExecuteIfBound();
+		}
+		return OnSetGraph.Add(Delegate);
+	}
+
+	bool FMetasoundGenerator::RemoveGraphSetCallback(const FDelegateHandle& Handle)
+	{
+		return OnSetGraph.Remove(Handle);
+	}
+
 	void FMetasoundGenerator::SetPendingGraph(MetasoundGeneratorPrivate::FMetasoundGeneratorData&& InData, bool bTriggerGraph)
 	{
 		using namespace MetasoundGeneratorPrivate;
diff --git a/Engine/Plugins/Runtime/Metasound/Source/MetasoundGenerator/Public/MetasoundGenerator.h b/Engine/Plugins/Runtime/Metasound/Source/MetasoundGenerator/Public/MetasoundGenerator.h
--- a/Engine/Plugins/Runtime/Metasound/Source/MetasoundGenerator/Public/MetasoundGenerator.h
+++ b/Engine/Plugins/Runtime/Metasound/Source/MetasoundGenerator/Public/MetasoundGenerator.h
@@ -216,9 +216,12 @@ namespace Metasound
 
 #endif // if ENABLE_METASOUND_GENERATOR_RENDER_TIMING
 		
-
 		// Called when a new graph has been "compiled" and set up as this generator's graph.
-		FOnSetGraph OnSetGraph;
+		// Note: We don't allow direct assignment to the FOnSetGraph delegate
+		// because we want to give the Delegate an initial immediate callback if the generator 
+		// already has a graph. 
+		FDelegateHandle AddGraphSetCallback(FOnSetGraph::FDelegate&& Delegate);
+		bool RemoveGraphSetCallback(const FDelegateHandle& Handle);
 
 	private:
 		bool TryUseCachedOperator(FMetasoundGeneratorInitParams& InParams, bool bTriggerGenerator);
@@ -311,6 +314,9 @@ namespace Metasound
 
 		TMpscQueue<TUniqueFunction<void()>> OutputAnalyzerModificationQueue;
 		TArray<TUniquePtr<Frontend::IVertexAnalyzer>> OutputAnalyzers;
+
+		FOnSetGraph OnSetGraph;
+
 #if ENABLE_METASOUND_GENERATOR_RENDER_TIMING
 		TUniquePtr<MetasoundGeneratorPrivate::FRenderTimer> RenderTimer;
 #endif // if ENABLE_METASOUND_GENERATOR_RENDER_TIMING
