- Migrate and deprecate modify interfaces transform to frontend document builder with broken up, useful functionality available now to the API (ex. SwapInput/Output)
 - New version no longer uses expensive controller system
- Fix-up & condense asset/builder initialization logic that utilizes modify interface transform functionality
- Misc clean-up
#rnx

diff --git a/Engine/Plugins/Runtime/Metasound/Source/MetasoundEditor/Private/MetasoundDetailCustomization.cpp b/Engine/Plugins/Runtime/Metasound/Source/MetasoundEditor/Private/MetasoundDetailCustomization.cpp
--- a/Engine/Plugins/Runtime/Metasound/Source/MetasoundEditor/Private/MetasoundDetailCustomization.cpp
+++ b/Engine/Plugins/Runtime/Metasound/Source/MetasoundEditor/Private/MetasoundDetailCustomization.cpp
@@ -13,6 +13,7 @@
 #include "Input/Events.h"
 #include "Interfaces/MetasoundFrontendInterfaceRegistry.h"
 #include "MetasoundAssetBase.h"
+#include "MetasoundBuilderSubsystem.h"
 #include "MetasoundEditor.h"
 #include "MetasoundEditorGraphBuilder.h"
 #include "MetasoundEditorSettings.h"
@@ -381,10 +382,11 @@ namespace Metasound
 							MetaSound.Get()->Modify();
 							MetaSoundAsset->GetGraphChecked().Modify();
 
-							FDocumentHandle DocumentHandle = MetaSoundAsset->GetDocumentHandle();
-							FModifyRootGraphInterfaces ModifyTransform({ }, { InterfaceToAdd });
-							ModifyTransform.SetDefaultNodeLocations(false); // Don't automatically add nodes to ed graph
-							ModifyTransform.Transform(DocumentHandle);
+							TScriptInterface<IMetaSoundDocumentInterface> MetaSoundPatchDocInterface = MetaSound.Get();
+							FMetaSoundFrontendDocumentBuilder Builder(MetaSoundPatchDocInterface);
+							FModifyInterfaceOptions Options({ }, { InterfaceToAdd });
+							Options.bSetDefaultNodeLocations = false; // Don't automatically add nodes to ed graph
+							Builder.ModifyInterfaces(MoveTemp(Options));
 						}
 
 						UpdateInterfaceNames();
diff --git a/Engine/Plugins/Runtime/Metasound/Source/MetasoundEditor/Private/MetasoundEditor.cpp b/Engine/Plugins/Runtime/Metasound/Source/MetasoundEditor/Private/MetasoundEditor.cpp
--- a/Engine/Plugins/Runtime/Metasound/Source/MetasoundEditor/Private/MetasoundEditor.cpp
+++ b/Engine/Plugins/Runtime/Metasound/Source/MetasoundEditor/Private/MetasoundEditor.cpp
@@ -1776,7 +1776,7 @@ namespace Metasound
 				check(MetaSoundClass);
 
 				TScriptInterface<IMetaSoundDocumentInterface> DocInterface = MetasoundAsset->GetOwningAsset();
-				FMetaSoundFrontendDocumentBuilder Builder(MetaSoundClass->GetClassPathName(), DocInterface);
+				FMetaSoundFrontendDocumentBuilder Builder(DocInterface);
 				Builder.ConvertFromPreset();
 
 				// Hack until toolbar is polished up & corner text properly dynamically updates
diff --git a/Engine/Plugins/Runtime/Metasound/Source/MetasoundEditor/Private/MetasoundEditorGraphBuilder.cpp b/Engine/Plugins/Runtime/Metasound/Source/MetasoundEditor/Private/MetasoundEditorGraphBuilder.cpp
--- a/Engine/Plugins/Runtime/Metasound/Source/MetasoundEditor/Private/MetasoundEditorGraphBuilder.cpp
+++ b/Engine/Plugins/Runtime/Metasound/Source/MetasoundEditor/Private/MetasoundEditorGraphBuilder.cpp
@@ -24,8 +24,10 @@
 #include "MetasoundEditorGraphValidation.h"
 #include "MetasoundEditorModule.h"
 #include "MetasoundEditorSettings.h"
+#include "MetasoundFactory.h"
 #include "MetasoundFrontendDataTypeRegistry.h"
 #include "MetasoundFrontendDocumentAccessPtr.h"
+#include "MetasoundFrontendDocumentBuilder.h"
 #include "MetasoundFrontendDocumentVersioning.h"
 #include "MetasoundFrontendQuery.h"
 #include "MetasoundFrontendQuerySteps.h"
@@ -119,26 +121,9 @@ namespace Metasound
 				}
 
 				// If no graph is set, MetaSound has been created outside of asset factory, so initialize it here.
-				// TODO: Move factory initialization and this code to single builder function (in header so cannot move
-				// until 5.1+).
 				if (!MetaSoundAsset->GetGraph())
 				{
-					FString Author = UKismetSystemLibrary::GetPlatformUserName();
-					if (const UMetasoundEditorSettings* EditorSettings = GetDefault<UMetasoundEditorSettings>())
-					{
-						if (!EditorSettings->DefaultAuthor.IsEmpty())
-						{
-							Author = EditorSettings->DefaultAuthor;
-						}
-					}
-
-					FGraphBuilder::InitMetaSound(InMetaSound, Author);
-
-					// Initial graph generation is not something to be managed by the transaction
-					// stack, so don't track dirty state until after initial setup if necessary.
-					UMetasoundEditorGraph* Graph = NewObject<UMetasoundEditorGraph>(&InMetaSound, FName(), RF_Transactional);
-					Graph->Schema = UMetasoundEditorGraphSchema::StaticClass();
-					MetaSoundAsset->SetGraph(Graph);
+					UMetaSoundBaseFactory::InitAsset(InMetaSound);
 				}
 
 				bEditorGraphModified |= FGraphBuilder::SynchronizeGraphMembers(InMetaSound);
diff --git a/Engine/Plugins/Runtime/Metasound/Source/MetasoundEditor/Private/MetasoundEditorGraphBuilder.h b/Engine/Plugins/Runtime/Metasound/Source/MetasoundEditor/Private/MetasoundEditorGraphBuilder.h
--- a/Engine/Plugins/Runtime/Metasound/Source/MetasoundEditor/Private/MetasoundEditorGraphBuilder.h
+++ b/Engine/Plugins/Runtime/Metasound/Source/MetasoundEditor/Private/MetasoundEditorGraphBuilder.h
@@ -233,6 +233,7 @@ namespace Metasound
 			static FLinearColor GetPinCategoryColor(const FEdGraphPinType& PinType);
 
 			// Initializes MetaSound with default inputs & outputs.
+			UE_DEPRECATED(5.3, "Use FMetaSoundFrontendBuilder::InitDocument instead.")
 			static void InitMetaSound(UObject& InMetaSound, const FString& InAuthor);
 
 			// Initializes a MetaSound Preset using the provided ReferencedMetaSound asset's
diff --git a/Engine/Plugins/Runtime/Metasound/Source/MetasoundEditor/Private/MetasoundEditorSubsystem.cpp b/Engine/Plugins/Runtime/Metasound/Source/MetasoundEditor/Private/MetasoundEditorSubsystem.cpp
--- a/Engine/Plugins/Runtime/Metasound/Source/MetasoundEditor/Private/MetasoundEditorSubsystem.cpp
+++ b/Engine/Plugins/Runtime/Metasound/Source/MetasoundEditor/Private/MetasoundEditorSubsystem.cpp
@@ -5,6 +5,7 @@
 #include "MetasoundBuilderSubsystem.h"
 #include "MetasoundEditorGraph.h"
 #include "MetasoundEditorGraphSchema.h"
+#include "MetasoundFactory.h"
 #include "MetasoundUObjectRegistry.h"
 
 #include UE_INLINE_GENERATED_CPP_BY_NAME(MetasoundEditorSubsystem)
@@ -33,11 +34,7 @@ TScriptInterface<IMetaSoundDocumentInterface> UMetaSoundEditorSubsystem::BuildTo
 			BuilderOptions.ExistingMetaSound = NewMetaSound;
 			TScriptInterface<IMetaSoundDocumentInterface> DocInterface = InBuilder->Build(Parent, BuilderOptions);
 
-			FMetasoundAssetBase* MetaSoundAsset = IMetasoundUObjectRegistry::Get().GetObjectAsAssetBase(NewMetaSound);
-			UMetasoundEditorGraph* Graph = NewObject<UMetasoundEditorGraph>(NewMetaSound, FName(), RF_Transactional);
-			Graph->Schema = UMetasoundEditorGraphSchema::StaticClass();
-			MetaSoundAsset->SetGraph(Graph);
-
+			UMetaSoundFactory::InitEdGraph(*NewMetaSound);
 			OutResult = EMetaSoundBuilderResult::Succeeded;
 			return NewMetaSound;
 		}
diff --git a/Engine/Plugins/Runtime/Metasound/Source/MetasoundEditor/Private/MetasoundFactory.cpp b/Engine/Plugins/Runtime/Metasound/Source/MetasoundEditor/Private/MetasoundFactory.cpp
--- a/Engine/Plugins/Runtime/Metasound/Source/MetasoundEditor/Private/MetasoundFactory.cpp
+++ b/Engine/Plugins/Runtime/Metasound/Source/MetasoundEditor/Private/MetasoundFactory.cpp
@@ -7,6 +7,7 @@
 #include "MetasoundEditorGraphBuilder.h"
 #include "MetasoundEditorGraphSchema.h"
 #include "MetasoundEditorSettings.h"
+#include "MetasoundFrontendDocumentBuilder.h"
 #include "MetasoundFrontendTransform.h"
 #include "MetasoundSource.h"
 
@@ -20,44 +21,10 @@ namespace Metasound
 		template <typename T>
 		T* CreateNewMetaSoundObject(UObject* InParent, FName InName, EObjectFlags InFlags, UObject* InReferencedMetaSoundObject)
 		{
-			using namespace Editor;
-			using namespace Frontend;
-
 			T* MetaSoundObject = NewObject<T>(InParent, InName, InFlags);
 			check(MetaSoundObject);
 
-			FString Author = UKismetSystemLibrary::GetPlatformUserName();
-			if (const UMetasoundEditorSettings* EditorSettings = GetDefault<UMetasoundEditorSettings>())
-			{
-				if (!EditorSettings->DefaultAuthor.IsEmpty())
-				{
-					Author = EditorSettings->DefaultAuthor;
-				}
-			}
-
-			FGraphBuilder::InitMetaSound(*MetaSoundObject, Author);
-
-			if (InReferencedMetaSoundObject)
-			{
-				FGraphBuilder::InitMetaSoundPreset(*InReferencedMetaSoundObject, *MetaSoundObject);
-			}
-
-			FMetasoundAssetBase* MetaSoundAsset = IMetasoundUObjectRegistry::Get().GetObjectAsAssetBase(MetaSoundObject);
-			check(MetaSoundAsset);
-
-			// Initial graph generation is not something to be managed by the transaction
-			// stack, so don't track dirty state until after initial setup if necessary.
-			UMetasoundEditorGraph* Graph = Cast<UMetasoundEditorGraph>(MetaSoundAsset->GetGraph());
-			if (!Graph)
-			{
-				Graph = NewObject<UMetasoundEditorGraph>(MetaSoundObject, FName(), RF_Transactional);
-				Graph->Schema = UMetasoundEditorGraphSchema::StaticClass();
-				MetaSoundAsset->SetGraph(Graph);
-
-				// Has to be done inline to have valid graph initially when opening editor for the first time
-				FGraphBuilder::SynchronizeGraph(*MetaSoundObject);
-			}
-
+			UMetaSoundBaseFactory::InitAsset(*MetaSoundObject, InReferencedMetaSoundObject);
 			return MetaSoundObject;
 		}
 	}
@@ -71,6 +38,58 @@ UMetaSoundBaseFactory::UMetaSoundBaseFactory(const FObjectInitializer& ObjectIni
 	bEditAfterNew = true;
 }
 
+void UMetaSoundBaseFactory::InitAsset(UObject& InNewMetaSound, UObject* InReferencedMetaSound)
+{
+	using namespace Metasound;
+	using namespace Metasound::Editor;
+	using namespace Metasound::Frontend;
+
+	TScriptInterface<IMetaSoundDocumentInterface> DocInterface = &InNewMetaSound;
+	FMetaSoundFrontendDocumentBuilder Builder(DocInterface);
+	Builder.InitDocument();
+	Builder.InitNodeLocations();
+
+	FString Author = UKismetSystemLibrary::GetPlatformUserName();
+	if (const UMetasoundEditorSettings* EditorSettings = GetDefault<UMetasoundEditorSettings>())
+	{
+		if (!EditorSettings->DefaultAuthor.IsEmpty())
+		{
+			Author = EditorSettings->DefaultAuthor;
+		}
+	}
+	Builder.SetAuthor(Author);
+
+	if (InReferencedMetaSound)
+	{
+		FGraphBuilder::InitMetaSoundPreset(*InReferencedMetaSound, InNewMetaSound);
+	}
+
+	// Initial graph generation is not something to be managed by the transaction
+	// stack, so don't track dirty state until after initial setup if necessary.
+	InitEdGraph(InNewMetaSound);
+}
+
+void UMetaSoundBaseFactory::InitEdGraph(UObject& InMetaSound)
+{
+	using namespace Metasound;
+	using namespace Metasound::Editor;
+
+	FMetasoundAssetBase* MetaSoundAsset = IMetasoundUObjectRegistry::Get().GetObjectAsAssetBase(&InMetaSound);
+	checkf(MetaSoundAsset, TEXT("EdGraph can only be initialized on registered MetaSoundAsset type"));
+
+	UMetasoundEditorGraph* Graph = Cast<UMetasoundEditorGraph>(MetaSoundAsset->GetGraph());
+	if (!Graph)
+	{
+		Graph = NewObject<UMetasoundEditorGraph>(&InMetaSound, FName(), RF_Transactional);
+		Graph->Schema = UMetasoundEditorGraphSchema::StaticClass();
+		MetaSoundAsset->SetGraph(Graph);
+
+		// Has to be done inline to have valid graph initially when opening editor for the first
+		// time (as opposed to being applied on tick when the document's modify context has updates)
+		FGraphBuilder::SynchronizeGraph(InMetaSound);
+	}
+}
+
 UMetaSoundFactory::UMetaSoundFactory(const FObjectInitializer& ObjectInitializer)
 	: Super(ObjectInitializer)
 {
diff --git a/Engine/Plugins/Runtime/Metasound/Source/MetasoundEditor/Public/MetasoundFactory.h b/Engine/Plugins/Runtime/Metasound/Source/MetasoundEditor/Public/MetasoundFactory.h
--- a/Engine/Plugins/Runtime/Metasound/Source/MetasoundEditor/Public/MetasoundFactory.h
+++ b/Engine/Plugins/Runtime/Metasound/Source/MetasoundEditor/Public/MetasoundFactory.h
@@ -20,6 +20,9 @@ class UMetaSoundBaseFactory : public UFactory
 	// using the provided MetaSound as a Reference
 	UPROPERTY(EditAnywhere, Transient, Category = Factory)
 	TObjectPtr<UObject> ReferencedMetaSoundObject;
+
+	static void InitAsset(UObject& InNewMetaSound, UObject* InReferencedMetaSound = nullptr);
+	static void InitEdGraph(UObject& InMetaSound);
 };
 
 UCLASS(hidecategories=Object, MinimalAPI)
diff --git a/Engine/Plugins/Runtime/Metasound/Source/MetasoundEngine/Private/Metasound.cpp b/Engine/Plugins/Runtime/Metasound/Source/MetasoundEngine/Private/Metasound.cpp
--- a/Engine/Plugins/Runtime/Metasound/Source/MetasoundEngine/Private/Metasound.cpp
+++ b/Engine/Plugins/Runtime/Metasound/Source/MetasoundEngine/Private/Metasound.cpp
@@ -6,6 +6,7 @@
 #include "Logging/TokenizedMessage.h"
 #include "MetasoundAssetBase.h"
 #include "MetasoundAudioFormats.h"
+#include "MetasoundBuilderSubsystem.h"
 #include "MetasoundEngineAsset.h"
 #include "MetasoundEngineEnvironment.h"
 #include "MetasoundEnvironment.h"
@@ -56,6 +57,16 @@ UMetaSoundPatch::UMetaSoundPatch(const FObjectInitializer& ObjectInitializer)
 {
 }
 
+const UClass& UMetaSoundPatch::GetBaseMetaSoundUClass() const
+{
+	return *UMetaSoundPatch::StaticClass();
+}
+
+const FMetasoundFrontendDocument& UMetaSoundPatch::GetDocument() const
+{
+	return RootMetaSoundDocument;
+}
+
 #if WITH_EDITOR
 void UMetaSoundPatch::PostDuplicate(EDuplicateMode::Type InDuplicateMode)
 {
diff --git a/Engine/Plugins/Runtime/Metasound/Source/MetasoundEngine/Private/MetasoundBuilderSubsystem.cpp b/Engine/Plugins/Runtime/Metasound/Source/MetasoundEngine/Private/MetasoundBuilderSubsystem.cpp
--- a/Engine/Plugins/Runtime/Metasound/Source/MetasoundEngine/Private/MetasoundBuilderSubsystem.cpp
+++ b/Engine/Plugins/Runtime/Metasound/Source/MetasoundEngine/Private/MetasoundBuilderSubsystem.cpp
@@ -4,6 +4,7 @@
 #include "Algo/Find.h"
 #include "Algo/Transform.h"
 #include "Components/AudioComponent.h"
+#include "Engine/Engine.h"
 #include "HAL/IConsoleManager.h"
 #include "Interfaces/MetasoundOutputFormatInterfaces.h"
 #include "Interfaces/MetasoundFrontendSourceInterface.h"
@@ -476,10 +477,10 @@ FMetaSoundNodeHandle UMetaSoundBuilderBase::FindGraphOutputNode(FName OutputName
 
 void UMetaSoundBuilderBase::CreateTransientDocument()
 {
-	TObjectPtr<UMetaSoundBuilderDocument> NewBuilder = NewObject<UMetaSoundBuilderDocument>();
-	TScriptInterface<IMetaSoundDocumentInterface> MetaSoundSourceDocInterface = NewBuilder;
-	const FTopLevelAssetPath BuilderPath = GetBuilderUClass().GetClassPathName();
-	Builder = FMetaSoundFrontendDocumentBuilder(BuilderPath, MetaSoundSourceDocInterface);
+	UMetaSoundBuilderDocument* NewDoc = NewObject<UMetaSoundBuilderDocument>();
+	NewDoc->SetBaseMetaSoundUClass(GetBuilderUClass());
+	TScriptInterface<IMetaSoundDocumentInterface> MetaSoundSourceDocInterface = NewDoc;
+	Builder = FMetaSoundFrontendDocumentBuilder(MetaSoundSourceDocInterface);
 	Builder.InitDocument();
 }
 
@@ -509,6 +510,18 @@ bool UMetaSoundBuilderBase::NodeOutputIsConnected(const FMetaSoundBuilderNodeOut
 	return Builder.IsNodeOutputConnected(OutputHandle.NodeID, OutputHandle.VertexID);
 }
 
+void UMetaSoundBuilderBase::RemoveGraphInput(FName Name, EMetaSoundBuilderResult& OutResult)
+{
+	const bool bRemoved = Builder.RemoveGraphInput(Name);
+	OutResult = bRemoved ? EMetaSoundBuilderResult::Succeeded : EMetaSoundBuilderResult::Failed;
+}
+
+void UMetaSoundBuilderBase::RemoveGraphOutput(FName Name, EMetaSoundBuilderResult& OutResult)
+{
+	const bool bRemoved = Builder.RemoveGraphOutput(Name);
+	OutResult = bRemoved ? EMetaSoundBuilderResult::Succeeded : EMetaSoundBuilderResult::Failed;
+}
+
 void UMetaSoundBuilderBase::RemoveInterface(FName InterfaceName, EMetaSoundBuilderResult& OutResult)
 {
 	const bool bInterfaceRemoved = Builder.RemoveInterface(InterfaceName);
@@ -726,7 +739,13 @@ void UMetaSoundSourceBuilder::SetFormat(EMetaSoundOutputAudioFormat OutputFormat
 		}
 	}
 
-	const bool bSuccess = Builder.ModifyInterfaces(OutputFormatsToRemove, OutputFormatsToAdd);
+	FModifyInterfaceOptions Options(OutputFormatsToRemove, OutputFormatsToAdd);
+
+#if WITH_EDITORONLY_DATA
+	Options.bSetDefaultNodeLocations = true;
+#endif // WITH_EDITORONLY_DATA
+
+	const bool bSuccess = Builder.ModifyInterfaces(MoveTemp(Options));
 	OutResult = bSuccess ? EMetaSoundBuilderResult::Succeeded : EMetaSoundBuilderResult::Failed;
 }
 
@@ -843,14 +862,63 @@ UMetaSoundSourceBuilder* UMetaSoundBuilderSubsystem::CreateSourceBuilder(
 	return &NewBuilder;
 }
 
+UMetaSoundBuilderSubsystem& UMetaSoundBuilderSubsystem::GetChecked()
+{
+	checkf(GEngine, TEXT("Cannot access UMetaSoundBuilderSubsystem without engine loaded"));
+	UMetaSoundBuilderSubsystem* BuilderSubsystem = GEngine->GetEngineSubsystem<UMetaSoundBuilderSubsystem>();
+	checkf(BuilderSubsystem, TEXT("Failed to find initialized 'UMetaSoundBuilderSubsystem"));
+	return *BuilderSubsystem;
+}
+
+const UMetaSoundBuilderSubsystem& UMetaSoundBuilderSubsystem::GetConstChecked()
+{
+	checkf(GEngine, TEXT("Cannot access UMetaSoundBuilderSubsystem without engine loaded"));
+	UMetaSoundBuilderSubsystem* BuilderSubsystem = GEngine->GetEngineSubsystem<UMetaSoundBuilderSubsystem>();
+	checkf(BuilderSubsystem, TEXT("Failed to find initialized 'UMetaSoundBuilderSubsystem"));
+	return *BuilderSubsystem;
+}
+
+UMetaSoundBuilderBase& UMetaSoundBuilderSubsystem::AttachBuilderToObjectChecked(TScriptInterface<IMetaSoundDocumentInterface> InInterface) const
+{
+	checkf(InInterface.GetInterface(), TEXT("Failed to attach MetaSound builder to asset: Interface not supplied"))
+	const UClass* BaseClass = &InInterface->GetBaseMetaSoundUClass();
+	if (BaseClass == UMetaSoundSource::StaticClass())
+	{
+		UMetaSoundSourceBuilder* NewBuilder = AttachSourceBuilderToAsset(CastChecked<UMetaSoundSource>(InInterface.GetObject()));
+		return *NewBuilder;
+	}
+	else if (BaseClass == UMetaSoundPatch::StaticClass())
+	{
+		UMetaSoundPatchBuilder* NewBuilder = AttachPatchBuilderToAsset(CastChecked<UMetaSoundPatch>(InInterface.GetObject()));
+		return *NewBuilder;
+	}
+	else
+	{
+		checkf(false, TEXT("UClass '%s' is not a base MetaSound that supports attachment via the MetaSoundBuilderSubsystem"), *BaseClass->GetFullName());
+		return Metasound::Engine::BuilderSubsystemPrivate::CreateTransientBuilder<UMetaSoundPatchBuilder>();
+	}
+}
+
+UMetaSoundPatchBuilder* UMetaSoundBuilderSubsystem::AttachPatchBuilderToAsset(UMetaSoundPatch* InPatch) const
+{
+	if (InPatch)
+	{
+		TObjectPtr<UMetaSoundPatchBuilder> NewBuilder = NewObject<UMetaSoundPatchBuilder>(InPatch);
+		TScriptInterface<IMetaSoundDocumentInterface> MetaSoundPatchDocInterface = InPatch;
+		NewBuilder->Builder = FMetaSoundFrontendDocumentBuilder(MetaSoundPatchDocInterface);
+		return NewBuilder;
+	}
+
+	return nullptr;
+}
+
 UMetaSoundSourceBuilder* UMetaSoundBuilderSubsystem::AttachSourceBuilderToAsset(UMetaSoundSource* InSource) const
 {
 	if (InSource)
 	{
 		TObjectPtr<UMetaSoundSourceBuilder> NewBuilder = NewObject<UMetaSoundSourceBuilder>(InSource);
 		TScriptInterface<IMetaSoundDocumentInterface> MetaSoundSourceDocInterface = InSource;
-		const FTopLevelAssetPath BuilderPath = NewBuilder->GetBuilderUClass().GetClassPathName();
-		NewBuilder->Builder = FMetaSoundFrontendDocumentBuilder(BuilderPath, MetaSoundSourceDocInterface);
+		NewBuilder->Builder = FMetaSoundFrontendDocumentBuilder(MetaSoundSourceDocInterface);
 		return NewBuilder;
 	}
 
diff --git a/Engine/Plugins/Runtime/Metasound/Source/MetasoundEngine/Private/MetasoundSource.cpp b/Engine/Plugins/Runtime/Metasound/Source/MetasoundEngine/Private/MetasoundSource.cpp
--- a/Engine/Plugins/Runtime/Metasound/Source/MetasoundEngine/Private/MetasoundSource.cpp
+++ b/Engine/Plugins/Runtime/Metasound/Source/MetasoundEngine/Private/MetasoundSource.cpp
@@ -144,6 +144,16 @@ UMetaSoundSource::UMetaSoundSource(const FObjectInitializer& ObjectInitializer)
 	SampleRate = 48000.f;
 }
 
+const UClass& UMetaSoundSource::GetBaseMetaSoundUClass() const
+{
+	return *UMetaSoundSource::StaticClass();
+}
+
+const FMetasoundFrontendDocument& UMetaSoundSource::GetDocument() const
+{
+	return RootMetasoundDocument;
+}
+
 #if WITH_EDITOR
 void UMetaSoundSource::PostEditUndo()
 {
@@ -178,9 +188,7 @@ void UMetaSoundSource::PostEditChangeOutputFormat()
 {
 	using namespace Metasound::Frontend;
 
-	const UMetaSoundBuilderSubsystem* BuilderSubsystem = GEngine->GetEngineSubsystem<UMetaSoundBuilderSubsystem>();
-	check(BuilderSubsystem);
-	UMetaSoundSourceBuilder* SourceBuilder = BuilderSubsystem->AttachSourceBuilderToAsset(this);
+	UMetaSoundSourceBuilder* SourceBuilder = UMetaSoundBuilderSubsystem::GetConstChecked().AttachSourceBuilderToAsset(this);
 	EMetaSoundBuilderResult Result;
 	SourceBuilder->SetFormat(OutputFormat, Result);
 	if (Result == EMetaSoundBuilderResult::Succeeded)
@@ -227,7 +235,6 @@ bool UMetaSoundSource::ConformObjectDataToInterfaces()
 	return bDidAlterObjectData;
 }
 
-
 void UMetaSoundSource::BeginDestroy()
 {
 	UnregisterGraphWithFrontend();
diff --git a/Engine/Plugins/Runtime/Metasound/Source/MetasoundEngine/Public/Metasound.h b/Engine/Plugins/Runtime/Metasound/Source/MetasoundEngine/Public/Metasound.h
--- a/Engine/Plugins/Runtime/Metasound/Source/MetasoundEngine/Public/Metasound.h
+++ b/Engine/Plugins/Runtime/Metasound/Source/MetasoundEngine/Public/Metasound.h
@@ -130,15 +130,8 @@ public:
 
 #endif // #if WITH_EDITORONLY_DATA
 
-	virtual const FMetasoundFrontendDocument& GetDocument() const override
-	{
-		return RootMetaSoundDocument;
-	}
-
-	virtual const UClass& GetBaseMetaSoundUClass() const final override
-	{
-		return *UMetaSoundPatch::StaticClass();
-	}
+	virtual const UClass& GetBaseMetaSoundUClass() const final override;
+	virtual const FMetasoundFrontendDocument& GetDocument() const override;
 
 #if WITH_EDITOR
 	virtual void PostDuplicate(EDuplicateMode::Type DuplicateMode) override;
diff --git a/Engine/Plugins/Runtime/Metasound/Source/MetasoundEngine/Public/MetasoundBuilderSubsystem.h b/Engine/Plugins/Runtime/Metasound/Source/MetasoundEngine/Public/MetasoundBuilderSubsystem.h
--- a/Engine/Plugins/Runtime/Metasound/Source/MetasoundEngine/Public/MetasoundBuilderSubsystem.h
+++ b/Engine/Plugins/Runtime/Metasound/Source/MetasoundEngine/Public/MetasoundBuilderSubsystem.h
@@ -16,10 +16,13 @@
 
 
 // Forward Declarations
+class UAudioComponent;
 class UMetaSound;
+class UMetaSoundPatch;
 class UMetaSoundSource;
+
 struct FMetasoundFrontendVersion;
-class UAudioComponent;
+
 enum class EMetaSoundOutputAudioFormat : uint8;
 
 namespace Metasound::Engine
@@ -265,6 +268,14 @@ public:
 	UFUNCTION(BlueprintCallable, Category = "Audio|MetaSound|Builder")
 	UPARAM(DisplayName = "Connected") bool NodeOutputIsConnected(const FMetaSoundBuilderNodeOutputHandle& OutputHandle) const;
 
+	// Removes graph input if it exists; sets result to succeeded if it was removed and failed if it was not.
+	UFUNCTION(BlueprintCallable, Category = "Audio|MetaSound|Builder")
+	void RemoveGraphInput(FName Name, EMetaSoundBuilderResult& OutResult);
+
+	// Removes graph output if it exists; sets result to succeeded if it was removed and failed if it was not.
+	UFUNCTION(BlueprintCallable, Category = "Audio|MetaSound|Builder")
+	void RemoveGraphOutput(FName Name, EMetaSoundBuilderResult& OutResult);
+
 	// Removes the interface with the given name from the builder's MetaSound. Removes any graph inputs
 	// and outputs associated with the given interface and their respective connections (if any).
 	UFUNCTION(BlueprintCallable, Category = "Audio|MetaSound|Builder", meta = (ExpandEnumAsExecs = "OutResult"))
@@ -421,6 +432,11 @@ class METASOUNDENGINE_API UMetaSoundBuilderSubsystem : public UEngineSubsystem
 	TMap<FName, TObjectPtr<UMetaSoundSourceBuilder>> SourceBuilders;
 
 public:
+	static UMetaSoundBuilderSubsystem& GetChecked();
+	static const UMetaSoundBuilderSubsystem& GetConstChecked();
+
+	UMetaSoundBuilderBase& AttachBuilderToObjectChecked(TScriptInterface<IMetaSoundDocumentInterface> InInterface) const;
+	UMetaSoundPatchBuilder* AttachPatchBuilderToAsset(UMetaSoundPatch* InPatch) const;
 	UMetaSoundSourceBuilder* AttachSourceBuilderToAsset(UMetaSoundSource* InSource) const;
 
 	UFUNCTION(BlueprintCallable, Category = "Audio|MetaSound|Builder",  meta = (ExpandEnumAsExecs = "OutResult"))
diff --git a/Engine/Plugins/Runtime/Metasound/Source/MetasoundEngine/Public/MetasoundSource.h b/Engine/Plugins/Runtime/Metasound/Source/MetasoundEngine/Public/MetasoundSource.h
--- a/Engine/Plugins/Runtime/Metasound/Source/MetasoundEngine/Public/MetasoundSource.h
+++ b/Engine/Plugins/Runtime/Metasound/Source/MetasoundEngine/Public/MetasoundSource.h
@@ -214,15 +214,8 @@ protected:
 		return MakeAccessPtr<FConstDocumentAccessPtr>(RootMetasoundDocument.AccessPoint, RootMetasoundDocument);
 	}
 
-	virtual const FMetasoundFrontendDocument& GetDocument() const override
-	{
-		return RootMetasoundDocument;
-	}
-
-	virtual const UClass& GetBaseMetaSoundUClass() const final override
-	{
-		return *UMetaSoundSource::StaticClass();
-	}
+	virtual const UClass& GetBaseMetaSoundUClass() const final override;
+	virtual const FMetasoundFrontendDocument& GetDocument() const override;
 
 	/** Gets all the default parameters for this Asset.  */
 	virtual bool GetAllDefaultParameters(TArray<FAudioParameter>& OutParameters) const override;
diff --git a/Engine/Plugins/Runtime/Metasound/Source/MetasoundFrontend/Private/MetasoundFrontendDocumentBuilder.cpp b/Engine/Plugins/Runtime/Metasound/Source/MetasoundFrontend/Private/MetasoundFrontendDocumentBuilder.cpp
--- a/Engine/Plugins/Runtime/Metasound/Source/MetasoundFrontend/Private/MetasoundFrontendDocumentBuilder.cpp
+++ b/Engine/Plugins/Runtime/Metasound/Source/MetasoundFrontend/Private/MetasoundFrontendDocumentBuilder.cpp
@@ -93,24 +93,390 @@ namespace Metasound::Frontend
 				UE_LOG(LogMetaSound, Error, TEXT("Node associated with graph vertex of type '%s' does not contain output vertex of matching type."), *InVertex.TypeName.ToString());
 			}
 		}
+
+		class FModifyInterfacesImpl
+		{
+		public:
+			FModifyInterfacesImpl(FModifyInterfaceOptions&& InOptions)
+				: Options(MoveTemp(InOptions))
+			{
+				for (const FMetasoundFrontendInterface& FromInterface : Options.InterfacesToRemove)
+				{
+					InputsToRemove.Append(FromInterface.Inputs);
+					OutputsToRemove.Append(FromInterface.Outputs);
+				}
+
+				for (const FMetasoundFrontendInterface& ToInterface : Options.InterfacesToAdd)
+				{
+					Algo::Transform(ToInterface.Inputs, InputsToAdd, [&ToInterface](const FMetasoundFrontendClassInput& Input)
+					{
+						FMetasoundFrontendClassInput NewInput = Input;
+						NewInput.NodeID = FGuid::NewGuid();
+						NewInput.VertexID = FGuid::NewGuid();
+						return FInputInterfacePair { MoveTemp(NewInput), &ToInterface };
+					});
+
+					Algo::Transform(ToInterface.Outputs, OutputsToAdd, [&ToInterface](const FMetasoundFrontendClassOutput& Output)
+					{
+						FMetasoundFrontendClassOutput NewOutput = Output;
+						NewOutput.NodeID = FGuid::NewGuid();
+						NewOutput.VertexID = FGuid::NewGuid();
+						return FOutputInterfacePair { MoveTemp(NewOutput), &ToInterface };
+					});
+				}
+
+				// Iterate in reverse to allow removal from `InputsToAdd`
+				for (int32 AddIndex = InputsToAdd.Num() - 1; AddIndex >= 0; AddIndex--)
+				{
+					const FMetasoundFrontendClassVertex& VertexToAdd = InputsToAdd[AddIndex].Key;
+
+					const int32 RemoveIndex = InputsToRemove.IndexOfByPredicate([&](const FMetasoundFrontendClassVertex& VertexToRemove)
+						{
+							if (VertexToAdd.TypeName != VertexToRemove.TypeName)
+							{
+								return false;
+							}
+
+							if (Options.NamePairingFunction)
+							{
+								return Options.NamePairingFunction(VertexToAdd.Name, VertexToRemove.Name);
+							}
+
+							FName ParamA;
+							FName ParamB;
+							FName Namespace;
+							VertexToAdd.SplitName(Namespace, ParamA);
+							VertexToRemove.SplitName(Namespace, ParamB);
+
+							return ParamA == ParamB;
+						});
+
+					if (INDEX_NONE != RemoveIndex)
+					{
+						PairedInputs.Add(FVertexPair { InputsToRemove[RemoveIndex], InputsToAdd[AddIndex].Key });
+						InputsToRemove.RemoveAtSwap(RemoveIndex);
+						InputsToAdd.RemoveAtSwap(AddIndex);
+					}
+				}
+
+				// Iterate in reverse to allow removal from `OutputsToAdd`
+				for (int32 AddIndex = OutputsToAdd.Num() - 1; AddIndex >= 0; AddIndex--)
+				{
+					const FMetasoundFrontendClassVertex& VertexToAdd = OutputsToAdd[AddIndex].Key;
+
+					const int32 RemoveIndex = OutputsToRemove.IndexOfByPredicate([&](const FMetasoundFrontendClassVertex& VertexToRemove)
+						{
+							if (VertexToAdd.TypeName != VertexToRemove.TypeName)
+							{
+								return false;
+							}
+
+							if (Options.NamePairingFunction)
+							{
+								return Options.NamePairingFunction(VertexToAdd.Name, VertexToRemove.Name);
+							}
+
+							FName ParamA;
+							FName ParamB;
+							FName Namespace;
+							VertexToAdd.SplitName(Namespace, ParamA);
+							VertexToRemove.SplitName(Namespace, ParamB);
+
+							return ParamA == ParamB;
+						});
+
+					if (INDEX_NONE != RemoveIndex)
+					{
+						PairedOutputs.Add(FVertexPair{ OutputsToRemove[RemoveIndex], OutputsToAdd[AddIndex].Key });
+						OutputsToRemove.RemoveAtSwap(RemoveIndex);
+						OutputsToAdd.RemoveAtSwap(AddIndex);
+					}
+				}
+			}
+
+		private:
+			bool AddMissingVertices(FMetaSoundFrontendDocumentBuilder& OutBuilder) const
+			{
+				if (!InputsToAdd.IsEmpty() || !OutputsToAdd.IsEmpty())
+				{
+					for (const FInputInterfacePair& Pair: InputsToAdd)
+					{
+						OutBuilder.AddGraphInput(Pair.Key);
+					}
+
+					for (const FOutputInterfacePair& Pair : OutputsToAdd)
+					{
+						OutBuilder.AddGraphOutput(Pair.Key);
+					}
+
+					return true;
+				}
+
+				return false;
+			}
+
+			bool RemoveUnsupportedVertices(FMetaSoundFrontendDocumentBuilder& OutBuilder) const
+			{
+				if (!InputsToRemove.IsEmpty() || !OutputsToRemove.IsEmpty())
+				{
+					auto RemoveMemberNodes = [&OutBuilder](
+						const FMetasoundFrontendClassVertex& ToRemove,
+						TFunctionRef<const FMetasoundFrontendClassVertex*(FName)> FindVertexFunc,
+						TFunctionRef<const FMetasoundFrontendNode* (FName)> FindNodeFunc)
+					{
+						if (const FMetasoundFrontendClassVertex* ClassInput = FindVertexFunc(ToRemove.Name))
+						{
+							if (FMetasoundFrontendClassVertex::IsFunctionalEquivalent(*ClassInput, ToRemove))
+							{
+								if (const FMetasoundFrontendNode* Node = FindNodeFunc(ToRemove.Name))
+								{
+									const bool bNodeRemoved = OutBuilder.RemoveNode(Node->GetID());
+									checkf(bNodeRemoved, TEXT("Failed to remove member while attempting to modify MetaSound graph interfaces"));
+								}
+							}
+						}
+					};
+
+					// Remove unsupported inputs
+					for (const FMetasoundFrontendClassVertex& InputToRemove : InputsToRemove)
+					{
+						RemoveMemberNodes(InputToRemove,
+							[&OutBuilder](FName Name) { return OutBuilder.FindGraphInput(Name); },
+							[&OutBuilder](FName Name) { return OutBuilder.FindGraphInputNode(Name); });
+					}
+
+					// Remove unsupported outputs
+					for (const FMetasoundFrontendClassVertex& OutputToRemove : OutputsToRemove)
+					{
+						RemoveMemberNodes(OutputToRemove,
+							[&OutBuilder](FName Name) { return OutBuilder.FindGraphOutput(Name); },
+							[&OutBuilder](FName Name) { return OutBuilder.FindGraphOutputNode(Name); });
+					}
+
+					return true;
+				}
+
+				return false;
+			}
+
+			bool SwapPairedVertices(FMetaSoundFrontendDocumentBuilder& OutBuilder) const
+			{
+				bool bDidEdit = false;
+				for (const FVertexPair& PairedInput : PairedInputs)
+				{
+					bDidEdit |= OutBuilder.SwapGraphInput(PairedInput.Get<0>(), PairedInput.Get<1>());
+				}
+
+				for (const FVertexPair& PairedOutput : PairedOutputs)
+				{
+					bDidEdit |= OutBuilder.SwapGraphOutput(PairedOutput.Get<0>(), PairedOutput.Get<1>());
+				}
+
+				return bDidEdit;
+			}
+
+#if WITH_EDITORONLY_DATA
+			void UpdateAddedVertexNodePositions(
+				EMetasoundFrontendClassType ClassType,
+				const FMetaSoundFrontendDocumentBuilder& InBuilder,
+				TSet<FName>& AddedNames,
+				TFunctionRef<int32(const FVertexName&)> InGetSortOrder,
+				TArrayView<FMetasoundFrontendNode> OutNodes)
+			{
+				// Add graph member nodes by sort order
+				TSortedMap<int32, FMetasoundFrontendNode*> SortOrderToNode;
+				for (FMetasoundFrontendNode& Node : OutNodes)
+				{
+					if (const FMetasoundFrontendClass* Class = InBuilder.FindDependency(Node.ClassID))
+					{
+						if (Class->Metadata.GetType() == ClassType)
+						{
+							const int32 Index = InGetSortOrder(Node.Name);
+							SortOrderToNode.Add(Index, &Node);
+						}
+					}
+				}
+
+				// Prime the first location as an offset prior to an existing location (as provided by a swapped member)
+				//  to avoid placing away from user's active area if possible.
+				FVector2D NextLocation = { 0.0f, 0.0f };
+				{
+					int32 NumBeforeDefined = 1;
+					for (const TPair<int32, FMetasoundFrontendNode*>& Pair : SortOrderToNode)
+					{
+						const FMetasoundFrontendNode* Node = Pair.Value;
+						const FName NodeName = Node->Name;
+						if (AddedNames.Contains(NodeName))
+						{
+							NumBeforeDefined++;
+						}
+						else
+						{
+							const TMap<FGuid, FVector2D>& Locations = Node->Style.Display.Locations;
+							if (!Locations.IsEmpty())
+							{
+								for (const TPair<FGuid, FVector2D>& Location : Locations)
+								{
+									NextLocation = Location.Value - (NumBeforeDefined * DisplayStyle::NodeLayout::DefaultOffsetY);
+									break;
+								}
+								break;
+							}
+						}
+					}
+				}
+
+				// Iterate through sorted map in sequence, slotting in new locations after existing swapped nodes with predefined locations.
+				for (TPair<int32, FMetasoundFrontendNode*>& Pair : SortOrderToNode)
+				{
+					FMetasoundFrontendNode* Node = Pair.Value;
+					const FName NodeName = Node->Name;
+					if (AddedNames.Contains(NodeName))
+					{
+						Node->Style.Display.Locations.Add(FGuid(), NextLocation);
+						NextLocation += DisplayStyle::NodeLayout::DefaultOffsetY;
+					}
+					else
+					{
+						for (const TPair<FGuid, FVector2D>& Location : Node->Style.Display.Locations)
+						{
+							NextLocation = Location.Value + DisplayStyle::NodeLayout::DefaultOffsetY;
+						}
+					}
+				}
+			}
+#endif // WITH_EDITORONLY_DATA
+
+		public:
+			bool Execute(FMetaSoundFrontendDocumentBuilder& OutBuilder, FMetasoundFrontendDocument& OutDoc)
+			{
+				bool bDidEdit = false;
+
+				for (const FMetasoundFrontendInterface& Interface : Options.InterfacesToRemove)
+				{
+					if (OutDoc.Interfaces.Remove(Interface.Version) > 0)
+					{
+						bDidEdit = true;
+#if WITH_EDITORONLY_DATA
+						OutDoc.Metadata.ModifyContext.AddInterfaceModified(Interface.Version.Name);
+#endif // WITH_EDITORONLY_DATA
+					}
+				}
+
+				for (const FMetasoundFrontendInterface& Interface : Options.InterfacesToAdd)
+				{
+					bool bAlreadyInSet = false;
+					OutDoc.Interfaces.Add(Interface.Version, &bAlreadyInSet);
+					if (!bAlreadyInSet)
+					{
+						bDidEdit = true;
+#if WITH_EDITORONLY_DATA
+						OutDoc.Metadata.ModifyContext.AddInterfaceModified(Interface.Version.Name);
+#endif // WITH_EDITORONLY_DATA
+					}
+				}
+
+				bDidEdit |= RemoveUnsupportedVertices(OutBuilder);
+				bDidEdit |= SwapPairedVertices(OutBuilder);
+				const bool bAddedVertices = AddMissingVertices(OutBuilder);
+				bDidEdit |= bAddedVertices;
+
+				if (bDidEdit)
+				{
+					OutBuilder.RemoveUnusedDependencies();
+				}
+
+#if WITH_EDITORONLY_DATA
+				if (bAddedVertices && Options.bSetDefaultNodeLocations)
+				{
+					TArray<FMetasoundFrontendNode>& Nodes = OutDoc.RootGraph.Graph.Nodes;
+					// Sort/Place Inputs
+					{
+						TSet<FName> NamesToSort;
+						Algo::Transform(InputsToAdd, NamesToSort, [](const FInputInterfacePair& Pair) { return Pair.Key.Name; });
+						auto GetInputSortOrder = [&OutBuilder](const FVertexName& InVertexName)
+						{
+							const FMetasoundFrontendClassInput* Input = OutBuilder.FindGraphInput(InVertexName);
+							checkf(Input, TEXT("Input must exist by this point of modifying the document's interfaces and respective members"));
+							return Input->Metadata.SortOrderIndex;
+						};
+						UpdateAddedVertexNodePositions(EMetasoundFrontendClassType::Input, OutBuilder, NamesToSort, GetInputSortOrder, Nodes);
+					}
+
+					// Sort/Place Outputs
+					{
+						TSet<FName> NamesToSort;
+						Algo::Transform(OutputsToAdd, NamesToSort, [](const FOutputInterfacePair& OutputInterfacePair) { return OutputInterfacePair.Key.Name; });
+						auto GetOutputSortOrder = [&OutBuilder](const FVertexName& InVertexName)
+						{
+							const FMetasoundFrontendClassOutput* Output = OutBuilder.FindGraphOutput(InVertexName);
+							checkf(Output, TEXT("Output must exist by this point of modifying the document's interfaces and respective members"));
+							return Output->Metadata.SortOrderIndex;
+						};
+						UpdateAddedVertexNodePositions(EMetasoundFrontendClassType::Output, OutBuilder, NamesToSort, GetOutputSortOrder, Nodes);
+					}
+				}
+#endif // WITH_EDITORONLY_DATA
+
+				return bDidEdit;
+			}
+
+			const FModifyInterfaceOptions Options;
+
+			using FVertexPair = TTuple<FMetasoundFrontendClassVertex, FMetasoundFrontendClassVertex>;
+			TArray<FVertexPair> PairedInputs;
+			TArray<FVertexPair> PairedOutputs;
+
+			using FInputInterfacePair = TPair<FMetasoundFrontendClassInput, const FMetasoundFrontendInterface*>;
+			using FOutputInterfacePair = TPair<FMetasoundFrontendClassOutput, const FMetasoundFrontendInterface*>;
+			TArray<FInputInterfacePair> InputsToAdd;
+			TArray<FOutputInterfacePair> OutputsToAdd;
+
+			TArray<FMetasoundFrontendClassInput> InputsToRemove;
+			TArray<FMetasoundFrontendClassOutput> OutputsToRemove;
+		};
 	} // namespace DocumentBuilderPrivate
-} // namespace Metasound::Frontend
 
-FMetaSoundFrontendDocumentBuilder::FMetaSoundFrontendDocumentBuilder()
-	: BuilderClassPath(UMetaSoundDocumentInterface::StaticClass()->GetClassPathName())
-{
-}
+	FModifyInterfaceOptions::FModifyInterfaceOptions(const TArray<FMetasoundFrontendInterface>& InInterfacesToRemove, const TArray<FMetasoundFrontendInterface>& InInterfacesToAdd)
+		: InterfacesToRemove(InInterfacesToRemove)
+		, InterfacesToAdd(InInterfacesToAdd)
+	{
+	}
+
+	FModifyInterfaceOptions::FModifyInterfaceOptions(const TArray<FMetasoundFrontendVersion>& InInterfaceVersionsToRemove, const TArray<FMetasoundFrontendVersion>& InInterfaceVersionsToAdd)
+	{
+		Algo::Transform(InInterfaceVersionsToRemove, InterfacesToRemove, [](const FMetasoundFrontendVersion& Version)
+		{
+			FMetasoundFrontendInterface Interface;
+			const bool bFromInterfaceFound = IInterfaceRegistry::Get().FindInterface(GetInterfaceRegistryKey(Version), Interface);
+			if (!ensureAlways(bFromInterfaceFound))
+			{
+				UE_LOG(LogMetaSound, Error, TEXT("Failed to find interface '%s' to remove"), *Version.ToString());
+			}
+			return Interface;
+		});
+
+		Algo::Transform(InInterfaceVersionsToAdd, InterfacesToAdd, [](const FMetasoundFrontendVersion& Version)
+		{
+			FMetasoundFrontendInterface Interface;
+			const bool bToInterfaceFound = IInterfaceRegistry::Get().FindInterface(GetInterfaceRegistryKey(Version), Interface);
+			if (!ensureAlways(bToInterfaceFound))
+			{
+				UE_LOG(LogMetaSound, Error, TEXT("Failed to find interface '%s' to add"), *Version.ToString());
+			}
+			return Interface;
+		});
+	}
+} // namespace Metasound::Frontend
 
-FMetaSoundFrontendDocumentBuilder::FMetaSoundFrontendDocumentBuilder(const FTopLevelAssetPath& InBuilderClassPath)
-	: BuilderClassPath(InBuilderClassPath)
-{
-}
 
-FMetaSoundFrontendDocumentBuilder::FMetaSoundFrontendDocumentBuilder(const FTopLevelAssetPath& InBuilderClassPath, TScriptInterface<IMetaSoundDocumentInterface> InDocumentInterface)
+FMetaSoundFrontendDocumentBuilder::FMetaSoundFrontendDocumentBuilder(TScriptInterface<IMetaSoundDocumentInterface> InDocumentInterface)
 	: DocumentInterface(InDocumentInterface)
-	, BuilderClassPath(InBuilderClassPath)
 {
-	ReloadCache();
+	if (DocumentInterface)
+	{
+		ReloadCache();
+	}
 }
 
 FMetaSoundFrontendDocumentBuilder::FMetaSoundFrontendDocumentBuilder(const FMetaSoundFrontendDocumentBuilder& InBuilder)
@@ -119,14 +485,12 @@ FMetaSoundFrontendDocumentBuilder::FMetaSoundFrontendDocumentBuilder(const FMeta
 	if (DocumentInterface)
 	{
 		ReloadCache();
-	}	
-	BuilderClassPath = InBuilder.BuilderClassPath;
+	}
 }
 
-FMetaSoundFrontendDocumentBuilder::FMetaSoundFrontendDocumentBuilder(const FTopLevelAssetPath& InBuilderClassPath, FMetaSoundFrontendDocumentBuilder&& InBuilder)
+FMetaSoundFrontendDocumentBuilder::FMetaSoundFrontendDocumentBuilder(FMetaSoundFrontendDocumentBuilder&& InBuilder)
 	: DocumentInterface(MoveTemp(InBuilder.DocumentInterface))
 	, DocumentCache(MoveTemp(InBuilder.DocumentCache))
-	, BuilderClassPath(MoveTemp(BuilderClassPath))
 {
 }
 
@@ -135,7 +499,6 @@ FMetaSoundFrontendDocumentBuilder& FMetaSoundFrontendDocumentBuilder::operator=(
 	using namespace Metasound::Frontend;
 	DocumentInterface = MoveTemp(InRHS.DocumentInterface);
 	DocumentCache = MoveTemp(InRHS.DocumentCache);
-	BuilderClassPath = MoveTemp(InRHS.BuilderClassPath);
 	return *this;
 }
 
@@ -148,7 +511,6 @@ FMetaSoundFrontendDocumentBuilder& FMetaSoundFrontendDocumentBuilder::operator=(
 	{
 		ReloadCache();
 	}
-	BuilderClassPath = InRHS.BuilderClassPath;
 	return *this;
 }
 
@@ -377,7 +739,8 @@ const FMetasoundFrontendNode* FMetaSoundFrontendDocumentBuilder::AddGraphInput(c
 {
 	using namespace Metasound::Frontend;
 
-	FMetasoundFrontendGraphClass& RootGraph = GetDocument().RootGraph;
+	FMetasoundFrontendDocument& Document = GetDocument();
+	FMetasoundFrontendGraphClass& RootGraph = Document.RootGraph;
 
 	checkf(InClassInput.NodeID.IsValid(), TEXT("Unassigned NodeID when adding graph input"));
 	checkf(InClassInput.VertexID.IsValid(), TEXT("Unassigned VertexID when adding graph input"));
@@ -434,12 +797,16 @@ const FMetasoundFrontendNode* FMetaSoundFrontendDocumentBuilder::AddGraphInput(c
 		};
 		if (FMetasoundFrontendNode* NewNode = AddNodeInternal(Class.Metadata, FinalizeNode, InClassInput.NodeID))
 		{
-			
 			FMetasoundFrontendClassInput& NewInput = RootGraph.Interface.Inputs.Add_GetRef(InClassInput);
 			if (!NewInput.VertexID.IsValid())
 			{
 				NewInput.VertexID = FGuid::NewGuid();
 			}
+
+#if WITH_EDITORONLY_DATA
+			Document.Metadata.ModifyContext.AddMemberIDModified(InClassInput.NodeID);
+#endif // WITH_EDITORONLY_DATA
+
 			return NewNode;
 		}
 	}
@@ -454,7 +821,8 @@ const FMetasoundFrontendNode* FMetaSoundFrontendDocumentBuilder::AddGraphOutput(
 	checkf(InClassOutput.NodeID.IsValid(), TEXT("Unassigned NodeID when adding graph output"));
 	checkf(InClassOutput.VertexID.IsValid(), TEXT("Unassigned VertexID when adding graph output"));
 
-	FMetasoundFrontendGraphClass& RootGraph = GetDocument().RootGraph;
+	FMetasoundFrontendDocument& Document = GetDocument();
+	FMetasoundFrontendGraphClass& RootGraph = Document.RootGraph;
 
 	if (InClassOutput.TypeName.IsNone())
 	{
@@ -513,6 +881,11 @@ const FMetasoundFrontendNode* FMetaSoundFrontendDocumentBuilder::AddGraphOutput(
 			{
 				NewOutput.VertexID = FGuid::NewGuid();
 			}
+
+#if WITH_EDITORONLY_DATA
+			Document.Metadata.ModifyContext.AddMemberIDModified(InClassOutput.NodeID);
+#endif // WITH_EDITORONLY_DATA
+
 			return NewNode;
 		}
 	}
@@ -533,6 +906,7 @@ bool FMetaSoundFrontendDocumentBuilder::AddInterface(FName InterfaceName)
 			return true;
 		}
 
+		const FTopLevelAssetPath BuilderClassPath = GetBuilderClassPath();
 		const FInterfaceRegistryKey Key = GetInterfaceRegistryKey(Interface.Version);
 		if (const IInterfaceRegistryEntry* Entry = IInterfaceRegistry::Get().FindInterfaceRegistryEntry(Key))
 		{
@@ -611,6 +985,10 @@ FMetasoundFrontendNode* FMetaSoundFrontendDocumentBuilder::AddNodeInternal(const
 		const int32 NewIndex = Nodes.Num() - 1;
 		IDocumentGraphNodeCache& NodeCache = DocumentCache->GetNodeCache();
 		NodeCache.OnNodeAdded(NewIndex);
+
+#if WITH_EDITORONLY_DATA
+		Document.Metadata.ModifyContext.AddNodeIDModified(InNodeID);
+#endif // WITH_EDITORONLY_DATA
 		return &Node;
 	}
 
@@ -647,6 +1025,24 @@ bool FMetaSoundFrontendDocumentBuilder::CanAddEdge(const FMetasoundFrontendEdge&
 	return false;
 }
 
+bool FMetaSoundFrontendDocumentBuilder::RemoveUnusedDependencies()
+{
+	bool bDidEdit = false;
+
+	TArray<FMetasoundFrontendClass>& Dependencies = GetDocument().Dependencies;
+	for (int32 i = Dependencies.Num() - 1; i >= 0; --i)
+	{
+		const FGuid& ClassID = Dependencies[i].ID;
+		if (!DocumentCache->GetNodeCache().ContainsNodesOfClassID(ClassID))
+		{
+			checkf(RemoveDependency(ClassID), TEXT("Failed to remove dependency that was found on document and was not referenced by nodes"));
+			bDidEdit = true;
+		}
+	}
+
+	return bDidEdit;
+}
+
 bool FMetaSoundFrontendDocumentBuilder::ContainsEdge(const FMetasoundFrontendEdge& InEdge) const
 {
 	using namespace Metasound::Frontend;
@@ -789,6 +1185,12 @@ bool FMetaSoundFrontendDocumentBuilder::FindInterfaceOutputNodes(FName Interface
 	return false;
 }
 
+const FMetasoundFrontendClassInput* FMetaSoundFrontendDocumentBuilder::FindGraphInput(FName InputName) const
+{
+	const TArray<FMetasoundFrontendClassInput>& Inputs = GetDocument().RootGraph.Interface.Inputs;
+	return Inputs.FindByPredicate([&InputName](const FMetasoundFrontendClassInput& Input) { return Input.Name == InputName; });
+}
+
 const FMetasoundFrontendNode* FMetaSoundFrontendDocumentBuilder::FindGraphInputNode(FName InputName) const
 {
 	using namespace Metasound::Frontend;
@@ -796,6 +1198,12 @@ const FMetasoundFrontendNode* FMetaSoundFrontendDocumentBuilder::FindGraphInputN
 	return NodeCache.FindInputNode(InputName);
 }
 
+const FMetasoundFrontendClassOutput* FMetaSoundFrontendDocumentBuilder::FindGraphOutput(FName OutputName) const
+{
+	const TArray<FMetasoundFrontendClassOutput>& Outputs = GetDocument().RootGraph.Interface.Outputs;
+	return Outputs.FindByPredicate([&OutputName](const FMetasoundFrontendClassOutput& Output) { return Output.Name == OutputName; });
+}
+
 const FMetasoundFrontendNode* FMetaSoundFrontendDocumentBuilder::FindGraphOutputNode(FName OutputName) const
 {
 	using namespace Metasound::Frontend;
@@ -911,9 +1319,17 @@ const FMetasoundFrontendDocument* FMetaSoundFrontendDocumentBuilder::FindOrLoadN
 	return nullptr;
 }
 
+const FTopLevelAssetPath FMetaSoundFrontendDocumentBuilder::GetBuilderClassPath() const
+{
+	IMetaSoundDocumentInterface* Interface = DocumentInterface.GetInterface();
+	checkf(Interface, TEXT("Failed to return class path; interface must always be valid while builder is operating on MetaSound UObject!"));
+	return Interface->GetBaseMetaSoundUClass().GetClassPathName();
+}
+
 FMetasoundFrontendDocument& FMetaSoundFrontendDocumentBuilder::GetDocument()
 {
 	IMetaSoundDocumentInterface* Interface = DocumentInterface.GetInterface();
+	checkf(Interface, TEXT("Failed to return document; interface must always be valid while builder is operating on MetaSound UObject!"));
 	return Interface->GetDocument();
 }
 
@@ -948,9 +1364,7 @@ void FMetaSoundFrontendDocumentBuilder::InitDocument()
 
 	METASOUND_TRACE_CPUPROFILER_EVENT_SCOPE(FMetaSoundFrontendDocumentBuilder::InitDocument);
 
-	IMetaSoundDocumentInterface* DocInterface = DocumentInterface.GetInterface();
-	check(DocInterface);
-	FMetasoundFrontendDocument& Document = DocInterface->GetDocument();
+	FMetasoundFrontendDocument& Document = GetDocument();
 
 	// 1. Set default class Metadata
 	{
@@ -967,14 +1381,10 @@ void FMetaSoundFrontendDocumentBuilder::InitDocument()
 
 	// 3. Add default interfaces for given UClass
 	{
-		TArray<FMetasoundFrontendVersion> InitVersions = ISearchEngine::Get().FindUClassDefaultInterfaceVersions(BuilderClassPath);
-		FModifyRootGraphInterfaces ModifyRootGraphInterfaces({ }, InitVersions);
-		ModifyRootGraphInterfaces.Transform(Document);
+		TArray<FMetasoundFrontendVersion> InitVersions = ISearchEngine::Get().FindUClassDefaultInterfaceVersions(GetBuilderClassPath());
+		FModifyInterfaceOptions Options({ }, InitVersions);
+		ModifyInterfaces(MoveTemp(Options));
 	}
-
-	// 5. Reload the whole cache as transforms above mutate cached collections. Can be safely removed once transforms above
-	// are re-implemented in builder and support cache to enable marginally better performance on document initialization.
-	ReloadCache();
 }
 
 void FMetaSoundFrontendDocumentBuilder::InitNodeLocations()
@@ -1023,6 +1433,11 @@ void FMetaSoundFrontendDocumentBuilder::InitNodeLocations()
 #endif // WITH_EDITORONLY_DATA
 }
 
+bool FMetaSoundFrontendDocumentBuilder::IsDependencyReferenced(const FGuid& InClassID) const
+{
+	return DocumentCache->GetNodeCache().ContainsNodesOfClassID(InClassID);
+}
+
 bool FMetaSoundFrontendDocumentBuilder::IsNodeInputConnected(const FGuid& InNodeID, const FGuid& InVertexID) const
 {
 	return DocumentCache->GetEdgeCache().IsNodeInputConnected(InNodeID, InVertexID);
@@ -1051,12 +1466,58 @@ bool FMetaSoundFrontendDocumentBuilder::IsInterfaceDeclared(const FMetasoundFron
 	return GetDocument().Interfaces.Contains(InInterfaceVersion);
 }
 
+bool FMetaSoundFrontendDocumentBuilder::ModifyInterfaces(Metasound::Frontend::FModifyInterfaceOptions&& InOptions)
+{
+	using namespace Metasound::Frontend;
+
+	DocumentBuilderPrivate::FModifyInterfacesImpl Context(MoveTemp(InOptions));
+	FMetasoundFrontendDocument& Doc = GetDocument();
+	return Context.Execute(*this, Doc);
+}
+
 void FMetaSoundFrontendDocumentBuilder::ReloadCache()
 {
 	using namespace Metasound::Frontend;
 	DocumentCache = MakeUnique<FDocumentCache>(GetDocument());
 }
 
+bool FMetaSoundFrontendDocumentBuilder::RemoveDependency(const FGuid& InClassID)
+{
+	using namespace Metasound::Frontend;
+
+	FMetasoundFrontendDocument& Document = GetDocument();
+	const IDocumentGraphNodeCache& NodeCache = DocumentCache->GetNodeCache();
+
+	if (const int32* IndexPtr = DocumentCache->FindDependencyIndex(InClassID))
+	{
+		TArray<FMetasoundFrontendClass>& Dependencies = Document.Dependencies;
+		const int32 Index = *IndexPtr;
+
+		TArray<const FMetasoundFrontendNode*> Nodes = NodeCache.FindNodesOfClassID(InClassID);
+		for (const FMetasoundFrontendNode* Node : Nodes)
+		{
+			if (!RemoveNode(Node->GetID()))
+			{
+				return false;
+			}
+		}
+
+		const int32 LastIndex = Dependencies.Num() - 1;
+		DocumentCache->OnRemovingDependency(Index);
+		if (Index != LastIndex)
+		{
+			DocumentCache->OnRemovingDependency(LastIndex);
+		}
+		Dependencies.RemoveAtSwap(Index, 1, false);
+		if (Index != LastIndex)
+		{
+			DocumentCache->OnDependencyAdded(Index);
+		}
+	}
+
+	return true;
+}
+
 bool FMetaSoundFrontendDocumentBuilder::RemoveDependency(EMetasoundFrontendClassType ClassType, const FMetasoundFrontendClassName& InClassName, const FMetasoundFrontendVersionNumber& InClassVersionNumber)
 {
 	using namespace Metasound::Frontend;
@@ -1064,10 +1525,10 @@ bool FMetaSoundFrontendDocumentBuilder::RemoveDependency(EMetasoundFrontendClass
 	FMetasoundFrontendDocument& Document = GetDocument();
 	const IDocumentGraphNodeCache& NodeCache = DocumentCache->GetNodeCache();
 
-	TArray<FMetasoundFrontendClass>& Dependencies = Document.Dependencies;
 	const FNodeRegistryKey ClassKey = NodeRegistryKey::CreateKey(ClassType, InClassName.GetFullName().ToString(), InClassVersionNumber.Major, InClassVersionNumber.Minor);
 	if (const int32* IndexPtr = DocumentCache->FindDependencyIndex(ClassKey))
 	{
+		TArray<FMetasoundFrontendClass>& Dependencies = Document.Dependencies;
 		const int32 Index = *IndexPtr;
 
 		TArray<const FMetasoundFrontendNode*> Nodes = NodeCache.FindNodesOfClassID(Dependencies[Index].ID);
@@ -1256,6 +1717,64 @@ bool FMetaSoundFrontendDocumentBuilder::RemoveEdgeToNodeInput(const FGuid& InNod
 	return false;
 }
 
+bool FMetaSoundFrontendDocumentBuilder::RemoveGraphInput(FName InInputName)
+{
+	if (const FMetasoundFrontendNode* Node = FindGraphInputNode(InInputName))
+	{
+		const FGuid ClassID = Node->ClassID;
+		if (RemoveNode(Node->GetID()))
+		{
+			TArray<FMetasoundFrontendClassInput>& Inputs = GetDocument().RootGraph.Interface.Inputs;
+			auto InputNameMatches = [InInputName](const FMetasoundFrontendClassInput& Input) { return Input.Name == InInputName; };
+			if (Inputs.RemoveAllSwap(InputNameMatches))
+			{
+				if (IsDependencyReferenced(ClassID))
+				{
+					return true;
+				}
+				else
+				{
+					if (RemoveDependency(ClassID))
+					{
+						return true;
+					}
+				}
+			}
+		}
+	}
+
+	return false;
+}
+
+bool FMetaSoundFrontendDocumentBuilder::RemoveGraphOutput(FName InOutputName)
+{
+	if (const FMetasoundFrontendNode* Node = FindGraphOutputNode(InOutputName))
+	{
+		const FGuid ClassID = Node->ClassID;
+		if (RemoveNode(Node->GetID()))
+		{
+			TArray<FMetasoundFrontendClassOutput>& Outputs = GetDocument().RootGraph.Interface.Outputs;
+			auto OutputNameMatches = [InOutputName](const FMetasoundFrontendClassOutput& Output) { return Output.Name == InOutputName; };
+			if (Outputs.RemoveAllSwap(OutputNameMatches))
+			{
+				if (IsDependencyReferenced(ClassID))
+				{
+					return true;
+				}
+				else
+				{
+					if (RemoveDependency(ClassID))
+					{
+						return true;
+					}
+				}
+			}
+		}
+	}
+
+	return false;
+}
+
 bool FMetaSoundFrontendDocumentBuilder::RemoveInterface(FName InterfaceName)
 {
 	using namespace Metasound::Frontend;
@@ -1269,6 +1788,7 @@ bool FMetaSoundFrontendDocumentBuilder::RemoveInterface(FName InterfaceName)
 			return true;
 		}
 
+		const FTopLevelAssetPath BuilderClassPath = GetBuilderClassPath();
 		const FInterfaceRegistryKey Key = GetInterfaceRegistryKey(Interface.Version);
 		if (const IInterfaceRegistryEntry* Entry = IInterfaceRegistry::Get().FindInterfaceRegistryEntry(Key))
 		{
@@ -1313,15 +1833,12 @@ bool FMetaSoundFrontendDocumentBuilder::RemoveNode(const FGuid& InNodeID)
 
 		for (const FMetasoundFrontendVertex& Vertex : Node.Interface.Inputs)
 		{
-			if (!RemoveEdgeToNodeInput(Node.GetID(), Vertex.VertexID))
-			{
-				return false;
-			}
+			RemoveEdgeToNodeInput(InNodeID, Vertex.VertexID);
 		}
 
 		for (const FMetasoundFrontendVertex& Vertex : Node.Interface.Outputs)
 		{
-			RemoveEdgesFromNodeOutput(Node.GetID(), Vertex.VertexID);
+			RemoveEdgesFromNodeOutput(InNodeID, Vertex.VertexID);
 		}
 
 		NodeCache.OnRemovingNode(Index);
@@ -1337,6 +1854,10 @@ bool FMetaSoundFrontendDocumentBuilder::RemoveNode(const FGuid& InNodeID)
 			NodeCache.OnNodeAdded(Index);
 		}
 
+#if WITH_EDITORONLY_DATA
+		Document.Metadata.ModifyContext.AddNodeIDModified(InNodeID);
+#endif // WITH_EDITORONLY_DATA
+
 		return true;
 	}
 
@@ -1385,17 +1906,149 @@ bool FMetaSoundFrontendDocumentBuilder::SetNodeInputDefault(const FGuid& InNodeI
 	return false;
 }
 
-bool FMetaSoundFrontendDocumentBuilder::ModifyInterfaces(TArray<FMetasoundFrontendVersion> OutputFormatsToRemove, TArray<FMetasoundFrontendVersion> OutputFormatsToAdd)
+bool FMetaSoundFrontendDocumentBuilder::SwapGraphInput(const FMetasoundFrontendClassVertex& InExistingInputVertex, const FMetasoundFrontendClassVertex& InNewInputVertex)
 {
 	using namespace Metasound::Frontend;
 
-	// Integrating this transform into the builder will allow for more selective updates and not having to rebuild the cache from scratch after performing the transform.
-	FModifyRootGraphInterfaces ModifyRootGraph(OutputFormatsToRemove, OutputFormatsToAdd);
-	if (ModifyRootGraph.Transform(GetDocument()))
+	// 1. Check if equivalent and early out if functionally do not match
 	{
-		ReloadCache();
-		return true;
+		const FMetasoundFrontendClassInput* ClassInput = FindGraphInput(InExistingInputVertex.Name);
+		if (!ClassInput || !FMetasoundFrontendVertex::IsFunctionalEquivalent(*ClassInput, InExistingInputVertex))
+		{
+			return false;
+		}
 	}
 
-	return false;
+	const IDocumentGraphNodeCache& NodeCache = DocumentCache->GetNodeCache();
+
+	// 2. Gather data from existing member/node needed to swap
+	TArray<FMetasoundFrontendEdge> RemovedEdges;
+
+	const FMetasoundFrontendNode* ExistingInputNode = NodeCache.FindInputNode(InExistingInputVertex.Name);
+	checkf(ExistingInputNode, TEXT("'SwapGraphInput' failed to find original input node"));
+	const FGuid NodeID = ExistingInputNode->GetID();
+
+#if WITH_EDITOR
+	TMap<FGuid, FVector2D> Locations;
+#endif // WITH_EDITOR
+	{
+#if WITH_EDITOR
+		Locations = ExistingInputNode->Style.Display.Locations;
+#endif // WITH_EDITOR
+
+		const FGuid VertexID = ExistingInputNode->Interface.Outputs.Last().VertexID;
+		TArray<const FMetasoundFrontendEdge*> Edges = DocumentCache->GetEdgeCache().FindEdges(ExistingInputNode->GetID(), VertexID);
+		Algo::Transform(Edges, RemovedEdges, [](const FMetasoundFrontendEdge* Edge) { return *Edge; });
+	}
+
+	// 3. Remove existing graph vertex
+	{
+		const bool bRemovedVertex = RemoveGraphOutput(InExistingInputVertex.Name);
+		checkf(bRemovedVertex, TEXT("Failed to swap MetaSound input expected to exist"));
+	}
+
+	// 4. Add new graph vertex
+	FMetasoundFrontendClassInput NewInput = InNewInputVertex;
+	NewInput.NodeID = NodeID;
+#if WITH_EDITOR
+	NewInput.Metadata.SetSerializeText(InExistingInputVertex.Metadata.GetSerializeText());
+#endif // WITH_EDITOR
+
+	const FMetasoundFrontendNode* NewInputNode = AddGraphInput(NewInput);
+	checkf(NewInputNode, TEXT("Failed to add new Input node when swapping graph inputs"));
+
+#if WITH_EDITOR
+	// 5a. Add to new copy existing node locations
+	if (!Locations.IsEmpty())
+	{
+		const int32* NodeIndex = NodeCache.FindNodeIndex(NewInputNode->GetID());
+		checkf(NodeIndex, TEXT("Cache was not updated to reflect newly added input node"));
+		FMetasoundFrontendNode& NewNode = GetDocument().RootGraph.Graph.Nodes[*NodeIndex];
+		NewNode.Style.Display.Locations = Locations;
+	}
+#endif // WITH_EDITOR
+
+	// 5b. Add to new copy existing node edges
+	for (const FMetasoundFrontendEdge& RemovedEdge : RemovedEdges)
+	{
+		FMetasoundFrontendEdge NewEdge = RemovedEdge;
+		NewEdge.FromNodeID = NewInputNode->GetID();
+		NewEdge.FromVertexID = NewInputNode->Interface.Outputs.Last().VertexID;
+		const bool bEdgeAdded = AddEdge(MoveTemp(NewEdge)) != nullptr;
+		checkf(bEdgeAdded, TEXT("Failed to add replacement edge when swapping paired inputs"));
+	}
+	return true;
+}
+
+bool FMetaSoundFrontendDocumentBuilder::SwapGraphOutput(const FMetasoundFrontendClassVertex& InExistingOutputVertex, const FMetasoundFrontendClassVertex& InNewOutputVertex)
+{
+	using namespace Metasound::Frontend;
+
+	// 1. Check if equivalent and early out if functionally do not match
+	{
+		const FMetasoundFrontendClassOutput* ClassOutput = FindGraphOutput(InExistingOutputVertex.Name);
+		if (!ClassOutput || !FMetasoundFrontendVertex::IsFunctionalEquivalent(*ClassOutput, InExistingOutputVertex))
+		{
+			return false;
+		}
+	}
+
+	const IDocumentGraphNodeCache& NodeCache = DocumentCache->GetNodeCache();
+	const FMetasoundFrontendNode* ExistingOutputNode = NodeCache.FindOutputNode(InExistingOutputVertex.Name);
+	checkf(ExistingOutputNode, TEXT("'SwapGraphOutput' failed to find original output node"));
+	const FGuid NodeID = ExistingOutputNode->GetID();
+
+	// 2. Gather data from existing member/node needed to swap
+	TArray<FMetasoundFrontendEdge> RemovedEdges;
+#if WITH_EDITOR
+	TMap<FGuid, FVector2D> Locations;
+#endif // WITH_EDITOR
+	{
+#if WITH_EDITOR
+		Locations = ExistingOutputNode->Style.Display.Locations;
+#endif // WITH_EDITOR
+
+		const FGuid VertexID = ExistingOutputNode->Interface.Inputs.Last().VertexID;
+		TArray<const FMetasoundFrontendEdge*> Edges = DocumentCache->GetEdgeCache().FindEdges(ExistingOutputNode->GetID(), VertexID);
+		Algo::Transform(Edges, RemovedEdges, [](const FMetasoundFrontendEdge* Edge) { return *Edge; });
+	}
+
+	// 3. Remove existing graph vertex
+	{
+		const bool bRemovedVertex = RemoveGraphOutput(InExistingOutputVertex.Name);
+		checkf(bRemovedVertex, TEXT("Failed to swap output expected to exist while swapping MetaSound outputs"));
+	}
+	
+	// 4. Add new graph vertex
+	FMetasoundFrontendClassOutput NewOutput = InNewOutputVertex;
+	NewOutput.NodeID = NodeID;
+#if WITH_EDITOR
+	NewOutput.Metadata.SetSerializeText(InExistingOutputVertex.Metadata.GetSerializeText());
+#endif // WITH_EDITOR
+
+	const FMetasoundFrontendNode* NewOutputNode = AddGraphOutput(NewOutput);
+	checkf(NewOutputNode, TEXT("Failed to add new output node when swapping graph outputs"));
+
+#if WITH_EDITOR
+	// 5a. Add to new copy existing node locations
+	if (!Locations.IsEmpty())
+	{
+		const int32* NodeIndex = NodeCache.FindNodeIndex(NewOutputNode->GetID());
+		checkf(NodeIndex, TEXT("Cache was not updated to reflect newly added input node"));
+		FMetasoundFrontendNode& NewNode = GetDocument().RootGraph.Graph.Nodes[*NodeIndex];
+		NewNode.Style.Display.Locations = Locations;
+	}
+#endif // WITH_EDITOR
+
+	// 5b. Add to new copy existing node edges
+	for (const FMetasoundFrontendEdge& RemovedEdge : RemovedEdges)
+	{
+		FMetasoundFrontendEdge NewEdge = RemovedEdge;
+		NewEdge.ToNodeID = NewOutputNode->GetID();
+		NewEdge.ToVertexID = NewOutputNode->Interface.Inputs.Last().VertexID;
+		const bool bEdgeAdded = AddEdge(MoveTemp(NewEdge)) != nullptr;
+		checkf(bEdgeAdded, TEXT("Failed to add replacement edge when swapping paired outputs"));
+	}
+
+	return true;
 }
diff --git a/Engine/Plugins/Runtime/Metasound/Source/MetasoundFrontend/Private/MetasoundFrontendTransform.cpp b/Engine/Plugins/Runtime/Metasound/Source/MetasoundFrontend/Private/MetasoundFrontendTransform.cpp
--- a/Engine/Plugins/Runtime/Metasound/Source/MetasoundFrontend/Private/MetasoundFrontendTransform.cpp
+++ b/Engine/Plugins/Runtime/Metasound/Source/MetasoundFrontend/Private/MetasoundFrontendTransform.cpp
@@ -366,12 +366,12 @@ namespace Metasound
 				bDidEdit |= SwapPairedVertices(GraphHandle);
 				bDidEdit |= RemoveUnsupportedVertices(GraphHandle);
 
-#if WITH_EDITOR
+#if WITH_EDITORONLY_DATA
 				if (bAddedVertices && bSetDefaultNodeLocations)
 				{
 					UpdateAddedVertexNodePositions(GraphHandle);
 				}
-#endif // WITH_EDITOR
+#endif // WITH_EDITORONLY_DATA
 			}
 
 			return bDidEdit;
@@ -399,7 +399,7 @@ namespace Metasound
 			return !InterfacesToRemove.IsEmpty() || !InterfacesToAdd.IsEmpty();
 		}
 
-#if WITH_EDITOR
+#if WITH_EDITORONLY_DATA
 		void FModifyRootGraphInterfaces::UpdateAddedVertexNodePositions(FGraphHandle GraphHandle) const
 		{
 			auto SortAndPlaceMemberNodes = [&GraphHandle](EMetasoundFrontendClassType ClassType, TSet<FName>& AddedNames, TFunctionRef<int32(const FVertexName&)> InGetSortOrder)
diff --git a/Engine/Plugins/Runtime/Metasound/Source/MetasoundFrontend/Public/MetasoundDocumentInterface.h b/Engine/Plugins/Runtime/Metasound/Source/MetasoundFrontend/Public/MetasoundDocumentInterface.h
--- a/Engine/Plugins/Runtime/Metasound/Source/MetasoundFrontend/Public/MetasoundDocumentInterface.h
+++ b/Engine/Plugins/Runtime/Metasound/Source/MetasoundFrontend/Public/MetasoundDocumentInterface.h
@@ -28,6 +28,8 @@ public:
 	// Returns read-only reference to the the MetaSoundFrontendDocument
 	// containing all MetaSound runtime & editor data.
 	virtual const FMetasoundFrontendDocument& GetDocument() const = 0;
+
+	// Returns the parent class registered with the MetaSound UObject registry.
 	virtual const UClass& GetBaseMetaSoundUClass() const = 0;
 
 	// Returns the parent class registered with the MetaSound UObject registry.
diff --git a/Engine/Plugins/Runtime/Metasound/Source/MetasoundFrontend/Public/MetasoundFrontendDocument.h b/Engine/Plugins/Runtime/Metasound/Source/MetasoundFrontend/Public/MetasoundFrontendDocument.h
--- a/Engine/Plugins/Runtime/Metasound/Source/MetasoundFrontend/Public/MetasoundFrontendDocument.h
+++ b/Engine/Plugins/Runtime/Metasound/Source/MetasoundFrontend/Public/MetasoundFrontendDocument.h
@@ -91,8 +91,8 @@ public:
 	// Performs union of provided interface set with the set of interfaces that have been modified since last context reset/construction
 	void AddInterfacesModified(const TSet<FName>& InInterfacesModified);
 
-	// Adds a MemberID to the set of interfaces that have been modified since last context reset/construction
-	void AddMemberIDModified(const FGuid& InMemberIDModified);
+	// Adds a MemberID to the set of MemberIDs that have been modified since last context reset/construction
+	void AddMemberIDModified(const FGuid& InMemberNodeIDModified);
 
 	// Performs union of provided MemberIDs set with the set of MemberIDs that have been modified since last context reset/construction
 	void AddMemberIDsModified(const TSet<FGuid>& InMemberIDsModified);
@@ -100,7 +100,7 @@ public:
 	// Performs union of provided NodeID set with the set of NodeIDs that have been modified since last context reset/construction
 	void AddNodeIDModified(const FGuid& InNodeIDModified);
 
-	// Performs union of provided interface set with the set of interfaces that have been modified since last context reset/construction
+	// Performs union of provided NodeID set with the set of NodeIDs that have been modified since last context reset/construction
 	void AddNodeIDsModified(const TSet<FGuid>& InNodeIDsModified);
 };
 #endif // WITH_EDITORONLY_DATA
diff --git a/Engine/Plugins/Runtime/Metasound/Source/MetasoundFrontend/Public/MetasoundFrontendDocumentBuilder.h b/Engine/Plugins/Runtime/Metasound/Source/MetasoundFrontend/Public/MetasoundFrontendDocumentBuilder.h
--- a/Engine/Plugins/Runtime/Metasound/Source/MetasoundFrontend/Public/MetasoundFrontendDocumentBuilder.h
+++ b/Engine/Plugins/Runtime/Metasound/Source/MetasoundFrontend/Public/MetasoundFrontendDocumentBuilder.h
@@ -5,7 +5,6 @@
 #include "Interfaces/MetasoundFrontendInterfaceRegistry.h"
 #include "MetasoundFrontendDocumentCacheInterface.h"
 #include "MetasoundDocumentInterface.h"
-#include "MetasoundFrontendController.h"
 #include "MetasoundFrontendDocument.h"
 #include "MetasoundFrontendRegistries.h"
 #include "MetasoundFrontendTransform.h"
@@ -49,26 +48,54 @@ namespace Metasound::Frontend
 			return HashCombineFast(NameHash, GuidHash);
 		}
 	};
+
+	struct METASOUNDFRONTEND_API FModifyInterfaceOptions
+	{
+		FModifyInterfaceOptions(const TArray<FMetasoundFrontendInterface>& InInterfacesToRemove, const TArray<FMetasoundFrontendInterface>& InInterfacesToAdd);
+		FModifyInterfaceOptions(const TArray<FMetasoundFrontendVersion>& InInterfaceVersionsToRemove, const TArray<FMetasoundFrontendVersion>& InInterfaceVersionsToAdd);
+
+		TArray<FMetasoundFrontendInterface> InterfacesToRemove;
+		TArray<FMetasoundFrontendInterface> InterfacesToAdd;
+
+		// Function used to determine if an old of a removed interface
+		// and new member of an added interface are considered equal and
+		// to be swapped, retaining preexisting connections (and locations
+		// if in editor and 'SetDefaultNodeLocations' option is set)
+		TFunction<bool(FName, FName)> NamePairingFunction;
+
+#if WITH_EDITORONLY_DATA
+		bool bSetDefaultNodeLocations = true;
+#endif // WITH_EDITORONLY_DATA
+	};
 } // namespace Metasound::Frontend
 
 
 UCLASS()
-class  METASOUNDFRONTEND_API UMetaSoundBuilderDocument : public UObject, public IMetaSoundDocumentInterface
+class METASOUNDFRONTEND_API UMetaSoundBuilderDocument : public UObject, public IMetaSoundDocumentInterface
 {
 	GENERATED_BODY()
 
 public:
+	// Returns the document
 	virtual const FMetasoundFrontendDocument& GetDocument() const override
 	{
 		return Document;
 	}
 
+	// Base MetaSoundClass that document is published to.
 	virtual const UClass& GetBaseMetaSoundUClass() const final override
 	{
-		return *UMetaSoundBuilderDocument::StaticClass();
+		checkf(MetaSoundUClass, TEXT("BaseMetaSoundUClass must be set upon creation of UMetaSoundBuilderDocument instance"));
+		return *MetaSoundUClass;
 	}
 
 private:
+	// Sets the base MetaSound UClass (to be called immediately following UObject construction)
+	void SetBaseMetaSoundUClass(const UClass& InMetaSoundClass)
+	{
+		MetaSoundUClass = &InMetaSoundClass;
+	}
+
 	virtual FMetasoundFrontendDocument& GetDocument() override
 	{
 		return Document;
@@ -77,7 +104,11 @@ private:
 	UPROPERTY(Transient)
 	FMetasoundFrontendDocument Document;
 
+	UPROPERTY(Transient)
+	TObjectPtr<const UClass> MetaSoundUClass = nullptr;
+
 	friend struct FMetaSoundFrontendDocumentBuilder;
+	friend class UMetaSoundBuilderBase;
 };
 
 // Builder used to support dynamically generating MetaSound documents at runtime. Builder contains caches that speed up
@@ -92,10 +123,9 @@ struct METASOUNDFRONTEND_API FMetaSoundFrontendDocumentBuilder
 public:
 	// Exists only to make UObject reflection happy.  Should never be
 	// used directly as builder class should be specified on construction.
-	FMetaSoundFrontendDocumentBuilder();
-	FMetaSoundFrontendDocumentBuilder(const FTopLevelAssetPath& InBuilderClassPath);
-	FMetaSoundFrontendDocumentBuilder(const FTopLevelAssetPath& InBuilderClassPath, FMetaSoundFrontendDocumentBuilder&& InBuilder);
-	FMetaSoundFrontendDocumentBuilder(const FTopLevelAssetPath& InBuilderClassPath, TScriptInterface<IMetaSoundDocumentInterface> InDocumentInterface);
+	FMetaSoundFrontendDocumentBuilder() = default;
+	FMetaSoundFrontendDocumentBuilder(FMetaSoundFrontendDocumentBuilder&& InBuilder);
+	FMetaSoundFrontendDocumentBuilder(TScriptInterface<IMetaSoundDocumentInterface> InDocumentInterface);
 
 	// Copying actively used builders is highly discouraged (as only one builder should be acting on a document at any given time
 	// to avoid internal cache state to become out-of-date) but must contractually be implemented in this case as a USTRUCT (builders
@@ -134,7 +164,9 @@ public:
 	bool FindInterfaceInputNodes(FName InterfaceName, TArray<const FMetasoundFrontendNode*>& OutInputs) const;
 	bool FindInterfaceOutputNodes(FName InterfaceName, TArray<const FMetasoundFrontendNode*>& OutOutputs) const;
 
+	const FMetasoundFrontendClassInput* FindGraphInput(FName InputName) const;
 	const FMetasoundFrontendNode* FindGraphInputNode(FName InputName) const;
+	const FMetasoundFrontendClassOutput* FindGraphOutput(FName OutputName) const;
 	const FMetasoundFrontendNode* FindGraphOutputNode(FName OutputName) const;
 
 	const FMetasoundFrontendVertex* FindNodeInput(const FGuid& InNodeID, const FGuid& InVertexID) const;
@@ -153,28 +185,35 @@ public:
 	void InitDocument();
 	void InitNodeLocations();
 
+	bool IsDependencyReferenced(const FGuid& InClassID) const;
 	bool IsNodeInputConnected(const FGuid& InNodeID, const FGuid& InVertexID) const;
 	bool IsNodeOutputConnected(const FGuid& InNodeID, const FGuid& InVertexID) const;
 
 	bool IsInterfaceDeclared(FName InInterfaceName) const;
 	bool IsInterfaceDeclared(const FMetasoundFrontendVersion& InInterfaceVersion) const;
 
-	bool ModifyInterfaces(TArray<FMetasoundFrontendVersion> OutputFormatsToRemove, TArray<FMetasoundFrontendVersion> OutputFormatsToAdd);
+	bool ModifyInterfaces(Metasound::Frontend::FModifyInterfaceOptions&& InOptions);
 
+	bool RemoveDependency(const FGuid& InClassID);
 	bool RemoveDependency(EMetasoundFrontendClassType ClassType, const FMetasoundFrontendClassName& InClassName, const FMetasoundFrontendVersionNumber& InClassVersionNumber);
 	bool RemoveEdge(const FMetasoundFrontendEdge& EdgeToRemove);
-	bool RemoveNamedEdges(const TSet<Metasound::Frontend::FNamedEdge>& InNamedEdgesToRemove, TArray<FMetasoundFrontendEdge>* OutRemovedEdges = nullptr);
 	bool RemoveEdgesByNodeClassInterfaceBindings(const FGuid& InOutputNodeID, const FGuid& InInputNodeID);
 	bool RemoveEdgesFromNodeOutput(const FGuid& InNodeID, const FGuid& InVertexID);
 	bool RemoveEdgeToNodeInput(const FGuid& InNodeID, const FGuid& InVertexID);
+	bool RemoveGraphInput(FName InInputName);
+	bool RemoveGraphOutput(FName InOutputName);
 	bool RemoveInterface(FName InName);
+	bool RemoveNamedEdges(const TSet<Metasound::Frontend::FNamedEdge>& InNamedEdgesToRemove, TArray<FMetasoundFrontendEdge>* OutRemovedEdges = nullptr);
 	bool RemoveNode(const FGuid& InNodeID);
+	bool RemoveUnusedDependencies();
 
 #if WITH_EDITOR
 	void SetAuthor(const FString& InAuthor);
 #endif // WITH_EDITOR
 
 	bool SetNodeInputDefault(const FGuid& InNodeID, const FGuid& InVertexID, const FMetasoundFrontendLiteral& InLiteral);
+	bool SwapGraphInput(const FMetasoundFrontendClassVertex& InExistingInputVertex, const FMetasoundFrontendClassVertex& NewInputVertex);
+	bool SwapGraphOutput(const FMetasoundFrontendClassVertex& InExistingOutputVertex, const FMetasoundFrontendClassVertex& NewOutputVertex);
 
 private:
 	using FFinalizeNodeFunctionRef = TFunctionRef<void(FMetasoundFrontendNode&, const Metasound::Frontend::FNodeRegistryKey&)>;
@@ -183,8 +222,8 @@ private:
 	TScriptInterface<IMetaSoundDocumentInterface> DocumentInterface;
 
 	TUniquePtr<Metasound::Frontend::IDocumentCache> DocumentCache;
-	FTopLevelAssetPath BuilderClassPath;
 
+	const FTopLevelAssetPath GetBuilderClassPath() const;
 	FMetasoundFrontendDocument& GetDocument();
 
 	FMetasoundFrontendNode* AddNodeInternal(const FMetasoundFrontendClassMetadata& InClassMetadata, FFinalizeNodeFunctionRef FinalizeNode, FGuid InNodeID = FGuid::NewGuid());
diff --git a/Engine/Plugins/Runtime/Metasound/Source/MetasoundFrontend/Public/MetasoundFrontendTransform.h b/Engine/Plugins/Runtime/Metasound/Source/MetasoundFrontend/Public/MetasoundFrontendTransform.h
--- a/Engine/Plugins/Runtime/Metasound/Source/MetasoundFrontend/Public/MetasoundFrontendTransform.h
+++ b/Engine/Plugins/Runtime/Metasound/Source/MetasoundFrontend/Public/MetasoundFrontendTransform.h
@@ -100,11 +100,11 @@ namespace Metasound
 			bool RemoveUnsupportedVertices(FGraphHandle GraphHandle) const;
 			bool UpdateInterfacesInternal(FDocumentHandle DocumentHandle) const;
 
-#if WITH_EDITOR
+#if WITH_EDITORONLY_DATA
 			void UpdateAddedVertexNodePositions(FGraphHandle GraphHandle) const;
 
 			bool bSetDefaultNodeLocations = true;
-#endif // WITH_EDITOR
+#endif // WITH_EDITORONLY_DATA
 
 			TArray<FMetasoundFrontendInterface> InterfacesToRemove;
 			TArray<FMetasoundFrontendInterface> InterfacesToAdd;
