[Metasound] Fix race which would sometimes cause generator handles to stop receiving output updates from generators.

While the "thread-safe" variety of delegates does protect the invocation list from data races, if you use AddUObject, it will end up calling TWeakObjectPtr::Get() when it attempts to compact the invocation list (as it does when broadcasting or adding), which as far as I can tell is not thread-safe. This means it can sometimes incorrectly report that the object is no longer valid and should be removed from the invocation list.

#rnx

diff --git a/Engine/Plugins/Runtime/Metasound/Source/MetasoundEngine/Private/MetasoundGeneratorHandle.cpp b/Engine/Plugins/Runtime/Metasound/Source/MetasoundEngine/Private/MetasoundGeneratorHandle.cpp
--- a/Engine/Plugins/Runtime/Metasound/Source/MetasoundEngine/Private/MetasoundGeneratorHandle.cpp
+++ b/Engine/Plugins/Runtime/Metasound/Source/MetasoundEngine/Private/MetasoundGeneratorHandle.cpp
@@ -41,7 +41,7 @@ void UMetasoundGeneratorHandle::BeginDestroy()
 
 	if (TSharedPtr<Metasound::FMetasoundGenerator> PinnedGenerator = CachedGeneratorPtr.Pin())
 	{
-		PinnedGenerator->OnOutputChanged.RemoveAll(this);
+		PinnedGenerator->OnOutputChanged.Remove(OutputChangedDelegateHandle);
 	}
 }
 
@@ -389,7 +389,7 @@ void UMetasoundGeneratorHandle::UpdateWatchers()
 	METASOUND_LLM_SCOPE;
 	METASOUND_TRACE_CPUPROFILER_EVENT_SCOPE(UMetasoundGeneratorHandle::UpdateWatchers);
 
-	while (TOptional<FOutputPayload> ChangedOutput = ChangedOutputs.Dequeue())
+	while (TOptional<FOutputPayload> ChangedOutput = ChangedOutputs->Dequeue())
 	{
 		if (const FOutputListener* OutputListener = OutputListeners.FindByPredicate(
 			[&ChangedOutput](const FOutputListener& ExistingListener)
@@ -472,9 +472,20 @@ void UMetasoundGeneratorHandle::CreateAnalyzer(
 	Generator->AddOutputVertexAnalyzer(AnalyzerAddress);
 
 	// Subscribe for output updates if we haven't already
-	if (!Generator->OnOutputChanged.IsBoundToObject(this))
-	{
-		Generator->OnOutputChanged.AddUObject(this, &UMetasoundGeneratorHandle::HandleOutputChanged);
+	if (!OutputChangedDelegateHandle.IsValid())
+	{
+		TWeakPtr<TSpscQueue<FOutputPayload>> ChangedOutputQueue{ ChangedOutputs };
+		OutputChangedDelegateHandle = Generator->OnOutputChanged.AddLambda([ChangedOutputQueue](
+			FName AnalyzerName,
+			FName OutputName,
+			FName AnalyzerOutputName,
+			TSharedPtr<Metasound::IOutputStorage> OutputData)
+		{
+			if (const TSharedPtr<TSpscQueue<FOutputPayload>> PinnedQueue = ChangedOutputQueue.Pin())
+			{
+				PinnedQueue->Enqueue(AnalyzerName, OutputName, AnalyzerOutputName, OutputData);
+			}
+		});
 	}
 }
 
@@ -502,14 +513,3 @@ void UMetasoundGeneratorHandle::CreateListener(
 		OutputListeners.Emplace(AnalyzerAddress, OnOutputValueChanged);
 	}
 }
-
-// NOTE: We're in an audio render thread here, so get out fast
-// and be mindful of touching things that get touched on the game thread
-void UMetasoundGeneratorHandle::HandleOutputChanged(
-	FName AnalyzerName,
-	FName OutputName,
-	FName AnalyzerOutputName,
-	TSharedPtr<Metasound::IOutputStorage> OutputData)
-{
-	ChangedOutputs.Enqueue(AnalyzerName, OutputName, AnalyzerOutputName, OutputData);
-}
diff --git a/Engine/Plugins/Runtime/Metasound/Source/MetasoundEngine/Public/MetasoundGeneratorHandle.h b/Engine/Plugins/Runtime/Metasound/Source/MetasoundEngine/Public/MetasoundGeneratorHandle.h
--- a/Engine/Plugins/Runtime/Metasound/Source/MetasoundEngine/Public/MetasoundGeneratorHandle.h
+++ b/Engine/Plugins/Runtime/Metasound/Source/MetasoundEngine/Public/MetasoundGeneratorHandle.h
@@ -150,11 +150,6 @@ private:
 	void CreateListener(
 		const Metasound::Frontend::FAnalyzerAddress& AnalyzerAddress,
 		const FOnMetasoundOutputValueChanged& OnOutputValueChanged);
-	void HandleOutputChanged(
-		FName AnalyzerName,
-		FName OutputName,
-		FName AnalyzerOutputName,
-		TSharedPtr<Metasound::IOutputStorage> OutputData);
 
 	struct FPassthroughAnalyzerInfo
 	{
@@ -203,6 +198,7 @@ private:
 				, OutputValue(AnalyzerOutputName, OutputData)
 		{}
 	};
-
-	TSpscQueue<FOutputPayload> ChangedOutputs;
+	
+	TSharedRef<TSpscQueue<FOutputPayload>> ChangedOutputs = MakeShared<TSpscQueue<FOutputPayload>>();
+	FDelegateHandle OutputChangedDelegateHandle;
 };
\ No newline at end of file
