Streaming support for StemsPlayerNode.

Media Streamer Transport Controller Node generates Transport events and controls a MidiClock based on the signals and timing from a Media Streamer Node.
Stems Player Node uses a Media Streamer Transport Controller Node when it has a valid Media Player set as its Media Stream Asset.
Stemmed Music Asset provides the Media Source for the Media Player, and the channel mapping for its channels.

diff --git a/Engine/Plugins/Runtime/Metasound/Source/MetasoundFrontend/Public/MetasoundTime.h b/Engine/Plugins/Runtime/Metasound/Source/MetasoundFrontend/Public/MetasoundTime.h
--- a/Engine/Plugins/Runtime/Metasound/Source/MetasoundFrontend/Public/MetasoundTime.h
+++ b/Engine/Plugins/Runtime/Metasound/Source/MetasoundFrontend/Public/MetasoundTime.h
@@ -146,7 +146,7 @@ namespace Metasound
 			friend FTime operator-(const FTime& InLHS, const ArithmeticType& InRHS)
 			{
 				static_assert(TIsArithmetic<ArithmeticType>::Value, "Must be arithmetic type.");
-				return InLHS.Time - InRHS;
+				return FTime(InLHS.Time - InRHS);
 			}
 
 			FTime operator+(const FTime& InRHS) const
