Metasound Rebindable Graph Operator
- Added Metasound::FRebindableGraphOperator
- Modified FOperatorBuilder to build either FGraphOperator or FRebindableGraphOperator
- Various API updates and fix-ups to support aforementioned changes.

diff --git a/Engine/Plugins/Runtime/Metasound/Source/MetasoundEngineTest/Private/EngineTestAutomatedNodeTests.cpp b/Engine/Plugins/Runtime/Metasound/Source/MetasoundEngineTest/Private/EngineTestAutomatedNodeTests.cpp
--- a/Engine/Plugins/Runtime/Metasound/Source/MetasoundEngineTest/Private/EngineTestAutomatedNodeTests.cpp
+++ b/Engine/Plugins/Runtime/Metasound/Source/MetasoundEngineTest/Private/EngineTestAutomatedNodeTests.cpp
@@ -916,7 +916,8 @@ bool FMetasoundAutomatedNodeTest::RunTest(const FString& InRegistryKey)
 		// Bind to inputs and output data of operator.
 		FVertexInterface VertexInterface = Node->GetVertexInterface();
 		FVertexInterfaceData VertexInterfaceData{VertexInterface};
-		Operator->Bind(VertexInterfaceData);
+		Operator->BindInputs(VertexInterfaceData.GetInputs());
+		Operator->BindOutputs(VertexInterfaceData.GetOutputs());
 
 		// Create output tester which will analyzer outputs of operator.
 		FOutputVertexDataTestController OutputTester{VertexInterface.GetOutputInterface(), VertexInterfaceData.GetOutputs()};
diff --git a/Engine/Plugins/Runtime/Metasound/Source/MetasoundFrontend/Private/MetasoundInputNode.cpp b/Engine/Plugins/Runtime/Metasound/Source/MetasoundFrontend/Private/MetasoundInputNode.cpp
--- a/Engine/Plugins/Runtime/Metasound/Source/MetasoundFrontend/Private/MetasoundInputNode.cpp
+++ b/Engine/Plugins/Runtime/Metasound/Source/MetasoundFrontend/Private/MetasoundInputNode.cpp
@@ -35,18 +35,12 @@ namespace Metasound
 		{
 		}
 
-		void FNonExecutableInputOperatorBase::Bind(FVertexInterfaceData& InOutVertexData) const
-		{
-			BindInputsInternal(InOutVertexData.GetInputs());
-			BindOutputsInternal(InOutVertexData.GetOutputs());
-		}
-
-		void FNonExecutableInputOperatorBase::BindInputsInternal(FInputVertexInterfaceData& InOutVertexData) const
+		void FNonExecutableInputOperatorBase::BindInputs(FInputVertexInterfaceData& InOutVertexData)
 		{
 			InOutVertexData.BindVertex(VertexName, DataRef);
 		}
 
-		void FNonExecutableInputOperatorBase::BindOutputsInternal(FOutputVertexInterfaceData& InOutVertexData) const
+		void FNonExecutableInputOperatorBase::BindOutputs(FOutputVertexInterfaceData& InOutVertexData)
 		{
 			InOutVertexData.BindVertex(VertexName, DataRef);
 		}
diff --git a/Engine/Plugins/Runtime/Metasound/Source/MetasoundFrontend/Private/MetasoundOutputNode.cpp b/Engine/Plugins/Runtime/Metasound/Source/MetasoundFrontend/Private/MetasoundOutputNode.cpp
--- a/Engine/Plugins/Runtime/Metasound/Source/MetasoundFrontend/Private/MetasoundOutputNode.cpp
+++ b/Engine/Plugins/Runtime/Metasound/Source/MetasoundFrontend/Private/MetasoundOutputNode.cpp
@@ -30,10 +30,14 @@ namespace Metasound
 			return {};
 		}
 
-		void FOutputOperator::Bind(FVertexInterfaceData& InVertexData) const
+		void FOutputOperator::BindInputs(FInputVertexInterfaceData& InVertexData) 
 		{
-			InVertexData.GetInputs().BindVertex(VertexName, DataReference);
-			InVertexData.GetOutputs().BindVertex(VertexName, DataReference);
+			InVertexData.BindVertex(VertexName, DataReference);
+		}
+
+		void FOutputOperator::BindOutputs(FOutputVertexInterfaceData& InVertexData) 
+		{
+			InVertexData.BindVertex(VertexName, DataReference);
 		}
 
 		FOutputNode::FOutputOperatorFactory::FOutputOperatorFactory(const FVertexName& InVertexName, const FName& InDataTypeName, EVertexAccessType InVertexAccessType)
diff --git a/Engine/Plugins/Runtime/Metasound/Source/MetasoundFrontend/Public/MetasoundInputNode.h b/Engine/Plugins/Runtime/Metasound/Source/MetasoundFrontend/Public/MetasoundInputNode.h
--- a/Engine/Plugins/Runtime/Metasound/Source/MetasoundFrontend/Public/MetasoundInputNode.h
+++ b/Engine/Plugins/Runtime/Metasound/Source/MetasoundFrontend/Public/MetasoundInputNode.h
@@ -28,7 +28,8 @@ namespace Metasound
 		class METASOUNDFRONTEND_API FNonExecutableInputOperatorBase : public FInputOperatorBase
 		{	
 		public:
-			virtual void Bind(FVertexInterfaceData& InOutVertexData) const override;
+			virtual void BindInputs(FInputVertexInterfaceData& InOutVertexData) override;
+			virtual void BindOutputs(FOutputVertexInterfaceData& InOutVertexData) override;
 
 			virtual IOperator::FExecuteFunction GetExecuteFunction() override;
 			virtual IOperator::FPostExecuteFunction GetPostExecuteFunction() override;
@@ -37,12 +38,6 @@ namespace Metasound
 		protected:
 			FNonExecutableInputOperatorBase(const FVertexName& InVertexName, FAnyDataReference&& InDataRef);
 
-		private:
-			void BindInputsInternal(FInputVertexInterfaceData& InOutVertexData) const;
-		protected:
-			void BindOutputsInternal(FOutputVertexInterfaceData& InOutVertexData) const;
-
-
 			FVertexName VertexName;
 			FAnyDataReference DataRef;
 		};
@@ -99,10 +94,14 @@ namespace Metasound
 			{
 			}
 
-			virtual void Bind(FVertexInterfaceData& InOutVertexData) const override
+			virtual void BindInputs(FInputVertexInterfaceData& InOutVertexData) override
+			{
+				InOutVertexData.BindWriteVertex(DataReferenceName, InputValue);
+			}
+
+			virtual void BindOutputs(FOutputVertexInterfaceData& InOutVertexData) override
 			{
-				BindInputsInternal(InOutVertexData.GetInputs());
-				BindOutputsInternal(InOutVertexData.GetOutputs());
+				InOutVertexData.BindReadVertex(DataReferenceName, OutputValue);
 			}
 
 			virtual FExecuteFunction GetExecuteFunction() override
@@ -115,19 +114,9 @@ namespace Metasound
 				return nullptr;
 			}
 
-		private:
-			void BindInputsInternal(FInputVertexInterfaceData& InOutVertexData) const
-			{
-				InOutVertexData.BindWriteVertex(DataReferenceName, InputValue);
-			}
 
 		protected:
 
-			void BindOutputsInternal(FOutputVertexInterfaceData& InOutVertexData) const
-			{
-				InOutVertexData.BindReadVertex(DataReferenceName, OutputValue);
-			}
-
 			static void Execute(IOperator* InOperator)
 			{
 				using FExecutableInputOperator = TExecutableInputOperator<DataType>;
@@ -195,10 +184,14 @@ namespace Metasound
 			{
 			}
 
-			virtual void Bind(FVertexInterfaceData& InOutVertexData) const override
+			virtual void BindInputs(FInputVertexInterfaceData& InOutVertexData) override
+			{
+				InOutVertexData.BindWriteVertex(DataReferenceName, Value);
+			}
+
+			virtual void BindOutputs(FOutputVertexInterfaceData& InOutVertexData) override
 			{
-				BindInputsInternal(InOutVertexData.GetInputs());
-				BindOutputsInternal(InOutVertexData.GetOutputs());
+				InOutVertexData.BindReadVertex(DataReferenceName, Value);
 			}
 
 			virtual FExecuteFunction GetExecuteFunction() override
@@ -216,19 +209,8 @@ namespace Metasound
 				return nullptr;
 			}
 
-		private:
-			void BindInputsInternal(FInputVertexInterfaceData& InOutVertexData) const
-			{
-				InOutVertexData.BindWriteVertex(DataReferenceName, Value);
-			}
-
 		protected:
 
-			void BindOutputsInternal(FOutputVertexInterfaceData& InOutVertexData) const
-			{
-				InOutVertexData.BindReadVertex(DataReferenceName, Value);
-			}
-
 			static void PostExecute(IOperator* InOperator)
 			{
 				using FPostExecutableInputOperator = TPostExecutableInputOperator<DataType>;
@@ -241,7 +223,6 @@ namespace Metasound
 
 			FVertexName DataReferenceName;
 			FDataWriteReference Value;
-
 		};
 
 		template<typename DataType>
diff --git a/Engine/Plugins/Runtime/Metasound/Source/MetasoundFrontend/Public/MetasoundLiteralNode.h b/Engine/Plugins/Runtime/Metasound/Source/MetasoundFrontend/Public/MetasoundLiteralNode.h
--- a/Engine/Plugins/Runtime/Metasound/Source/MetasoundFrontend/Public/MetasoundLiteralNode.h
+++ b/Engine/Plugins/Runtime/Metasound/Source/MetasoundFrontend/Public/MetasoundLiteralNode.h
@@ -33,10 +33,14 @@ namespace Metasound
 
 		virtual ~TLiteralOperator() = default;
 
-		virtual void Bind(FVertexInterfaceData& InVertexData) const
+		virtual void BindInputs(FInputVertexInterfaceData& InVertexData) override
+		{
+		}
+
+		virtual void BindOutputs(FOutputVertexInterfaceData& InVertexData) override
 		{
 			using namespace LiteralNodeNames;
-			InVertexData.GetOutputs().BindValueVertex(METASOUND_GET_PARAM_NAME(OutputValue), Value);
+			InVertexData.BindValueVertex(METASOUND_GET_PARAM_NAME(OutputValue), Value);
 		}
 
 	private:
@@ -67,10 +71,14 @@ namespace Metasound
 
 		virtual ~TExecutableLiteralOperator() = default;
 
-		virtual void Bind(FVertexInterfaceData& InVertexData) const
+		virtual void BindInputs(FInputVertexInterfaceData& InVertexData) override
+		{
+		}
+
+		virtual void BindOutputs(FOutputVertexInterfaceData& InVertexData) override
 		{
 			using namespace LiteralNodeNames;
-			InVertexData.GetOutputs().BindReadVertex(METASOUND_GET_PARAM_NAME(OutputValue), OutputValue);
+			InVertexData.BindValueVertex(METASOUND_GET_PARAM_NAME(OutputValue), OutputValue);
 		}
 
 		virtual FDataReferenceCollection GetInputs() const override
@@ -125,10 +133,11 @@ namespace Metasound
 		{
 		public:
 			using TInputOperator<DataType>::TInputOperator;
-			virtual void Bind(FVertexInterfaceData& InVertexData) const override
+
+			virtual void BindInputs(FInputVertexInterfaceData& InVertexData) override
 			{
 				// Do not bind inputs as there are no inputs exposed from the literal node. 
-				TInputOperator<DataType>::BindOutputsInternal(InVertexData.GetOutputs());
+				// This overrides the implementation already in the TInputOperator<>
 			}
 		};
 
diff --git a/Engine/Plugins/Runtime/Metasound/Source/MetasoundFrontend/Public/MetasoundOutputNode.h b/Engine/Plugins/Runtime/Metasound/Source/MetasoundFrontend/Public/MetasoundOutputNode.h
--- a/Engine/Plugins/Runtime/Metasound/Source/MetasoundFrontend/Public/MetasoundOutputNode.h
+++ b/Engine/Plugins/Runtime/Metasound/Source/MetasoundFrontend/Public/MetasoundOutputNode.h
@@ -55,7 +55,8 @@ namespace Metasound
 
 				virtual FDataReferenceCollection GetInputs() const override;
 				virtual FDataReferenceCollection GetOutputs() const override;
-				virtual void Bind(FVertexInterfaceData& InVertexData) const;
+				virtual void BindInputs(FInputVertexInterfaceData& InVertexData) override;
+				virtual void BindOutputs(FOutputVertexInterfaceData& InVertexData) override;
 
 			private:
 				FVertexName VertexName;
diff --git a/Engine/Plugins/Runtime/Metasound/Source/MetasoundGenerator/Private/MetasoundGeneratorBuilder.cpp b/Engine/Plugins/Runtime/Metasound/Source/MetasoundGenerator/Private/MetasoundGeneratorBuilder.cpp
--- a/Engine/Plugins/Runtime/Metasound/Source/MetasoundGenerator/Private/MetasoundGeneratorBuilder.cpp
+++ b/Engine/Plugins/Runtime/Metasound/Source/MetasoundGenerator/Private/MetasoundGeneratorBuilder.cpp
@@ -191,11 +191,12 @@ namespace Metasound
 
 			METASOUND_TRACE_CPUPROFILER_EVENT_SCOPE_TEXT(TEXT("GeneratorBuilder::BuildGeneratorData"));
 
-			checkf(InGraphOperatorAndInputs.Operator.IsValid(), TEXT("Graph operator must be a valid object"));
+			checkf(InGraphOperatorAndInputs.Operator.IsValid(), TEXT("Graph operator must be a valid object"))
 
 			// Gather relevant input and output references
 			FVertexInterfaceData VertexData(InInitParams.Graph->GetVertexInterface());
-			InGraphOperatorAndInputs.Operator->Bind(VertexData);
+			InGraphOperatorAndInputs.Operator->BindInputs(VertexData.GetInputs());
+			InGraphOperatorAndInputs.Operator->BindOutputs(VertexData.GetOutputs());
 
 			// Replace input data with writable inputs
 			VertexData.GetInputs() = InGraphOperatorAndInputs.Inputs;
diff --git a/Engine/Plugins/Runtime/Metasound/Source/MetasoundGraphCore/Private/MetasoundGraphAlgo.cpp b/Engine/Plugins/Runtime/Metasound/Source/MetasoundGraphCore/Private/MetasoundGraphAlgo.cpp
--- a/Engine/Plugins/Runtime/Metasound/Source/MetasoundGraphCore/Private/MetasoundGraphAlgo.cpp
+++ b/Engine/Plugins/Runtime/Metasound/Source/MetasoundGraphCore/Private/MetasoundGraphAlgo.cpp
@@ -1,15 +1,23 @@
 // Copyright Epic Games, Inc. All Rights Reserved.
 
 #include "MetasoundGraphAlgo.h"
+#include "MetasoundGraphAlgoPrivate.h"
 
 #include "DSP/DirectedGraphAlgo.h"
+#include "Containers/Array.h"
+#include "Containers/Map.h"
+#include "Containers/Set.h"
+#include "MetasoundLog.h"
+#include "MetasoundNodeInterface.h"
+#include "MetasoundVertexData.h"
+#include "Templates/Tuple.h"
 
 namespace Metasound
 {
 	// The core graph algorithms utilize integers to represent vertices
 	// and pairs of vertices to represent edges. This class provides 
 	// convenience functions for converting back and forth between 
-	// metasound types and the core graph algorithm types.
+	// MetaSound types and the core graph algorithm types.
 	class FDirectedGraphAlgoAdapter
 	{
 		using FNodePair = TTuple<const INode*, const INode*>;
@@ -227,173 +235,379 @@ namespace Metasound
 			mutable FLazyCache Lazy;
 	};
 
+	
 	TPimplPtr<FDirectedGraphAlgoAdapter> FDirectedGraphAlgo::CreateDirectedGraphAlgoAdapter(const IGraph& InGraph)
 	{
-		return MakePimpl<FDirectedGraphAlgoAdapter>(InGraph);
+		return DirectedGraphAlgo::CreateDirectedGraphAlgoAdapter(InGraph);
 	}
 
 	bool FDirectedGraphAlgo::DepthFirstTopologicalSort(const IGraph& InGraph, TArray<const INode*>& OutNodeOrder)
 	{
-		return DepthFirstTopologicalSort(FDirectedGraphAlgoAdapter(InGraph), OutNodeOrder);
+		return DirectedGraphAlgo::DepthFirstTopologicalSort(InGraph, OutNodeOrder);
 	}
 
 	bool FDirectedGraphAlgo::DepthFirstTopologicalSort(const FDirectedGraphAlgoAdapter& InAdapter, TArray<const INode*>& OutNodeOrder)
 	{
-		using namespace Audio;
+		return DirectedGraphAlgo::DepthFirstTopologicalSort(InAdapter, OutNodeOrder);
+	}
 
-		TArray<int32> VertexOrder;
+	bool FDirectedGraphAlgo::KahnTopologicalSort(const IGraph& InGraph, TArray<const INode*>& OutNodeOrder)
+	{
+		return DirectedGraphAlgo::KahnTopologicalSort(InGraph, OutNodeOrder);
+	}
 
-		// Call algo implementation
-		bool bResult = Audio::FDirectedGraphAlgo::DepthFirstTopologicalSort(InAdapter.GetUniqueVertices(), InAdapter.GetUniqueEdges(), VertexOrder);
+	bool FDirectedGraphAlgo::KahnTopologicalSort(const FDirectedGraphAlgoAdapter& InAdapter, TArray<const INode*>& OutNodeOrder)
+	{
+		return DirectedGraphAlgo::KahnTopologicalSort(InAdapter, OutNodeOrder);
+	}
 
-		if (!bResult)
-		{
-			return false;
-		}
+	bool FDirectedGraphAlgo::TarjanStronglyConnectedComponents(const IGraph& InGraph, TArray<DirectedGraphAlgo::FStronglyConnectedComponent>& OutComponents, bool bExcludeSingleVertex)
+	{
+		return DirectedGraphAlgo::TarjanStronglyConnectedComponents(InGraph, OutComponents, bExcludeSingleVertex);
+	}
 
-		// Convert back to Metasound types.
-		OutNodeOrder.Reserve(OutNodeOrder.Num() + VertexOrder.Num());
+	bool FDirectedGraphAlgo::TarjanStronglyConnectedComponents(const FDirectedGraphAlgoAdapter& InAdapter, TArray<DirectedGraphAlgo::FStronglyConnectedComponent>& OutComponents, bool bExcludeSingleVertex)
+	{
+		return DirectedGraphAlgo::TarjanStronglyConnectedComponents(InAdapter, OutComponents, bExcludeSingleVertex);
+	}
 
-		for (int32 Vertex : VertexOrder)
-		{
-			OutNodeOrder.Add(InAdapter.GetNode(Vertex));
-		}
 
-		return bResult;
+	void FDirectedGraphAlgo::FindReachableNodesFromInput(const IGraph& InGraph, TSet<const INode*>& OutNodes)
+	{
+		DirectedGraphAlgo::FindReachableNodesFromInput(InGraph, OutNodes);
 	}
 
-	bool FDirectedGraphAlgo::KahnTopologicalSort(const IGraph& InGraph, TArray<const INode*>& OutNodeOrder)
+	void FDirectedGraphAlgo::FindReachableNodesFromInput(const FDirectedGraphAlgoAdapter& InAdapter, TSet<const INode*>& OutNodes)
 	{
-		return KahnTopologicalSort(FDirectedGraphAlgoAdapter(InGraph), OutNodeOrder);
+		DirectedGraphAlgo::FindReachableNodesFromInput(InAdapter, OutNodes);
 	}
 
-	bool FDirectedGraphAlgo::KahnTopologicalSort(const FDirectedGraphAlgoAdapter& InAdapter, TArray<const INode*>& OutNodeOrder)
+	void FDirectedGraphAlgo::FindReachableNodesFromOutput(const IGraph& InGraph, TSet<const INode*>& OutNodes)
+	{
+		DirectedGraphAlgo::FindReachableNodesFromOutput(InGraph, OutNodes);
+	}
+
+	void FDirectedGraphAlgo::FindReachableNodesFromOutput(const FDirectedGraphAlgoAdapter& InAdapter, TSet<const INode*>& OutNodes)
 	{
-		using namespace Audio;
+		DirectedGraphAlgo::FindReachableNodesFromOutput(InAdapter, OutNodes);
+	}
 
-		TArray<int32> VertexOrder;
+	void FDirectedGraphAlgo::FindReachableNodes(const IGraph& InGraph, TSet<const INode*>& OutNodes)
+	{
+		DirectedGraphAlgo::FindReachableNodes(InGraph, OutNodes);
+	}
 
-		// Call algo implementation
-		bool bResult = Audio::FDirectedGraphAlgo::KahnTopologicalSort(InAdapter.GetUniqueVertices(), InAdapter.GetUniqueEdges(), VertexOrder);
+	void FDirectedGraphAlgo::FindReachableNodes(const FDirectedGraphAlgoAdapter& InAdapter, TSet<const INode*>& OutNodes)
+	{
+		DirectedGraphAlgo::FindReachableNodes(InAdapter, OutNodes);
+	}
 
-		if (!bResult)
+	namespace DirectedGraphAlgo
+	{
+		
+		TPimplPtr<FDirectedGraphAlgoAdapter> CreateDirectedGraphAlgoAdapter(const IGraph& InGraph)
 		{
-			return false;
+			return MakePimpl<FDirectedGraphAlgoAdapter>(InGraph);
 		}
 
-		// Convert back to Metasound types.
-		OutNodeOrder.Reserve(OutNodeOrder.Num() + VertexOrder.Num());
-
-		for (int32 Vertex : VertexOrder)
+		bool DepthFirstTopologicalSort(const IGraph& InGraph, TArray<const INode*>& OutNodeOrder)
 		{
-			OutNodeOrder.Add(InAdapter.GetNode(Vertex));
+			return DepthFirstTopologicalSort(FDirectedGraphAlgoAdapter(InGraph), OutNodeOrder);
 		}
 
-		return bResult;
-	}
+		bool DepthFirstTopologicalSort(const FDirectedGraphAlgoAdapter& InAdapter, TArray<const INode*>& OutNodeOrder)
+		{
+			using namespace Audio;
 
-	bool FDirectedGraphAlgo::TarjanStronglyConnectedComponents(const IGraph& InGraph, TArray<FStronglyConnectedComponent>& OutComponents, bool bExcludeSingleVertex)
-	{
-		return TarjanStronglyConnectedComponents(FDirectedGraphAlgoAdapter(InGraph), OutComponents, bExcludeSingleVertex);
-	}
+			TArray<int32> VertexOrder;
 
-	bool FDirectedGraphAlgo::TarjanStronglyConnectedComponents(const FDirectedGraphAlgoAdapter& InAdapter, TArray<FStronglyConnectedComponent>& OutComponents, bool bExcludeSingleVertex)
-	{
-		TArray<Audio::FStronglyConnectedComponent> StronglyConnectedComponents;
+			// Call algo implementation
+			bool bResult = Audio::FDirectedGraphAlgo::DepthFirstTopologicalSort(InAdapter.GetUniqueVertices(), InAdapter.GetUniqueEdges(), VertexOrder);
 
-		// Run tarjan on metasound derived graph edges 
-		if (Audio::FDirectedGraphAlgo::TarjanStronglyConnectedComponents(InAdapter.GetEdgeSet(), StronglyConnectedComponents, bExcludeSingleVertex))
+			if (!bResult)
+			{
+				return false;
+			}
+
+			// Convert back to Metasound types.
+			OutNodeOrder.Reserve(OutNodeOrder.Num() + VertexOrder.Num());
+
+			for (int32 Vertex : VertexOrder)
+			{
+				OutNodeOrder.Add(InAdapter.GetNode(Vertex));
+			}
+
+			return bResult;
+		}
+
+		bool KahnTopologicalSort(const IGraph& InGraph, TArray<const INode*>& OutNodeOrder)
+		{
+			return KahnTopologicalSort(FDirectedGraphAlgoAdapter(InGraph), OutNodeOrder);
+		}
+
+		bool KahnTopologicalSort(const FDirectedGraphAlgoAdapter& InAdapter, TArray<const INode*>& OutNodeOrder)
 		{
-			// If strongly connected components are found, they must be converted
-			// back into metasound types. 
-			for (const Audio::FStronglyConnectedComponent& Component : StronglyConnectedComponents)
+			using namespace Audio;
+
+			TArray<int32> VertexOrder;
+
+			// Call algo implementation
+			bool bResult = Audio::FDirectedGraphAlgo::KahnTopologicalSort(InAdapter.GetUniqueVertices(), InAdapter.GetUniqueEdges(), VertexOrder);
+
+			if (!bResult)
 			{
-				FStronglyConnectedComponent& MetasoundGraphComponent = OutComponents.AddDefaulted_GetRef();
+				return false;
+			}
 
-				for (int32 Vertex : Component.Vertices)
-				{
-					MetasoundGraphComponent.Nodes.Add(InAdapter.GetNode(Vertex));
-				}
+			// Convert back to Metasound types.
+			OutNodeOrder.Reserve(OutNodeOrder.Num() + VertexOrder.Num());
+
+			for (int32 Vertex : VertexOrder)
+			{
+				OutNodeOrder.Add(InAdapter.GetNode(Vertex));
+			}
+
+			return bResult;
+		}
+
+		bool TarjanStronglyConnectedComponents(const IGraph& InGraph, TArray<FStronglyConnectedComponent>& OutComponents, bool bExcludeSingleVertex)
+		{
+			return TarjanStronglyConnectedComponents(FDirectedGraphAlgoAdapter(InGraph), OutComponents, bExcludeSingleVertex);
+		}
+
+		bool TarjanStronglyConnectedComponents(const FDirectedGraphAlgoAdapter& InAdapter, TArray<FStronglyConnectedComponent>& OutComponents, bool bExcludeSingleVertex)
+		{
+			TArray<Audio::FStronglyConnectedComponent> StronglyConnectedComponents;
 
-				for (const Audio::FDirectedEdge& Edge : Component.Edges)
+			// Run tarjan on metasound derived graph edges 
+			if (Audio::FDirectedGraphAlgo::TarjanStronglyConnectedComponents(InAdapter.GetEdgeSet(), StronglyConnectedComponents, bExcludeSingleVertex))
+			{
+				// If strongly connected components are found, they must be converted
+				// back into metasound types. 
+				for (const Audio::FStronglyConnectedComponent& Component : StronglyConnectedComponents)
 				{
-					TArray<FDataEdge> MetasoundComponentEdges;
+					FStronglyConnectedComponent& MetasoundGraphComponent = OutComponents.AddDefaulted_GetRef();
+
+					for (int32 Vertex : Component.Vertices)
+					{
+						MetasoundGraphComponent.Nodes.Add(InAdapter.GetNode(Vertex));
+					}
 
-					InAdapter.GetDataEdges(Edge, MetasoundComponentEdges);
+					for (const Audio::FDirectedEdge& Edge : Component.Edges)
+					{
+						TArray<FDataEdge> MetasoundComponentEdges;
 
-					MetasoundGraphComponent.Edges.Append(MetasoundComponentEdges);
+						InAdapter.GetDataEdges(Edge, MetasoundComponentEdges);
+
+						MetasoundGraphComponent.Edges.Append(MetasoundComponentEdges);
+					}
 				}
+
+				return true;
 			}
 
-			return true;
+			// No strongly connected components found.
+			return false;
 		}
 
-		// No strongly connected components found.
-		return false;
-	}
 
+		void FindReachableNodesFromInput(const IGraph& InGraph, TSet<const INode*>& OutNodes)
+		{
+			FindReachableNodesFromInput(FDirectedGraphAlgoAdapter(InGraph), OutNodes);
+		}
 
-	void FDirectedGraphAlgo::FindReachableNodesFromInput(const IGraph& InGraph, TSet<const INode*>& OutNodes)
-	{
-		FindReachableNodesFromInput(FDirectedGraphAlgoAdapter(InGraph), OutNodes);
-	}
+		void FindReachableNodesFromInput(const FDirectedGraphAlgoAdapter& InAdapter, TSet<const INode*>& OutNodes)
+		{
+			TSet<int32> VisitedVertices;
 
-	void FDirectedGraphAlgo::FindReachableNodesFromInput(const FDirectedGraphAlgoAdapter& InAdapter, TSet<const INode*>& OutNodes)
-	{
-		TSet<int32> VisitedVertices;
+			for (int32 Vertex : InAdapter.GetUniqueInputVertices())
+			{
+				Audio::FDirectedGraphAlgo::DepthFirstTraversal(Vertex, InAdapter.GetTree(), [&](int32 InVertexBeingVisited) -> bool 
+					{
+						bool bIsAlreadyInSet = false;
+						VisitedVertices.Add(InVertexBeingVisited, &bIsAlreadyInSet);
+
+						return !bIsAlreadyInSet;
+					}
+				);
+			}
+
+			InAdapter.GetNodeSet(VisitedVertices, OutNodes);
+		}
 
-		for (int32 Vertex : InAdapter.GetUniqueInputVertices())
+		void FindReachableNodesFromOutput(const IGraph& InGraph, TSet<const INode*>& OutNodes)
 		{
-			Audio::FDirectedGraphAlgo::DepthFirstTraversal(Vertex, InAdapter.GetTree(), [&](int32 InVertexBeingVisited) -> bool 
-				{
-					bool bIsAlreadyInSet = false;
-					VisitedVertices.Add(InVertexBeingVisited, &bIsAlreadyInSet);
+			FindReachableNodesFromOutput(FDirectedGraphAlgoAdapter(InGraph), OutNodes);
+		}
 
-					return !bIsAlreadyInSet;
-				}
-			);
+		void FindReachableNodesFromOutput(const FDirectedGraphAlgoAdapter& InAdapter, TSet<const INode*>& OutNodes)
+		{
+			TSet<int32> VisitedVertices;
+
+			for (int32 Vertex : InAdapter.GetUniqueOutputVertices())
+			{ 
+				Audio::FDirectedGraphAlgo::DepthFirstTraversal(Vertex, InAdapter.GetTransposeTree(), [&](int32 InVertexBeingVisited) -> bool 
+					{
+						bool bIsAlreadyInSet = false;
+						VisitedVertices.Add(InVertexBeingVisited, &bIsAlreadyInSet);
+
+						return !bIsAlreadyInSet;
+					}
+				);
+			}
+
+			InAdapter.GetNodeSet(VisitedVertices, OutNodes);
 		}
 
-		InAdapter.GetNodeSet(VisitedVertices, OutNodes);
-	}
+		void FindReachableNodes(const IGraph& InGraph, TSet<const INode*>& OutNodes)
+		{
+			FindReachableNodes(FDirectedGraphAlgoAdapter(InGraph), OutNodes);
+		}
 
-	void FDirectedGraphAlgo::FindReachableNodesFromOutput(const IGraph& InGraph, TSet<const INode*>& OutNodes)
-	{
-		FindReachableNodesFromOutput(FDirectedGraphAlgoAdapter(InGraph), OutNodes);
-	}
+		void FindReachableNodes(const FDirectedGraphAlgoAdapter& InAdapter, TSet<const INode*>& OutNodes)
+		{
+			TSet<const INode*> NodesFromInput;
+			FindReachableNodesFromInput(InAdapter, NodesFromInput);
 
-	void FDirectedGraphAlgo::FindReachableNodesFromOutput(const FDirectedGraphAlgoAdapter& InAdapter, TSet<const INode*>& OutNodes)
-	{
-		TSet<int32> VisitedVertices;
+			TSet<const INode*> NodesFromOutput;
+			FindReachableNodesFromOutput(InAdapter, NodesFromOutput);
 
-		for (int32 Vertex : InAdapter.GetUniqueOutputVertices())
-		{ 
-			Audio::FDirectedGraphAlgo::DepthFirstTraversal(Vertex, InAdapter.GetTransposeTree(), [&](int32 InVertexBeingVisited) -> bool 
-				{
-					bool bIsAlreadyInSet = false;
-					VisitedVertices.Add(InVertexBeingVisited, &bIsAlreadyInSet);
+			OutNodes = NodesFromInput.Union(NodesFromOutput);
+		}
+
+		namespace GraphAlgoPrivate
+		{
+			struct FInputToUpdate
+			{
+				FOperatorID OperatorID;
+				FVertexName VertexName;
+				FAnyDataReference DataReference;
+			};
+
+			// Apply updates to data references through all the operators by following connections described in the FOperatorInfo map.
+			void PropagateBindUpdate(FOperatorID InInitialOperatorID, const FVertexName& InVertexName, const FAnyDataReference& InNewReference, TSortedMap<FOperatorID, FGraphOperatorData::FOperatorInfo>& InOperatorInfoMap, FOnPostOperatorRebind OnPostOperatorRebind)
+			{
+				TArray<FInputToUpdate> UpdateStack;
+				UpdateStack.Emplace(FInputToUpdate{InInitialOperatorID, InVertexName, InNewReference});
 
-					return !bIsAlreadyInSet;
+				TArray<FVertexDataState> InitialOutputState;
+				TSortedVertexNameMap<FAnyDataReference> OutputUpdates;
+				while (UpdateStack.Num())
+				{
+					FInputToUpdate Current = UpdateStack.Pop();
+					if (FGraphOperatorData::FOperatorInfo* OpInfo = InOperatorInfoMap.Find(Current.OperatorID))
+					{
+						// Get current outputs
+						InitialOutputState.Reset();
+						GetVertexInterfaceDataState(OpInfo->VertexData.GetOutputs(), InitialOutputState);
+
+						// Set new input.
+						OpInfo->VertexData.GetInputs().SetVertex(Current.VertexName, MoveTemp(Current.DataReference));
+
+						// Bind inputs and outputs
+						OpInfo->Operator->BindInputs(OpInfo->VertexData.GetInputs());
+						OpInfo->Operator->BindOutputs(OpInfo->VertexData.GetOutputs());
+						OnPostOperatorRebind(Current.OperatorID, *(OpInfo->Operator));
+
+						// See if binding altered the outputs. 
+						OutputUpdates.Reset();
+						CompareVertexInterfaceDataToPriorState(OpInfo->VertexData.GetOutputs(), InitialOutputState, OutputUpdates);
+
+						// Any updates to the outputs need to be propagated through the graph.
+						for (const TPair<FVertexName, FAnyDataReference>& OutputUpdate : OutputUpdates)
+						{
+							const FVertexName& OutputVertexName = OutputUpdate.Get<0>();
+							const FAnyDataReference& OutputDataReference = OutputUpdate.Get<1>();
+
+							if (const TArray<FGraphOperatorData::FVertexDestination>* Destinations = OpInfo->OutputConnections.Find(OutputVertexName))
+							{
+								for (const FGraphOperatorData::FVertexDestination& Destination : *Destinations)
+								{
+									UpdateStack.Push({Destination.OperatorID, Destination.VertexName, OutputDataReference});
+								}
+							}
+						}
+					}
+					else
+					{
+						UE_LOG(LogMetaSound, Error, TEXT("Failed to rebind graph operator state. Could not find operator info with ID %s"), *LexToString(Current.OperatorID));
+					}
 				}
-			);
+			}
+
 		}
 
-		InAdapter.GetNodeSet(VisitedVertices, OutNodes);
-	}
+		FOperatorID GetOperatorID(const INode& InNode)
+		{
+			return GetOperatorID(&InNode);
+		}
 
-	void FDirectedGraphAlgo::FindReachableNodes(const IGraph& InGraph, TSet<const INode*>& OutNodes)
-	{
-		FindReachableNodes(FDirectedGraphAlgoAdapter(InGraph), OutNodes);
-	}
+		FOperatorID GetOperatorID(const INode* InNode)
+		{
+			return reinterpret_cast<FOperatorID>(InNode);
+		}
+		
+		FGraphOperatorData::FGraphOperatorData(const FOperatorSettings& InOperatorSettings)
+		: OperatorSettings(InOperatorSettings)
+		{
+		}
+		
+		void RebindGraphInputs(FInputVertexInterfaceData& InOutVertexData, FRebindGraphDataParams& InParams, FOnPostOperatorRebind OnPostOperatorRebind)
+		{
+			// Bind and diff the graph's interface to determine if there is an update to any vertices
+			FInputVertexInterfaceData& InputVertexData = InParams.VertexData.GetInputs();
+			TArray<FVertexDataState> InitialVertexDataState;
+			GetVertexInterfaceDataState(InputVertexData, InitialVertexDataState);
 
-	void FDirectedGraphAlgo::FindReachableNodes(const FDirectedGraphAlgoAdapter& InAdapter, TSet<const INode*>& OutNodes)
-	{
-		TSet<const INode*> NodesFromInput;
-		FindReachableNodesFromInput(InAdapter, NodesFromInput);
+			InOutVertexData.Bind(InputVertexData);
 
-		TSet<const INode*> NodesFromOutput;
-		FindReachableNodesFromOutput(InAdapter, NodesFromOutput);
+			TSortedVertexNameMap<FAnyDataReference> GraphInputsToUpdate;
+			CompareVertexInterfaceDataToPriorState(InputVertexData, InitialVertexDataState, GraphInputsToUpdate);
 
-		OutNodes = NodesFromInput.Union(NodesFromOutput);
+			if (GraphInputsToUpdate.Num() > 0)
+			{
+				// Update the graph data by propagating the updates through the inputs nodes. 
+				for (const TPair<FVertexName, FAnyDataReference>& InputToUpdate : GraphInputsToUpdate)
+				{
+					const FVertexName& VertexName = InputToUpdate.Get<0>();
+					if (const FOperatorID* OperatorID = InParams.InputVertexMap.Find(VertexName))
+					{
+						GraphAlgoPrivate::PropagateBindUpdate(*OperatorID, VertexName, InputToUpdate.Get<1>(), InParams.OperatorMap, OnPostOperatorRebind);
+					}
+					else
+					{
+						UE_LOG(LogMetaSound, Error, TEXT("No input operator exists for input vertex %s"), *VertexName.ToString());
+					}
+				}
+
+				// Refresh output vertex interface data.
+				for (const TPair<FVertexName, FOperatorID>& OutputVertexInfo : InParams.OutputVertexMap)
+				{
+					const FVertexName& VertexName = OutputVertexInfo.Get<0>();
+					const FOperatorID& OperatorID = OutputVertexInfo.Get<1>();
+
+					if (const FGraphOperatorData::FOperatorInfo* OperatorInfo = InParams.OperatorMap.Find(OperatorID))
+					{
+						if (const FAnyDataReference* Ref = OperatorInfo->VertexData.GetOutputs().FindDataReference(OutputVertexInfo.Get<0>()))
+						{
+							InParams.VertexData.GetOutputs().SetVertex(VertexName, *Ref);
+						}
+						else if (InParams.VertexData.GetOutputs().IsVertexBound(VertexName))
+						{
+							UE_LOG(LogMetaSound, Error, TEXT("Output vertex (%s) lost data reference after rebinding graph"), *VertexName.ToString());
+						}
+					}
+					else
+					{
+						UE_LOG(LogMetaSound, Error, TEXT("Failed to rebind graph operator state outputs. Could not find output operator info with ID %s"), *LexToString(OperatorID));
+					}
+				}
+			}
+		}
+
+		void RebindGraphOutputs(FOutputVertexInterfaceData& InOutVertexData, FRebindGraphDataParams& InParams, FOnPostOperatorRebind OnPostOperatorRebind)
+		{
+			InOutVertexData.Bind(InParams.VertexData.GetOutputs());
+		}
 	}
 }
diff --git a/Engine/Plugins/Runtime/Metasound/Source/MetasoundGraphCore/Private/MetasoundGraphAlgoPrivate.h b/Engine/Plugins/Runtime/Metasound/Source/MetasoundGraphCore/Private/MetasoundGraphAlgoPrivate.h
new file mode 100644
--- /dev/null
+++ b/Engine/Plugins/Runtime/Metasound/Source/MetasoundGraphCore/Private/MetasoundGraphAlgoPrivate.h
@@ -0,0 +1,87 @@
+// Copyright Epic Games, Inc. All Rights Reserved.
+
+#pragma once
+
+#include "Containers/Array.h"
+#include "Containers/SortedMap.h"
+#include "MetasoundOperatorInterface.h"
+#include "MetasoundVertexData.h"
+#include "Templates/UniquePtr.h"
+
+namespace Metasound
+{
+	// Forward Declare
+	class INode;
+
+	namespace DirectedGraphAlgo
+	{
+		// ID used to lookup operators.
+		using FOperatorID = uintptr_t;
+
+		// Retrieve an operator ID from a node. 
+		FOperatorID GetOperatorID(const INode& InNode);
+		FOperatorID GetOperatorID(const INode* InNode);
+
+		/** FGraphOoperatorData contains all the objects needed to implement a runtime instance of a MetaSound graph.  */
+		struct FGraphOperatorData
+		{
+			// Represents an input of an operator
+			struct FVertexDestination
+			{
+				FOperatorID OperatorID;
+				FVertexName VertexName;
+			};
+
+			struct FOperatorInfo
+			{
+				TUniquePtr<IOperator> Operator;
+
+				// Vertex Data bound to operator.
+				FVertexInterfaceData VertexData;
+
+				// Map where the operators output vertex is the key, and an array of output
+				// connections is the value.
+				TSortedVertexNameMap<TArray<FVertexDestination>> OutputConnections;
+			};
+
+			FGraphOperatorData(const FOperatorSettings& InOperatorSettings);
+
+			FOperatorSettings OperatorSettings;
+
+			// Vertex Data bound to the graph
+			FVertexInterfaceData VertexData;
+
+			// Sorted order of operators for execution.
+			TArray<FOperatorID> OperatorOrder;
+
+			TSortedMap<FOperatorID, FOperatorInfo> OperatorMap;
+
+			// Map with input vertex name as key, and OperatorID of input node as value.
+			TSortedVertexNameMap<FOperatorID> InputVertexMap;
+
+			// Map with output vertex name as key, and OperatorID of output node as value.
+			TSortedVertexNameMap<FOperatorID> OutputVertexMap;
+		};
+
+		// Arguments for calling Rebind
+		struct FRebindGraphDataParams
+		{
+			FVertexInterfaceData& VertexData; 											//< Current vertex data of graph
+			TSortedMap<FOperatorID, FGraphOperatorData::FOperatorInfo>& OperatorMap; 	//< Map of operator data in graph
+			TSortedVertexNameMap<FOperatorID>& InputVertexMap; 							//< Map of graph input nodes
+			TSortedVertexNameMap<FOperatorID>& OutputVertexMap; 						//< Map of graph output nodes
+		};
+
+		using FOnPostOperatorRebind = TFunctionRef<void(FOperatorID, IOperator&)>;
+
+		// Rebind the graph inputs, updating internal operator bindings as needed
+		//
+		// When an internal node is rebound, the OnPostOperatorRebind will be called. 
+		void RebindGraphInputs(FInputVertexInterfaceData& InOutVertexData, FRebindGraphDataParams& InOutParams, FOnPostOperatorRebind OnPostOperatorRebind);
+
+		// Rebind the graph inputs, updating internal operator bindings as needed
+		//
+		// When an internal node is rebound, the OnPostOperatorRebind will be called. 
+		void RebindGraphOutputs(FOutputVertexInterfaceData& InOutVertexData, FRebindGraphDataParams& InOutParams, FOnPostOperatorRebind OnPostOperatorRebind);
+	}
+}
diff --git a/Engine/Plugins/Runtime/Metasound/Source/MetasoundGraphCore/Private/MetasoundGraphLinter.cpp b/Engine/Plugins/Runtime/Metasound/Source/MetasoundGraphCore/Private/MetasoundGraphLinter.cpp
--- a/Engine/Plugins/Runtime/Metasound/Source/MetasoundGraphCore/Private/MetasoundGraphLinter.cpp
+++ b/Engine/Plugins/Runtime/Metasound/Source/MetasoundGraphCore/Private/MetasoundGraphLinter.cpp
@@ -199,7 +199,7 @@ namespace Metasound
 
 	bool FGraphLinter::ValidateNoCyclesInGraph(const IGraph& InGraph, TArray<FBuildErrorPtr>& OutErrors) 
 	{
-		TPimplPtr<FDirectedGraphAlgoAdapter> Adapter = FDirectedGraphAlgo::CreateDirectedGraphAlgoAdapter(InGraph);
+		TPimplPtr<FDirectedGraphAlgoAdapter> Adapter = DirectedGraphAlgo::CreateDirectedGraphAlgoAdapter(InGraph);
 
 		check(Adapter.IsValid());
 
@@ -212,18 +212,18 @@ namespace Metasound
 
 		bool bIsValid = true;
 
-		TArray<FStronglyConnectedComponent> Cycles;
+		TArray<DirectedGraphAlgo::FStronglyConnectedComponent> Cycles;
 
 		// In graph theory, a single vertex is technically a strongly connected
 		// component. The graph linter is only interested in strongly connected
 		// components of more than one vertex since this denotes a cycle.
 		bool bExcludeSingleVertex = true;
 
-		if(FDirectedGraphAlgo::TarjanStronglyConnectedComponents(InAdapter, Cycles, bExcludeSingleVertex))
+		if(DirectedGraphAlgo::TarjanStronglyConnectedComponents(InAdapter, Cycles, bExcludeSingleVertex))
 		{
 			bIsValid = false;
 
-			for (const FStronglyConnectedComponent& Cycle : Cycles)
+			for (const DirectedGraphAlgo::FStronglyConnectedComponent& Cycle : Cycles)
 			{
 				AddBuildError<FGraphCycleError>(OutErrors, Cycle.Nodes, Cycle.Edges);
 			}
diff --git a/Engine/Plugins/Runtime/Metasound/Source/MetasoundGraphCore/Private/MetasoundGraphOperator.cpp b/Engine/Plugins/Runtime/Metasound/Source/MetasoundGraphCore/Private/MetasoundGraphOperator.cpp
--- a/Engine/Plugins/Runtime/Metasound/Source/MetasoundGraphCore/Private/MetasoundGraphOperator.cpp
+++ b/Engine/Plugins/Runtime/Metasound/Source/MetasoundGraphCore/Private/MetasoundGraphOperator.cpp
@@ -7,6 +7,7 @@
 #include "MetasoundDataReference.h"
 #include "MetasoundDataReferenceCollection.h"
 #include "MetasoundExecutableOperator.h"
+#include "MetasoundGraphAlgoPrivate.h"
 #include "MetasoundOperatorInterface.h"
 
 namespace Metasound
@@ -64,7 +65,7 @@ namespace Metasound
 				}
 			}
 
-			if (!NewIter)
+			if (NewIter)
 			{
 				UE_LOG(LogMetaSound, Warning, TEXT("Cannot bind to FGraphOperator because vertex %s does not exist in current vertex data"), *(NewIter->VertexName.ToString()));
 				return false;
@@ -77,6 +78,20 @@ namespace Metasound
 		}
 	}
 
+	FGraphOperator::FGraphOperator(TUniquePtr<DirectedGraphAlgo::FGraphOperatorData>&& InOperatorState)
+	{
+		using namespace DirectedGraphAlgo;
+
+		// Append operators in order.
+		for (const FOperatorID OperatorID : InOperatorState->OperatorOrder)
+		{
+			AppendOperator(MoveTemp(InOperatorState->OperatorMap[OperatorID].Operator));
+		}
+
+		// Copy over vertex data
+		SetVertexInterfaceData(MoveTemp(InOperatorState->VertexData));
+	}
+
 	void FGraphOperator::AppendOperator(FOperatorPtr InOperator)
 	{
 		if (InOperator.IsValid())
diff --git a/Engine/Plugins/Runtime/Metasound/Source/MetasoundGraphCore/Private/MetasoundOperatorBuilder.cpp b/Engine/Plugins/Runtime/Metasound/Source/MetasoundGraphCore/Private/MetasoundOperatorBuilder.cpp
--- a/Engine/Plugins/Runtime/Metasound/Source/MetasoundGraphCore/Private/MetasoundOperatorBuilder.cpp
+++ b/Engine/Plugins/Runtime/Metasound/Source/MetasoundGraphCore/Private/MetasoundOperatorBuilder.cpp
@@ -8,15 +8,52 @@
 #include "MetasoundBuilderInterface.h"
 #include "MetasoundBuildError.h"
 #include "MetasoundGraphAlgo.h"
+#include "MetasoundGraphAlgoPrivate.h"
 #include "MetasoundGraphLinter.h"
 #include "MetasoundGraphOperator.h"
 #include "MetasoundOperatorInterface.h"
+#include "MetasoundRebindableGraphOperator.h"
 #include "MetasoundThreadLocalDebug.h"
+#include "MetasoundTrace.h"
 #include "Templates/PimplPtr.h"
 
 
 namespace Metasound
 {
+	namespace OperatorBuilder
+	{
+		// Shared context used in the builder to maintain state of current build.
+		struct FBuildContext
+		{
+			const IGraph& Graph;
+			const FDirectedGraphAlgoAdapter& AlgoAdapter;
+			const FOperatorSettings& Settings;
+			const FMetasoundEnvironment& Environment;
+			const FOperatorBuilderSettings& BuilderSettings;
+			
+			FBuildResults& Results;
+			TUniquePtr<DirectedGraphAlgo::FGraphOperatorData> GraphOperatorData;
+
+			FBuildContext(
+				const IGraph& InGraph,
+				const FDirectedGraphAlgoAdapter& InAlgoAdapter,
+				const FOperatorSettings& InSettings,
+				const FMetasoundEnvironment& InEnvironment,
+				const FOperatorBuilderSettings& bInBuilderSettings,
+				FBuildResults& OutResults)
+			: Graph(InGraph)
+			, AlgoAdapter(InAlgoAdapter)
+			, Settings(InSettings)
+			, Environment(InEnvironment)
+			, BuilderSettings(bInBuilderSettings)
+			, Results(OutResults)
+			, GraphOperatorData(MakeUnique<DirectedGraphAlgo::FGraphOperatorData>(InSettings))
+			{
+				GraphOperatorData->VertexData = FVertexInterfaceData(InGraph.GetVertexInterface());
+			}
+		};
+	}
+
 	namespace OperatorBuilderPrivate
 	{
 		using FBuildErrorPtr = TUniquePtr<IOperatorBuildError>;
@@ -29,7 +66,6 @@ namespace Metasound
 				AddBuildError<FInternalError>(OutErrors, __FILE__, __LINE__);
 			}
 		}
-
 	}
 
 	FOperatorBuilder::FOperatorBuilder(const FOperatorBuilderSettings& InBuilderSettings)
@@ -43,6 +79,8 @@ namespace Metasound
 
 	TUniquePtr<IOperator> FOperatorBuilder::BuildGraphOperator(const FBuildGraphOperatorParams& InParams, FBuildResults& OutResults) const
 	{
+		METASOUND_TRACE_CPUPROFILER_EVENT_SCOPE(Metasound::FOperatorBuilder::BuildGraphOperator);
+
 		using namespace OperatorBuilderPrivate;
 
 		FBuildStatus BuildStatus = FBuildStatus::NoError;
@@ -82,7 +120,7 @@ namespace Metasound
 		}
 
 		// Create algo adapter view of graph to cache graph operations.
-		TPimplPtr<FDirectedGraphAlgoAdapter> AlgoAdapter = FDirectedGraphAlgo::CreateDirectedGraphAlgoAdapter(InParams.Graph);
+		TPimplPtr<FDirectedGraphAlgoAdapter> AlgoAdapter = DirectedGraphAlgo::CreateDirectedGraphAlgoAdapter(InParams.Graph);
 		
 		if (!AlgoAdapter.IsValid())
 		{
@@ -90,7 +128,7 @@ namespace Metasound
 			return TUniquePtr<IOperator>(nullptr);
 		}
 
-		FBuildContext BuildContext(InParams.Graph, *AlgoAdapter, InParams.OperatorSettings, InParams.Environment, BuilderSettings, OutResults);
+		OperatorBuilder::FBuildContext BuildContext(InParams.Graph, *AlgoAdapter, InParams.OperatorSettings, InParams.Environment, BuilderSettings, OutResults);
 
 		TArray<const INode*> SortedNodes;
 
@@ -98,6 +136,7 @@ namespace Metasound
 		BuildStatus |= DepthFirstTopologicalSort(BuildContext, SortedNodes);
 
 		// TODO: Add FindReachableNodesFromVariables in Prune.
+		// TODO: will need to prune edges as well.
 		// Otherwise, subgraphs incorrectly get pruned.
 		// BuildStatus |= PruneNodges(BuildContext, SortedNodes);
 
@@ -107,6 +146,17 @@ namespace Metasound
 			return TUniquePtr<IOperator>(nullptr);
 		}
 
+		InitializeOperatorInfo(InParams.Graph, SortedNodes, *BuildContext.GraphOperatorData);
+
+		// Assign external inputs to various vertex interfaces.
+		BuildStatus |= FOperatorBuilder::GatherExternalInputDataReferences(BuildContext, InParams.InputData);
+
+		// Check build status in case build routine should be exited early.
+		if (BuildStatus > GetMaxErrorLevel())
+		{
+			return TUniquePtr<IOperator>(nullptr);
+		}
+
 		// Create node operators from factories.
 		BuildStatus |= CreateOperators(BuildContext, SortedNodes, InParams.InputData);
 
@@ -115,15 +165,21 @@ namespace Metasound
 			return TUniquePtr<IOperator>(nullptr);
 		}
 
+		if (BuilderSettings.bPopulateInternalDataReferences)
+		{
+			GatherInternalGraphDataReferences(BuildContext, SortedNodes, BuildContext.Results.InternalDataReferences);
+		}
+
 		// Create graph operator from collection of node operators.
 		return CreateGraphOperator(BuildContext);
 	}
 
-	FOperatorBuilder::FBuildStatus FOperatorBuilder::DepthFirstTopologicalSort(FBuildContext& InOutContext, TArray<const INode*>& OutNodes) const
+	FOperatorBuilder::FBuildStatus FOperatorBuilder::DepthFirstTopologicalSort(OperatorBuilder::FBuildContext& InOutContext, TArray<const INode*>& OutNodes) const
 	{
+		METASOUND_TRACE_CPUPROFILER_EVENT_SCOPE(Metasound::FOperatorBuilder::DepthFirstTopologicalSort);
 		using namespace OperatorBuilderPrivate;
 
-		bool bSuccess = FDirectedGraphAlgo::DepthFirstTopologicalSort(InOutContext.AlgoAdapter, OutNodes);
+		bool bSuccess = DirectedGraphAlgo::DepthFirstTopologicalSort(InOutContext.AlgoAdapter, OutNodes);
 
 		if (!bSuccess)
 		{
@@ -136,11 +192,12 @@ namespace Metasound
 		return FBuildStatus::NoError;
 	}
 
-	FOperatorBuilder::FBuildStatus FOperatorBuilder::KahnsTopologicalSort(FBuildContext& InOutContext, TArray<const INode*>& OutNodes) const
+	FOperatorBuilder::FBuildStatus FOperatorBuilder::KahnsTopologicalSort(OperatorBuilder::FBuildContext& InOutContext, TArray<const INode*>& OutNodes) const
 	{
+		METASOUND_TRACE_CPUPROFILER_EVENT_SCOPE(Metasound::FOperatorBuilder::KahnsTopologicalSort);
 		using namespace OperatorBuilderPrivate;
 
-		bool bSuccess = FDirectedGraphAlgo::KahnTopologicalSort(InOutContext.AlgoAdapter, OutNodes);
+		bool bSuccess = DirectedGraphAlgo::KahnTopologicalSort(InOutContext.AlgoAdapter, OutNodes);
 
 		if (!bSuccess)
 		{
@@ -153,8 +210,9 @@ namespace Metasound
 		return FBuildStatus::NoError;
 	}
 
-	FOperatorBuilder::FBuildStatus FOperatorBuilder::PruneNodes(FBuildContext& InOutContext, TArray<const INode*>& InOutNodes) const
+	FOperatorBuilder::FBuildStatus FOperatorBuilder::PruneNodes(OperatorBuilder::FBuildContext& InOutContext, TArray<const INode*>& InOutNodes) const
 	{
+		METASOUND_TRACE_CPUPROFILER_EVENT_SCOPE(Metasound::FOperatorBuilder::PruneNodes);
 		using namespace OperatorBuilderPrivate;
 
 		FBuildStatus BuildStatus = FBuildStatus::NoError;
@@ -164,15 +222,15 @@ namespace Metasound
 		switch (BuilderSettings.PruningMode)
 		{
 			case EOperatorBuilderNodePruning::PruneNodesWithoutExternalDependency:
-				FDirectedGraphAlgo::FindReachableNodes(InOutContext.AlgoAdapter, ReachableNodes);
+				DirectedGraphAlgo::FindReachableNodes(InOutContext.AlgoAdapter, ReachableNodes);
 				break;
 
 			case EOperatorBuilderNodePruning::PruneNodesWithoutOutputDependency:
-				FDirectedGraphAlgo::FindReachableNodesFromOutput(InOutContext.AlgoAdapter, ReachableNodes);
+				DirectedGraphAlgo::FindReachableNodesFromOutput(InOutContext.AlgoAdapter, ReachableNodes);
 				break;
 
 			case EOperatorBuilderNodePruning::PruneNodesWithoutInputDependency:
-				FDirectedGraphAlgo::FindReachableNodesFromInput(InOutContext.AlgoAdapter, ReachableNodes);
+				DirectedGraphAlgo::FindReachableNodesFromInput(InOutContext.AlgoAdapter, ReachableNodes);
 				break;
 
 			case EOperatorBuilderNodePruning::None:
@@ -229,87 +287,69 @@ namespace Metasound
 		return BuildStatus;
 	}
 
-	FOperatorBuilder::FBuildStatus FOperatorBuilder::GatherInputDataReferences(FBuildContext& InOutContext, const INode* InNode, const FNodeEdgeMultiMap& InEdgeMap, FInputVertexInterfaceData& OutVertexData) const
+	void FOperatorBuilder::InitializeOperatorInfo(const IGraph& InGraph, TArray<const INode*>& InSortedNodes, DirectedGraphAlgo::FGraphOperatorData& InOutGraphOperatorData) const
 	{
-		using namespace OperatorBuilderPrivate;
+		METASOUND_TRACE_CPUPROFILER_EVENT_SCOPE(Metasound::FOperatorBuilder::InitializeOperatorInfo);
+		using namespace DirectedGraphAlgo;
 
-		FBuildStatus BuildStatus;
-
-		// Find all input edges associated with this node. 
-		TArray<const FDataEdge*> Edges;
-		InEdgeMap.MultiFind(InNode, Edges);
+		// Create FOperatorInfos from Nodes
+		TSortedMap<FOperatorID, FGraphOperatorData::FOperatorInfo>& OperatorMap = InOutGraphOperatorData.OperatorMap;
+		TArray<FOperatorID>& OperatorOrder = InOutGraphOperatorData.OperatorOrder;
 
-		// Add parameters to collection based on edge description
-		for (const FDataEdge* Edge : Edges)
+		for (const INode* Node : InSortedNodes)
 		{
-			if (!InOutContext.DataReferences.Contains(Edge->From.Node))
-			{
-				// This is likely due to a failed topological sort and is more of an internal error
-				// than a user error.
-				AddBuildError<FInternalError>(InOutContext.Results.Errors, TEXT(__FILE__), __LINE__);
-
-				return FBuildStatus::NonFatalError;
-			}
-
-			const FOutputVertexInterfaceData& FromData = InOutContext.DataReferences[Edge->From.Node].GetOutputs();
-
-			if (const FAnyDataReference* DataReference = FromData.FindDataReference(Edge->From.Vertex.VertexName))
-			{
-				check(DataReference->GetDataTypeName() == Edge->From.Vertex.DataTypeName);
-				check(DataReference->GetDataTypeName() == Edge->To.Vertex.DataTypeName);
-				OutVertexData.SetVertex(Edge->To.Vertex.VertexName, FAnyDataReference{*DataReference});
-			}
-			else 
-			{
-				// Does not contain any reference
-				// This is likely a node programming error where the edges reported by the INode interface
-				// did not match the readable parameter refs created by the operators outputs. Or, the edge description is invalid.
+			FOperatorID OperatorID = GetOperatorID(Node);
+			OperatorOrder.Add(OperatorID);
+			OperatorMap.Add(OperatorID, FGraphOperatorData::FOperatorInfo{nullptr, Node->GetVertexInterface()});
+		}
 
-				AddBuildError<FMissingOutputDataReferenceError>(InOutContext.Results.Errors, Edge->From);
+		// Set the output destinations on operator infos
+		for (const FDataEdge& Edge : InGraph.GetDataEdges())
+		{
+			const FOperatorID FromOperatorID = GetOperatorID(Edge.From.Node);
+			FGraphOperatorData::FOperatorInfo& OperatorInfo = OperatorMap.FindChecked(FromOperatorID);
 
-				BuildStatus |= FBuildStatus::NonFatalError;
-				continue;
-			}
+			OperatorInfo.OutputConnections.FindOrAdd(Edge.From.Vertex.VertexName).Add(FGraphOperatorData::FVertexDestination{GetOperatorID(Edge.To.Node), Edge.To.Vertex.VertexName});
 		}
-
-		return BuildStatus;
 	}
 
-	FOperatorBuilder::FBuildStatus FOperatorBuilder::GatherExternalInputDataReferences(FBuildContext& InOutContext, const INode* InNode, const FNodeDestinationMap& InNodeDestinationMap, const FInputVertexInterfaceData& InExternalInputData, FInputVertexInterfaceData& OutVertexData) const
+	FOperatorBuilder::FBuildStatus FOperatorBuilder::GatherExternalInputDataReferences(OperatorBuilder::FBuildContext& InOutContext, const FInputVertexInterfaceData& InExternalInputData) const
 	{
-		using namespace OperatorBuilderPrivate;
+		METASOUND_TRACE_CPUPROFILER_EVENT_SCOPE(Metasound::FOperatorBuilder::GatherExternalInputDataReferences);
+		using namespace DirectedGraphAlgo;
 
 		FBuildStatus BuildStatus;
-
-		// Check if there are any graph inputs connected to the node's inputs.
-		if (const FInputDataDestination* const* DestinationPtr = InNodeDestinationMap.Find(InNode))
+		// Gather external input data to graph destinations
+		for (const TPair<FNodeDataVertexKey, FInputDataDestination>& InputDestinationKV : InOutContext.Graph.GetInputDataDestinations())
 		{
-			if (const FInputDataDestination* Destination = *DestinationPtr)
+			const FInputDataDestination& Destination = InputDestinationKV.Value;
+
+			FOperatorID OperatorID = GetOperatorID(Destination.Node);
+			FGraphOperatorData::FOperatorInfo& OperatorInfo = InOutContext.GraphOperatorData->OperatorMap.FindChecked(OperatorID);
+
+			if (const FAnyDataReference* DataReference = InExternalInputData.FindDataReference(Destination.Vertex.VertexName))
 			{
-				if (const FAnyDataReference* DataReference = InExternalInputData.FindDataReference(Destination->Vertex.VertexName))
+				if (DataReference->GetDataTypeName() == Destination.Vertex.DataTypeName)
 				{
-					if (DataReference->GetDataTypeName() == Destination->Vertex.DataTypeName)
-					{
-						OutVertexData.SetVertex(Destination->Vertex.VertexName, FAnyDataReference{*DataReference});
-					}
-					else
-					{
-						// Mismatch in datatypes. This likely corresponds to a corrupt 
-						// Metasound Graph. The graph's inputs should route directly
-						// to TInputNode<>s with matching data types. 
-						
-						// Create source for reporting connection since external inputs do not have nodes. 
-						const FInputDataVertex& GraphVertex = InExternalInputData.GetVertex(Destination->Vertex.VertexName);
-
-						FOutputDataSource Source;
-						Source.Vertex.VertexName = GraphVertex.VertexName; 
-						Source.Vertex.DataTypeName = GraphVertex.DataTypeName;
+					OperatorInfo.VertexData.GetInputs().SetVertex(Destination.Vertex.VertexName, *DataReference);
+				}
+				else
+				{
+					// Mismatch in datatypes. This likely corresponds to a corrupt 
+					// Metasound Graph. The graph's inputs should route directly
+					// to TInputNode<>s with matching data types. 
+					
+					// Create source for reporting connection since external inputs do not have nodes. 
+					const FInputDataVertex& GraphVertex = InExternalInputData.GetVertex(Destination.Vertex.VertexName);
+
+					FOutputDataSource Source;
+					Source.Vertex.VertexName = GraphVertex.VertexName; 
+					Source.Vertex.DataTypeName = GraphVertex.DataTypeName;
 #if WITH_EDITORONLY_DATA
-						Source.Vertex.Metadata = GraphVertex.Metadata;
+					Source.Vertex.Metadata = GraphVertex.Metadata;
 #endif
 
-						AddBuildError<FInvalidConnectionDataTypeError>(InOutContext.Results.Errors, FDataEdge{Source, *Destination});
-					}
+					AddBuildError<FInvalidConnectionDataTypeError>(InOutContext.Results.Errors, FDataEdge{Source, Destination});
 				}
 			}
 		}
@@ -317,18 +357,24 @@ namespace Metasound
 		return BuildStatus;
 	}
 
-	void FOperatorBuilder::GatherInternalGraphDataReferences(FBuildContext& InOutContext, TMap<FGuid, FDataReferenceCollection>& OutNodeVertexData) const
+	void FOperatorBuilder::GatherInternalGraphDataReferences(OperatorBuilder::FBuildContext& InOutContext, const TArray<const INode*>& InNodes, TMap<FGuid, FDataReferenceCollection>& OutNodeVertexData) const
 	{
-		for (TPair<const INode*, FVertexInterfaceData>& ReferencePair : InOutContext.DataReferences)
+		METASOUND_TRACE_CPUPROFILER_EVENT_SCOPE(Metasound::FOperatorBuilder::GatherInternalGraphDataReferences);
+		using namespace DirectedGraphAlgo;
+
+		for (const INode* NodePtr : InNodes)
 		{
-			const INode* Node = ReferencePair.Key;
-			check(Node);
-			OutNodeVertexData.Emplace(Node->GetInstanceID(), ReferencePair.Value.GetOutputs().ToDataReferenceCollection());
+			check(NodePtr);
+			if (const FGraphOperatorData::FOperatorInfo* OpInfo = InOutContext.GraphOperatorData->OperatorMap.Find(GetOperatorID(NodePtr)))
+			{
+				OutNodeVertexData.Emplace(NodePtr->GetInstanceID(), OpInfo->VertexData.GetOutputs().ToDataReferenceCollection());
+			}
 		}
 	}
 
 	FOperatorBuilder::FBuildStatus FOperatorBuilder::ValidateOperatorOutputsAreBound(const INode& InNode, const FOutputVertexInterfaceData& InVertexData) const
 	{
+		METASOUND_TRACE_CPUPROFILER_EVENT_SCOPE(Metasound::FOperatorBuilder::ValidateOperatorOutputsAreBound);
 		using FOutputVertexBinding = MetasoundVertexDataPrivate::FOutputBinding;
 
 		FBuildStatus BuildStatus = FBuildStatus::NoError;
@@ -351,84 +397,85 @@ namespace Metasound
 		return BuildStatus;
 	}
 
-	FOperatorBuilder::FBuildStatus FOperatorBuilder::CreateOperators(FBuildContext& InOutContext, const TArray<const INode*>& InSortedNodes, const FInputVertexInterfaceData& InExternalInputData) const
+	FOperatorBuilder::FBuildStatus FOperatorBuilder::CreateOperators(OperatorBuilder::FBuildContext& InOutContext, const TArray<const INode*>& InSortedNodes, const FInputVertexInterfaceData& InExternalInputData) const
 	{
-		FBuildStatus BuildStatus;
+		METASOUND_TRACE_CPUPROFILER_EVENT_SCOPE(Metasound::FOperatorBuilder::CreateOperators);
+		using namespace DirectedGraphAlgo;
 
-		FNodeEdgeMultiMap NodeInputEdges;
+		FBuildStatus BuildStatus;
 
-		// Gather input edges for each node.
-		for (const FDataEdge& Edge : InOutContext.Graph.GetDataEdges())
-		{
-			NodeInputEdges.Add(Edge.To.Node, &Edge);
-		}
-
-		// Map input nodes to graph destinations
-		FNodeDestinationMap NodeDestinations;
-		for (const auto& InputDestinationKV : InOutContext.Graph.GetInputDataDestinations())
-		{
-			const FInputDataDestination& Destination = InputDestinationKV.Value;
-			NodeDestinations.Add(Destination.Node, &Destination);
-		}
+		// Create FOperatorInfos from Nodes
+		TSortedMap<FOperatorID, FGraphOperatorData::FOperatorInfo>& OperatorMap = InOutContext.GraphOperatorData->OperatorMap;
 
 		// Call operator factory for each node.
 		for (const INode* Node : InSortedNodes)
 		{
-			ThreadLocalDebug::SetActiveNodeClass(Node->GetMetadata());
-
-			const FVertexInterface& NodeInterface = Node->GetVertexInterface();
-			FInputVertexInterfaceData InputData(NodeInterface.GetInputInterface());
+			//ThreadLocalDebug::SetActiveNodeClass(Node->GetMetadata());
 
-			// Gather the input parameters for this IOperator from the output parameters of already created IOperators. 
-			BuildStatus |= GatherInputDataReferences(InOutContext, Node, NodeInputEdges, InputData);
-			// Gather the input parameters for this IOperator from the graph inputs.
-			BuildStatus |= GatherExternalInputDataReferences(InOutContext, Node, NodeDestinations, InExternalInputData, InputData);
+			FOperatorID OperatorID = GetOperatorID(Node);
+			FGraphOperatorData::FOperatorInfo& OperatorInfo = OperatorMap.FindChecked(OperatorID);
 
-			if (BuildStatus >= FBuildStatus::FatalError)
 			{
-				return BuildStatus;
-			}
+				METASOUND_TRACE_CPUPROFILER_EVENT_SCOPE(Metasound::FOperatorBuilder::CreateOperators::CreateAndBind);
+				FBuildOperatorParams CreateParams{*Node, InOutContext.Settings, OperatorInfo.VertexData.GetInputs(), InOutContext.Environment, this};
+				FOperatorFactorySharedRef Factory = Node->GetDefaultOperatorFactory();
+				OperatorInfo.Operator = Factory->CreateOperator(CreateParams, InOutContext.Results);
 
-			FOperatorFactorySharedRef Factory = Node->GetDefaultOperatorFactory();
-
-			// This is here to handle the deprecated field FCreateOperator::InputDataReferences.
-			// It can be removed once the deprecated member is removed.
-			FDataReferenceCollection InputCollection = InputData.ToDataReferenceCollection();
+				if (!OperatorInfo.Operator.IsValid())
+				{
+					return FBuildStatus::FatalError;
+				}
 
-			FBuildOperatorParams CreateParams{*Node, InOutContext.Settings, InputData, InOutContext.Environment, this};
+				// Bind vertex to operator data
+				OperatorInfo.Operator->BindInputs(OperatorInfo.VertexData.GetInputs());
+				OperatorInfo.Operator->BindOutputs(OperatorInfo.VertexData.GetOutputs());
+				
+				// Check if outputs are bound correctly.
+				if (BuilderSettings.bValidateOperatorOutputsAreBound)
+				{
+					BuildStatus |= ValidateOperatorOutputsAreBound(*Node, OperatorInfo.VertexData.GetOutputs());
+				}
+			}
 
-			FOperatorPtr Operator = Factory->CreateOperator(CreateParams, InOutContext.Results);
 
-			if (!Operator.IsValid())
+			// Route outputs of operator to downstream operators' VertexData
 			{
-				return FBuildStatus::FatalError;
-			}
+				METASOUND_TRACE_CPUPROFILER_EVENT_SCOPE(Metasound::FOperatorBuilder::CreateOperators::RouteOutputs);
+				for (const TPair<FVertexName, TArray<FGraphOperatorData::FVertexDestination>>& OutputRouting : OperatorInfo.OutputConnections)
+				{
+					const FVertexName& VertexName = OutputRouting.Key;
 
-			// Bind vertex to operator data
-			FVertexInterfaceData BoundOperatorData(NodeInterface);
-			Operator->Bind(BoundOperatorData);
-			
-			// Check if outputs are bound correctly.
-			if (BuilderSettings.bValidateOperatorOutputsAreBound)
-			{
-				BuildStatus |= ValidateOperatorOutputsAreBound(*Node, BoundOperatorData.GetOutputs());
-			}
+					if (const FAnyDataReference* Ref = OperatorInfo.VertexData.GetOutputs().FindDataReference(VertexName))
+					{
+						for (const FGraphOperatorData::FVertexDestination& Destination : OutputRouting.Value)
+						{
+							OperatorMap.FindChecked(Destination.OperatorID).VertexData.GetInputs().SetVertex(Destination.VertexName, *Ref);
+						}
+					}
+					else
+					{
+						// Does not contain any reference
+						// This is likely a node programming error where the edges reported by the INode interface
+						// did not match the readable parameter refs created by the operators outputs. Or, the edge description is invalid.
 
-			// Store bound interface for later use.
-			InOutContext.DataReferences.Emplace(Node, MoveTemp(BoundOperatorData));
+						AddBuildError<FMissingOutputDataReferenceError>(InOutContext.Results.Errors, FOutputDataSource{*Node, Node->GetVertexInterface().GetOutputVertex(VertexName)});
 
-			// Add operator to operator array
-			InOutContext.Operators.Add(MoveTemp(Operator));
+						BuildStatus |= FBuildStatus::NonFatalError;
+					}
+				}
+			}
 
-			ThreadLocalDebug::ResetActiveNodeClass();
+			//ThreadLocalDebug::ResetActiveNodeClass();
 		}
 
 		return BuildStatus;
 	}
 
-	FOperatorBuilder::FBuildStatus FOperatorBuilder::GatherGraphDataReferences(FBuildContext& InOutContext, FVertexInterfaceData& OutVertexData) const
+	FOperatorBuilder::FBuildStatus FOperatorBuilder::GatherGraphDataReferences(OperatorBuilder::FBuildContext& InOutContext, FVertexInterfaceData& OutVertexData) const
 	{
+		METASOUND_TRACE_CPUPROFILER_EVENT_SCOPE(Metasound::FOperatorBuilder::GatherGraphDataReferences);
 		using namespace OperatorBuilderPrivate;
+		using namespace DirectedGraphAlgo;
 		using FDestinationElement = FInputDataDestinationCollection::ElementType;
 		using FSourceElement = FOutputDataSourceCollection::ElementType;
 
@@ -440,9 +487,10 @@ namespace Metasound
 			bool bFoundDataReference = false;
 			const FInputDataDestination& InputDestination = Element.Value;
 
-			if (const FVertexInterfaceData* VertexData = InOutContext.DataReferences.Find(InputDestination.Node))
+			const FOperatorID OperatorID = GetOperatorID(InputDestination.Node);
+			if (const FGraphOperatorData::FOperatorInfo* OperatorInfo = InOutContext.GraphOperatorData->OperatorMap.Find(OperatorID))
 			{
-				const FInputVertexInterfaceData& NodeInputData = VertexData->GetInputs();
+				const FInputVertexInterfaceData& NodeInputData = OperatorInfo->VertexData.GetInputs();
 				if (const FAnyDataReference* DataReference = NodeInputData.FindDataReference(InputDestination.Vertex.VertexName))
 				{
 					if (DataReference->GetDataTypeName() == InputDestination.Vertex.DataTypeName)
@@ -451,6 +499,7 @@ namespace Metasound
 						OutVertexData.GetInputs().SetVertex(InputDestination.Vertex.VertexName, *DataReference);
 					}
 				}
+				InOutContext.GraphOperatorData->InputVertexMap.Emplace(InputDestination.Vertex.VertexName, OperatorID);
 			}
 
 			if (!bFoundDataReference)
@@ -466,9 +515,10 @@ namespace Metasound
 			bool bFoundDataReference = false;
 			const FOutputDataSource& OutputSource = Element.Value;
 
-			if (const FVertexInterfaceData* VertexData = InOutContext.DataReferences.Find(OutputSource.Node))
+			const FOperatorID OperatorID = GetOperatorID(OutputSource.Node);
+			if (const FGraphOperatorData::FOperatorInfo* OperatorInfo = InOutContext.GraphOperatorData->OperatorMap.Find(OperatorID))
 			{
-				const FOutputVertexInterfaceData& NodeOutputData = VertexData->GetOutputs();
+				const FOutputVertexInterfaceData& NodeOutputData = OperatorInfo->VertexData.GetOutputs();
 				if (const FAnyDataReference* DataReference = NodeOutputData.FindDataReference(OutputSource.Vertex.VertexName))
 				{
 					if (DataReference->GetDataTypeName() == OutputSource.Vertex.DataTypeName)
@@ -477,6 +527,7 @@ namespace Metasound
 						OutVertexData.GetOutputs().SetVertex(OutputSource.Vertex.VertexName, FAnyDataReference{*DataReference});
 					}
 				}
+				InOutContext.GraphOperatorData->OutputVertexMap.Emplace(OutputSource.Vertex.VertexName, OperatorID);
 			}
 
 			if (!bFoundDataReference)
@@ -489,32 +540,28 @@ namespace Metasound
 		return BuildStatus;
 	}
 
-	TUniquePtr<IOperator> FOperatorBuilder::CreateGraphOperator(FBuildContext& InOutContext) const
+	TUniquePtr<IOperator> FOperatorBuilder::CreateGraphOperator(OperatorBuilder::FBuildContext& InOutContext) const
 	{
-		FVertexInterfaceData BoundGraphData(InOutContext.Graph.GetVertexInterface());
+		METASOUND_TRACE_CPUPROFILER_EVENT_SCOPE(Metasound::FOperatorBuilder::CreateGraphOperator);
+		FBuildStatus BuildStatus = GatherGraphDataReferences(InOutContext, InOutContext.GraphOperatorData->VertexData);
 
-		FBuildStatus BuildStatus = GatherGraphDataReferences(InOutContext, BoundGraphData);
+		TUniquePtr<IOperator> GraphOperator;
 
 		if (BuildStatus > GetMaxErrorLevel())
 		{
-			return TUniquePtr<IOperator>(nullptr);
+			return GraphOperator;
 		}
 
-		TUniquePtr<FGraphOperator> GraphOperator = MakeUnique<FGraphOperator>();
-
-		GraphOperator->SetVertexInterfaceData(MoveTemp(BoundGraphData));
-		
-		if (BuilderSettings.bPopulateInternalDataReferences)
+		if (InOutContext.BuilderSettings.bEnableOperatorRebind)
 		{
-			GatherInternalGraphDataReferences(InOutContext, InOutContext.Results.InternalDataReferences);
+			 GraphOperator = MakeUnique<FRebindableGraphOperator>(MoveTemp(InOutContext.GraphOperatorData));
 		}
-
-		for (FOperatorPtr& Ptr : InOutContext.Operators)
+		else
 		{
-			GraphOperator->AppendOperator(MoveTemp(Ptr));
+			 GraphOperator = MakeUnique<FGraphOperator>(MoveTemp(InOutContext.GraphOperatorData));
 		}
 
-		return MoveTemp(GraphOperator);
+		return GraphOperator;
 	}
 
 	FOperatorBuilder::FBuildStatus::EStatus FOperatorBuilder::GetMaxErrorLevel() const
diff --git a/Engine/Plugins/Runtime/Metasound/Source/MetasoundGraphCore/Private/MetasoundRebindableGraphOperator.cpp b/Engine/Plugins/Runtime/Metasound/Source/MetasoundGraphCore/Private/MetasoundRebindableGraphOperator.cpp
new file mode 100644
--- /dev/null
+++ b/Engine/Plugins/Runtime/Metasound/Source/MetasoundGraphCore/Private/MetasoundRebindableGraphOperator.cpp
@@ -0,0 +1,223 @@
+// Copyright Epic Games, Inc. All Rights Reserved.
+
+#include "MetasoundRebindableGraphOperator.h"
+
+#include "Containers/Array.h"
+#include "Containers/SortedMap.h"
+#include "MetasoundExecutableOperator.h"
+#include "MetasoundGraphAlgoPrivate.h"
+#include "MetasoundOperatorInterface.h"
+#include "MetasoundVertex.h"
+#include "MetasoundVertexData.h"
+#include "Misc/Guid.h"
+#include "Templates/UniquePtr.h"
+
+namespace Metasound
+{
+	FRebindableGraphOperator::FRebindableGraphOperator(TUniquePtr<DirectedGraphAlgo::FGraphOperatorData> InOperatorData)
+	: GraphOperatorData(MoveTemp(InOperatorData))
+	{
+		using namespace DirectedGraphAlgo;
+
+		if (GraphOperatorData.IsValid())
+		{
+			for (DirectedGraphAlgo::FOperatorID OperatorID : GraphOperatorData->OperatorOrder)
+			{
+				if (FGraphOperatorData::FOperatorInfo* OperatorInfo = GraphOperatorData->OperatorMap.Find(OperatorID))
+				{
+					check(OperatorInfo->Operator.IsValid());
+					IOperator& Operator = *(OperatorInfo->Operator);
+
+					if (IOperator::FExecuteFunction ExecuteFunc = Operator.GetExecuteFunction())
+					{
+						ExecuteStack.Emplace(FExecuteEntry{OperatorID, Operator, ExecuteFunc});
+					}
+
+					if (IOperator::FPostExecuteFunction PostExecuteFunc = Operator.GetPostExecuteFunction())
+					{
+						PostExecuteStack.Emplace(FPostExecuteEntry{OperatorID, Operator, PostExecuteFunc});
+					}
+
+					if (IOperator::FResetFunction ResetFunc = Operator.GetResetFunction())
+					{
+						ResetStack.Emplace(FResetEntry{OperatorID, Operator, ResetFunc});
+					}
+				}
+			}
+		}
+	}
+
+	FDataReferenceCollection FRebindableGraphOperator::GetInputs() const
+	{
+		checkNoEntry();
+		return FDataReferenceCollection();
+	}
+
+	FDataReferenceCollection FRebindableGraphOperator::GetOutputs() const
+	{
+		checkNoEntry();
+		return FDataReferenceCollection();
+	}
+
+	// Bind the graph's interface data references to FVertexInterfaceData.
+	void FRebindableGraphOperator::BindInputs(FInputVertexInterfaceData& InOutVertexData)
+	{
+		using namespace DirectedGraphAlgo;
+
+		if (GraphOperatorData.IsValid())
+		{
+			auto UpdateStackOnRebind = [&](FOperatorID InOperatorID, IOperator& InOperator) 
+			{ 
+				UpdateStacks(InOperatorID, InOperator);
+			};
+
+			FRebindGraphDataParams Params
+			{
+				GraphOperatorData->VertexData,
+				GraphOperatorData->OperatorMap,
+				GraphOperatorData->InputVertexMap,
+				GraphOperatorData->OutputVertexMap
+			};
+
+			RebindGraphInputs(InOutVertexData, Params, UpdateStackOnRebind);
+		}
+	}
+
+	void FRebindableGraphOperator::BindOutputs(FOutputVertexInterfaceData& InOutVertexData)
+	{
+		using namespace DirectedGraphAlgo;
+		if (GraphOperatorData.IsValid())
+		{
+			auto UpdateStackOnRebind = [&](FOperatorID InOperatorID, IOperator& InOperator) 
+			{ 
+				UpdateStacks(InOperatorID, InOperator);
+			};
+
+			FRebindGraphDataParams Params
+			{
+				GraphOperatorData->VertexData,
+				GraphOperatorData->OperatorMap,
+				GraphOperatorData->InputVertexMap,
+				GraphOperatorData->OutputVertexMap
+			};
+
+			RebindGraphOutputs(InOutVertexData, Params, UpdateStackOnRebind);
+		}
+	}
+
+
+	void FRebindableGraphOperator::Execute()
+	{
+		for (FExecuteEntry& Entry : ExecuteStack)
+		{
+			Entry.Execute();
+		}
+	}
+
+	void FRebindableGraphOperator::PostExecute()
+	{
+		for (FPostExecuteEntry& Entry : PostExecuteStack)
+		{
+			Entry.PostExecute();
+		}
+	}
+
+	void FRebindableGraphOperator::Reset(const IOperator::FResetParams& InParams)
+	{
+		for (FResetEntry& Entry : ResetStack)
+		{
+			Entry.Reset(InParams);
+		}
+	}
+
+	IOperator::FPostExecuteFunction FRebindableGraphOperator::GetPostExecuteFunction() 
+	{
+		return &StaticPostExecute;
+	}
+
+	void FRebindableGraphOperator::StaticPostExecute(IOperator* InOperator)
+	{
+		check(InOperator);
+		static_cast<FRebindableGraphOperator*>(InOperator)->PostExecute();
+	}
+
+	template<typename EntryType, typename FunctionType>
+	void UpdateOperatorAndStack(const TArray<DirectedGraphAlgo::FOperatorID>& InOperatorOrder, TArray<EntryType>& InStack, DirectedGraphAlgo::FOperatorID InOperatorID, IOperator& InOperator, FunctionType InFunction)
+	{
+		using namespace DirectedGraphAlgo;
+
+		int32 EntryIndex = InStack.IndexOfByPredicate([&](const EntryType& InEntry) { return InEntry.OperatorID == InOperatorID; });
+
+		bool bEntryShouldExist = (nullptr != InFunction);
+		bool bEntryCurrentlyExists = (EntryIndex != INDEX_NONE);
+
+		if (bEntryCurrentlyExists && bEntryShouldExist)
+		{
+			// Update the existing entry
+			InStack[EntryIndex].Function = InFunction;
+		}
+		else if (bEntryCurrentlyExists)
+		{
+			// Remove the existing entry
+			InStack.RemoveAt(EntryIndex);
+		}
+		else if (bEntryShouldExist)
+		{
+			// Add new entry
+			int32 OperatorIndex = InOperatorOrder.Find(InOperatorID);
+			if (INDEX_NONE == OperatorIndex)
+			{
+				UE_LOG(LogMetaSound, Error, TEXT("Failed to update operator stack. OperatorID not found in operator order array."));
+				return;
+			}
+
+			int32 PreceedingOperatorIndex = OperatorIndex - 1;
+			int32 InsertLocation = 0; // Default to inserting at the beginning
+			while (PreceedingOperatorIndex >= 0)
+			{
+				FOperatorID PreceedingOperatorID = InOperatorOrder[PreceedingOperatorIndex];
+				PreceedingOperatorIndex--;
+
+				int32 PreceedingEntryIndex = InStack.IndexOfByPredicate([&](const EntryType& InEntry) { return InEntry.OperatorID == PreceedingOperatorID; });
+				if (PreceedingEntryIndex != INDEX_NONE)
+				{
+					InsertLocation = PreceedingEntryIndex + 1;
+					break;
+				}
+			}
+
+			InStack.Insert(EntryType{InOperatorID, InOperator, InFunction}, InsertLocation);
+		}
+	}
+
+	void FRebindableGraphOperator::UpdateStacks(DirectedGraphAlgo::FOperatorID InOperatorID, IOperator& InOperator)
+	{
+		check(GraphOperatorData);
+		UpdateOperatorAndStack(GraphOperatorData->OperatorOrder, ExecuteStack, InOperatorID, InOperator, InOperator.GetExecuteFunction());
+		UpdateOperatorAndStack(GraphOperatorData->OperatorOrder, PostExecuteStack, InOperatorID, InOperator, InOperator.GetPostExecuteFunction());
+		UpdateOperatorAndStack(GraphOperatorData->OperatorOrder, ResetStack, InOperatorID, InOperator, InOperator.GetResetFunction());
+	}
+
+	FRebindableGraphOperator::FExecuteEntry::FExecuteEntry(DirectedGraphAlgo::FOperatorID InOperatorID, IOperator& InOperator, IOperator::FExecuteFunction InFunc)
+	: OperatorID(InOperatorID)
+	, Operator(&InOperator)
+	, Function(InFunc)
+	{
+	}
+
+	FRebindableGraphOperator::FPostExecuteEntry::FPostExecuteEntry(DirectedGraphAlgo::FOperatorID InOperatorID, IOperator& InOperator, IOperator::FPostExecuteFunction InFunc)
+	: OperatorID(InOperatorID)
+	, Operator(&InOperator)
+	, Function(InFunc)
+	{
+	}
+
+	FRebindableGraphOperator::FResetEntry::FResetEntry(DirectedGraphAlgo::FOperatorID InOperatorID, IOperator& InOperator, IOperator::FResetFunction InFunc)
+	: OperatorID(InOperatorID)
+	, Operator(&InOperator)
+	, Function(InFunc)
+	{
+	}
+}
+
+
diff --git a/Engine/Plugins/Runtime/Metasound/Source/MetasoundGraphCore/Private/MetasoundRebindableGraphOperator.h b/Engine/Plugins/Runtime/Metasound/Source/MetasoundGraphCore/Private/MetasoundRebindableGraphOperator.h
new file mode 100644
--- /dev/null
+++ b/Engine/Plugins/Runtime/Metasound/Source/MetasoundGraphCore/Private/MetasoundRebindableGraphOperator.h
@@ -0,0 +1,97 @@
+// Copyright Epic Games, Inc. All Rights Reserved.
+
+#pragma once
+
+#include "Containers/Array.h"
+#include "Containers/SortedMap.h"
+#include "MetasoundExecutableOperator.h"
+#include "MetasoundGraphAlgoPrivate.h"
+#include "MetasoundOperatorInterface.h"
+#include "MetasoundVertex.h"
+#include "MetasoundVertexData.h"
+#include "Misc/Guid.h"
+#include "Templates/UniquePtr.h"
+
+namespace Metasound
+{
+	class FRebindableGraphOperator : public TExecutableOperator<FRebindableGraphOperator>
+	{
+	public:
+
+		FRebindableGraphOperator(TUniquePtr<DirectedGraphAlgo::FGraphOperatorData> InOperatorState);
+
+		virtual FDataReferenceCollection GetInputs() const override;
+
+		virtual FDataReferenceCollection GetOutputs() const override;
+
+		// Bind the graph's interface data references to FVertexInterfaceData.
+		virtual void BindInputs(FInputVertexInterfaceData& InOutVertexData) override;
+		virtual void BindOutputs(FOutputVertexInterfaceData& InOutVertexData) override;
+
+		void Execute();
+		void PostExecute();
+		void Reset(const IOperator::FResetParams& InParams);
+
+		virtual IOperator::FPostExecuteFunction GetPostExecuteFunction() override;
+
+	private:
+		static void StaticPostExecute(IOperator* InOperator);
+
+		void UpdateStacks(DirectedGraphAlgo::FOperatorID InOperatorID, IOperator& InOperator);
+
+		struct FExecuteEntry
+		{
+			FExecuteEntry(DirectedGraphAlgo::FOperatorID InOperatorID, IOperator& InOperator, IOperator::FExecuteFunction InFunc);
+
+			void Execute()
+			{
+				check(Operator);
+				check(Function);
+				Function(Operator);
+			}
+
+			DirectedGraphAlgo::FOperatorID OperatorID;
+			IOperator* Operator;
+			IOperator::FExecuteFunction Function;	
+		};
+
+		struct FPostExecuteEntry
+		{
+			FPostExecuteEntry(DirectedGraphAlgo::FOperatorID InOperatorID, IOperator& InOperator, IOperator::FPostExecuteFunction InFunc);
+
+			void PostExecute()
+			{
+				check(Operator);
+				check(Function);
+				Function(Operator);
+			}
+
+			DirectedGraphAlgo::FOperatorID OperatorID;
+			IOperator* Operator;
+			IOperator::FPostExecuteFunction Function;	
+		};
+
+		struct FResetEntry
+		{
+			FResetEntry(DirectedGraphAlgo::FOperatorID InOperatorID, IOperator& InOperator, IOperator::FResetFunction InFunc);
+
+			void Reset(const IOperator::FResetParams& InParams)
+			{
+				check(Operator);
+				check(Function);
+				Function(Operator, InParams);
+			}
+
+			DirectedGraphAlgo::FOperatorID OperatorID;
+			IOperator* Operator;
+			IOperator::FResetFunction Function;	
+		};
+
+		TArray<FExecuteEntry> ExecuteStack;
+		TArray<FPostExecuteEntry> PostExecuteStack;
+		TArray<FResetEntry> ResetStack;
+
+		TUniquePtr<DirectedGraphAlgo::FGraphOperatorData> GraphOperatorData;
+	};
+}
+
diff --git a/Engine/Plugins/Runtime/Metasound/Source/MetasoundGraphCore/Private/MetasoundVertexData.cpp b/Engine/Plugins/Runtime/Metasound/Source/MetasoundGraphCore/Private/MetasoundVertexData.cpp
--- a/Engine/Plugins/Runtime/Metasound/Source/MetasoundGraphCore/Private/MetasoundVertexData.cpp
+++ b/Engine/Plugins/Runtime/Metasound/Source/MetasoundGraphCore/Private/MetasoundVertexData.cpp
@@ -391,7 +391,7 @@ namespace Metasound
 		}
 
 		template<typename BindingType>
-		void CompareVertexInterfaceDataToPriorState(const TArray<BindingType>& InBindings, const TArray<FVertexDataState>& InPriorState, TSortedMap<FDataReferenceID, FAnyDataReference>& OutUpdates)
+		void CompareVertexInterfaceDataToPriorState(const TArray<BindingType>& InBindings, const TArray<FVertexDataState>& InPriorState, TSortedVertexNameMap<FAnyDataReference>& OutUpdates)
 		{
 			for (const BindingType& Binding : InBindings)
 			{
@@ -413,7 +413,7 @@ namespace Metasound
 					{
 						if (OtherState->ID != Binding.GetDataReferenceID())
 						{
-							OutUpdates.Add(OtherState->ID, *CurrentReference);
+							OutUpdates.Add(OtherState->VertexName, *CurrentReference);
 						}
 					}
 				}
@@ -745,12 +745,12 @@ namespace Metasound
 		MetasoundVertexDataPrivate::GetVertexInterfaceDataState(InVertexInterface.Bindings, OutState);
 	}
 
-	void CompareVertexInterfaceDataToPriorState(const FInputVertexInterfaceData& InCurrentInterface, const TArray<FVertexDataState>& InPriorState, TSortedMap<FDataReferenceID, FAnyDataReference>& OutUpdates)
+	void CompareVertexInterfaceDataToPriorState(const FInputVertexInterfaceData& InCurrentInterface, const TArray<FVertexDataState>& InPriorState, TSortedVertexNameMap<FAnyDataReference>& OutUpdates)
 	{
 		MetasoundVertexDataPrivate::CompareVertexInterfaceDataToPriorState(InCurrentInterface.Bindings, InPriorState, OutUpdates);
 	}
 
-	void CompareVertexInterfaceDataToPriorState(const FOutputVertexInterfaceData& InCurrentInterface, const TArray<FVertexDataState>& InPriorState, TSortedMap<FDataReferenceID, FAnyDataReference>& OutUpdates)
+	void CompareVertexInterfaceDataToPriorState(const FOutputVertexInterfaceData& InCurrentInterface, const TArray<FVertexDataState>& InPriorState, TSortedVertexNameMap<FAnyDataReference>& OutUpdates)
 	{
 		MetasoundVertexDataPrivate::CompareVertexInterfaceDataToPriorState(InCurrentInterface.Bindings, InPriorState, OutUpdates);
 	}
diff --git a/Engine/Plugins/Runtime/Metasound/Source/MetasoundGraphCore/Public/MetasoundExecutableOperator.h b/Engine/Plugins/Runtime/Metasound/Source/MetasoundGraphCore/Public/MetasoundExecutableOperator.h
--- a/Engine/Plugins/Runtime/Metasound/Source/MetasoundGraphCore/Public/MetasoundExecutableOperator.h
+++ b/Engine/Plugins/Runtime/Metasound/Source/MetasoundGraphCore/Public/MetasoundExecutableOperator.h
@@ -292,6 +292,22 @@ namespace Metasound
 				return EmptyCollection;
 			}
 
+			virtual void BindInputs(FInputVertexInterfaceData& InVertexData) override
+			{
+				if (Operator.IsValid())
+				{
+					Operator->BindInputs(InVertexData);
+				}
+			}
+
+			virtual void BindOutputs(FOutputVertexInterfaceData& InVertexData) override
+			{
+				if (Operator.IsValid())
+				{
+					Operator->BindOutputs(InVertexData);
+				}
+			}
+
 			virtual IOperator::FExecuteFunction GetExecuteFunction() override
 			{
 				return ExecuteFunction;
diff --git a/Engine/Plugins/Runtime/Metasound/Source/MetasoundGraphCore/Public/MetasoundGraphAlgo.h b/Engine/Plugins/Runtime/Metasound/Source/MetasoundGraphCore/Public/MetasoundGraphAlgo.h
--- a/Engine/Plugins/Runtime/Metasound/Source/MetasoundGraphCore/Public/MetasoundGraphAlgo.h
+++ b/Engine/Plugins/Runtime/Metasound/Source/MetasoundGraphCore/Public/MetasoundGraphAlgo.h
@@ -2,27 +2,37 @@
 
 #pragma once
 
-#include "CoreMinimal.h"
+#include "Containers/Array.h"
+#include "Containers/Set.h"
 #include "MetasoundNodeInterface.h"
+#include "Misc/CoreMiscDefines.h"
 #include "Templates/PimplPtr.h"
 
 namespace Metasound
 {
-	/** A strongly connected component containing Metasound INodes and 
-	 * FDataEdges.
-	 */
-	struct FStronglyConnectedComponent
+	// Forward declare 
+	class FDirectedGraphAlgoAdapter;
+	class INode;
+	class IGraph;
+
+	struct UE_DEPRECATED(5.3, "Use the equivalent methods declared in `namespace DirectedGrpahAlgo {...}`") FDirectedGraphAlgo;
+
+	namespace DirectedGraphAlgo
 	{
-		/** Nodes in the strongly connected component. */
-		TArray<const INode*> Nodes;
+		/** A strongly connected component containing Metasound INodes and 
+		 * FDataEdges.
+		 */
+		struct FStronglyConnectedComponent
+		{
+			/** Nodes in the strongly connected component. */
+			TArray<const INode*> Nodes;
 
-		/** FDataEdges in the strongly connected component. */
-		TArray<FDataEdge> Edges;
-	};
+			/** FDataEdges in the strongly connected component. */
+			TArray<FDataEdge> Edges;
+		};
+	}
 
-	// Forward declare adapter which caches transformations on an IGraph and 
-	// for graph algorithm optimizations.
-	class FDirectedGraphAlgoAdapter;
+	using FStronglyConnectedComponent UE_DEPRECATED(5.3, "Use the equivalent type declared in `namespace DirectedGrpahAlgo {...}`") = DirectedGraphAlgo::FStronglyConnectedComponent;
 
 	struct METASOUNDGRAPHCORE_API FDirectedGraphAlgo
 	{
@@ -91,7 +101,7 @@ namespace Metasound
 		 * @return True if one or more strongly connected components are added
 		 *         to OutComponents. False otherwise. 
 		 */
-		static bool TarjanStronglyConnectedComponents(const IGraph& InGraph, TArray<FStronglyConnectedComponent>& OutComponents, bool bExcludeSingleVertex = true);
+		static bool TarjanStronglyConnectedComponents(const IGraph& InGraph, TArray<DirectedGraphAlgo::FStronglyConnectedComponent>& OutComponents, bool bExcludeSingleVertex = true);
 
 		/** Find strongly connected components given a Metasound IGraph.
 		 *
@@ -106,7 +116,7 @@ namespace Metasound
 		 * @return True if one or more strongly connected components are added
 		 *         to OutComponents. False otherwise. 
 		 */
-		static bool TarjanStronglyConnectedComponents(const FDirectedGraphAlgoAdapter& InAdapter, TArray<FStronglyConnectedComponent>& OutComponents, bool bExcludeSingleVertex = true);
+		static bool TarjanStronglyConnectedComponents(const FDirectedGraphAlgoAdapter& InAdapter, TArray<DirectedGraphAlgo::FStronglyConnectedComponent>& OutComponents, bool bExcludeSingleVertex = true);
 
 		/** Finds all nodes which can be reached by traversing the graph beginning
 		 * from FInputDataDestination nodes. 
@@ -157,6 +167,140 @@ namespace Metasound
 		 * @param OutNodes  - The nodes which were reached during traversal.
 		 */
 		static void FindReachableNodes(const FDirectedGraphAlgoAdapter& InAdapter, TSet<const INode*>& OutNodes);
+	};
+
+	namespace DirectedGraphAlgo
+	{
+		/** The FDirectedGraphAlgoAdapter caches internal representations of the
+		 * passed in IGraph. Using this can reduce redundant calcuations when 
+		 * multiple algorithms are run on the same IGraph (without altering the IGraph).
+		 *
+		 * This function creates the adapter. 
+		 */
+		METASOUNDGRAPHCORE_API TPimplPtr<FDirectedGraphAlgoAdapter> CreateDirectedGraphAlgoAdapter(const IGraph& InGraph);
+
+		/** Sort the nodes of a directed acyclic graph using a depth-first 
+		 * topological sorting algorithm.
+		 *
+		 * @param InGraph      - The graph containing edges and nodes to sort. 
+		 * @param OutNodeOrder - An array of nodes from InGraph put in a valid 
+		 *                       topological sorted order.
+		 *
+		 * @return True if sorting was successful, false otherwise. 
+		 */
+		METASOUNDGRAPHCORE_API bool DepthFirstTopologicalSort(const IGraph& InGraph, TArray<const INode*>& OutNodeOrder);
+
+		/** Sort the nodes of a directed acyclic graph using a depth-first 
+		 * topological sorting algorithm.
+		 *
+		 * @param InAdapter    - The adapter containing edges and nodes to sort. 
+		 * @param OutNodeOrder - An array of nodes from InGraph put in a valid 
+		 *                       topological sorted order.
+		 *
+		 * @return True if sorting was successful, false otherwise. 
+		 */
+		METASOUNDGRAPHCORE_API bool DepthFirstTopologicalSort(const FDirectedGraphAlgoAdapter& InAdapter, TArray<const INode*>& OutNodeOrder);
+
+		/** Sort the nodes of a directed acyclic graph using Kahn's topological
+		 * sorting algorithm.
+		 *
+		 * @param InGraph      - The graph containing edges and nodes to sort. 
+		 * @param OutNodeOrder - An array of nodes from InGraph put in a valid 
+		 *                       topological sorted order.
+		 *
+		 * @return True if sorting was successful, false otherwise. 
+		 */
+		METASOUNDGRAPHCORE_API bool KahnTopologicalSort(const IGraph& InGraph, TArray<const INode*>& OutNodeOrder);
+
+		/** Sort the nodes of a directed acyclic graph using Kahn's topological
+		 * sorting algorithm.
+		 *
+		 * @param InAdapter    - The adapter containing edges and nodes to sort. 
+		 * @param OutNodeOrder - An array of nodes from InGraph put in a valid 
+		 *                       topological sorted order.
+		 *
+		 * @return True if sorting was successful, false otherwise. 
+		 */
+		METASOUNDGRAPHCORE_API bool KahnTopologicalSort(const FDirectedGraphAlgoAdapter& InAdapter, TArray<const INode*>& OutNodeOrder);
+
+		/** Find strongly connected components given a Metasound IGraph.
+		 *
+		 * @param InGraph              - Graph to be analyzed. 
+		 * @param OutComponents        - Strongly connected components found in
+		 * 								 the graph are be added to this array.
+		 * @param bExcludeSingleVertex - If true, single vertices are not be 
+		 *                               returned as strongly connected components. 
+		 *                               If false, single vertices may be returned
+		 *                               as strongly connected components. 
+		 *
+		 * @return True if one or more strongly connected components are added
+		 *         to OutComponents. False otherwise. 
+		 */
+		METASOUNDGRAPHCORE_API bool TarjanStronglyConnectedComponents(const IGraph& InGraph, TArray<FStronglyConnectedComponent>& OutComponents, bool bExcludeSingleVertex = true);
+
+		/** Find strongly connected components given a Metasound IGraph.
+		 *
+		 * @param InAdapter            - The adapter containing edges and nodes to analyze.
+		 * @param OutComponents        - Strongly connected components found in
+		 * 								 the graph are be added to this array.
+		 * @param bExcludeSingleVertex - If true, single vertices are not be 
+		 *                               returned as strongly connected components. 
+		 *                               If false, single vertices may be returned
+		 *                               as strongly connected components. 
+		 *
+		 * @return True if one or more strongly connected components are added
+		 *         to OutComponents. False otherwise. 
+		 */
+		METASOUNDGRAPHCORE_API bool TarjanStronglyConnectedComponents(const FDirectedGraphAlgoAdapter& InAdapter, TArray<FStronglyConnectedComponent>& OutComponents, bool bExcludeSingleVertex = true);
+
+		/** Finds all nodes which can be reached by traversing the graph beginning
+		 * from FInputDataDestination nodes. 
+		 *
+		 * @param InGraph  - The graph to be traversed.
+		 * @param OutNodes - The nodes which were reached during traversal.
+		 */
+		METASOUNDGRAPHCORE_API void FindReachableNodesFromInput(const IGraph& InGraph, TSet<const INode*>& OutNodes);
 
+		/** Finds all nodes which can be reached by traversing the graph beginning
+		 * from FInputDataDestination nodes. 
+		 *
+		 * @param InAdapter - The adapter to be traversed.
+		 * @param OutNodes  - The nodes which were reached during traversal.
+		 */
+		METASOUNDGRAPHCORE_API void FindReachableNodesFromInput(const FDirectedGraphAlgoAdapter& InAdapter, TSet<const INode*>& OutNodes);
+
+		/** Finds all nodes which can be reached by traversing the transpose graph 
+		 * beginning from FOutputDataSource nodes. 
+		 *
+		 * @param InGraph  - The graph to be traversed.
+		 * @param OutNodes - The nodes which were reached during traversal.
+		 */
+		METASOUNDGRAPHCORE_API void FindReachableNodesFromOutput(const IGraph& InGraph, TSet<const INode*>& OutNodes);
+
+		/** Finds all nodes which can be reached by traversing the transpose graph 
+		 * beginning from FOutputDataSource nodes. 
+		 *
+		 * @param InAdapter - The adapter to be traversed.
+		 * @param OutNodes  - The nodes which were reached during traversal.
+		 */
+		METASOUNDGRAPHCORE_API void FindReachableNodesFromOutput(const FDirectedGraphAlgoAdapter& InAdapter, TSet<const INode*>& OutNodes);
+
+		/** Finds all nodes which can be reached by traversing the graph 
+		 * beginning from FInputDataDestinations or by traversing the transpose
+		 * graph beginning from FOutputDataSource nodes. 
+		 *
+		 * @param InGraph  - The graph to be traversed.
+		 * @param OutNodes - The nodes which were reached during traversal.
+		 */
+		METASOUNDGRAPHCORE_API void FindReachableNodes(const IGraph& InGraph, TSet<const INode*>& OutNodes);
+
+		/** Finds all nodes which can be reached by traversing the graph 
+		 * beginning from FInputDataDestinations or by traversing the transpose
+		 * graph beginning from FOutputDataSource nodes. 
+		 *
+		 * @param InAdapter - The adapter to be traversed.
+		 * @param OutNodes  - The nodes which were reached during traversal.
+		 */
+		METASOUNDGRAPHCORE_API void FindReachableNodes(const FDirectedGraphAlgoAdapter& InAdapter, TSet<const INode*>& OutNodes);
 	};
 }
diff --git a/Engine/Plugins/Runtime/Metasound/Source/MetasoundGraphCore/Public/MetasoundGraphOperator.h b/Engine/Plugins/Runtime/Metasound/Source/MetasoundGraphCore/Public/MetasoundGraphOperator.h
--- a/Engine/Plugins/Runtime/Metasound/Source/MetasoundGraphCore/Public/MetasoundGraphOperator.h
+++ b/Engine/Plugins/Runtime/Metasound/Source/MetasoundGraphCore/Public/MetasoundGraphOperator.h
@@ -2,14 +2,22 @@
 
 #pragma once
 
-#include "CoreMinimal.h"
+#include "Containers/Array.h"
 #include "MetasoundDataReference.h"
 #include "MetasoundDataReferenceCollection.h"
-#include "MetasoundOperatorInterface.h"
 #include "MetasoundExecutableOperator.h"
+#include "MetasoundOperatorInterface.h"
+#include "MetasoundVertexData.h"
+#include "Templates/UniquePtr.h"
 
 namespace Metasound
 {
+	// Forward declare
+	namespace DirectedGraphAlgo
+	{
+		struct FGraphOperatorData;
+	}
+
 	class METASOUNDGRAPHCORE_API FGraphOperator : public TExecutableOperator<FGraphOperator>
 	{
 		public:
@@ -19,6 +27,7 @@ namespace Metasound
 			using FResetParams = IOperator::FResetParams;
 
 			FGraphOperator() = default;
+			FGraphOperator(TUniquePtr<DirectedGraphAlgo::FGraphOperatorData>&& InOperatorData);
 
 			virtual ~FGraphOperator() = default;
 
diff --git a/Engine/Plugins/Runtime/Metasound/Source/MetasoundGraphCore/Public/MetasoundOperatorBuilder.h b/Engine/Plugins/Runtime/Metasound/Source/MetasoundGraphCore/Public/MetasoundOperatorBuilder.h
--- a/Engine/Plugins/Runtime/Metasound/Source/MetasoundGraphCore/Public/MetasoundOperatorBuilder.h
+++ b/Engine/Plugins/Runtime/Metasound/Source/MetasoundGraphCore/Public/MetasoundOperatorBuilder.h
@@ -20,6 +20,12 @@ namespace Metasound
 	class FOperatorSettings;
 	class FMetasoundEnvironment;
 
+	namespace OperatorBuilder
+	{
+		// Forward Declare
+		struct FBuildContext;
+	}
+
 	/** FOperatorBuilder builds an IOperator from an IGraph. */
 	class METASOUNDGRAPHCORE_API FOperatorBuilder : public IOperatorBuilder
 	{
@@ -45,11 +51,6 @@ namespace Metasound
 
 		private:
 
-			using FNodeEdgeMultiMap = TMultiMap<const INode*, const FDataEdge*>;
-			using FNodeDestinationMap = TMap<const INode*, const FInputDataDestination*>;
-			using FNodeVertexInterfaceDataMap = TMap<const INode*, FVertexInterfaceData>;
-			using FOperatorPtr = TUniquePtr<IOperator>;
-
 			// Handles build status of current build operation.
 			struct FBuildStatus
 			{
@@ -96,65 +97,36 @@ namespace Metasound
 				EStatus Value = NoError;
 			};
 
-			struct FBuildContext
-			{
-				const IGraph& Graph;
-				const FDirectedGraphAlgoAdapter& AlgoAdapter;
-				const FOperatorSettings& Settings;
-				const FMetasoundEnvironment& Environment;
-				const FOperatorBuilderSettings& BuilderSettings;
-
-				FBuildResults& Results;
-
-				TArray<FOperatorPtr> Operators;
-				FNodeVertexInterfaceDataMap DataReferences;
-
-				FBuildContext(
-					const IGraph& InGraph,
-					const FDirectedGraphAlgoAdapter& InAlgoAdapter,
-					const FOperatorSettings& InSettings,
-					const FMetasoundEnvironment& InEnvironment,
-					const FOperatorBuilderSettings& bInBuilderSettings,
-					FBuildResults& OutResults)
-				:	Graph(InGraph)
-				,	AlgoAdapter(InAlgoAdapter)
-				,	Settings(InSettings)
-				,	Environment(InEnvironment)
-				,	BuilderSettings(bInBuilderSettings)
-				,	Results(OutResults)
-				{
-				}
-			};
 
 			// Perform topological sort using depth first algorithm.
-			FBuildStatus DepthFirstTopologicalSort(FBuildContext& InOutContext, TArray<const INode*>& OutNodes) const;
+			FBuildStatus DepthFirstTopologicalSort(OperatorBuilder::FBuildContext& InOutContext, TArray<const INode*>& OutNodes) const;
 
 			// Perform topological sort using kahns algorithm.
-			FBuildStatus KahnsTopologicalSort(FBuildContext& InOutContext, TArray<const INode*>& OutNodes) const;
+			FBuildStatus KahnsTopologicalSort(OperatorBuilder::FBuildContext& InOutContext, TArray<const INode*>& OutNodes) const;
 
 			// Prune unreachable nodes from InOutNodes
-			FBuildStatus PruneNodes(FBuildContext& InOutContext, TArray<const INode*>& InOutNodes) const;
+			FBuildStatus PruneNodes(OperatorBuilder::FBuildContext& InOutContext, TArray<const INode*>& InOutNodes) const;
 
-			// Get all input data references for a given node for inputs provided internally to the graph.
-			FBuildStatus GatherInputDataReferences(FBuildContext& InOutContext, const INode* InNode, const FNodeEdgeMultiMap& InEdgeMap, FInputVertexInterfaceData& OutVertexData) const;
+			// Initialize Operator Info
+			void InitializeOperatorInfo(const IGraph& InGraph ,TArray<const INode*>& InSortedNodes, DirectedGraphAlgo::FGraphOperatorData& InOutGraphOperatorData) const;
 
-			// Get all input data references for a given node for inputs provided externally to the graph.
-			FBuildStatus GatherExternalInputDataReferences(FBuildContext& InContext, const INode* InNode, const FNodeDestinationMap& InNodeDestinationMap, const FInputVertexInterfaceData& InExternalCollection, FInputVertexInterfaceData& OutVertexData) const;
+			// Get and route input data references for inputs provided externally to the graph.
+			FBuildStatus GatherExternalInputDataReferences(OperatorBuilder::FBuildContext& InOutContext, const FInputVertexInterfaceData& InExternalInputData) const;
 
 			// Graphs all internal graph references and places them in output map.
-			void GatherInternalGraphDataReferences(FBuildContext& InOutContext, TMap<FGuid, FDataReferenceCollection>& OutNodeVertexData) const;
+			void GatherInternalGraphDataReferences(OperatorBuilder::FBuildContext& InOutContext, const TArray<const INode*>& InNodes, TMap<FGuid, FDataReferenceCollection>& OutNodeVertexData) const;
 
 			// Validates whether all operator outputs are bound to data references. 
 			FBuildStatus ValidateOperatorOutputsAreBound(const INode& InNode, const FOutputVertexInterfaceData& InVertexData) const;
 
 			// Get all input/output data references for a given graph.
-			FBuildStatus GatherGraphDataReferences(FBuildContext& InOutContext, FVertexInterfaceData& OutVertexData) const;
+			FBuildStatus GatherGraphDataReferences(OperatorBuilder::FBuildContext& InOutContext, FVertexInterfaceData& OutVertexData) const;
 
 			// Call the operator factories for the nodes
-			FBuildStatus CreateOperators(FBuildContext& InOutContext, const TArray<const INode*>& InSortedNodes, const FInputVertexInterfaceData& InExternalInputData) const;
+			FBuildStatus CreateOperators(OperatorBuilder::FBuildContext& InOutContext, const TArray<const INode*>& InSortedNodes, const FInputVertexInterfaceData& InExternalInputData) const;
 
 			// Create the final graph operator from the provided build context.
-			TUniquePtr<IOperator> CreateGraphOperator(FBuildContext& InOutContext) const;
+			TUniquePtr<IOperator> CreateGraphOperator(OperatorBuilder::FBuildContext& InOutContext) const;
 
 			FBuildStatus::EStatus GetMaxErrorLevel() const;
 
diff --git a/Engine/Plugins/Runtime/Metasound/Source/MetasoundGraphCore/Public/MetasoundOperatorBuilderSettings.h b/Engine/Plugins/Runtime/Metasound/Source/MetasoundGraphCore/Public/MetasoundOperatorBuilderSettings.h
--- a/Engine/Plugins/Runtime/Metasound/Source/MetasoundGraphCore/Public/MetasoundOperatorBuilderSettings.h
+++ b/Engine/Plugins/Runtime/Metasound/Source/MetasoundGraphCore/Public/MetasoundOperatorBuilderSettings.h
@@ -77,6 +77,14 @@ namespace Metasound
 		  * for analyzing internal graph state by managing sound generator). */
 		bool bPopulateInternalDataReferences = false;
 
+		/** If true, the operator created by the builder will be able to update
+		 * the underlying data of it's input and output data references through 
+		 * the IOperator::BindInputs(...) and IOperator BindOutputs(...) methods.
+		 * This requires that the created operator maintain additional data and
+		 * results in a modest increase in memory usage.
+		 */
+		bool bEnableOperatorRebind = false;
+
 		/** Return the default settings for the current build environment. */
 		static const FOperatorBuilderSettings& GetDefaultSettings();
 
diff --git a/Engine/Plugins/Runtime/Metasound/Source/MetasoundGraphCore/Public/MetasoundOperatorInterface.h b/Engine/Plugins/Runtime/Metasound/Source/MetasoundGraphCore/Public/MetasoundOperatorInterface.h
--- a/Engine/Plugins/Runtime/Metasound/Source/MetasoundGraphCore/Public/MetasoundOperatorInterface.h
+++ b/Engine/Plugins/Runtime/Metasound/Source/MetasoundGraphCore/Public/MetasoundOperatorInterface.h
@@ -64,6 +64,7 @@ namespace Metasound
 			FVertexInterfaceData Data;
 			Data.GetInputs() = InVertexData;
 			Bind(Data);
+			InVertexData = Data.GetInputs();
 		}
 
 		virtual void BindOutputs(FOutputVertexInterfaceData& InVertexData) 
@@ -71,6 +72,7 @@ namespace Metasound
 			FVertexInterfaceData Data;
 			Data.GetOutputs() = InVertexData;
 			Bind(Data);
+			InVertexData = Data.GetOutputs();
 		}
 
 		/** Pointer to initialize function for an operator.
diff --git a/Engine/Plugins/Runtime/Metasound/Source/MetasoundGraphCore/Public/MetasoundVertexData.h b/Engine/Plugins/Runtime/Metasound/Source/MetasoundGraphCore/Public/MetasoundVertexData.h
--- a/Engine/Plugins/Runtime/Metasound/Source/MetasoundGraphCore/Public/MetasoundVertexData.h
+++ b/Engine/Plugins/Runtime/Metasound/Source/MetasoundGraphCore/Public/MetasoundVertexData.h
@@ -404,6 +404,22 @@ namespace Metasound
 		};
 	}
 
+	/** Convenience for using a TSortedMap with FVertexName Key type.
+	 *
+	 * This template makes it convenient to create a TSortedMap with an FVertexName 
+	 * while also avoiding compilation errors incurred from using the FName default
+	 * "less than" operator in the TSortedMap implementation. 
+	 *
+	 * - FVertexName is an alias to FName. 
+	 * - TSortedMap<FName, ValueType> fails to compile since the "less than" operator
+	 *   specific implementation needs to be chosen (FastLess vs LexicalLess)
+	 * - Due to the template argument order of TSortedMap this also forces you to
+	 *   choose the allocator. 
+	 * - This is all a bit of an annoyance to do every time we use a TSortedMap 
+	 *   with an FVertexName as the key.
+	 */
+	template<typename ValueType>
+	using TSortedVertexNameMap = TSortedMap<FVertexName, ValueType, FDefaultAllocator, FNameFastLess>;
 
 	/** MetaSound Binding 
 	 *
@@ -730,7 +746,7 @@ namespace Metasound
 	
 	private:
 		friend METASOUNDGRAPHCORE_API void GetVertexInterfaceDataState(const FInputVertexInterfaceData& InVertexInterface, TArray<FVertexDataState>& OutState);
-		friend METASOUNDGRAPHCORE_API void CompareVertexInterfaceDataToPriorState(const FInputVertexInterfaceData& InCurrentInterface, const TArray<FVertexDataState>& InPriorState, TSortedMap<FDataReferenceID, FAnyDataReference>& OutUpdates);
+		friend METASOUNDGRAPHCORE_API void CompareVertexInterfaceDataToPriorState(const FInputVertexInterfaceData& InCurrentInterface, const TArray<FVertexDataState>& InPriorState, TSortedVertexNameMap<FAnyDataReference>& OutUpdates);
 
 		void Apply(const FVertexName& InVertexName, TFunctionRef<FInputBinding ()> InCreateFunc, TFunctionRef<void (FInputBinding&)> InApplyFunc);
 
@@ -990,7 +1006,7 @@ namespace Metasound
 	
 	private:
 		friend METASOUNDGRAPHCORE_API void GetVertexInterfaceDataState(const FOutputVertexInterfaceData& InVertexInterface, TArray<FVertexDataState>& OutState);
-		friend METASOUNDGRAPHCORE_API void CompareVertexInterfaceDataToPriorState(const FOutputVertexInterfaceData& InCurrentInterface, const TArray<FVertexDataState>& InPriorState, TSortedMap<FDataReferenceID, FAnyDataReference>& OutUpdates);
+		friend METASOUNDGRAPHCORE_API void CompareVertexInterfaceDataToPriorState(const FOutputVertexInterfaceData& InCurrentInterface, const TArray<FVertexDataState>& InPriorState, TSortedVertexNameMap<FAnyDataReference>& OutUpdates);
 
 		void Apply(const FVertexName& InVertexName, TFunctionRef<FOutputBinding ()> InCreateFunc, TFunctionRef<void (FOutputBinding&)> InBindFunc);
 
@@ -1066,7 +1082,7 @@ namespace Metasound
 	METASOUNDGRAPHCORE_API void GetVertexInterfaceDataState(const FOutputVertexInterfaceData& InVertexInterface, TArray<FVertexDataState>& OutState);
 
 	/** Compares the current data bound to the vertex interface with a prior cached state. */
-	METASOUNDGRAPHCORE_API void CompareVertexInterfaceDataToPriorState(const FInputVertexInterfaceData& InCurrentInterface, const TArray<FVertexDataState>& InPriorState, TSortedMap<FDataReferenceID, FAnyDataReference>& OutUpdates);
+	METASOUNDGRAPHCORE_API void CompareVertexInterfaceDataToPriorState(const FInputVertexInterfaceData& InCurrentInterface, const TArray<FVertexDataState>& InPriorState, TSortedVertexNameMap<FAnyDataReference>& OutUpdates);
 	/** Compares the current data bound to the vertex interface with a prior cached state. */
-	METASOUNDGRAPHCORE_API void CompareVertexInterfaceDataToPriorState(const FOutputVertexInterfaceData& InCurrentInterface, const TArray<FVertexDataState>& InPriorState, TSortedMap<FDataReferenceID, FAnyDataReference>& OutUpdates);
+	METASOUNDGRAPHCORE_API void CompareVertexInterfaceDataToPriorState(const FOutputVertexInterfaceData& InCurrentInterface, const TArray<FVertexDataState>& InPriorState, TSortedVertexNameMap<FAnyDataReference>& OutUpdates);
 }
