Fix MetaSound Convert From Preset not marking graph as editable

#rnx

diff --git a/Engine/Plugins/Runtime/Metasound/Source/MetasoundFrontend/Private/MetasoundFrontendDocumentBuilder.cpp b/Engine/Plugins/Runtime/Metasound/Source/MetasoundFrontend/Private/MetasoundFrontendDocumentBuilder.cpp
--- a/Engine/Plugins/Runtime/Metasound/Source/MetasoundFrontend/Private/MetasoundFrontendDocumentBuilder.cpp
+++ b/Engine/Plugins/Runtime/Metasound/Source/MetasoundFrontend/Private/MetasoundFrontendDocumentBuilder.cpp
@@ -1067,7 +1067,7 @@ bool FMetaSoundFrontendDocumentBuilder::ConvertFromPreset()
 	PresetOptions.bIsPreset = false;
 
 #if WITH_EDITOR
-	FMetasoundFrontendGraphStyle Style = Document.RootGraph.Graph.Style;
+	FMetasoundFrontendGraphStyle& Style = Document.RootGraph.Graph.Style;
 	Style.bIsGraphEditable = true;
 #endif // WITH_EDITOR
 
