Fix MetaSound wave proxy reader using wrong duration for streaming sounds using Ogg Vorbis

#rnx

#localization none
#tests local, preflight 646675958377b53e1d0c5e15

diff --git a/Engine/Plugins/Runtime/Metasound/Source/MetasoundEngine/Private/MetasoundWaveProxyReader.cpp b/Engine/Plugins/Runtime/Metasound/Source/MetasoundEngine/Private/MetasoundWaveProxyReader.cpp
--- a/Engine/Plugins/Runtime/Metasound/Source/MetasoundEngine/Private/MetasoundWaveProxyReader.cpp
+++ b/Engine/Plugins/Runtime/Metasound/Source/MetasoundEngine/Private/MetasoundWaveProxyReader.cpp
@@ -447,7 +447,10 @@ namespace Metasound
 		Audio::FFormatDescriptorSection FormatDesc;
 		DecoderInput->FindSection(FormatDesc);
 		SampleRate = FormatDesc.NumFramesPerSec;
-		NumFramesInWave = FormatDesc.NumFrames;
+		if (FormatDesc.NumFrames > 0)
+		{
+			NumFramesInWave = FormatDesc.NumFrames;
+		}
 		// end hack
 
 		CurrentFrameIndex = FMath::Clamp(static_cast<int32>(InStartTimeInSeconds * GetSampleRate()), 0, GetNumFramesInWave());
