[Metasounds] [DynamicOperator] Implement rebinding functionality in built-in metasound nodes

diff --git a/Engine/Plugins/Runtime/Metasound/Source/MetasoundStandardNodes/Private/MetasoundADEnvelopeNode.cpp b/Engine/Plugins/Runtime/Metasound/Source/MetasoundStandardNodes/Private/MetasoundADEnvelopeNode.cpp
--- a/Engine/Plugins/Runtime/Metasound/Source/MetasoundStandardNodes/Private/MetasoundADEnvelopeNode.cpp
+++ b/Engine/Plugins/Runtime/Metasound/Source/MetasoundStandardNodes/Private/MetasoundADEnvelopeNode.cpp
@@ -326,32 +326,43 @@ namespace Metasound
 
 		virtual ~TADEnvelopeNodeOperator() = default;
 
-		virtual FDataReferenceCollection GetInputs() const override
+
+		virtual void BindInputs(FInputVertexInterfaceData& InOutVertexData) override
 		{
 			using namespace ADEnvelopeVertexNames;
 
-			FDataReferenceCollection Inputs;
-			Inputs.AddDataReadReference(METASOUND_GET_PARAM_NAME(InputTrigger), TriggerAttackIn);
-			Inputs.AddDataReadReference(METASOUND_GET_PARAM_NAME(InputAttackTime), AttackTime);
-			Inputs.AddDataReadReference(METASOUND_GET_PARAM_NAME(InputDecayTime), DecayTime);
-			Inputs.AddDataReadReference(METASOUND_GET_PARAM_NAME(InputAttackCurve), AttackCurveFactor);
-			Inputs.AddDataReadReference(METASOUND_GET_PARAM_NAME(InputDecayCurve), DecayCurveFactor);
-			Inputs.AddDataReadReference(METASOUND_GET_PARAM_NAME(InputLooping), bLooping);
-			Inputs.AddDataReadReference(METASOUND_GET_PARAM_NAME(InputHardReset), bHardReset);
-
-			return Inputs;
+			InOutVertexData.BindReadVertex(METASOUND_GET_PARAM_NAME(InputTrigger), TriggerAttackIn);
+			InOutVertexData.BindReadVertex(METASOUND_GET_PARAM_NAME(InputAttackTime), AttackTime);
+			InOutVertexData.BindReadVertex(METASOUND_GET_PARAM_NAME(InputDecayTime), DecayTime);
+			InOutVertexData.BindReadVertex(METASOUND_GET_PARAM_NAME(InputAttackCurve), AttackCurveFactor);
+			InOutVertexData.BindReadVertex(METASOUND_GET_PARAM_NAME(InputDecayCurve), DecayCurveFactor);
+			InOutVertexData.BindReadVertex(METASOUND_GET_PARAM_NAME(InputLooping), bLooping);
+			InOutVertexData.BindReadVertex(METASOUND_GET_PARAM_NAME(InputHardReset), bHardReset);
 		}
 
-		virtual FDataReferenceCollection GetOutputs() const override
+		virtual void BindOutputs(FOutputVertexInterfaceData& InOutVertexData) override
 		{
 			using namespace ADEnvelopeVertexNames;
 
-			FDataReferenceCollection Outputs;
-			Outputs.AddDataReadReference(METASOUND_GET_PARAM_NAME(OutputOnTrigger), OnAttackTrigger);
-			Outputs.AddDataReadReference(METASOUND_GET_PARAM_NAME(OutputOnDone), OnDone);
-			Outputs.AddDataReadReference(METASOUND_GET_PARAM_NAME(OutputEnvelopeValue), OutputEnvelope);
+			InOutVertexData.BindReadVertex(METASOUND_GET_PARAM_NAME(OutputOnTrigger), OnAttackTrigger);
+			InOutVertexData.BindReadVertex(METASOUND_GET_PARAM_NAME(OutputOnDone), OnDone);
+			InOutVertexData.BindReadVertex(METASOUND_GET_PARAM_NAME(OutputEnvelopeValue), OutputEnvelope);
+		}
 
-			return Outputs;
+		virtual FDataReferenceCollection GetInputs() const override
+		{
+			// This should never be called. Bind(...) is called instead. This method
+			// exists as a stop-gap until the API can be deprecated and removed.
+			checkNoEntry();
+			return {};
+		}
+
+		virtual FDataReferenceCollection GetOutputs() const override
+		{
+			// This should never be called. Bind(...) is called instead. This method
+			// exists as a stop-gap until the API can be deprecated and removed.
+			checkNoEntry();
+			return {};
 		}
 
 		void UpdateParams()
diff --git a/Engine/Plugins/Runtime/Metasound/Source/MetasoundStandardNodes/Private/MetasoundADSREnvelopeNode.cpp b/Engine/Plugins/Runtime/Metasound/Source/MetasoundStandardNodes/Private/MetasoundADSREnvelopeNode.cpp
--- a/Engine/Plugins/Runtime/Metasound/Source/MetasoundStandardNodes/Private/MetasoundADSREnvelopeNode.cpp
+++ b/Engine/Plugins/Runtime/Metasound/Source/MetasoundStandardNodes/Private/MetasoundADSREnvelopeNode.cpp
@@ -605,38 +605,48 @@ namespace Metasound
 
 		virtual ~TADSREnvelopeNodeOperator() = default;
 
-		virtual FDataReferenceCollection GetInputs() const override
+		virtual void BindInputs(FInputVertexInterfaceData& InOutVertexData) override
 		{
 			using namespace ADSREnvelopeVertexNames;
 
-			FDataReferenceCollection Inputs;
-			Inputs.AddDataReadReference(METASOUND_GET_PARAM_NAME(InputAttackTrigger), TriggerAttackIn);
-			Inputs.AddDataReadReference(METASOUND_GET_PARAM_NAME(InputAttackTrigger), TriggerReleaseIn);
-			Inputs.AddDataReadReference(METASOUND_GET_PARAM_NAME(InputAttackTime), AttackTime);
-			Inputs.AddDataReadReference(METASOUND_GET_PARAM_NAME(InputDecayTime), DecayTime);
-			Inputs.AddDataReadReference(METASOUND_GET_PARAM_NAME(InputSustainLevel), SustainLevel);
-			Inputs.AddDataReadReference(METASOUND_GET_PARAM_NAME(InputReleaseTime), ReleaseTime);
-			Inputs.AddDataReadReference(METASOUND_GET_PARAM_NAME(InputAttackCurve), AttackCurveFactor);
-			Inputs.AddDataReadReference(METASOUND_GET_PARAM_NAME(InputDecayCurve), DecayCurveFactor);
-			Inputs.AddDataReadReference(METASOUND_GET_PARAM_NAME(InputReleaseCurve), ReleaseCurveFactor);
-			Inputs.AddDataReadReference(METASOUND_GET_PARAM_NAME(InputHardReset), bHardReset);
-
-			return Inputs;
+			InOutVertexData.BindReadVertex(METASOUND_GET_PARAM_NAME(InputAttackTrigger), TriggerAttackIn);
+			InOutVertexData.BindReadVertex(METASOUND_GET_PARAM_NAME(InputAttackTrigger), TriggerReleaseIn);
+			InOutVertexData.BindReadVertex(METASOUND_GET_PARAM_NAME(InputAttackTime), AttackTime);
+			InOutVertexData.BindReadVertex(METASOUND_GET_PARAM_NAME(InputDecayTime), DecayTime);
+			InOutVertexData.BindReadVertex(METASOUND_GET_PARAM_NAME(InputSustainLevel), SustainLevel);
+			InOutVertexData.BindReadVertex(METASOUND_GET_PARAM_NAME(InputReleaseTime), ReleaseTime);
+			InOutVertexData.BindReadVertex(METASOUND_GET_PARAM_NAME(InputAttackCurve), AttackCurveFactor);
+			InOutVertexData.BindReadVertex(METASOUND_GET_PARAM_NAME(InputDecayCurve), DecayCurveFactor);
+			InOutVertexData.BindReadVertex(METASOUND_GET_PARAM_NAME(InputReleaseCurve), ReleaseCurveFactor);
+			InOutVertexData.BindReadVertex(METASOUND_GET_PARAM_NAME(InputHardReset), bHardReset);
 		}
 
-		virtual FDataReferenceCollection GetOutputs() const override
+		virtual void BindOutputs(FOutputVertexInterfaceData& InOutVertexData) override
 		{
 			using namespace ADSREnvelopeVertexNames;
 
-			FDataReferenceCollection Outputs;
-			Outputs.AddDataReadReference(METASOUND_GET_PARAM_NAME(OutputOnAttackTrigger), TriggerAttackIn);
-			Outputs.AddDataReadReference(METASOUND_GET_PARAM_NAME(OutputOnDecayTrigger), OnDecayTrigger);
-			Outputs.AddDataReadReference(METASOUND_GET_PARAM_NAME(OutputOnSustainTrigger), OnSustainTrigger);
-			Outputs.AddDataReadReference(METASOUND_GET_PARAM_NAME(OutputOnReleaseTrigger), TriggerReleaseIn);
-			Outputs.AddDataReadReference(METASOUND_GET_PARAM_NAME(OutputOnDone), OnDone);
-			Outputs.AddDataReadReference(METASOUND_GET_PARAM_NAME(OutputEnvelopeValue), OutputEnvelope);
+			InOutVertexData.BindReadVertex(METASOUND_GET_PARAM_NAME(OutputOnAttackTrigger), TriggerAttackIn);
+			InOutVertexData.BindReadVertex(METASOUND_GET_PARAM_NAME(OutputOnDecayTrigger), OnDecayTrigger);
+			InOutVertexData.BindReadVertex(METASOUND_GET_PARAM_NAME(OutputOnSustainTrigger), OnSustainTrigger);
+			InOutVertexData.BindReadVertex(METASOUND_GET_PARAM_NAME(OutputOnReleaseTrigger), TriggerReleaseIn);
+			InOutVertexData.BindReadVertex(METASOUND_GET_PARAM_NAME(OutputOnDone), OnDone);
+			InOutVertexData.BindReadVertex(METASOUND_GET_PARAM_NAME(OutputEnvelopeValue), OutputEnvelope);
+		}
 
-			return Outputs;
+		virtual FDataReferenceCollection GetInputs() const override
+		{
+			// This should never be called. Bind(...) is called instead. This method
+			// exists as a stop-gap until the API can be deprecated and removed.
+			checkNoEntry();
+			return {};
+		}
+
+		virtual FDataReferenceCollection GetOutputs() const override
+		{
+			// This should never be called. Bind(...) is called instead. This method
+			// exists as a stop-gap until the API can be deprecated and removed.
+			checkNoEntry();
+			return {};
 		}
 
 		void UpdateParams()
diff --git a/Engine/Plugins/Runtime/Metasound/Source/MetasoundStandardNodes/Private/MetasoundAbsoluteValueNodes.cpp b/Engine/Plugins/Runtime/Metasound/Source/MetasoundStandardNodes/Private/MetasoundAbsoluteValueNodes.cpp
--- a/Engine/Plugins/Runtime/Metasound/Source/MetasoundStandardNodes/Private/MetasoundAbsoluteValueNodes.cpp
+++ b/Engine/Plugins/Runtime/Metasound/Source/MetasoundStandardNodes/Private/MetasoundAbsoluteValueNodes.cpp
@@ -180,24 +180,33 @@ namespace Metasound
 
 		virtual ~TAbsNodeOperator() = default;
 
-		virtual FDataReferenceCollection GetInputs() const override
+
+		virtual void BindInputs(FInputVertexInterfaceData& InOutVertexData) override
 		{
 			using namespace AbsVertexNames;
-
-			FDataReferenceCollection Inputs;
-			Inputs.AddDataReadReference(METASOUND_GET_PARAM_NAME(InputValue), Input);
-
-			return Inputs;
+			InOutVertexData.BindReadVertex(METASOUND_GET_PARAM_NAME(InputValue), Input);
 		}
 
-		virtual FDataReferenceCollection GetOutputs() const override
+		virtual void BindOutputs(FOutputVertexInterfaceData& InOutVertexData) override
 		{
 			using namespace AbsVertexNames;
+			InOutVertexData.BindReadVertex(METASOUND_GET_PARAM_NAME(OutputValue), OutputValue);
+		}
 
-			FDataReferenceCollection Outputs;
-			Outputs.AddDataReadReference(METASOUND_GET_PARAM_NAME(OutputValue), OutputValue);
+		virtual FDataReferenceCollection GetInputs() const override
+		{
+			// This should never be called. Bind(...) is called instead. This method
+			// exists as a stop-gap until the API can be deprecated and removed.
+			checkNoEntry();
+			return {};
+		}
 
-			return Outputs;
+		virtual FDataReferenceCollection GetOutputs() const override
+		{
+			// This should never be called. Bind(...) is called instead. This method
+			// exists as a stop-gap until the API can be deprecated and removed.
+			checkNoEntry();
+			return {};
 		}
 
 		void GetAbs()
diff --git a/Engine/Plugins/Runtime/Metasound/Source/MetasoundStandardNodes/Private/MetasoundAdditiveSynthNode.cpp b/Engine/Plugins/Runtime/Metasound/Source/MetasoundStandardNodes/Private/MetasoundAdditiveSynthNode.cpp
--- a/Engine/Plugins/Runtime/Metasound/Source/MetasoundStandardNodes/Private/MetasoundAdditiveSynthNode.cpp
+++ b/Engine/Plugins/Runtime/Metasound/Source/MetasoundStandardNodes/Private/MetasoundAdditiveSynthNode.cpp
@@ -116,27 +116,39 @@ namespace Metasound
 			CreateSinusoids();
 		}
 
-		FDataReferenceCollection GetInputs() const
+		virtual void BindInputs(FInputVertexInterfaceData& InOutVertexData) override
 		{
 			using namespace AdditiveSynthVertexNames;
 
-			FDataReferenceCollection InputDataReferences;
-			InputDataReferences.AddDataReadReference(METASOUND_GET_PARAM_NAME(BaseFrequency), BaseFrequency);
-			InputDataReferences.AddDataReadReference(METASOUND_GET_PARAM_NAME(HarmonicMultipliers), HarmonicMultipliers);
-			InputDataReferences.AddDataReadReference(METASOUND_GET_PARAM_NAME(Amplitudes), Amplitudes);
-			InputDataReferences.AddDataReadReference(METASOUND_GET_PARAM_NAME(Phases), Phases);
-			InputDataReferences.AddDataReadReference(METASOUND_GET_PARAM_NAME(PanAmounts), PanAmounts);
-			return InputDataReferences;
+			InOutVertexData.BindReadVertex(METASOUND_GET_PARAM_NAME(BaseFrequency), BaseFrequency);
+			InOutVertexData.BindReadVertex(METASOUND_GET_PARAM_NAME(HarmonicMultipliers), HarmonicMultipliers);
+			InOutVertexData.BindReadVertex(METASOUND_GET_PARAM_NAME(Amplitudes), Amplitudes);
+			InOutVertexData.BindReadVertex(METASOUND_GET_PARAM_NAME(Phases), Phases);
+			InOutVertexData.BindReadVertex(METASOUND_GET_PARAM_NAME(PanAmounts), PanAmounts);
 		}
 
-		virtual FDataReferenceCollection GetOutputs() const
+		virtual void BindOutputs(FOutputVertexInterfaceData& InOutVertexData) override
 		{
 			using namespace AdditiveSynthVertexNames;
 
-			FDataReferenceCollection OutputDataReferences;
-			OutputDataReferences.AddDataReadReference(METASOUND_GET_PARAM_NAME(LeftAudioOut), LeftAudioOut);
-			OutputDataReferences.AddDataReadReference(METASOUND_GET_PARAM_NAME(RightAudioOut), RightAudioOut);
-			return OutputDataReferences;
+			InOutVertexData.BindReadVertex(METASOUND_GET_PARAM_NAME(LeftAudioOut), LeftAudioOut);
+			InOutVertexData.BindReadVertex(METASOUND_GET_PARAM_NAME(RightAudioOut), RightAudioOut);
+		}
+
+		virtual FDataReferenceCollection GetInputs() const override
+		{
+			// This should never be called. Bind(...) is called instead. This method
+			// exists as a stop-gap until the API can be deprecated and removed.
+			checkNoEntry();
+			return {};
+		}
+
+		virtual FDataReferenceCollection GetOutputs() const override
+		{
+			// This should never be called. Bind(...) is called instead. This method
+			// exists as a stop-gap until the API can be deprecated and removed.
+			checkNoEntry();
+			return {};
 		}
 
 		void Reset(const IOperator::FResetParams& InParams)
diff --git a/Engine/Plugins/Runtime/Metasound/Source/MetasoundStandardNodes/Private/MetasoundAudioFormats.cpp b/Engine/Plugins/Runtime/Metasound/Source/MetasoundStandardNodes/Private/MetasoundAudioFormats.cpp
--- a/Engine/Plugins/Runtime/Metasound/Source/MetasoundStandardNodes/Private/MetasoundAudioFormats.cpp
+++ b/Engine/Plugins/Runtime/Metasound/Source/MetasoundStandardNodes/Private/MetasoundAudioFormats.cpp
@@ -111,25 +111,34 @@ namespace Metasound
 
 			virtual ~FOutputOperator() {}
 
-			virtual FDataReferenceCollection GetInputs() const override
+
+			virtual void BindInputs(FInputVertexInterfaceData& InOutVertexData) override
 			{
 				using namespace StereoAudioFormatVertexKeys;
 
-				FDataReferenceCollection Inputs;
+				InOutVertexData.BindReadVertex(METASOUND_GET_PARAM_NAME(LeftChannelVertex), Left);
+				InOutVertexData.BindReadVertex(METASOUND_GET_PARAM_NAME(RightChannelVertex), Right);
+			}
 
-				Inputs.AddDataReadReference(METASOUND_GET_PARAM_NAME(LeftChannelVertex), Left);
-				Inputs.AddDataReadReference(METASOUND_GET_PARAM_NAME(RightChannelVertex), Right);
+			virtual void BindOutputs(FOutputVertexInterfaceData& InOutVertexData) override
+			{
+				InOutVertexData.BindReadVertex(OutputName, Stereo);
+			}
 
-				return Inputs;
+			virtual FDataReferenceCollection GetInputs() const override
+			{
+				// This should never be called. Bind(...) is called instead. This method
+				// exists as a stop-gap until the API can be deprecated and removed.
+				checkNoEntry();
+				return {};
 			}
 
 			virtual FDataReferenceCollection GetOutputs() const override
 			{
-				FDataReferenceCollection Outputs;
-
-				Outputs.AddDataReadReference(OutputName, Stereo);
-
-				return Outputs;
+				// This should never be called. Bind(...) is called instead. This method
+				// exists as a stop-gap until the API can be deprecated and removed.
+				checkNoEntry();
+				return {};
 			}
 
 		private:
diff --git a/Engine/Plugins/Runtime/Metasound/Source/MetasoundStandardNodes/Private/MetasoundBPMToSecondsNode.cpp b/Engine/Plugins/Runtime/Metasound/Source/MetasoundStandardNodes/Private/MetasoundBPMToSecondsNode.cpp
--- a/Engine/Plugins/Runtime/Metasound/Source/MetasoundStandardNodes/Private/MetasoundBPMToSecondsNode.cpp
+++ b/Engine/Plugins/Runtime/Metasound/Source/MetasoundStandardNodes/Private/MetasoundBPMToSecondsNode.cpp
@@ -32,6 +32,9 @@ namespace Metasound
 
 		FBPMToSecondsOperator(const FCreateOperatorParams& InParams, const FFloatReadRef& InBPM, const FFloatReadRef& InBeatMultiplier, const FFloatReadRef& InDivOfWholeNote);
 
+
+		virtual void BindInputs(FInputVertexInterfaceData& InOutVertexData) override;
+		virtual void BindOutputs(FOutputVertexInterfaceData& InOutVertexData) override;
 		virtual FDataReferenceCollection GetInputs() const override;
 		virtual FDataReferenceCollection GetOutputs() const override;
 		void Reset(const IOperator::FResetParams& InParams);
@@ -69,25 +72,36 @@ namespace Metasound
 		Reset(InParams);
 	}
 
-	FDataReferenceCollection FBPMToSecondsOperator::GetInputs() const
+	void FBPMToSecondsOperator::BindInputs(FInputVertexInterfaceData& InOutVertexData)
 	{
 		using namespace BPMToSecondsVertexNames;
 
-		FDataReferenceCollection InputDataReferences;
-		InputDataReferences.AddDataReadReference(METASOUND_GET_PARAM_NAME(InputBPM), BPM);
-		InputDataReferences.AddDataReadReference(METASOUND_GET_PARAM_NAME(InputBeatMultiplier), BeatMultiplier);
-		InputDataReferences.AddDataReadReference(METASOUND_GET_PARAM_NAME(InputDivOfWholeNote), DivOfWholeNote);
-
-		return InputDataReferences;
+		InOutVertexData.BindReadVertex(METASOUND_GET_PARAM_NAME(InputBPM), BPM);
+		InOutVertexData.BindReadVertex(METASOUND_GET_PARAM_NAME(InputBeatMultiplier), BeatMultiplier);
+		InOutVertexData.BindReadVertex(METASOUND_GET_PARAM_NAME(InputDivOfWholeNote), DivOfWholeNote);
 	}
 
-	FDataReferenceCollection FBPMToSecondsOperator::GetOutputs() const
+	void FBPMToSecondsOperator::BindOutputs(FOutputVertexInterfaceData& InOutVertexData)
 	{
 		using namespace BPMToSecondsVertexNames;
 
-		FDataReferenceCollection OutputDataReferences;
-		OutputDataReferences.AddDataReadReference(METASOUND_GET_PARAM_NAME(OutputTimeSeconds), TimeSeconds);
-		return OutputDataReferences;
+		InOutVertexData.BindReadVertex(METASOUND_GET_PARAM_NAME(OutputTimeSeconds), TimeSeconds);
+	}
+
+	FDataReferenceCollection FBPMToSecondsOperator::GetInputs() const
+	{
+		// This should never be called. Bind(...) is called instead. This method
+		// exists as a stop-gap until the API can be deprecated and removed.
+		checkNoEntry();
+		return {};
+	}
+
+	FDataReferenceCollection FBPMToSecondsOperator::GetOutputs() const
+	{
+		// This should never be called. Bind(...) is called instead. This method
+		// exists as a stop-gap until the API can be deprecated and removed.
+		checkNoEntry();
+		return {};
 	}
 
 	void FBPMToSecondsOperator::UpdateTime()
diff --git a/Engine/Plugins/Runtime/Metasound/Source/MetasoundStandardNodes/Private/MetasoundBandsplitterNode.cpp b/Engine/Plugins/Runtime/Metasound/Source/MetasoundStandardNodes/Private/MetasoundBandsplitterNode.cpp
--- a/Engine/Plugins/Runtime/Metasound/Source/MetasoundStandardNodes/Private/MetasoundBandsplitterNode.cpp
+++ b/Engine/Plugins/Runtime/Metasound/Source/MetasoundStandardNodes/Private/MetasoundBandsplitterNode.cpp
@@ -121,41 +121,50 @@ namespace Metasound
 			return Metadata;
 		}
 
-		virtual FDataReferenceCollection GetInputs() const override
+
+		virtual void BindInputs(FInputVertexInterfaceData& InOutVertexData) override
 		{
 			using namespace BandSplitterNode;
 
-			FDataReferenceCollection InputPins;
-
 			for (uint32 Chan = 0; Chan < NumChannels; ++Chan)
 			{
-				InputPins.AddDataReadReference(GetAudioInputName(Chan), AudioInputs[Chan]);
+				InOutVertexData.BindReadVertex(GetAudioInputName(Chan), AudioInputs[Chan]);
 			}
 
-			InputPins.AddDataReadReference(METASOUND_GET_PARAM_NAME(InputFilterOrder), FilterOrder);
-			InputPins.AddDataReadReference(METASOUND_GET_PARAM_NAME(InputPhaseCompensate), bPhaseCompensate);
+			InOutVertexData.BindReadVertex(METASOUND_GET_PARAM_NAME(InputFilterOrder), FilterOrder);
+			InOutVertexData.BindReadVertex(METASOUND_GET_PARAM_NAME(InputPhaseCompensate), bPhaseCompensate);
 
 			for (uint32 BandIndex = 0; BandIndex < NumBands - 1; ++BandIndex)
 			{
-				InputPins.AddDataReadReference(GetCrossoverInputName(BandIndex), CrossoverFrequencies[BandIndex]);
+				InOutVertexData.BindReadVertex(GetCrossoverInputName(BandIndex), CrossoverFrequencies[BandIndex]);
 			}
-
-			return InputPins;
 		}
 
-		virtual FDataReferenceCollection GetOutputs() const override
+		virtual void BindOutputs(FOutputVertexInterfaceData& InOutVertexData) override
 		{
-			FDataReferenceCollection OutputPins;
-
 			for (uint32 BandIndex = 0; BandIndex < NumBands; ++BandIndex)
 			{
 				for (uint32 ChannelIndex = 0; ChannelIndex < NumChannels; ++ChannelIndex)
 				{
-					OutputPins.AddDataReadReference(GetAudioOutputName(BandIndex, ChannelIndex), AudioOutputs[BandIndex * NumChannels + ChannelIndex]);
+					InOutVertexData.BindReadVertex(GetAudioOutputName(BandIndex, ChannelIndex), AudioOutputs[BandIndex * NumChannels + ChannelIndex]);
 				}
 			}
+		}
 
-			return OutputPins;
+		virtual FDataReferenceCollection GetInputs() const override
+		{
+			// This should never be called. Bind(...) is called instead. This method
+			// exists as a stop-gap until the API can be deprecated and removed.
+			checkNoEntry();
+			return {};
+		}
+
+		virtual FDataReferenceCollection GetOutputs() const override
+		{
+			// This should never be called. Bind(...) is called instead. This method
+			// exists as a stop-gap until the API can be deprecated and removed.
+			checkNoEntry();
+			return {};
 		}
 
 		static const FVertexInterface& GetDefaultInterface()
diff --git a/Engine/Plugins/Runtime/Metasound/Source/MetasoundStandardNodes/Private/MetasoundBasicFilters.cpp b/Engine/Plugins/Runtime/Metasound/Source/MetasoundStandardNodes/Private/MetasoundBasicFilters.cpp
--- a/Engine/Plugins/Runtime/Metasound/Source/MetasoundStandardNodes/Private/MetasoundBasicFilters.cpp
+++ b/Engine/Plugins/Runtime/Metasound/Source/MetasoundStandardNodes/Private/MetasoundBasicFilters.cpp
@@ -97,23 +97,32 @@ namespace Metasound
 
 		static TUniquePtr<IOperator> CreateOperator(const FCreateOperatorParams& InParams, FBuildErrorArray& OutErrors);
 
-		virtual FDataReferenceCollection GetInputs() const override
+		virtual void BindInputs(FInputVertexInterfaceData& InOutVertexData) override
 		{
-			FDataReferenceCollection InputDataReferences;
-			InputDataReferences.AddDataReadReference(METASOUND_GET_PARAM_NAME(ParamAudioInput), FAudioBufferReadRef(AudioInput));
-			InputDataReferences.AddDataReadReference(METASOUND_GET_PARAM_NAME(ParamCutoffFrequency), FFloatReadRef(Frequency));
-			InputDataReferences.AddDataReadReference(METASOUND_GET_PARAM_NAME(ParamResonance), FFloatReadRef(Resonance));
+			InOutVertexData.BindReadVertex(METASOUND_GET_PARAM_NAME(ParamAudioInput), AudioInput);
+			InOutVertexData.BindReadVertex(METASOUND_GET_PARAM_NAME(ParamCutoffFrequency), Frequency);
+			InOutVertexData.BindReadVertex(METASOUND_GET_PARAM_NAME(ParamResonance), Resonance);
+		}
 
-			return InputDataReferences;
+		virtual void BindOutputs(FOutputVertexInterfaceData& InOutVertexData) override
+		{
+			InOutVertexData.BindReadVertex(METASOUND_GET_PARAM_NAME(ParamAudioOutput), AudioOutput);
 		}
 
-		virtual FDataReferenceCollection GetOutputs() const override
+		virtual FDataReferenceCollection GetInputs() const override
 		{
-			// expose read access to our output buffer for other processors in the graph
-			FDataReferenceCollection OutputDataReferences;
-			OutputDataReferences.AddDataReadReference(METASOUND_GET_PARAM_NAME(ParamAudioOutput), FAudioBufferReadRef(AudioOutput));
+			// This should never be called. Bind(...) is called instead. This method
+			// exists as a stop-gap until the API can be deprecated and removed.
+			checkNoEntry();
+			return {};
+		}
 
-			return OutputDataReferences;
+		virtual FDataReferenceCollection GetOutputs() const override
+		{
+			// This should never be called. Bind(...) is called instead. This method
+			// exists as a stop-gap until the API can be deprecated and removed.
+			checkNoEntry();
+			return {};
 		}
 
 		void Reset(const IOperator::FResetParams& InParams);
@@ -272,27 +281,36 @@ namespace Metasound
 
 		static TUniquePtr<IOperator> CreateOperator(const FCreateOperatorParams& InParams, FBuildErrorArray& OutErrors);
 
-		virtual FDataReferenceCollection GetInputs() const override
+		virtual void BindInputs(FInputVertexInterfaceData& InOutVertexData) override
 		{
-			FDataReferenceCollection InputDataReferences;
-			InputDataReferences.AddDataReadReference(METASOUND_GET_PARAM_NAME(ParamAudioInput), FAudioBufferReadRef(AudioInput));
-			InputDataReferences.AddDataReadReference(METASOUND_GET_PARAM_NAME(ParamCutoffFrequency), FFloatReadRef(Frequency));
-			InputDataReferences.AddDataReadReference(METASOUND_GET_PARAM_NAME(ParamResonance), FFloatReadRef(Resonance));
-			InputDataReferences.AddDataReadReference(METASOUND_GET_PARAM_NAME(ParamBandStopControl), FFloatReadRef(BandStopControl));
+			InOutVertexData.BindReadVertex(METASOUND_GET_PARAM_NAME(ParamAudioInput), AudioInput);
+			InOutVertexData.BindReadVertex(METASOUND_GET_PARAM_NAME(ParamCutoffFrequency), Frequency);
+			InOutVertexData.BindReadVertex(METASOUND_GET_PARAM_NAME(ParamResonance), Resonance);
+			InOutVertexData.BindReadVertex(METASOUND_GET_PARAM_NAME(ParamBandStopControl), BandStopControl);
+		}
 
-			return InputDataReferences;
+		virtual void BindOutputs(FOutputVertexInterfaceData& InOutVertexData) override
+		{
+			InOutVertexData.BindReadVertex(METASOUND_GET_PARAM_NAME(ParamLowPassOutput), LowPassOutput);
+			InOutVertexData.BindReadVertex(METASOUND_GET_PARAM_NAME(ParamHighPassOutput), HighPassOutput);
+			InOutVertexData.BindReadVertex(METASOUND_GET_PARAM_NAME(ParamBandPassOutput), BandPassOutput);
+			InOutVertexData.BindReadVertex(METASOUND_GET_PARAM_NAME(ParamBandStopOutput), BandStopOutput);
+		}
+
+		virtual FDataReferenceCollection GetInputs() const override
+		{
+			// This should never be called. Bind(...) is called instead. This method
+			// exists as a stop-gap until the API can be deprecated and removed.
+			checkNoEntry();
+			return {};
 		}
 
 		virtual FDataReferenceCollection GetOutputs() const override
 		{
-			// expose read access to our output buffer for other processors in the graph
-			FDataReferenceCollection OutputDataReferences;
-			OutputDataReferences.AddDataReadReference(METASOUND_GET_PARAM_NAME(ParamLowPassOutput), FAudioBufferReadRef(LowPassOutput));
-			OutputDataReferences.AddDataReadReference(METASOUND_GET_PARAM_NAME(ParamHighPassOutput), FAudioBufferReadRef(HighPassOutput));
-			OutputDataReferences.AddDataReadReference(METASOUND_GET_PARAM_NAME(ParamBandPassOutput), FAudioBufferReadRef(BandPassOutput));
-			OutputDataReferences.AddDataReadReference(METASOUND_GET_PARAM_NAME(ParamBandStopOutput), FAudioBufferReadRef(BandStopOutput));
-
-			return OutputDataReferences;
+			// This should never be called. Bind(...) is called instead. This method
+			// exists as a stop-gap until the API can be deprecated and removed.
+			checkNoEntry();
+			return {};
 		}
 
 		void Reset(const IOperator::FResetParams& InParams);
@@ -491,22 +509,31 @@ namespace Metasound
 
 		static TUniquePtr<IOperator> CreateOperator(const FCreateOperatorParams& InParams, FBuildErrorArray& OutErrors);
 
-		virtual FDataReferenceCollection GetInputs() const override
+		virtual void BindInputs(FInputVertexInterfaceData& InOutVertexData) override
 		{
-			FDataReferenceCollection InputDataReferences;
-			InputDataReferences.AddDataReadReference(METASOUND_GET_PARAM_NAME(ParamAudioInput), FAudioBufferReadRef(AudioInput));
-			InputDataReferences.AddDataReadReference(METASOUND_GET_PARAM_NAME(ParamCutoffFrequency), FFloatReadRef(Frequency));
+			InOutVertexData.BindReadVertex(METASOUND_GET_PARAM_NAME(ParamAudioInput), AudioInput);
+			InOutVertexData.BindReadVertex(METASOUND_GET_PARAM_NAME(ParamCutoffFrequency), Frequency);
+		}
 
-			return InputDataReferences;
+		virtual void BindOutputs(FOutputVertexInterfaceData& InOutVertexData) override
+		{
+			InOutVertexData.BindReadVertex(METASOUND_GET_PARAM_NAME(ParamAudioOutput), AudioOutput);
 		}
 
-		virtual FDataReferenceCollection GetOutputs() const override
+		virtual FDataReferenceCollection GetInputs() const override
 		{
-			// expose read access to our output buffer for other processors in the graph
-			FDataReferenceCollection OutputDataReferences;
-			OutputDataReferences.AddDataReadReference(METASOUND_GET_PARAM_NAME(ParamAudioOutput), FAudioBufferReadRef(AudioOutput));
+			// This should never be called. Bind(...) is called instead. This method
+			// exists as a stop-gap until the API can be deprecated and removed.
+			checkNoEntry();
+			return {};
+		}
 
-			return OutputDataReferences;
+		virtual FDataReferenceCollection GetOutputs() const override
+		{
+			// This should never be called. Bind(...) is called instead. This method
+			// exists as a stop-gap until the API can be deprecated and removed.
+			checkNoEntry();
+			return {};
 		}
 
 		void Reset(const IOperator::FResetParams& InParams);
@@ -636,22 +663,38 @@ namespace Metasound
 
 		static TUniquePtr<IOperator> CreateOperator(const FCreateOperatorParams& InParams, FBuildErrorArray& OutErrors);
 
-		virtual FDataReferenceCollection GetInputs() const override
+		virtual void BindInputs(FInputVertexInterfaceData& InOutVertexData) override
 		{
-			FDataReferenceCollection InputDataReferences;
-			InputDataReferences.AddDataReadReference(METASOUND_GET_PARAM_NAME(ParamAudioInput), FAudioBufferReadRef(AudioInput));
-			InputDataReferences.AddDataReadReference(METASOUND_GET_PARAM_NAME(ParamCutoffFrequency), FFloatReadRef(Frequency));
+			InOutVertexData.BindReadVertex(METASOUND_GET_PARAM_NAME(ParamAudioInput), AudioInput);
+			InOutVertexData.BindReadVertex(METASOUND_GET_PARAM_NAME(ParamCutoffFrequency), Frequency);
+		}
 
-			return InputDataReferences;
+		virtual void BindOutputs(FOutputVertexInterfaceData& InOutVertexData) override
+		{
+			InOutVertexData.BindReadVertex(METASOUND_GET_PARAM_NAME(ParamAudioOutput), AudioOutput);
 		}
 
-		virtual FDataReferenceCollection GetOutputs() const override
+		virtual void Bind(FVertexInterfaceData& InVertexData) const override
+		{
+			FInputVertexInterfaceData& Inputs = InVertexData.GetInputs();
+
+			FOutputVertexInterfaceData& Outputs = InVertexData.GetOutputs();
+		}
+
+		virtual FDataReferenceCollection GetInputs() const override
 		{
-			// expose read access to our output buffer for other processors in the graph
-			FDataReferenceCollection OutputDataReferences;
-			OutputDataReferences.AddDataReadReference(METASOUND_GET_PARAM_NAME(ParamAudioOutput), FAudioBufferReadRef(AudioOutput));
+			// This should never be called. Bind(...) is called instead. This method
+			// exists as a stop-gap until the API can be deprecated and removed.
+			checkNoEntry();
+			return {};
+		}
 
-			return OutputDataReferences;
+		virtual FDataReferenceCollection GetOutputs() const override
+		{
+			// This should never be called. Bind(...) is called instead. This method
+			// exists as a stop-gap until the API can be deprecated and removed.
+			checkNoEntry();
+			return {};
 		}
 
 		void Reset(const IOperator::FResetParams& InParams);
@@ -769,7 +812,7 @@ namespace Metasound
 
 
 #pragma region Biquad Filter
-		class FBiquadFilterOperator : public TExecutableOperator<FBiquadFilterOperator>
+	class FBiquadFilterOperator : public TExecutableOperator<FBiquadFilterOperator>
 	{
 	private:
 		static constexpr float InvalidValue = -1.f;
@@ -781,25 +824,41 @@ namespace Metasound
 
 		static TUniquePtr<IOperator> CreateOperator(const FCreateOperatorParams& InParams, FBuildErrorArray& OutErrors);
 
+		virtual void BindInputs(FInputVertexInterfaceData& InOutVertexData) override
+		{
+			InOutVertexData.BindReadVertex(METASOUND_GET_PARAM_NAME(ParamAudioInput), AudioInput);
+			InOutVertexData.BindReadVertex(METASOUND_GET_PARAM_NAME(ParamCutoffFrequency), Frequency);
+			InOutVertexData.BindReadVertex(METASOUND_GET_PARAM_NAME(ParamBandwidth), Bandwidth);
+			InOutVertexData.BindReadVertex(METASOUND_GET_PARAM_NAME(ParamGainDb), FilterGainDb);
+			InOutVertexData.BindReadVertex(METASOUND_GET_PARAM_NAME(ParamFilterType), FilterType);
+		}
+
+		virtual void BindOutputs(FOutputVertexInterfaceData& InOutVertexData) override
+		{
+			InOutVertexData.BindReadVertex(METASOUND_GET_PARAM_NAME(ParamAudioOutput), AudioOutput);
+		}
+
+		virtual void Bind(FVertexInterfaceData& InVertexData) const override
+		{
+			FInputVertexInterfaceData& Inputs = InVertexData.GetInputs();
+			FOutputVertexInterfaceData& Outputs = InVertexData.GetOutputs();
+
+		}
+
 		virtual FDataReferenceCollection GetInputs() const override
 		{
-			FDataReferenceCollection InputDataReferences;
-			InputDataReferences.AddDataReadReference(METASOUND_GET_PARAM_NAME(ParamAudioInput), FAudioBufferReadRef(AudioInput));
-			InputDataReferences.AddDataReadReference(METASOUND_GET_PARAM_NAME(ParamCutoffFrequency), FFloatReadRef(Frequency));
-			InputDataReferences.AddDataReadReference(METASOUND_GET_PARAM_NAME(ParamBandwidth), FFloatReadRef(Bandwidth));
-			InputDataReferences.AddDataReadReference(METASOUND_GET_PARAM_NAME(ParamGainDb), FFloatReadRef(FilterGainDb));
-			InputDataReferences.AddDataReadReference(METASOUND_GET_PARAM_NAME(ParamFilterType), FEnumBiQuadFilterReadRef(FilterType));
-
-			return InputDataReferences;
+			// This should never be called. Bind(...) is called instead. This method
+			// exists as a stop-gap until the API can be deprecated and removed.
+			checkNoEntry();
+			return {};
 		}
 
 		virtual FDataReferenceCollection GetOutputs() const override
 		{
-			// expose read access to our output buffer for other processors in the graph
-			FDataReferenceCollection OutputDataReferences;
-			OutputDataReferences.AddDataReadReference(METASOUND_GET_PARAM_NAME(ParamAudioOutput), FAudioBufferReadRef(AudioOutput));
-
-			return OutputDataReferences;
+			// This should never be called. Bind(...) is called instead. This method
+			// exists as a stop-gap until the API can be deprecated and removed.
+			checkNoEntry();
+			return {};
 		}
 
 		void Reset(const IOperator::FResetParams& InParams);
diff --git a/Engine/Plugins/Runtime/Metasound/Source/MetasoundStandardNodes/Private/MetasoundBitcrusherNode.cpp b/Engine/Plugins/Runtime/Metasound/Source/MetasoundStandardNodes/Private/MetasoundBitcrusherNode.cpp
--- a/Engine/Plugins/Runtime/Metasound/Source/MetasoundStandardNodes/Private/MetasoundBitcrusherNode.cpp
+++ b/Engine/Plugins/Runtime/Metasound/Source/MetasoundStandardNodes/Private/MetasoundBitcrusherNode.cpp
@@ -92,29 +92,37 @@ namespace Metasound
 			return Interface;
 		}
 
-		virtual FDataReferenceCollection GetInputs() const override
+
+		virtual void BindInputs(FInputVertexInterfaceData& InOutVertexData) override
 		{
 			using namespace BitcrusherVertexNames;
 
-			FDataReferenceCollection InputDataReferences;
-
-			InputDataReferences.AddDataReadReference(METASOUND_GET_PARAM_NAME(InputAudio), AudioInput);
-			InputDataReferences.AddDataReadReference(METASOUND_GET_PARAM_NAME(InputSampleRate), CrushedSampleRate);
-			InputDataReferences.AddDataReadReference(METASOUND_GET_PARAM_NAME(InputBitDepth), CrushedBitDepth);
-
-			return InputDataReferences;
-
+			InOutVertexData.BindReadVertex(METASOUND_GET_PARAM_NAME(InputAudio), AudioInput);
+			InOutVertexData.BindReadVertex(METASOUND_GET_PARAM_NAME(InputSampleRate), CrushedSampleRate);
+			InOutVertexData.BindReadVertex(METASOUND_GET_PARAM_NAME(InputBitDepth), CrushedBitDepth);
 		}
 
-		virtual FDataReferenceCollection GetOutputs() const override
+		virtual void BindOutputs(FOutputVertexInterfaceData& InOutVertexData) override
 		{
 			using namespace BitcrusherVertexNames;
 
-			FDataReferenceCollection OutputDataReferences;
+			InOutVertexData.BindReadVertex(METASOUND_GET_PARAM_NAME(OutputAudio), AudioOutput);
+		}
 
-			OutputDataReferences.AddDataReadReference(METASOUND_GET_PARAM_NAME(OutputAudio), AudioOutput);
+		virtual FDataReferenceCollection GetInputs() const override
+		{
+			// This should never be called. Bind(...) is called instead. This method
+			// exists as a stop-gap until the API can be deprecated and removed.
+			checkNoEntry();
+			return {};
+		}
 
-			return OutputDataReferences;
+		virtual FDataReferenceCollection GetOutputs() const override
+		{
+			// This should never be called. Bind(...) is called instead. This method
+			// exists as a stop-gap until the API can be deprecated and removed.
+			checkNoEntry();
+			return {};
 		}
 
 		static TUniquePtr<IOperator> CreateOperator(const FCreateOperatorParams& InParams, FBuildErrorArray& OutErrors)
diff --git a/Engine/Plugins/Runtime/Metasound/Source/MetasoundStandardNodes/Private/MetasoundClampNode.cpp b/Engine/Plugins/Runtime/Metasound/Source/MetasoundStandardNodes/Private/MetasoundClampNode.cpp
--- a/Engine/Plugins/Runtime/Metasound/Source/MetasoundStandardNodes/Private/MetasoundClampNode.cpp
+++ b/Engine/Plugins/Runtime/Metasound/Source/MetasoundStandardNodes/Private/MetasoundClampNode.cpp
@@ -213,26 +213,36 @@ namespace Metasound
 
 		virtual ~TClampNodeOperator() = default;
 
-		virtual FDataReferenceCollection GetInputs() const override
+		virtual void BindInputs(FInputVertexInterfaceData& InOutVertexData) override
 		{
 			using namespace ClampVertexNames;
 
-			FDataReferenceCollection Inputs;
-			Inputs.AddDataReadReference(METASOUND_GET_PARAM_NAME(InputValue), InputValue);
-			Inputs.AddDataReadReference(METASOUND_GET_PARAM_NAME(InputMinValue), InputMinValue);
-			Inputs.AddDataReadReference(METASOUND_GET_PARAM_NAME(InputMaxValue), InputMaxValue);
-
-			return Inputs;
+			InOutVertexData.BindReadVertex(METASOUND_GET_PARAM_NAME(InputValue), InputValue);
+			InOutVertexData.BindReadVertex(METASOUND_GET_PARAM_NAME(InputMinValue), InputMinValue);
+			InOutVertexData.BindReadVertex(METASOUND_GET_PARAM_NAME(InputMaxValue), InputMaxValue);
 		}
 
-		virtual FDataReferenceCollection GetOutputs() const override
+		virtual void BindOutputs(FOutputVertexInterfaceData& InOutVertexData) override
 		{
 			using namespace ClampVertexNames;
 
-			FDataReferenceCollection Outputs;
-			Outputs.AddDataReadReference(METASOUND_GET_PARAM_NAME(OutputValue), OutputValue);
+			InOutVertexData.BindReadVertex(METASOUND_GET_PARAM_NAME(OutputValue), OutputValue);
+		}
+
+		virtual FDataReferenceCollection GetInputs() const override
+		{
+			// This should never be called. Bind(...) is called instead. This method
+			// exists as a stop-gap until the API can be deprecated and removed.
+			checkNoEntry();
+			return {};
+		}
 
-			return Outputs;
+		virtual FDataReferenceCollection GetOutputs() const override
+		{
+			// This should never be called. Bind(...) is called instead. This method
+			// exists as a stop-gap until the API can be deprecated and removed.
+			checkNoEntry();
+			return {};
 		}
 
 		void GetClamped()
diff --git a/Engine/Plugins/Runtime/Metasound/Source/MetasoundStandardNodes/Private/MetasoundCompressorNode.cpp b/Engine/Plugins/Runtime/Metasound/Source/MetasoundStandardNodes/Private/MetasoundCompressorNode.cpp
--- a/Engine/Plugins/Runtime/Metasound/Source/MetasoundStandardNodes/Private/MetasoundCompressorNode.cpp
+++ b/Engine/Plugins/Runtime/Metasound/Source/MetasoundStandardNodes/Private/MetasoundCompressorNode.cpp
@@ -171,39 +171,46 @@ namespace Metasound
 			return Interface;
 		}
 
-		virtual FDataReferenceCollection GetInputs() const override
+		virtual void BindInputs(FInputVertexInterfaceData& InOutVertexData) override
 		{
 			using namespace CompressorVertexNames;
 
-			FDataReferenceCollection InputDataReferences;
-
-			InputDataReferences.AddDataReadReference(METASOUND_GET_PARAM_NAME(InputAudio), AudioInput);
-			InputDataReferences.AddDataReadReference(METASOUND_GET_PARAM_NAME(InputRatio), RatioInput);
-			InputDataReferences.AddDataReadReference(METASOUND_GET_PARAM_NAME(InputThreshold), ThresholdDbInput);
-			InputDataReferences.AddDataReadReference(METASOUND_GET_PARAM_NAME(InputAttackTime), AttackTimeInput);
-			InputDataReferences.AddDataReadReference(METASOUND_GET_PARAM_NAME(InputReleaseTime), ReleaseTimeInput);
-			InputDataReferences.AddDataReadReference(METASOUND_GET_PARAM_NAME(InputLookaheadTime), LookaheadTimeInput);
-			InputDataReferences.AddDataReadReference(METASOUND_GET_PARAM_NAME(InputKnee), KneeInput);
-			InputDataReferences.AddDataReadReference(METASOUND_GET_PARAM_NAME(InputSidechain), SidechainInput);
-			InputDataReferences.AddDataReadReference(METASOUND_GET_PARAM_NAME(InputEnvelopeMode), EnvelopeModeInput);
-			InputDataReferences.AddDataReadReference(METASOUND_GET_PARAM_NAME(InputIsUpwards), bIsUpwardsInput);
-			InputDataReferences.AddDataReadReference(METASOUND_GET_PARAM_NAME(InputIsAnalog), bIsAnalogInput);
-			InputDataReferences.AddDataReadReference(METASOUND_GET_PARAM_NAME(InputWetDryMix), WetDryMixInput);
-
-			return InputDataReferences;
-
+			InOutVertexData.BindReadVertex(METASOUND_GET_PARAM_NAME(InputAudio), AudioInput);
+			InOutVertexData.BindReadVertex(METASOUND_GET_PARAM_NAME(InputRatio), RatioInput);
+			InOutVertexData.BindReadVertex(METASOUND_GET_PARAM_NAME(InputThreshold), ThresholdDbInput);
+			InOutVertexData.BindReadVertex(METASOUND_GET_PARAM_NAME(InputAttackTime), AttackTimeInput);
+			InOutVertexData.BindReadVertex(METASOUND_GET_PARAM_NAME(InputReleaseTime), ReleaseTimeInput);
+			InOutVertexData.BindReadVertex(METASOUND_GET_PARAM_NAME(InputLookaheadTime), LookaheadTimeInput);
+			InOutVertexData.BindReadVertex(METASOUND_GET_PARAM_NAME(InputKnee), KneeInput);
+			InOutVertexData.BindReadVertex(METASOUND_GET_PARAM_NAME(InputSidechain), SidechainInput);
+			InOutVertexData.BindReadVertex(METASOUND_GET_PARAM_NAME(InputEnvelopeMode), EnvelopeModeInput);
+			InOutVertexData.BindReadVertex(METASOUND_GET_PARAM_NAME(InputIsUpwards), bIsUpwardsInput);
+			InOutVertexData.BindReadVertex(METASOUND_GET_PARAM_NAME(InputIsAnalog), bIsAnalogInput);
+			InOutVertexData.BindReadVertex(METASOUND_GET_PARAM_NAME(InputWetDryMix), WetDryMixInput);
 		}
 
-		virtual FDataReferenceCollection GetOutputs() const override
+		virtual void BindOutputs(FOutputVertexInterfaceData& InOutVertexData) override
 		{
 			using namespace CompressorVertexNames;
 
-			FDataReferenceCollection OutputDataReferences;
+			InOutVertexData.BindReadVertex(METASOUND_GET_PARAM_NAME(OutputAudio), AudioOutput);
+			InOutVertexData.BindReadVertex(METASOUND_GET_PARAM_NAME(OutputEnvelope), EnvelopeOutput);
+		}
 
-			OutputDataReferences.AddDataReadReference(METASOUND_GET_PARAM_NAME(OutputAudio), AudioOutput);
-			OutputDataReferences.AddDataReadReference(METASOUND_GET_PARAM_NAME(OutputEnvelope), EnvelopeOutput);
+		virtual FDataReferenceCollection GetInputs() const override
+		{
+			// This should never be called. Bind(...) is called instead. This method
+			// exists as a stop-gap until the API can be deprecated and removed.
+			checkNoEntry();
+			return {};
+		}
 
-			return OutputDataReferences;
+		virtual FDataReferenceCollection GetOutputs() const override
+		{
+			// This should never be called. Bind(...) is called instead. This method
+			// exists as a stop-gap until the API can be deprecated and removed.
+			checkNoEntry();
+			return {};
 		}
 
 		static TUniquePtr<IOperator> CreateOperator(const FCreateOperatorParams& InParams, FBuildErrorArray& OutErrors)
diff --git a/Engine/Plugins/Runtime/Metasound/Source/MetasoundStandardNodes/Private/MetasoundConversionNodes.cpp b/Engine/Plugins/Runtime/Metasound/Source/MetasoundStandardNodes/Private/MetasoundConversionNodes.cpp
--- a/Engine/Plugins/Runtime/Metasound/Source/MetasoundStandardNodes/Private/MetasoundConversionNodes.cpp
+++ b/Engine/Plugins/Runtime/Metasound/Source/MetasoundStandardNodes/Private/MetasoundConversionNodes.cpp
@@ -239,24 +239,32 @@ namespace Metasound
 			}
 		}
 
-		virtual FDataReferenceCollection GetInputs() const override
+		virtual void BindInputs(FInputVertexInterfaceData& InOutVertexData) override
 		{
 			using namespace ConversionNodeVertexNames;
-
-			FDataReferenceCollection Inputs;
-			Inputs.AddDataReadReference(METASOUND_GET_PARAM_NAME(InputValue), InputValue);
-
-			return Inputs;
+			InOutVertexData.BindReadVertex(METASOUND_GET_PARAM_NAME(InputValue), InputValue);
 		}
 
-		virtual FDataReferenceCollection GetOutputs() const override
+		virtual void BindOutputs(FOutputVertexInterfaceData& InOutVertexData) override
 		{
 			using namespace ConversionNodeVertexNames;
+			InOutVertexData.BindReadVertex(METASOUND_GET_PARAM_NAME(OutputValue), OutputValue);
+		}
 
-			FDataReferenceCollection Outputs;
-			Outputs.AddDataReadReference(METASOUND_GET_PARAM_NAME(OutputValue), OutputValue);
+		virtual FDataReferenceCollection GetInputs() const override
+		{
+			// This should never be called. Bind(...) is called instead. This method
+			// exists as a stop-gap until the API can be deprecated and removed.
+			checkNoEntry();
+			return {};
+		}
 
-			return Outputs;
+		virtual FDataReferenceCollection GetOutputs() const override
+		{
+			// This should never be called. Bind(...) is called instead. This method
+			// exists as a stop-gap until the API can be deprecated and removed.
+			checkNoEntry();
+			return {};
 		}
 
 		void Reset(const IOperator::FResetParams& InParams)
diff --git a/Engine/Plugins/Runtime/Metasound/Source/MetasoundStandardNodes/Private/MetasoundCrossfadeNode.cpp b/Engine/Plugins/Runtime/Metasound/Source/MetasoundStandardNodes/Private/MetasoundCrossfadeNode.cpp
--- a/Engine/Plugins/Runtime/Metasound/Source/MetasoundStandardNodes/Private/MetasoundCrossfadeNode.cpp
+++ b/Engine/Plugins/Runtime/Metasound/Source/MetasoundStandardNodes/Private/MetasoundCrossfadeNode.cpp
@@ -232,26 +232,38 @@ namespace Metasound
 
 		virtual ~TCrossfadeOperator() = default;
 
-		virtual FDataReferenceCollection GetInputs() const override
+
+		virtual void BindInputs(FInputVertexInterfaceData& InOutVertexData) override
 		{
 			using namespace CrossfadeVertexNames;
-			FDataReferenceCollection Inputs;
-			Inputs.AddDataReadReference(METASOUND_GET_PARAM_NAME(InputCrossfadeValue), CrossfadeValue);
+			InOutVertexData.BindReadVertex(METASOUND_GET_PARAM_NAME(InputCrossfadeValue), CrossfadeValue);
+
 			for (uint32 i = 0; i < NumInputs; ++i)
 			{
-				Inputs.AddDataReadReference(GetInputName(i), InputValues[i]);
+				InOutVertexData.BindReadVertex(GetInputName(i), InputValues[i]);
 			}
-			return Inputs;
 		}
 
-		virtual FDataReferenceCollection GetOutputs() const override
+		virtual void BindOutputs(FOutputVertexInterfaceData& InOutVertexData) override
 		{
 			using namespace CrossfadeVertexNames;
+			InOutVertexData.BindReadVertex(METASOUND_GET_PARAM_NAME(OutputTrigger), OutputValue);
+		}
 
-			FDataReferenceCollection Outputs;
-			Outputs.AddDataReadReference(METASOUND_GET_PARAM_NAME(OutputTrigger), OutputValue);
+		virtual FDataReferenceCollection GetInputs() const override
+		{
+			// This should never be called. Bind(...) is called instead. This method
+			// exists as a stop-gap until the API can be deprecated and removed.
+			checkNoEntry();
+			return {};
+		}
 
-			return Outputs;
+		virtual FDataReferenceCollection GetOutputs() const override
+		{
+			// This should never be called. Bind(...) is called instead. This method
+			// exists as a stop-gap until the API can be deprecated and removed.
+			checkNoEntry();
+			return {};
 		}
 
 		void PerformCrossfadeOutput()
diff --git a/Engine/Plugins/Runtime/Metasound/Source/MetasoundStandardNodes/Private/MetasoundDebugLogNode.cpp b/Engine/Plugins/Runtime/Metasound/Source/MetasoundStandardNodes/Private/MetasoundDebugLogNode.cpp
--- a/Engine/Plugins/Runtime/Metasound/Source/MetasoundStandardNodes/Private/MetasoundDebugLogNode.cpp
+++ b/Engine/Plugins/Runtime/Metasound/Source/MetasoundStandardNodes/Private/MetasoundDebugLogNode.cpp
@@ -113,24 +113,32 @@ namespace Metasound
 
 			virtual ~TPrintLogOperator() = default;
 
-
-			virtual FDataReferenceCollection GetInputs() const override
+			virtual void BindInputs(FInputVertexInterfaceData& InOutVertexData) override
 			{
 				using namespace PrintLogVertexNames;
-				FDataReferenceCollection Inputs;
+				InOutVertexData.BindReadVertex(METASOUND_GET_PARAM_NAME(InputTrigger), Trigger);
+				InOutVertexData.BindReadVertex(METASOUND_GET_PARAM_NAME(InputLabel), Label);
+				InOutVertexData.BindReadVertex(METASOUND_GET_PARAM_NAME(InputValueToLog), ValueToLog);
+			}
 
-				Inputs.AddDataReadReference(METASOUND_GET_PARAM_NAME(InputTrigger), Trigger);
-				Inputs.AddDataReadReference(METASOUND_GET_PARAM_NAME(InputLabel), Label);
-				Inputs.AddDataReadReference(METASOUND_GET_PARAM_NAME(InputValueToLog), ValueToLog);
+			virtual void BindOutputs(FOutputVertexInterfaceData&) override
+			{
+			}
 
-				return Inputs;
+			virtual FDataReferenceCollection GetInputs() const override
+			{
+				// This should never be called. Bind(...) is called instead. This method
+				// exists as a stop-gap until the API can be deprecated and removed.
+				checkNoEntry();
+				return {};
 			}
 
 			virtual FDataReferenceCollection GetOutputs() const override
 			{
-				FDataReferenceCollection Outputs;
-
-				return Outputs;
+				// This should never be called. Bind(...) is called instead. This method
+				// exists as a stop-gap until the API can be deprecated and removed.
+				checkNoEntry();
+				return {};
 			}
 
 			void Reset(const IOperator::FResetParams& InParams)
diff --git a/Engine/Plugins/Runtime/Metasound/Source/MetasoundStandardNodes/Private/MetasoundDebugTriggerOnNanNode.cpp b/Engine/Plugins/Runtime/Metasound/Source/MetasoundStandardNodes/Private/MetasoundDebugTriggerOnNanNode.cpp
--- a/Engine/Plugins/Runtime/Metasound/Source/MetasoundStandardNodes/Private/MetasoundDebugTriggerOnNanNode.cpp
+++ b/Engine/Plugins/Runtime/Metasound/Source/MetasoundStandardNodes/Private/MetasoundDebugTriggerOnNanNode.cpp
@@ -39,8 +39,9 @@ namespace Metasound
 
 		void Reset(const IOperator::FResetParams&);
 
+		virtual void BindInputs(FInputVertexInterfaceData& InOutVertexData) override;
+		virtual void BindOutputs(FOutputVertexInterfaceData& InOutVertexData) override;
 		virtual FDataReferenceCollection GetInputs() const override;
-
 		virtual FDataReferenceCollection GetOutputs() const override;
 
 		void Execute();
@@ -72,25 +73,36 @@ namespace Metasound
 		bHasTriggered = false;
 	}
 
-	FDataReferenceCollection FTriggerOnNanOperator::GetInputs() const
+	void FTriggerOnNanOperator::BindInputs(FInputVertexInterfaceData& InOutVertexData)
 	{
 		using namespace TriggerOnNanVertexNames;
 
-		FDataReferenceCollection InputDataReferences;
-		InputDataReferences.AddDataReadReference(METASOUND_GET_PARAM_NAME(AudioInput), AudioInput);
-		InputDataReferences.AddDataReadReference(METASOUND_GET_PARAM_NAME(TriggerOnce), bTriggerOnce);
-		return InputDataReferences;
+		InOutVertexData.BindReadVertex(METASOUND_GET_PARAM_NAME(AudioInput), AudioInput);
+		InOutVertexData.BindReadVertex(METASOUND_GET_PARAM_NAME(TriggerOnce), bTriggerOnce);
 	}
 
-	FDataReferenceCollection FTriggerOnNanOperator::GetOutputs() const
+	void FTriggerOnNanOperator::BindOutputs(FOutputVertexInterfaceData& InOutVertexData)
 	{
 		using namespace TriggerOnNanVertexNames;
-
-		FDataReferenceCollection OutputDataReferences;
-		OutputDataReferences.AddDataReadReference(METASOUND_GET_PARAM_NAME(OutputTriggerOnNan), TriggerOutput);
-		return OutputDataReferences;
+		InOutVertexData.BindReadVertex(METASOUND_GET_PARAM_NAME(OutputTriggerOnNan), TriggerOutput);
 	}
 
+    FDataReferenceCollection FTriggerOnNanOperator::GetInputs() const
+    {
+    	// This should never be called. Bind(...) is called instead. This method
+    	// exists as a stop-gap until the API can be deprecated and removed.
+    	checkNoEntry();
+    	return {};
+    }
+
+    FDataReferenceCollection FTriggerOnNanOperator::GetOutputs() const
+    {
+    	// This should never be called. Bind(...) is called instead. This method
+    	// exists as a stop-gap until the API can be deprecated and removed.
+    	checkNoEntry();
+    	return {};
+    }
+
 	void FTriggerOnNanOperator::Execute()
 	{
 		const bool bShouldTriggerOnce = *bTriggerOnce;
diff --git a/Engine/Plugins/Runtime/Metasound/Source/MetasoundStandardNodes/Private/MetasoundDecibelsToLinearGain.cpp b/Engine/Plugins/Runtime/Metasound/Source/MetasoundStandardNodes/Private/MetasoundDecibelsToLinearGain.cpp
--- a/Engine/Plugins/Runtime/Metasound/Source/MetasoundStandardNodes/Private/MetasoundDecibelsToLinearGain.cpp
+++ b/Engine/Plugins/Runtime/Metasound/Source/MetasoundStandardNodes/Private/MetasoundDecibelsToLinearGain.cpp
@@ -33,6 +33,8 @@ namespace Metasound
 
 		FDecibelsToLinearGainOperator(const FCreateOperatorParams& InParams, const FFloatReadRef& InDecibelGain);
 
+		virtual void BindInputs(FInputVertexInterfaceData& InOutVertexData) override;
+		virtual void BindOutputs(FOutputVertexInterfaceData& InOutVertexData) override;
 		virtual FDataReferenceCollection GetInputs() const override;
 		virtual FDataReferenceCollection GetOutputs() const override;
 		void Reset(const IOperator::FResetParams& InParams);
@@ -52,24 +54,33 @@ namespace Metasound
 	{
 		Reset(InParams);
 	}
-	
-	FDataReferenceCollection FDecibelsToLinearGainOperator::GetInputs() const
+
+	void FDecibelsToLinearGainOperator::BindInputs(FInputVertexInterfaceData& InOutVertexData)
 	{
 		using namespace DecibelsToLinearGainVertexNames;
+		InOutVertexData.BindReadVertex(METASOUND_GET_PARAM_NAME(InputDecibelGain), DecibelGainInput);
+	}
 
-		FDataReferenceCollection InputDataReferences;
-		InputDataReferences.AddDataReadReference(METASOUND_GET_PARAM_NAME(InputDecibelGain), DecibelGainInput);
+	void FDecibelsToLinearGainOperator::BindOutputs(FOutputVertexInterfaceData& InOutVertexData)
+	{
+		using namespace DecibelsToLinearGainVertexNames;
+		InOutVertexData.BindReadVertex(METASOUND_GET_PARAM_NAME(OutputLinearGain), LinearGainOutput);
+	}
 
-		return InputDataReferences;
+	FDataReferenceCollection FDecibelsToLinearGainOperator::GetInputs() const
+	{
+		// This should never be called. Bind(...) is called instead. This method
+		// exists as a stop-gap until the API can be deprecated and removed.
+		checkNoEntry();
+		return {};
 	}
 
 	FDataReferenceCollection FDecibelsToLinearGainOperator::GetOutputs() const
 	{
-		using namespace DecibelsToLinearGainVertexNames;
-
-		FDataReferenceCollection OutputDataReferences;
-		OutputDataReferences.AddDataReadReference(METASOUND_GET_PARAM_NAME(OutputLinearGain), LinearGainOutput);
-		return OutputDataReferences;
+		// This should never be called. Bind(...) is called instead. This method
+		// exists as a stop-gap until the API can be deprecated and removed.
+		checkNoEntry();
+		return {};
 	}
 
 	void FDecibelsToLinearGainOperator::Reset(const IOperator::FResetParams& InParams)
diff --git a/Engine/Plugins/Runtime/Metasound/Source/MetasoundStandardNodes/Private/MetasoundDelayNode.cpp b/Engine/Plugins/Runtime/Metasound/Source/MetasoundStandardNodes/Private/MetasoundDelayNode.cpp
--- a/Engine/Plugins/Runtime/Metasound/Source/MetasoundStandardNodes/Private/MetasoundDelayNode.cpp
+++ b/Engine/Plugins/Runtime/Metasound/Source/MetasoundStandardNodes/Private/MetasoundDelayNode.cpp
@@ -48,7 +48,8 @@ namespace Metasound
 			const FFloatReadRef& InFeedback,
 			float InMaxDelayTimeSeconds);
 
-		virtual void Bind(FVertexInterfaceData& InVertexData) const override;
+		virtual void BindInputs(FInputVertexInterfaceData& InOutVertexData) override;
+		virtual void BindOutputs(FOutputVertexInterfaceData& InOutVertexData) override;
 		virtual FDataReferenceCollection GetInputs() const override;
 		virtual FDataReferenceCollection GetOutputs() const override;
 		void Reset(const IOperator::FResetParams& InParams);
@@ -108,43 +109,39 @@ namespace Metasound
 		Reset(InParams);
 	}
 
-	void FDelayOperator::Bind(FVertexInterfaceData& InVertexData) const
+
+	void FDelayOperator::BindInputs(FInputVertexInterfaceData& InOutVertexData)
 	{
 		using namespace Delay;
 
-		FInputVertexInterfaceData& Inputs = InVertexData.GetInputs();
-
-		Inputs.BindReadVertex(METASOUND_GET_PARAM_NAME(InParamAudioInput), FAudioBufferReadRef(AudioInput));
-		Inputs.BindReadVertex(METASOUND_GET_PARAM_NAME(InParamDelayTime), FTimeReadRef(DelayTime));
-		Inputs.BindReadVertex(METASOUND_GET_PARAM_NAME(InParamDryLevel), FFloatReadRef(DryLevel));
-		Inputs.BindReadVertex(METASOUND_GET_PARAM_NAME(InParamWetLevel), FFloatReadRef(WetLevel));
-		Inputs.BindReadVertex(METASOUND_GET_PARAM_NAME(InParamFeedbackAmount), FFloatReadRef(Feedback));
-		Inputs.SetValue(METASOUND_GET_PARAM_NAME(InParamMaxDelayTime), FTime::FromSeconds(MaxDelayTimeSeconds));
+		InOutVertexData.BindReadVertex(METASOUND_GET_PARAM_NAME(InParamAudioInput), AudioInput);
+		InOutVertexData.BindReadVertex(METASOUND_GET_PARAM_NAME(InParamDelayTime), DelayTime);
+		InOutVertexData.BindReadVertex(METASOUND_GET_PARAM_NAME(InParamDryLevel), DryLevel);
+		InOutVertexData.BindReadVertex(METASOUND_GET_PARAM_NAME(InParamWetLevel), WetLevel);
+		InOutVertexData.BindReadVertex(METASOUND_GET_PARAM_NAME(InParamFeedbackAmount), Feedback);
+		InOutVertexData.SetValue(METASOUND_GET_PARAM_NAME(InParamMaxDelayTime), FTime::FromSeconds(MaxDelayTimeSeconds));
+	}
 
-		FOutputVertexInterfaceData& Outputs = InVertexData.GetOutputs();
-		Outputs.BindReadVertex(METASOUND_GET_PARAM_NAME(OutParamAudio), AudioOutput);
+	void FDelayOperator::BindOutputs(FOutputVertexInterfaceData& InOutVertexData)
+	{
+		using namespace Delay;
+		InOutVertexData.BindReadVertex(METASOUND_GET_PARAM_NAME(OutParamAudio), AudioOutput);
 	}
 
 	FDataReferenceCollection FDelayOperator::GetInputs() const
 	{
 		// This should never be called. Bind(...) is called instead. This method
-		// exists as a stop-gap until the API can be deprecated and removed. 
+		// exists as a stop-gap until the API can be deprecated and removed.
 		checkNoEntry();
-
-		FDataReferenceCollection InputDataReferences;
-		return InputDataReferences;
+		return {};
 	}
 
 	FDataReferenceCollection FDelayOperator::GetOutputs() const
 	{
 		// This should never be called. Bind(...) is called instead. This method
-		// exists as a stop-gap until the API can be deprecated and removed. 
+		// exists as a stop-gap until the API can be deprecated and removed.
 		checkNoEntry();
-
-		using namespace Delay;
-
-		FDataReferenceCollection OutputDataReferences;
-		return OutputDataReferences;
+		return {};
 	}
 
 	float FDelayOperator::GetInputDelayTimeMsec() const
diff --git a/Engine/Plugins/Runtime/Metasound/Source/MetasoundStandardNodes/Private/MetasoundDiffuserNode.cpp b/Engine/Plugins/Runtime/Metasound/Source/MetasoundStandardNodes/Private/MetasoundDiffuserNode.cpp
--- a/Engine/Plugins/Runtime/Metasound/Source/MetasoundStandardNodes/Private/MetasoundDiffuserNode.cpp
+++ b/Engine/Plugins/Runtime/Metasound/Source/MetasoundStandardNodes/Private/MetasoundDiffuserNode.cpp
@@ -117,26 +117,33 @@ namespace Metasound
 			}
 		}
 
-		virtual FDataReferenceCollection GetInputs() const override
+		virtual void Bind(FVertexInterfaceData& InVertexData) const override
 		{
 			using namespace DiffuserNode;
 
-			FDataReferenceCollection InputDataReferences;
-			InputDataReferences.AddDataReadReference(METASOUND_GET_PARAM_NAME(InputAudio), AudioInput);
-			InputDataReferences.AddDataReadReference(METASOUND_GET_PARAM_NAME(InputDiffusionDepth), Depth);
-			InputDataReferences.AddDataReadReference(METASOUND_GET_PARAM_NAME(InputFeedbackGain), Feedback);
+			FInputVertexInterfaceData& Inputs = InVertexData.GetInputs();
+			Inputs.BindReadVertex(METASOUND_GET_PARAM_NAME(InputAudio), AudioInput);
+			Inputs.BindReadVertex(METASOUND_GET_PARAM_NAME(InputDiffusionDepth), Depth);
+			Inputs.BindReadVertex(METASOUND_GET_PARAM_NAME(InputFeedbackGain), Feedback);
 
-			return InputDataReferences;
+			FOutputVertexInterfaceData& Outputs = InVertexData.GetOutputs();
+			Outputs.BindReadVertex(METASOUND_GET_PARAM_NAME(OutputAudio), AudioOutput);
 		}
 
-		virtual FDataReferenceCollection GetOutputs() const override
+		virtual FDataReferenceCollection GetInputs() const override
 		{
-			using namespace DiffuserNode;
-
-			FDataReferenceCollection OutputDataReferences;
-			OutputDataReferences.AddDataReadReference(METASOUND_GET_PARAM_NAME(OutputAudio), AudioOutput);
+			// This should never be called. Bind(...) is called instead. This method
+			// exists as a stop-gap until the API can be deprecated and removed.
+			checkNoEntry();
+			return {};
+		}
 
-			return OutputDataReferences;
+		virtual FDataReferenceCollection GetOutputs() const override
+		{
+			// This should never be called. Bind(...) is called instead. This method
+			// exists as a stop-gap until the API can be deprecated and removed.
+			checkNoEntry();
+			return {};
 		}
 
 		void Reset(const IOperator::FResetParams& InParams)
diff --git a/Engine/Plugins/Runtime/Metasound/Source/MetasoundStandardNodes/Private/MetasoundDopplerPitchShiftNode.cpp b/Engine/Plugins/Runtime/Metasound/Source/MetasoundStandardNodes/Private/MetasoundDopplerPitchShiftNode.cpp
--- a/Engine/Plugins/Runtime/Metasound/Source/MetasoundStandardNodes/Private/MetasoundDopplerPitchShiftNode.cpp
+++ b/Engine/Plugins/Runtime/Metasound/Source/MetasoundStandardNodes/Private/MetasoundDopplerPitchShiftNode.cpp
@@ -32,6 +32,7 @@ namespace Metasound
 
 		FDopplerPitchShiftOperator(const FCreateOperatorParams& InParams, const FAudioBufferReadRef& InAudioInput, const FFloatReadRef& InPitchShift, const FFloatReadRef& InDelayLength);
 
+		virtual void Bind(FVertexInterfaceData& InVertexData) const override;
 		virtual FDataReferenceCollection GetInputs() const override;
 		virtual FDataReferenceCollection GetOutputs() const override;
 		void Reset(const IOperator::FResetParams& InParams);
@@ -59,25 +60,33 @@ namespace Metasound
 		Reset(InParams);
 	}
 
-	FDataReferenceCollection FDopplerPitchShiftOperator::GetInputs() const
+	void FDopplerPitchShiftOperator::Bind(FVertexInterfaceData& InVertexData) const
 	{
 		using namespace DopplerPitchShift;
 
-		FDataReferenceCollection InputDataReferences;
-		InputDataReferences.AddDataReadReference(METASOUND_GET_PARAM_NAME(InParamAudioInput), FAudioBufferReadRef(AudioInput));
-		InputDataReferences.AddDataReadReference(METASOUND_GET_PARAM_NAME(InParamPitchShift), FFloatReadRef(PitchShift));
-		InputDataReferences.AddDataReadReference(METASOUND_GET_PARAM_NAME(InParamDelayLength), FFloatReadRef(DelayLength));
+		FInputVertexInterfaceData& Inputs = InVertexData.GetInputs();
+		Inputs.BindReadVertex(METASOUND_GET_PARAM_NAME(InParamAudioInput), FAudioBufferReadRef(AudioInput));
+		Inputs.BindReadVertex(METASOUND_GET_PARAM_NAME(InParamPitchShift), FFloatReadRef(PitchShift));
+		Inputs.BindReadVertex(METASOUND_GET_PARAM_NAME(InParamDelayLength), FFloatReadRef(DelayLength));
 
-		return InputDataReferences;
+		FOutputVertexInterfaceData& Outputs = InVertexData.GetOutputs();
+		Outputs.BindReadVertex(METASOUND_GET_PARAM_NAME(OutParamAudio), FAudioBufferReadRef(AudioOutput));
 	}
 
-	FDataReferenceCollection FDopplerPitchShiftOperator::GetOutputs() const
+	FDataReferenceCollection FDopplerPitchShiftOperator::GetInputs() const
 	{
-		using namespace DopplerPitchShift;
+		// This should never be called. Bind(...) is called instead. This method
+		// exists as a stop-gap until the API can be deprecated and removed.
+		checkNoEntry();
+		return {};
+	}
 
-		FDataReferenceCollection OutputDataReferences;
-		OutputDataReferences.AddDataReadReference(METASOUND_GET_PARAM_NAME(OutParamAudio), FAudioBufferReadRef(AudioOutput));
-		return OutputDataReferences;
+	FDataReferenceCollection FDopplerPitchShiftOperator::GetOutputs() const
+	{
+		// This should never be called. Bind(...) is called instead. This method
+		// exists as a stop-gap until the API can be deprecated and removed.
+		checkNoEntry();
+		return {};
 	}
 
 	void FDopplerPitchShiftOperator::Reset(const IOperator::FResetParams& InParams)
diff --git a/Engine/Plugins/Runtime/Metasound/Source/MetasoundStandardNodes/Private/MetasoundDynamicFilterNode.cpp b/Engine/Plugins/Runtime/Metasound/Source/MetasoundStandardNodes/Private/MetasoundDynamicFilterNode.cpp
--- a/Engine/Plugins/Runtime/Metasound/Source/MetasoundStandardNodes/Private/MetasoundDynamicFilterNode.cpp
+++ b/Engine/Plugins/Runtime/Metasound/Source/MetasoundStandardNodes/Private/MetasoundDynamicFilterNode.cpp
@@ -123,36 +123,48 @@ namespace Metasound
 			return Metadata;
 		}
 
-		virtual FDataReferenceCollection GetInputs() const override
+
+		virtual void BindInputs(FInputVertexInterfaceData& InOutVertexData) override
 		{
 			using namespace DynamicFilterNode;
 
-			FDataReferenceCollection InputDataReferences;
-			InputDataReferences.AddDataReadReference(METASOUND_GET_PARAM_NAME(InputAudio), AudioInput);
-			InputDataReferences.AddDataReadReference(METASOUND_GET_PARAM_NAME(InputSidechain), SidechainInput);
-			InputDataReferences.AddDataReadReference(METASOUND_GET_PARAM_NAME(InputFilterType), FilterType);
-			InputDataReferences.AddDataReadReference(METASOUND_GET_PARAM_NAME(InputFrequency), Frequency);
-			InputDataReferences.AddDataReadReference(METASOUND_GET_PARAM_NAME(InputQ), Q);
-			InputDataReferences.AddDataReadReference(METASOUND_GET_PARAM_NAME(InputThreshold), ThresholdDB);
-			InputDataReferences.AddDataReadReference(METASOUND_GET_PARAM_NAME(InputRatio), Ratio);
-			InputDataReferences.AddDataReadReference(METASOUND_GET_PARAM_NAME(InputKnee), KneeDB);
-			InputDataReferences.AddDataReadReference(METASOUND_GET_PARAM_NAME(InputRange), Range);
-			InputDataReferences.AddDataReadReference(METASOUND_GET_PARAM_NAME(InputGain), GainDB);
-			InputDataReferences.AddDataReadReference(METASOUND_GET_PARAM_NAME(InputAttackTime), AttackTime);
-			InputDataReferences.AddDataReadReference(METASOUND_GET_PARAM_NAME(InputReleaseTime), ReleaseTime);
-			InputDataReferences.AddDataReadReference(METASOUND_GET_PARAM_NAME(InputEnvelopeMode), EnvelopeMode);
-			InputDataReferences.AddDataReadReference(METASOUND_GET_PARAM_NAME(InputAnalogMode), bAnalogMode);
-
-			return InputDataReferences;
+			InOutVertexData.BindReadVertex(METASOUND_GET_PARAM_NAME(InputAudio), AudioInput);
+			InOutVertexData.BindReadVertex(METASOUND_GET_PARAM_NAME(InputSidechain), SidechainInput);
+			InOutVertexData.BindReadVertex(METASOUND_GET_PARAM_NAME(InputFilterType), FilterType);
+			InOutVertexData.BindReadVertex(METASOUND_GET_PARAM_NAME(InputFrequency), Frequency);
+			InOutVertexData.BindReadVertex(METASOUND_GET_PARAM_NAME(InputQ), Q);
+			InOutVertexData.BindReadVertex(METASOUND_GET_PARAM_NAME(InputThreshold), ThresholdDB);
+			InOutVertexData.BindReadVertex(METASOUND_GET_PARAM_NAME(InputRatio), Ratio);
+			InOutVertexData.BindReadVertex(METASOUND_GET_PARAM_NAME(InputKnee), KneeDB);
+			InOutVertexData.BindReadVertex(METASOUND_GET_PARAM_NAME(InputRange), Range);
+			InOutVertexData.BindReadVertex(METASOUND_GET_PARAM_NAME(InputGain), GainDB);
+			InOutVertexData.BindReadVertex(METASOUND_GET_PARAM_NAME(InputAttackTime), AttackTime);
+			InOutVertexData.BindReadVertex(METASOUND_GET_PARAM_NAME(InputReleaseTime), ReleaseTime);
+			InOutVertexData.BindReadVertex(METASOUND_GET_PARAM_NAME(InputEnvelopeMode), EnvelopeMode);
+			InOutVertexData.BindReadVertex(METASOUND_GET_PARAM_NAME(InputAnalogMode), bAnalogMode);
 		}
 
-		virtual FDataReferenceCollection GetOutputs() const override
+		virtual void BindOutputs(FOutputVertexInterfaceData& InOutVertexData) override
 		{
 			using namespace DynamicFilterNode;
 
-			FDataReferenceCollection OutputDataReferences;
-			OutputDataReferences.AddDataReadReference(METASOUND_GET_PARAM_NAME(OutputAudio), AudioOutput);
-			return OutputDataReferences;
+			InOutVertexData.BindReadVertex(METASOUND_GET_PARAM_NAME(OutputAudio), AudioOutput);
+		}
+
+		virtual FDataReferenceCollection GetInputs() const override
+		{
+			// This should never be called. Bind(...) is called instead. This method
+			// exists as a stop-gap until the API can be deprecated and removed.
+			checkNoEntry();
+			return {};
+		}
+
+		virtual FDataReferenceCollection GetOutputs() const override
+		{
+			// This should never be called. Bind(...) is called instead. This method
+			// exists as a stop-gap until the API can be deprecated and removed.
+			checkNoEntry();
+			return {};
 		}
 
 		static const FVertexInterface& DeclareVertexInterface()
diff --git a/Engine/Plugins/Runtime/Metasound/Source/MetasoundStandardNodes/Private/MetasoundEnvelopeFollowerNode.cpp b/Engine/Plugins/Runtime/Metasound/Source/MetasoundStandardNodes/Private/MetasoundEnvelopeFollowerNode.cpp
--- a/Engine/Plugins/Runtime/Metasound/Source/MetasoundStandardNodes/Private/MetasoundEnvelopeFollowerNode.cpp
+++ b/Engine/Plugins/Runtime/Metasound/Source/MetasoundStandardNodes/Private/MetasoundEnvelopeFollowerNode.cpp
@@ -48,6 +48,8 @@ namespace Metasound
 			const FTimeReadRef& InReleaseTime,
 			const FEnvelopePeakModeReadRef& InEnvelopeMode);
 
+		virtual void BindInputs(FInputVertexInterfaceData& InOutVertexData) override;
+		virtual void BindOutputs(FOutputVertexInterfaceData& InOutVertexData) override;
 		virtual FDataReferenceCollection GetInputs() const override;
 		virtual FDataReferenceCollection GetOutputs() const override;
 		void Reset(const IOperator::FResetParams& InParams);
@@ -93,27 +95,38 @@ namespace Metasound
 		Reset(InParams);
 	}
 
-	FDataReferenceCollection FEnvelopeFollowerOperator::GetInputs() const
+	void FEnvelopeFollowerOperator::BindInputs(FInputVertexInterfaceData& InOutVertexData)
 	{
 		using namespace EnvelopeFollowerVertexNames;
 
-		FDataReferenceCollection InputDataReferences;
-		InputDataReferences.AddDataReadReference(METASOUND_GET_PARAM_NAME(InParamAudioInput), AudioInput);
-		InputDataReferences.AddDataReadReference(METASOUND_GET_PARAM_NAME(InParamAttackTime), AttackTimeInput);
-		InputDataReferences.AddDataReadReference(METASOUND_GET_PARAM_NAME(InParamReleaseTime), ReleaseTimeInput);
-		InputDataReferences.AddDataReadReference(METASOUND_GET_PARAM_NAME(InParamFollowMode), FollowModeInput);
-
-		return InputDataReferences;
+		InOutVertexData.BindReadVertex(METASOUND_GET_PARAM_NAME(InParamAudioInput), AudioInput);
+		InOutVertexData.BindReadVertex(METASOUND_GET_PARAM_NAME(InParamAttackTime), AttackTimeInput);
+		InOutVertexData.BindReadVertex(METASOUND_GET_PARAM_NAME(InParamReleaseTime), ReleaseTimeInput);
+		InOutVertexData.BindReadVertex(METASOUND_GET_PARAM_NAME(InParamFollowMode), FollowModeInput);
 	}
 
-	FDataReferenceCollection FEnvelopeFollowerOperator::GetOutputs() const
+	void FEnvelopeFollowerOperator::BindOutputs(FOutputVertexInterfaceData& InOutVertexData)
 	{
 		using namespace EnvelopeFollowerVertexNames;
 
-		FDataReferenceCollection OutputDataReferences;
-		OutputDataReferences.AddDataReadReference(METASOUND_GET_PARAM_NAME(OutParamEnvelope), EnvelopeFloatOutput);
-		OutputDataReferences.AddDataReadReference(METASOUND_GET_PARAM_NAME(OutputAudioEnvelope), EnvelopeAudioOutput);
-		return OutputDataReferences;
+		InOutVertexData.BindReadVertex(METASOUND_GET_PARAM_NAME(OutParamEnvelope), EnvelopeFloatOutput);
+		InOutVertexData.BindReadVertex(METASOUND_GET_PARAM_NAME(OutputAudioEnvelope), EnvelopeAudioOutput);
+	}
+
+	FDataReferenceCollection FEnvelopeFollowerOperator::GetInputs() const
+	{
+		// This should never be called. Bind(...) is called instead. This method
+		// exists as a stop-gap until the API can be deprecated and removed.
+		checkNoEntry();
+		return {};
+	}
+
+	FDataReferenceCollection FEnvelopeFollowerOperator::GetOutputs() const
+	{
+		// This should never be called. Bind(...) is called instead. This method
+		// exists as a stop-gap until the API can be deprecated and removed.
+		checkNoEntry();
+		return {};
 	}
 
 	void FEnvelopeFollowerOperator::Reset(const IOperator::FResetParams& InParams)
diff --git a/Engine/Plugins/Runtime/Metasound/Source/MetasoundStandardNodes/Private/MetasoundFilterQToBandwidthNode.cpp b/Engine/Plugins/Runtime/Metasound/Source/MetasoundStandardNodes/Private/MetasoundFilterQToBandwidthNode.cpp
--- a/Engine/Plugins/Runtime/Metasound/Source/MetasoundStandardNodes/Private/MetasoundFilterQToBandwidthNode.cpp
+++ b/Engine/Plugins/Runtime/Metasound/Source/MetasoundStandardNodes/Private/MetasoundFilterQToBandwidthNode.cpp
@@ -87,24 +87,33 @@ namespace Metasound
 
 		virtual ~FConvertQToBandwidthNodeOperator() = default;
 
-		virtual FDataReferenceCollection GetInputs() const override
+		virtual void BindInputs(FInputVertexInterfaceData& InOutVertexData) override
 		{
 			using namespace QToBandwidth;
 
-			FDataReferenceCollection Inputs;
-			Inputs.AddDataReadReference(METASOUND_GET_PARAM_NAME(InputQ), InValue);
-
-			return Inputs;
+			InOutVertexData.BindReadVertex(METASOUND_GET_PARAM_NAME(InputQ), InValue);
 		}
 
-		virtual FDataReferenceCollection GetOutputs() const override
+		virtual void BindOutputs(FOutputVertexInterfaceData& InOutVertexData) override
 		{
 			using namespace QToBandwidth;
+			InOutVertexData.BindReadVertex(METASOUND_GET_PARAM_NAME(OutputBandwidth), OutValue);
+		}
 
-			FDataReferenceCollection Outputs;
-			Outputs.AddDataReadReference(METASOUND_GET_PARAM_NAME(OutputBandwidth), OutValue);
+		virtual FDataReferenceCollection GetInputs() const override
+		{
+			// This should never be called. Bind(...) is called instead. This method
+			// exists as a stop-gap until the API can be deprecated and removed.
+			checkNoEntry();
+			return {};
+		}
 
-			return Outputs;
+		virtual FDataReferenceCollection GetOutputs() const override
+		{
+			// This should never be called. Bind(...) is called instead. This method
+			// exists as a stop-gap until the API can be deprecated and removed.
+			checkNoEntry();
+			return {};
 		}
 
 		void Reset(const IOperator::FResetParams& InParams)
diff --git a/Engine/Plugins/Runtime/Metasound/Source/MetasoundStandardNodes/Private/MetasoundFlangerNode.cpp b/Engine/Plugins/Runtime/Metasound/Source/MetasoundStandardNodes/Private/MetasoundFlangerNode.cpp
--- a/Engine/Plugins/Runtime/Metasound/Source/MetasoundStandardNodes/Private/MetasoundFlangerNode.cpp
+++ b/Engine/Plugins/Runtime/Metasound/Source/MetasoundStandardNodes/Private/MetasoundFlangerNode.cpp
@@ -109,27 +109,37 @@ namespace Metasound
 			Reset(InParams);
 		}
 
-		FDataReferenceCollection GetInputs() const
+		virtual void BindInputs(FInputVertexInterfaceData& InOutVertexData) override
 		{
 			using namespace FlangerVertexNames;
 
-			FDataReferenceCollection InputDataReferences;
-			InputDataReferences.AddDataReadReference(METASOUND_GET_PARAM_NAME(InputAudio), AudioIn);
-			InputDataReferences.AddDataReadReference(METASOUND_GET_PARAM_NAME(InputModulationRate), ModulationRate);
-			InputDataReferences.AddDataReadReference(METASOUND_GET_PARAM_NAME(InputModulationDepth), ModulationDepth);
-			InputDataReferences.AddDataReadReference(METASOUND_GET_PARAM_NAME(InputCenterDelay), CenterDelay);
-			InputDataReferences.AddDataReadReference(METASOUND_GET_PARAM_NAME(InputMixLevel), MixLevel);
-
-			return InputDataReferences;
+			InOutVertexData.BindReadVertex(METASOUND_GET_PARAM_NAME(InputAudio), AudioIn);
+			InOutVertexData.BindReadVertex(METASOUND_GET_PARAM_NAME(InputModulationRate), ModulationRate);
+			InOutVertexData.BindReadVertex(METASOUND_GET_PARAM_NAME(InputModulationDepth), ModulationDepth);
+			InOutVertexData.BindReadVertex(METASOUND_GET_PARAM_NAME(InputCenterDelay), CenterDelay);
+			InOutVertexData.BindReadVertex(METASOUND_GET_PARAM_NAME(InputMixLevel), MixLevel);
 		}
 
-		virtual FDataReferenceCollection GetOutputs() const
+		virtual void BindOutputs(FOutputVertexInterfaceData& InOutVertexData) override
 		{
 			using namespace FlangerVertexNames;
+			InOutVertexData.BindReadVertex(METASOUND_GET_PARAM_NAME(OutputAudio), AudioOut);
+		}
+
+		virtual FDataReferenceCollection GetInputs() const override
+		{
+			// This should never be called. Bind(...) is called instead. This method
+			// exists as a stop-gap until the API can be deprecated and removed.
+			checkNoEntry();
+			return {};
+		}
 
-			FDataReferenceCollection OutputDataReferences;
-			OutputDataReferences.AddDataReadReference(METASOUND_GET_PARAM_NAME(OutputAudio), AudioOut);
-			return OutputDataReferences;
+		virtual FDataReferenceCollection GetOutputs() const override
+		{
+			// This should never be called. Bind(...) is called instead. This method
+			// exists as a stop-gap until the API can be deprecated and removed.
+			checkNoEntry();
+			return {};
 		}
 
 		void Reset(const IOperator::FResetParams& InParams)
diff --git a/Engine/Plugins/Runtime/Metasound/Source/MetasoundStandardNodes/Private/MetasoundFreqMultiplierToSemitoneNode.cpp b/Engine/Plugins/Runtime/Metasound/Source/MetasoundStandardNodes/Private/MetasoundFreqMultiplierToSemitoneNode.cpp
--- a/Engine/Plugins/Runtime/Metasound/Source/MetasoundStandardNodes/Private/MetasoundFreqMultiplierToSemitoneNode.cpp
+++ b/Engine/Plugins/Runtime/Metasound/Source/MetasoundStandardNodes/Private/MetasoundFreqMultiplierToSemitoneNode.cpp
@@ -33,6 +33,8 @@ namespace Metasound
 
 		FFrequencyMultiplierToSemitoneOperator(const FOperatorSettings& InSettings, const FFloatReadRef& InFrequencyMultiplier);
 
+		virtual void BindInputs(FInputVertexInterfaceData& InOutVertexData) override;
+		virtual void BindOutputs(FOutputVertexInterfaceData& InOutVertexData) override;
 		virtual FDataReferenceCollection GetInputs() const override;
 		virtual FDataReferenceCollection GetOutputs() const override;
 		void Reset(const IOperator::FResetParams& InParams);
@@ -56,24 +58,33 @@ namespace Metasound
 	{
 	}
 
-	
-	FDataReferenceCollection FFrequencyMultiplierToSemitoneOperator::GetInputs() const
+
+	void FFrequencyMultiplierToSemitoneOperator::BindInputs(FInputVertexInterfaceData& InOutVertexData)
 	{
 		using namespace FrequencyMultiplierToSemitoneVertexNames;
+		InOutVertexData.BindReadVertex(METASOUND_GET_PARAM_NAME(InputFrequencyMultiplier), FrequencyMultiplierInput);
+	}
 
-		FDataReferenceCollection InputDataReferences;
-		InputDataReferences.AddDataReadReference(METASOUND_GET_PARAM_NAME(InputFrequencyMultiplier), FrequencyMultiplierInput);
+	void FFrequencyMultiplierToSemitoneOperator::BindOutputs(FOutputVertexInterfaceData& InOutVertexData)
+	{
+		using namespace FrequencyMultiplierToSemitoneVertexNames;
+		InOutVertexData.BindReadVertex(METASOUND_GET_PARAM_NAME(OutputSemitone), SemitoneOutput);
+	}
 
-		return InputDataReferences;
+	FDataReferenceCollection FFrequencyMultiplierToSemitoneOperator::GetInputs() const
+	{
+		// This should never be called. Bind(...) is called instead. This method
+		// exists as a stop-gap until the API can be deprecated and removed.
+		checkNoEntry();
+		return {};
 	}
 
 	FDataReferenceCollection FFrequencyMultiplierToSemitoneOperator::GetOutputs() const
 	{
-		using namespace FrequencyMultiplierToSemitoneVertexNames;
-
-		FDataReferenceCollection OutputDataReferences;
-		OutputDataReferences.AddDataReadReference(METASOUND_GET_PARAM_NAME(OutputSemitone), SemitoneOutput);
-		return OutputDataReferences;
+		// This should never be called. Bind(...) is called instead. This method
+		// exists as a stop-gap until the API can be deprecated and removed.
+		checkNoEntry();
+		return {};
 	}
 
 	void FFrequencyMultiplierToSemitoneOperator::Reset(const IOperator::FResetParams& InParams)
diff --git a/Engine/Plugins/Runtime/Metasound/Source/MetasoundStandardNodes/Private/MetasoundFreqToMIDINode.cpp b/Engine/Plugins/Runtime/Metasound/Source/MetasoundStandardNodes/Private/MetasoundFreqToMIDINode.cpp
--- a/Engine/Plugins/Runtime/Metasound/Source/MetasoundStandardNodes/Private/MetasoundFreqToMIDINode.cpp
+++ b/Engine/Plugins/Runtime/Metasound/Source/MetasoundStandardNodes/Private/MetasoundFreqToMIDINode.cpp
@@ -30,6 +30,8 @@ namespace Metasound
 
 		FFreqToMidiOperator(const FOperatorSettings& InSettings, const FFloatReadRef& InMidiNote);
 
+		virtual void BindInputs(FInputVertexInterfaceData& InOutVertexData) override;
+		virtual void BindOutputs(FOutputVertexInterfaceData& InOutVertexData) override;
 		virtual FDataReferenceCollection GetInputs() const override;
 		virtual FDataReferenceCollection GetOutputs() const override;
 		void Reset(const IOperator::FResetParams& InParams);
@@ -52,24 +54,33 @@ namespace Metasound
 		, PrevFreq(*InFreq)
 	{
 	}
-	
-	FDataReferenceCollection FFreqToMidiOperator::GetInputs() const
+
+	void FFreqToMidiOperator::BindInputs(FInputVertexInterfaceData& InOutVertexData)
 	{
 		using namespace FrequencyToMidiVertexNames;
+		InOutVertexData.BindReadVertex(METASOUND_GET_PARAM_NAME(InputFreq), FreqInput);
+	}
 
-		FDataReferenceCollection InputDataReferences;
-		InputDataReferences.AddDataReadReference(METASOUND_GET_PARAM_NAME(InputFreq), FreqInput);
+	void FFreqToMidiOperator::BindOutputs(FOutputVertexInterfaceData& InOutVertexData)
+	{
+		using namespace FrequencyToMidiVertexNames;
+		InOutVertexData.BindReadVertex(METASOUND_GET_PARAM_NAME(OutputMidi), MidiOutput);
+	}
 
-		return InputDataReferences;
+	FDataReferenceCollection FFreqToMidiOperator::GetInputs() const
+	{
+		// This should never be called. Bind(...) is called instead. This method
+		// exists as a stop-gap until the API can be deprecated and removed.
+		checkNoEntry();
+		return {};
 	}
 
 	FDataReferenceCollection FFreqToMidiOperator::GetOutputs() const
 	{
-		using namespace FrequencyToMidiVertexNames;
-
-		FDataReferenceCollection OutputDataReferences;
-		OutputDataReferences.AddDataReadReference(METASOUND_GET_PARAM_NAME(OutputMidi), MidiOutput);
-		return OutputDataReferences;
+		// This should never be called. Bind(...) is called instead. This method
+		// exists as a stop-gap until the API can be deprecated and removed.
+		checkNoEntry();
+		return {};
 	}
 
 	void FFreqToMidiOperator::Reset(const IOperator::FResetParams& InParams)
diff --git a/Engine/Plugins/Runtime/Metasound/Source/MetasoundStandardNodes/Private/MetasoundGrainDelayNode.cpp b/Engine/Plugins/Runtime/Metasound/Source/MetasoundStandardNodes/Private/MetasoundGrainDelayNode.cpp
--- a/Engine/Plugins/Runtime/Metasound/Source/MetasoundStandardNodes/Private/MetasoundGrainDelayNode.cpp
+++ b/Engine/Plugins/Runtime/Metasound/Source/MetasoundStandardNodes/Private/MetasoundGrainDelayNode.cpp
@@ -67,6 +67,8 @@ namespace Metasound
 		static const FNodeClassMetadata& GetNodeInfo();
 		static const FVertexInterface& GetVertexInterface();
 		static TUniquePtr<IOperator> CreateOperator(const FCreateOperatorParams& InParams, FBuildErrorArray& OutErrors);
+		virtual void BindInputs(FInputVertexInterfaceData& InOutVertexData) override;
+		virtual void BindOutputs(FOutputVertexInterfaceData& InOutVertexData) override;
 		virtual FDataReferenceCollection GetInputs() const override;
 		virtual FDataReferenceCollection GetOutputs() const override;
 		void Reset(const IOperator::FResetParams& InParams);
@@ -218,34 +220,45 @@ namespace Metasound
 			GrainMaxCount,
 			GrainDelayFeedbackAmount);
 	}
-	
-	FDataReferenceCollection FGrainDelayOperator::GetOutputs() const
+
+	void FGrainDelayOperator::BindInputs(FInputVertexInterfaceData& InOutVertexData)
 	{
 		using namespace GrainDelay;
 
-		FDataReferenceCollection OutputDataReferences;
-		OutputDataReferences.AddDataReadReference(METASOUND_GET_PARAM_NAME(OutAudio), FAudioBufferReadRef(AudioOutput));
-		return OutputDataReferences;
+		InOutVertexData.BindReadVertex(METASOUND_GET_PARAM_NAME(InAudio), AudioInput);
+		InOutVertexData.BindReadVertex(METASOUND_GET_PARAM_NAME(InGrainSpawnTrigger), GrainSpawnTrigger);
+		InOutVertexData.BindReadVertex(METASOUND_GET_PARAM_NAME(InGrainDelay), GrainDelay);
+		InOutVertexData.BindReadVertex(METASOUND_GET_PARAM_NAME(InGrainDelayDelta), GrainDelayDelta);
+		InOutVertexData.BindReadVertex(METASOUND_GET_PARAM_NAME(InGrainDuration), GrainDuration);
+		InOutVertexData.BindReadVertex(METASOUND_GET_PARAM_NAME(InGrainDurationDelta), GrainDurationDelta);
+		InOutVertexData.BindReadVertex(METASOUND_GET_PARAM_NAME(InGrainPitchShift), GrainPitchShift);
+		InOutVertexData.BindReadVertex(METASOUND_GET_PARAM_NAME(InGrainPitchShiftDelta), GrainPitchShiftDelta);
+		InOutVertexData.BindReadVertex(METASOUND_GET_PARAM_NAME(InGrainDelayEnvelope), GrainDelayEnvelope);
+		InOutVertexData.BindReadVertex(METASOUND_GET_PARAM_NAME(InGrainMaxCount), GrainMaxCount);
+		InOutVertexData.BindReadVertex(METASOUND_GET_PARAM_NAME(InGrainDelayFeedbackAmount), GrainDelayFeedbackAmount);
 	}
-	
-	FDataReferenceCollection FGrainDelayOperator::GetInputs() const
+
+	void FGrainDelayOperator::BindOutputs(FOutputVertexInterfaceData& InOutVertexData)
 	{
 		using namespace GrainDelay;
 
-		FDataReferenceCollection InputDataReferences;
-		InputDataReferences.AddDataReadReference(METASOUND_GET_PARAM_NAME(InAudio), FAudioBufferReadRef(AudioInput));
-		InputDataReferences.AddDataReadReference(METASOUND_GET_PARAM_NAME(InGrainSpawnTrigger), FTriggerReadRef(GrainSpawnTrigger));
-		InputDataReferences.AddDataReadReference(METASOUND_GET_PARAM_NAME(InGrainDelay), FFloatReadRef(GrainDelay));
-		InputDataReferences.AddDataReadReference(METASOUND_GET_PARAM_NAME(InGrainDelayDelta), FFloatReadRef(GrainDelayDelta));
-		InputDataReferences.AddDataReadReference(METASOUND_GET_PARAM_NAME(InGrainDuration), FFloatReadRef(GrainDuration));
-		InputDataReferences.AddDataReadReference(METASOUND_GET_PARAM_NAME(InGrainDurationDelta), FFloatReadRef(GrainDurationDelta));
-		InputDataReferences.AddDataReadReference(METASOUND_GET_PARAM_NAME(InGrainPitchShift), FFloatReadRef(GrainPitchShift));
-		InputDataReferences.AddDataReadReference(METASOUND_GET_PARAM_NAME(InGrainPitchShiftDelta), FFloatReadRef(GrainPitchShiftDelta));
-		InputDataReferences.AddDataReadReference(METASOUND_GET_PARAM_NAME(InGrainDelayEnvelope), FGrainDelayEnvelopeReadRef(GrainDelayEnvelope));
-		InputDataReferences.AddDataReadReference(METASOUND_GET_PARAM_NAME(InGrainMaxCount), FInt32ReadRef(GrainMaxCount));
-		InputDataReferences.AddDataReadReference(METASOUND_GET_PARAM_NAME(InGrainDelayFeedbackAmount), FFloatReadRef(GrainDelayFeedbackAmount));
-
-		return InputDataReferences;
+		InOutVertexData.BindReadVertex(METASOUND_GET_PARAM_NAME(OutAudio), AudioOutput);
+	}
+
+	FDataReferenceCollection FGrainDelayOperator::GetInputs() const
+	{
+		// This should never be called. Bind(...) is called instead. This method
+		// exists as a stop-gap until the API can be deprecated and removed.
+		checkNoEntry();
+		return {};
+	}
+
+	FDataReferenceCollection FGrainDelayOperator::GetOutputs() const
+	{
+		// This should never be called. Bind(...) is called instead. This method
+		// exists as a stop-gap until the API can be deprecated and removed.
+		checkNoEntry();
+		return {};
 	}
 
 	void FGrainDelayOperator::Reset(const IOperator::FResetParams& InParams)
diff --git a/Engine/Plugins/Runtime/Metasound/Source/MetasoundStandardNodes/Private/MetasoundITDPannerNode.cpp b/Engine/Plugins/Runtime/Metasound/Source/MetasoundStandardNodes/Private/MetasoundITDPannerNode.cpp
--- a/Engine/Plugins/Runtime/Metasound/Source/MetasoundStandardNodes/Private/MetasoundITDPannerNode.cpp
+++ b/Engine/Plugins/Runtime/Metasound/Source/MetasoundStandardNodes/Private/MetasoundITDPannerNode.cpp
@@ -46,6 +46,8 @@ namespace Metasound
 			const FFloatReadRef& InDistanceFactor,
 			const FFloatReadRef& InHeadWidth);
 
+		virtual void BindInputs(FInputVertexInterfaceData& InOutVertexData) override;
+		virtual void BindOutputs(FOutputVertexInterfaceData& InOutVertexData) override;
 		virtual FDataReferenceCollection GetInputs() const override;
 		virtual FDataReferenceCollection GetOutputs() const override;
 		void Reset(const IOperator::FResetParams& InParams);
@@ -108,28 +110,38 @@ namespace Metasound
 		PrevRightGain = CurrRightGain;
 	}
 
-	FDataReferenceCollection FITDPannerOperator::GetInputs() const
+	void FITDPannerOperator::BindInputs(FInputVertexInterfaceData& InOutVertexData)
 	{
 		using namespace ITDPannerVertexNames;
 
-		FDataReferenceCollection InputDataReferences;
-		InputDataReferences.AddDataReadReference(METASOUND_GET_PARAM_NAME(InputAudio), AudioInput);
-		InputDataReferences.AddDataReadReference(METASOUND_GET_PARAM_NAME(InputPanAngle), PanningAngle);
-		InputDataReferences.AddDataReadReference(METASOUND_GET_PARAM_NAME(InputDistanceFactor), DistanceFactor);
-		InputDataReferences.AddDataReadReference(METASOUND_GET_PARAM_NAME(InputHeadWidth), HeadWidth);
-
-		return InputDataReferences;
+		InOutVertexData.BindReadVertex(METASOUND_GET_PARAM_NAME(InputAudio), AudioInput);
+		InOutVertexData.BindReadVertex(METASOUND_GET_PARAM_NAME(InputPanAngle), PanningAngle);
+		InOutVertexData.BindReadVertex(METASOUND_GET_PARAM_NAME(InputDistanceFactor), DistanceFactor);
+		InOutVertexData.BindReadVertex(METASOUND_GET_PARAM_NAME(InputHeadWidth), HeadWidth);
 	}
 
-	FDataReferenceCollection FITDPannerOperator::GetOutputs() const
+	void FITDPannerOperator::BindOutputs(FOutputVertexInterfaceData& InOutVertexData)
 	{
 		using namespace ITDPannerVertexNames;
 
-		FDataReferenceCollection OutputDataReferences;
-		OutputDataReferences.AddDataReadReference(METASOUND_GET_PARAM_NAME(OutputAudioLeft), AudioLeftOutput);
-		OutputDataReferences.AddDataReadReference(METASOUND_GET_PARAM_NAME(OutputAudioRight), AudioRightOutput);
+		InOutVertexData.BindReadVertex(METASOUND_GET_PARAM_NAME(OutputAudioLeft), AudioLeftOutput);
+		InOutVertexData.BindReadVertex(METASOUND_GET_PARAM_NAME(OutputAudioRight), AudioRightOutput);
+	}
+
+	FDataReferenceCollection FITDPannerOperator::GetInputs() const
+	{
+		// This should never be called. Bind(...) is called instead. This method
+		// exists as a stop-gap until the API can be deprecated and removed.
+		checkNoEntry();
+		return {};
+	}
 
-		return OutputDataReferences;
+	FDataReferenceCollection FITDPannerOperator::GetOutputs() const
+	{
+		// This should never be called. Bind(...) is called instead. This method
+		// exists as a stop-gap until the API can be deprecated and removed.
+		checkNoEntry();
+		return {};
 	}
 
 	void FITDPannerOperator::UpdateParams(bool bIsInit)
diff --git a/Engine/Plugins/Runtime/Metasound/Source/MetasoundStandardNodes/Private/MetasoundInterpToNode.cpp b/Engine/Plugins/Runtime/Metasound/Source/MetasoundStandardNodes/Private/MetasoundInterpToNode.cpp
--- a/Engine/Plugins/Runtime/Metasound/Source/MetasoundStandardNodes/Private/MetasoundInterpToNode.cpp
+++ b/Engine/Plugins/Runtime/Metasound/Source/MetasoundStandardNodes/Private/MetasoundInterpToNode.cpp
@@ -46,6 +46,8 @@ namespace Metasound
 
 		FInterpToOperator(const FCreateOperatorParams& InSettings, const FFloatReadRef& InTargetValue, const FTimeReadRef& InInterpTime);
 
+		virtual void BindInputs(FInputVertexInterfaceData& InOutVertexData) override;
+		virtual void BindOutputs(FOutputVertexInterfaceData& InOutVertexData) override;
 		virtual FDataReferenceCollection GetInputs() const override;
 		virtual FDataReferenceCollection GetOutputs() const override;
 		void Reset(const IOperator::FResetParams& InParams);
@@ -79,24 +81,35 @@ namespace Metasound
 		Reset(InParams);
 	}
 
-	FDataReferenceCollection FInterpToOperator::GetInputs() const
+	void FInterpToOperator::BindInputs(FInputVertexInterfaceData& InOutVertexData)
 	{
 		using namespace InterpToVertexNames;
 
-		FDataReferenceCollection InputDataReferences;
-		InputDataReferences.AddDataReadReference(METASOUND_GET_PARAM_NAME(InParamTarget), FFloatReadRef(TargetValue));
-		InputDataReferences.AddDataReadReference(METASOUND_GET_PARAM_NAME(InParamInterpTime), FTimeReadRef(InterpTime));
-
-		return InputDataReferences;
+		InOutVertexData.BindReadVertex(METASOUND_GET_PARAM_NAME(InParamTarget), TargetValue);
+		InOutVertexData.BindReadVertex(METASOUND_GET_PARAM_NAME(InParamInterpTime), InterpTime);
 	}
 
-	FDataReferenceCollection FInterpToOperator::GetOutputs() const
+	void FInterpToOperator::BindOutputs(FOutputVertexInterfaceData& InOutVertexData)
 	{
 		using namespace InterpToVertexNames;
 
-		FDataReferenceCollection OutputDataReferences;
-		OutputDataReferences.AddDataReadReference(METASOUND_GET_PARAM_NAME(OutParamValue), FFloatReadRef(ValueOutput));
-		return OutputDataReferences;
+		InOutVertexData.BindReadVertex(METASOUND_GET_PARAM_NAME(OutParamValue), ValueOutput);
+	}
+
+	FDataReferenceCollection FInterpToOperator::GetInputs() const
+	{
+		// This should never be called. Bind(...) is called instead. This method
+		// exists as a stop-gap until the API can be deprecated and removed.
+		checkNoEntry();
+		return {};
+	}
+
+	FDataReferenceCollection FInterpToOperator::GetOutputs() const
+	{
+		// This should never be called. Bind(...) is called instead. This method
+		// exists as a stop-gap until the API can be deprecated and removed.
+		checkNoEntry();
+		return {};
 	}
 
 	void FInterpToOperator::Reset(const IOperator::FResetParams& InParams)
diff --git a/Engine/Plugins/Runtime/Metasound/Source/MetasoundStandardNodes/Private/MetasoundLimiterNode.cpp b/Engine/Plugins/Runtime/Metasound/Source/MetasoundStandardNodes/Private/MetasoundLimiterNode.cpp
--- a/Engine/Plugins/Runtime/Metasound/Source/MetasoundStandardNodes/Private/MetasoundLimiterNode.cpp
+++ b/Engine/Plugins/Runtime/Metasound/Source/MetasoundStandardNodes/Private/MetasoundLimiterNode.cpp
@@ -120,30 +120,38 @@ namespace Metasound
 			return Interface;
 		}
 
-		virtual FDataReferenceCollection GetInputs() const override
+		virtual void BindInputs(FInputVertexInterfaceData& InOutVertexData) override
 		{
 			using namespace LimiterVertexNames;
 
-			FDataReferenceCollection InputDataReferences;
-
-			InputDataReferences.AddDataReadReference(METASOUND_GET_PARAM_NAME(InputAudio), AudioInput);
-			InputDataReferences.AddDataReadReference(METASOUND_GET_PARAM_NAME(InputInGainDb), InGainDbInput);
-			InputDataReferences.AddDataReadReference(METASOUND_GET_PARAM_NAME(InputThresholdDb), ThresholdDbInput);
-			InputDataReferences.AddDataReadReference(METASOUND_GET_PARAM_NAME(InputReleaseTime), ReleaseTimeInput);
-			InputDataReferences.AddDataReadReference(METASOUND_GET_PARAM_NAME(InputKneeMode), KneeModeInput);
-
-			return InputDataReferences;
+			InOutVertexData.BindReadVertex(METASOUND_GET_PARAM_NAME(InputAudio), AudioInput);
+			InOutVertexData.BindReadVertex(METASOUND_GET_PARAM_NAME(InputInGainDb), InGainDbInput);
+			InOutVertexData.BindReadVertex(METASOUND_GET_PARAM_NAME(InputThresholdDb), ThresholdDbInput);
+			InOutVertexData.BindReadVertex(METASOUND_GET_PARAM_NAME(InputReleaseTime), ReleaseTimeInput);
+			InOutVertexData.BindReadVertex(METASOUND_GET_PARAM_NAME(InputKneeMode), KneeModeInput);
 		}
 
-		virtual FDataReferenceCollection GetOutputs() const override
+		virtual void BindOutputs(FOutputVertexInterfaceData& InOutVertexData) override
 		{
 			using namespace LimiterVertexNames;
 
-			FDataReferenceCollection OutputDataReferences;
+			InOutVertexData.BindReadVertex(METASOUND_GET_PARAM_NAME(OutputAudio), AudioOutput);
+		}
 
-			OutputDataReferences.AddDataReadReference(METASOUND_GET_PARAM_NAME(OutputAudio), AudioOutput);
+		virtual FDataReferenceCollection GetInputs() const override
+		{
+			// This should never be called. Bind(...) is called instead. This method
+			// exists as a stop-gap until the API can be deprecated and removed.
+			checkNoEntry();
+			return {};
+		}
 
-			return OutputDataReferences;
+		virtual FDataReferenceCollection GetOutputs() const override
+		{
+			// This should never be called. Bind(...) is called instead. This method
+			// exists as a stop-gap until the API can be deprecated and removed.
+			checkNoEntry();
+			return {};
 		}
 
 		static TUniquePtr<IOperator> CreateOperator(const FCreateOperatorParams& InParams, FBuildErrorArray& OutErrors)
diff --git a/Engine/Plugins/Runtime/Metasound/Source/MetasoundStandardNodes/Private/MetasoundLinearGainToDecibels.cpp b/Engine/Plugins/Runtime/Metasound/Source/MetasoundStandardNodes/Private/MetasoundLinearGainToDecibels.cpp
--- a/Engine/Plugins/Runtime/Metasound/Source/MetasoundStandardNodes/Private/MetasoundLinearGainToDecibels.cpp
+++ b/Engine/Plugins/Runtime/Metasound/Source/MetasoundStandardNodes/Private/MetasoundLinearGainToDecibels.cpp
@@ -31,6 +31,8 @@ namespace Metasound
 
 		FLinearGainToDecibelsOperator(const FCreateOperatorParams& InParams, const FFloatReadRef& InLinearGain);
 
+		virtual void BindInputs(FInputVertexInterfaceData& InOutVertexData) override;
+		virtual void BindOutputs(FOutputVertexInterfaceData& InOutVertexData) override;
 		virtual FDataReferenceCollection GetInputs() const override;
 		virtual FDataReferenceCollection GetOutputs() const override;
 		void Reset(const IOperator::FResetParams& InParams);
@@ -51,24 +53,34 @@ namespace Metasound
 		Reset(InParams);
 	}
 
-	
-	FDataReferenceCollection FLinearGainToDecibelsOperator::GetInputs() const
+	void FLinearGainToDecibelsOperator::BindInputs(FInputVertexInterfaceData& InOutVertexData)
 	{
 		using namespace LinearGainToDecibelsVertexNames;
 
-		FDataReferenceCollection InputDataReferences;
-		InputDataReferences.AddDataReadReference(METASOUND_GET_PARAM_NAME(InputLinearGain), LinearGainInput);
-
-		return InputDataReferences;
+		InOutVertexData.BindReadVertex(METASOUND_GET_PARAM_NAME(InputLinearGain), LinearGainInput);
 	}
 
-	FDataReferenceCollection FLinearGainToDecibelsOperator::GetOutputs() const
+	void FLinearGainToDecibelsOperator::BindOutputs(FOutputVertexInterfaceData& InOutVertexData)
 	{
 		using namespace LinearGainToDecibelsVertexNames;
 
-		FDataReferenceCollection OutputDataReferences;
-		OutputDataReferences.AddDataReadReference(METASOUND_GET_PARAM_NAME(OutputDecibels), DecibelOutput);
-		return OutputDataReferences;
+		InOutVertexData.BindReadVertex(METASOUND_GET_PARAM_NAME(OutputDecibels), DecibelOutput);
+	}
+
+	FDataReferenceCollection FLinearGainToDecibelsOperator::GetInputs() const
+	{
+		// This should never be called. Bind(...) is called instead. This method
+		// exists as a stop-gap until the API can be deprecated and removed.
+		checkNoEntry();
+		return {};
+	}
+
+	FDataReferenceCollection FLinearGainToDecibelsOperator::GetOutputs() const
+	{
+		// This should never be called. Bind(...) is called instead. This method
+		// exists as a stop-gap until the API can be deprecated and removed.
+		checkNoEntry();
+		return {};
 	}
 
 	void FLinearGainToDecibelsOperator::Reset(const IOperator::FResetParams& InParams)
diff --git a/Engine/Plugins/Runtime/Metasound/Source/MetasoundStandardNodes/Private/MetasoundLinearToLogFrequencyNode.cpp b/Engine/Plugins/Runtime/Metasound/Source/MetasoundStandardNodes/Private/MetasoundLinearToLogFrequencyNode.cpp
--- a/Engine/Plugins/Runtime/Metasound/Source/MetasoundStandardNodes/Private/MetasoundLinearToLogFrequencyNode.cpp
+++ b/Engine/Plugins/Runtime/Metasound/Source/MetasoundStandardNodes/Private/MetasoundLinearToLogFrequencyNode.cpp
@@ -108,28 +108,38 @@ namespace Metasound
 		{
 		}
 
-		virtual FDataReferenceCollection GetInputs() const override
+		virtual void BindInputs(FInputVertexInterfaceData& InOutVertexData) override
 		{
 			using namespace LinearToLogFrequencyVertexNames;
 
-			FDataReferenceCollection InputDataReferences;
-			InputDataReferences.AddDataReadReference(METASOUND_GET_PARAM_NAME(InputValue), Value);
-			InputDataReferences.AddDataReadReference(METASOUND_GET_PARAM_NAME(InputDomainMin), DomainMin);
-			InputDataReferences.AddDataReadReference(METASOUND_GET_PARAM_NAME(InputDomainMax), DomainMax);
-			InputDataReferences.AddDataReadReference(METASOUND_GET_PARAM_NAME(InputRangeMin), RangeMin);
-			InputDataReferences.AddDataReadReference(METASOUND_GET_PARAM_NAME(InputRangeMax), RangeMax);
+			InOutVertexData.BindReadVertex(METASOUND_GET_PARAM_NAME(InputValue), Value);
+			InOutVertexData.BindReadVertex(METASOUND_GET_PARAM_NAME(InputDomainMin), DomainMin);
+			InOutVertexData.BindReadVertex(METASOUND_GET_PARAM_NAME(InputDomainMax), DomainMax);
+			InOutVertexData.BindReadVertex(METASOUND_GET_PARAM_NAME(InputRangeMin), RangeMin);
+			InOutVertexData.BindReadVertex(METASOUND_GET_PARAM_NAME(InputRangeMax), RangeMax);
+		}
 
-			return InputDataReferences;
+		virtual void BindOutputs(FOutputVertexInterfaceData& InOutVertexData) override
+		{
+			using namespace LinearToLogFrequencyVertexNames;
 
+			InOutVertexData.BindReadVertex(METASOUND_GET_PARAM_NAME(OutputFreq), FreqOutput);
 		}
 
-		virtual FDataReferenceCollection GetOutputs() const override
+		virtual FDataReferenceCollection GetInputs() const override
 		{
-			using namespace LinearToLogFrequencyVertexNames;
+			// This should never be called. Bind(...) is called instead. This method
+			// exists as a stop-gap until the API can be deprecated and removed.
+			checkNoEntry();
+			return {};
+		}
 
-			FDataReferenceCollection OutputDataReferences;
-			OutputDataReferences.AddDataReadReference(METASOUND_GET_PARAM_NAME(OutputFreq), FreqOutput);
-			return OutputDataReferences;
+		virtual FDataReferenceCollection GetOutputs() const override
+		{
+			// This should never be called. Bind(...) is called instead. This method
+			// exists as a stop-gap until the API can be deprecated and removed.
+			checkNoEntry();
+			return {};
 		}
 
 		void Reset(const IOperator::FResetParams& InParams)
diff --git a/Engine/Plugins/Runtime/Metasound/Source/MetasoundStandardNodes/Private/MetasoundLowFrequencyNoise.cpp b/Engine/Plugins/Runtime/Metasound/Source/MetasoundStandardNodes/Private/MetasoundLowFrequencyNoise.cpp
--- a/Engine/Plugins/Runtime/Metasound/Source/MetasoundStandardNodes/Private/MetasoundLowFrequencyNoise.cpp
+++ b/Engine/Plugins/Runtime/Metasound/Source/MetasoundStandardNodes/Private/MetasoundLowFrequencyNoise.cpp
@@ -214,38 +214,43 @@ namespace Metasound
 			CurrentRandomRate = GenerateRandomRate();
 		}
 
-		FDataReferenceCollection GetInputs() const override
+		virtual void BindInputs(FInputVertexInterfaceData& InOutVertexData) override
 		{
-			checkNoEntry();
-			FDataReferenceCollection InputDataReferences;		
-			return InputDataReferences;
+			using namespace LowFrequencyNoiseVertexNames;
+
+			InOutVertexData.BindReadVertex(METASOUND_GET_PARAM_NAME(RatePin), Pins.Rate);
+			InOutVertexData.BindReadVertex(METASOUND_GET_PARAM_NAME(RateJitterPin), Pins.RateJitter);
+			InOutVertexData.BindReadVertex(METASOUND_GET_PARAM_NAME(StepLimitPin), Pins.StepLimit);
+			InOutVertexData.BindReadVertex(METASOUND_GET_PARAM_NAME(SeedPin), Pins.Seed);
+			InOutVertexData.BindReadVertex(METASOUND_GET_PARAM_NAME(ResetSeedPin), Pins.ResetSeed);
+			InOutVertexData.BindReadVertex(METASOUND_GET_PARAM_NAME(SyncPin), Pins.Sync);
+			InOutVertexData.BindReadVertex(METASOUND_GET_PARAM_NAME(InterpTypePin), Pins.InterpType);
+			InOutVertexData.BindReadVertex(METASOUND_GET_PARAM_NAME(MinOutputValuePin), Pins.MinOutput);
+			InOutVertexData.BindReadVertex(METASOUND_GET_PARAM_NAME(MaxOutputValuePin), Pins.MaxOutput);
 		}
 
-		FDataReferenceCollection GetOutputs() const override
+		virtual void BindOutputs(FOutputVertexInterfaceData& InOutVertexData) override
 		{
-			checkNoEntry();
-			FDataReferenceCollection OutputDataReferences;
-			return OutputDataReferences;
+			using namespace LowFrequencyNoiseVertexNames;
+
+			InOutVertexData.BindReadVertex(METASOUND_GET_PARAM_NAME(ScaledOutputPin), ScaledOutput);
+			InOutVertexData.BindReadVertex(METASOUND_GET_PARAM_NAME(NormalizedOutputPin), NormalizedOutput);
 		}
 
-		virtual void Bind(FVertexInterfaceData& InVertexData) const override
+		virtual FDataReferenceCollection GetInputs() const override
 		{
-			using namespace LowFrequencyNoiseVertexNames;
+			// This should never be called. Bind(...) is called instead. This method
+			// exists as a stop-gap until the API can be deprecated and removed.
+			checkNoEntry();
+			return {};
+		}
 
-			FInputVertexInterfaceData& Inputs = InVertexData.GetInputs();
-			Inputs.BindReadVertex(METASOUND_GET_PARAM_NAME(RatePin), Pins.Rate);
-			Inputs.BindReadVertex(METASOUND_GET_PARAM_NAME(RateJitterPin), Pins.RateJitter);
-			Inputs.BindReadVertex(METASOUND_GET_PARAM_NAME(StepLimitPin), Pins.StepLimit);
-			Inputs.BindReadVertex(METASOUND_GET_PARAM_NAME(SeedPin), Pins.Seed);
-			Inputs.BindReadVertex(METASOUND_GET_PARAM_NAME(ResetSeedPin), Pins.ResetSeed);
-			Inputs.BindReadVertex(METASOUND_GET_PARAM_NAME(SyncPin), Pins.Sync);
-			Inputs.BindReadVertex(METASOUND_GET_PARAM_NAME(InterpTypePin), Pins.InterpType);
-			Inputs.BindReadVertex(METASOUND_GET_PARAM_NAME(MinOutputValuePin), Pins.MinOutput);
-			Inputs.BindReadVertex(METASOUND_GET_PARAM_NAME(MaxOutputValuePin), Pins.MaxOutput);
-			
-			FOutputVertexInterfaceData& Outputs = InVertexData.GetOutputs();
-			Outputs.BindReadVertex(METASOUND_GET_PARAM_NAME(ScaledOutputPin), ScaledOutput);
-			Outputs.BindReadVertex(METASOUND_GET_PARAM_NAME(NormalizedOutputPin), NormalizedOutput);
+		virtual FDataReferenceCollection GetOutputs() const override
+		{
+			// This should never be called. Bind(...) is called instead. This method
+			// exists as a stop-gap until the API can be deprecated and removed.
+			checkNoEntry();
+			return {};
 		}
 
 		void SetSeed(int32 InSeedValue)
diff --git a/Engine/Plugins/Runtime/Metasound/Source/MetasoundStandardNodes/Private/MetasoundMIDIToFreqNode.cpp b/Engine/Plugins/Runtime/Metasound/Source/MetasoundStandardNodes/Private/MetasoundMIDIToFreqNode.cpp
--- a/Engine/Plugins/Runtime/Metasound/Source/MetasoundStandardNodes/Private/MetasoundMIDIToFreqNode.cpp
+++ b/Engine/Plugins/Runtime/Metasound/Source/MetasoundStandardNodes/Private/MetasoundMIDIToFreqNode.cpp
@@ -68,6 +68,8 @@ namespace Metasound
 
 		TMidiToFreqOperator(const FCreateOperatorParams& InParams, const TDataReadReference<ValueType>& InMidiNote);
 
+		virtual void BindInputs(FInputVertexInterfaceData& InOutVertexData) override;
+		virtual void BindOutputs(FOutputVertexInterfaceData& InOutVertexData) override;
 		virtual FDataReferenceCollection GetInputs() const override;
 		virtual FDataReferenceCollection GetOutputs() const override;
 		void Reset(const IOperator::FResetParams& InParams);
@@ -94,24 +96,37 @@ namespace Metasound
 	}
 
 	template<typename ValueType>
-	FDataReferenceCollection TMidiToFreqOperator<ValueType>::GetInputs() const
+	void TMidiToFreqOperator<ValueType>::BindInputs(FInputVertexInterfaceData& InOutVertexData)
 	{
 		using namespace MidiToFrequencyVertexNames;
 
-		FDataReferenceCollection InputDataReferences;
-		InputDataReferences.AddDataReadReference(METASOUND_GET_PARAM_NAME(InputMidi), MidiNote);
-
-		return InputDataReferences;
+		InOutVertexData.BindReadVertex(METASOUND_GET_PARAM_NAME(InputMidi), MidiNote);
 	}
 
 	template<typename ValueType>
-	FDataReferenceCollection TMidiToFreqOperator<ValueType>::GetOutputs() const
+	void TMidiToFreqOperator<ValueType>::BindOutputs(FOutputVertexInterfaceData& InOutVertexData)
 	{
 		using namespace MidiToFrequencyVertexNames;
 
-		FDataReferenceCollection OutputDataReferences;
-		OutputDataReferences.AddDataReadReference(METASOUND_GET_PARAM_NAME(OutputFreq), FreqOutput);
-		return OutputDataReferences;
+		InOutVertexData.BindReadVertex(METASOUND_GET_PARAM_NAME(OutputFreq), FreqOutput);
+	}
+
+	template<typename ValueType>
+	FDataReferenceCollection TMidiToFreqOperator<ValueType>::GetInputs() const
+	{
+		// This should never be called. Bind(...) is called instead. This method
+		// exists as a stop-gap until the API can be deprecated and removed.
+		checkNoEntry();
+		return {};
+	}
+
+	template<typename ValueType>
+	FDataReferenceCollection TMidiToFreqOperator<ValueType>::GetOutputs() const
+	{
+		// This should never be called. Bind(...) is called instead. This method
+		// exists as a stop-gap until the API can be deprecated and removed.
+		checkNoEntry();
+		return {};
 	}
 
 	template<typename ValueType>
diff --git a/Engine/Plugins/Runtime/Metasound/Source/MetasoundStandardNodes/Private/MetasoundMapRangeNode.cpp b/Engine/Plugins/Runtime/Metasound/Source/MetasoundStandardNodes/Private/MetasoundMapRangeNode.cpp
--- a/Engine/Plugins/Runtime/Metasound/Source/MetasoundStandardNodes/Private/MetasoundMapRangeNode.cpp
+++ b/Engine/Plugins/Runtime/Metasound/Source/MetasoundStandardNodes/Private/MetasoundMapRangeNode.cpp
@@ -74,24 +74,19 @@ namespace Metasound
 			{
 			}
 
-			FDataReferenceCollection GetInputs() const
+			void BindInputs(FInputVertexInterfaceData& InOutVertexData)
 			{
-				FDataReferenceCollection Inputs;
-				Inputs.AddDataReadReference(METASOUND_GET_PARAM_NAME(InputValueName), Value);
-				Inputs.AddDataReadReference(METASOUND_GET_PARAM_NAME(InputInRangeAName), InRangeA);
-				Inputs.AddDataReadReference(METASOUND_GET_PARAM_NAME(InputInRangeBName), InRangeB);
-				Inputs.AddDataReadReference(METASOUND_GET_PARAM_NAME(InputOutRangeAName), OutRangeA);
-				Inputs.AddDataReadReference(METASOUND_GET_PARAM_NAME(InputOutRangeBName), OutRangeB);
-				Inputs.AddDataReadReference(METASOUND_GET_PARAM_NAME(InputClampedName), bClamped);
-
-				return Inputs;
+				InOutVertexData.BindReadVertex(METASOUND_GET_PARAM_NAME(InputValueName), Value);
+				InOutVertexData.BindReadVertex(METASOUND_GET_PARAM_NAME(InputInRangeAName), InRangeA);
+				InOutVertexData.BindReadVertex(METASOUND_GET_PARAM_NAME(InputInRangeBName), InRangeB);
+				InOutVertexData.BindReadVertex(METASOUND_GET_PARAM_NAME(InputOutRangeAName), OutRangeA);
+				InOutVertexData.BindReadVertex(METASOUND_GET_PARAM_NAME(InputOutRangeBName), OutRangeB);
+				InOutVertexData.BindReadVertex(METASOUND_GET_PARAM_NAME(InputClampedName), bClamped);
 			}
 
-			FDataReferenceCollection GetOutputs() const
+			void BindOutputs(FOutputVertexInterfaceData& InOutVertexData)
 			{
-				FDataReferenceCollection Outputs;
-				Outputs.AddDataReadReference(METASOUND_GET_PARAM_NAME(OutputValueName), OutputValue);
-				return Outputs;
+				InOutVertexData.BindReadVertex(METASOUND_GET_PARAM_NAME(OutputValueName), OutputValue);
 			}
 
 			void DoMapping()
@@ -158,26 +153,22 @@ namespace Metasound
 			{
 			}
 
-			FDataReferenceCollection GetInputs() const
+			void BindInputs(FInputVertexInterfaceData& InOutVertexData)
 			{
-				FDataReferenceCollection Inputs;
-				Inputs.AddDataReadReference(METASOUND_GET_PARAM_NAME(InputValueName),Value);
-				Inputs.AddDataReadReference(METASOUND_GET_PARAM_NAME(InputInRangeAName), InRangeA);
-				Inputs.AddDataReadReference(METASOUND_GET_PARAM_NAME(InputInRangeBName), InRangeB);
-				Inputs.AddDataReadReference(METASOUND_GET_PARAM_NAME(InputOutRangeAName), OutRangeA);
-				Inputs.AddDataReadReference(METASOUND_GET_PARAM_NAME(InputOutRangeBName), OutRangeB);
-				Inputs.AddDataReadReference(METASOUND_GET_PARAM_NAME(InputClampedName), bClamped);
-
-				return Inputs;
+				InOutVertexData.BindReadVertex(METASOUND_GET_PARAM_NAME(InputValueName),Value);
+				InOutVertexData.BindReadVertex(METASOUND_GET_PARAM_NAME(InputInRangeAName), InRangeA);
+				InOutVertexData.BindReadVertex(METASOUND_GET_PARAM_NAME(InputInRangeBName), InRangeB);
+				InOutVertexData.BindReadVertex(METASOUND_GET_PARAM_NAME(InputOutRangeAName), OutRangeA);
+				InOutVertexData.BindReadVertex(METASOUND_GET_PARAM_NAME(InputOutRangeBName), OutRangeB);
+				InOutVertexData.BindReadVertex(METASOUND_GET_PARAM_NAME(InputClampedName), bClamped);
 			}
 
-			FDataReferenceCollection GetOutputs() const
+			void BindOutputs(FOutputVertexInterfaceData& InOutVertexData)
 			{
-				FDataReferenceCollection Outputs;
-				Outputs.AddDataReadReference(METASOUND_GET_PARAM_NAME(OutputValueName), OutputValue);
-				return Outputs;
+				InOutVertexData.BindReadVertex(METASOUND_GET_PARAM_NAME(OutputValueName), OutputValue);
 			}
 
+
 			void DoMapping()
 			{
 				if (*bClamped)
@@ -243,26 +234,23 @@ namespace Metasound
 			{
 			}
 
-			FDataReferenceCollection GetInputs() const
+
+			void BindInputs(FInputVertexInterfaceData& InOutVertexData)
 			{
-				FDataReferenceCollection Inputs;
-				Inputs.AddDataReadReference(METASOUND_GET_PARAM_NAME(InputValueName), Value);
-				Inputs.AddDataReadReference(METASOUND_GET_PARAM_NAME(InputInRangeAName), InRangeA);
-				Inputs.AddDataReadReference(METASOUND_GET_PARAM_NAME(InputInRangeBName), InRangeB);
-				Inputs.AddDataReadReference(METASOUND_GET_PARAM_NAME(InputOutRangeAName), OutRangeA);
-				Inputs.AddDataReadReference(METASOUND_GET_PARAM_NAME(InputOutRangeBName), OutRangeB);
-				Inputs.AddDataReadReference(METASOUND_GET_PARAM_NAME(InputClampedName), bClamped);
-
-				return Inputs;
+				InOutVertexData.BindReadVertex(METASOUND_GET_PARAM_NAME(InputValueName), Value);
+				InOutVertexData.BindReadVertex(METASOUND_GET_PARAM_NAME(InputInRangeAName), InRangeA);
+				InOutVertexData.BindReadVertex(METASOUND_GET_PARAM_NAME(InputInRangeBName), InRangeB);
+				InOutVertexData.BindReadVertex(METASOUND_GET_PARAM_NAME(InputOutRangeAName), OutRangeA);
+				InOutVertexData.BindReadVertex(METASOUND_GET_PARAM_NAME(InputOutRangeBName), OutRangeB);
+				InOutVertexData.BindReadVertex(METASOUND_GET_PARAM_NAME(InputClampedName), bClamped);
 			}
 
-			FDataReferenceCollection GetOutputs() const
+			void BindOutputs(FOutputVertexInterfaceData& InOutVertexData)
 			{
-				FDataReferenceCollection Outputs;
-				Outputs.AddDataReadReference(METASOUND_GET_PARAM_NAME(OutputValueName), OutputValue);
-				return Outputs;
+				InOutVertexData.BindReadVertex(METASOUND_GET_PARAM_NAME(OutputValueName), OutputValue);
 			}
 
+
 			void DoMapping()
 			{
 				FAudioBuffer& OutBuffer = *OutputValue;
@@ -344,14 +332,31 @@ namespace Metasound
 
 		virtual ~TMapRangeOperator() = default;
 
+
+		virtual void BindInputs(FInputVertexInterfaceData& InOutVertexData) override
+		{
+			MappingObject.BindInputs(InOutVertexData);
+		}
+
+		virtual void BindOutputs(FOutputVertexInterfaceData& InOutVertexData) override
+		{
+			MappingObject.BindOutputs(InOutVertexData);
+		}
+
 		virtual FDataReferenceCollection GetInputs() const override
 		{
-			return MappingObject.GetInputs();
+			// This should never be called. Bind(...) is called instead. This method
+			// exists as a stop-gap until the API can be deprecated and removed.
+			checkNoEntry();
+			return {};
 		}
 
 		virtual FDataReferenceCollection GetOutputs() const override
 		{
-			return MappingObject.GetOutputs();
+			// This should never be called. Bind(...) is called instead. This method
+			// exists as a stop-gap until the API can be deprecated and removed.
+			checkNoEntry();
+			return {};
 		}
 
 		void Reset(const IOperator::FResetParams& InParams)
diff --git a/Engine/Plugins/Runtime/Metasound/Source/MetasoundStandardNodes/Private/MetasoundMathNodes.cpp b/Engine/Plugins/Runtime/Metasound/Source/MetasoundStandardNodes/Private/MetasoundMathNodes.cpp
--- a/Engine/Plugins/Runtime/Metasound/Source/MetasoundStandardNodes/Private/MetasoundMathNodes.cpp
+++ b/Engine/Plugins/Runtime/Metasound/Source/MetasoundStandardNodes/Private/MetasoundMathNodes.cpp
@@ -280,24 +280,35 @@ namespace Metasound
 				Reset(InParams);
 			}
 
-			virtual FDataReferenceCollection GetInputs() const override
+			virtual void BindInputs(FInputVertexInterfaceData& InOutVertexData) override
 			{
-				using namespace MathOpNames; 
+				using namespace MathOpNames;
 
-				FDataReferenceCollection InputDataReferences;
-				InputDataReferences.AddDataReadReference(METASOUND_GET_PARAM_NAME(PrimaryOperand), PrimaryOperandRef);
-				InputDataReferences.AddDataReadReference(METASOUND_GET_PARAM_NAME(AdditionalOperands), AdditionalOperandRefs[0]);
-				return InputDataReferences;
+				InOutVertexData.BindReadVertex(METASOUND_GET_PARAM_NAME(PrimaryOperand), PrimaryOperandRef);
+				InOutVertexData.BindReadVertex(METASOUND_GET_PARAM_NAME(AdditionalOperands), AdditionalOperandRefs[0]);
 			}
 
-			virtual FDataReferenceCollection GetOutputs() const override
+			virtual void BindOutputs(FOutputVertexInterfaceData& InOutVertexData) override
 			{
 				using namespace MathOpNames;
-				
-				FDataReferenceCollection OutputDataReferences;
-				OutputDataReferences.AddDataReadReference(METASOUND_GET_PARAM_NAME(OutMath), TDataClassReadRef(ValueRef));
 
-				return OutputDataReferences;
+				InOutVertexData.BindReadVertex(METASOUND_GET_PARAM_NAME(OutMath), TDataClassReadRef(ValueRef));
+			}
+
+			virtual FDataReferenceCollection GetInputs() const override
+			{
+				// This should never be called. Bind(...) is called instead. This method
+				// exists as a stop-gap until the API can be deprecated and removed.
+				checkNoEntry();
+				return {};
+			}
+
+			virtual FDataReferenceCollection GetOutputs() const override
+			{
+				// This should never be called. Bind(...) is called instead. This method
+				// exists as a stop-gap until the API can be deprecated and removed.
+				checkNoEntry();
+				return {};
 			}
 		};
 
diff --git a/Engine/Plugins/Runtime/Metasound/Source/MetasoundStandardNodes/Private/MetasoundMaxNode.cpp b/Engine/Plugins/Runtime/Metasound/Source/MetasoundStandardNodes/Private/MetasoundMaxNode.cpp
--- a/Engine/Plugins/Runtime/Metasound/Source/MetasoundStandardNodes/Private/MetasoundMaxNode.cpp
+++ b/Engine/Plugins/Runtime/Metasound/Source/MetasoundStandardNodes/Private/MetasoundMaxNode.cpp
@@ -191,25 +191,35 @@ namespace Metasound
 
 		virtual ~TMaxNodeOperator() = default;
 
-		virtual FDataReferenceCollection GetInputs() const override
+		virtual void BindInputs(FInputVertexInterfaceData& InOutVertexData) override
 		{
 			using namespace MaxVertexNames;
 
-			FDataReferenceCollection Inputs;
-			Inputs.AddDataReadReference(METASOUND_GET_PARAM_NAME(InputAValue), InputA);
-			Inputs.AddDataReadReference(METASOUND_GET_PARAM_NAME(InputBValue), InputB);
-
-			return Inputs;
+			InOutVertexData.BindReadVertex(METASOUND_GET_PARAM_NAME(InputAValue), InputA);
+			InOutVertexData.BindReadVertex(METASOUND_GET_PARAM_NAME(InputBValue), InputB);
 		}
 
-		virtual FDataReferenceCollection GetOutputs() const override
+		virtual void BindOutputs(FOutputVertexInterfaceData& InOutVertexData) override
 		{
 			using namespace MaxVertexNames;
 
-			FDataReferenceCollection Outputs;
-			Outputs.AddDataReadReference(METASOUND_GET_PARAM_NAME(OutputValue), OutputValue);
+			InOutVertexData.BindReadVertex(METASOUND_GET_PARAM_NAME(OutputValue), OutputValue);
+		}
+
+		virtual FDataReferenceCollection GetInputs() const override
+		{
+			// This should never be called. Bind(...) is called instead. This method
+			// exists as a stop-gap until the API can be deprecated and removed.
+			checkNoEntry();
+			return {};
+		}
 
-			return Outputs;
+		virtual FDataReferenceCollection GetOutputs() const override
+		{
+			// This should never be called. Bind(...) is called instead. This method
+			// exists as a stop-gap until the API can be deprecated and removed.
+			checkNoEntry();
+			return {};
 		}
 
 		void GetMax()
diff --git a/Engine/Plugins/Runtime/Metasound/Source/MetasoundStandardNodes/Private/MetasoundMidSideNodes.cpp b/Engine/Plugins/Runtime/Metasound/Source/MetasoundStandardNodes/Private/MetasoundMidSideNodes.cpp
--- a/Engine/Plugins/Runtime/Metasound/Source/MetasoundStandardNodes/Private/MetasoundMidSideNodes.cpp
+++ b/Engine/Plugins/Runtime/Metasound/Source/MetasoundStandardNodes/Private/MetasoundMidSideNodes.cpp
@@ -104,30 +104,38 @@ namespace Metasound
 			return Interface;
 		}
 
-		virtual FDataReferenceCollection GetInputs() const override
+		virtual void BindInputs(FInputVertexInterfaceData& InOutVertexData) override
 		{
 			using namespace MidSideEncodeVertexNames;
 
-			FDataReferenceCollection InputDataReferences;
-			
-			InputDataReferences.AddDataReadReference(METASOUND_GET_PARAM_NAME(InputAudioLeft), AudioInputLeft);
-			InputDataReferences.AddDataReadReference(METASOUND_GET_PARAM_NAME(InputAudioRight), AudioInputRight);
-			InputDataReferences.AddDataReadReference(METASOUND_GET_PARAM_NAME(InputSpreadAmount), SpreadAmount);
-			InputDataReferences.AddDataReadReference(METASOUND_GET_PARAM_NAME(InputEqualPower), bEqualPower);
-			
-			return InputDataReferences;
+			InOutVertexData.BindReadVertex(METASOUND_GET_PARAM_NAME(InputAudioLeft), AudioInputLeft);
+			InOutVertexData.BindReadVertex(METASOUND_GET_PARAM_NAME(InputAudioRight), AudioInputRight);
+			InOutVertexData.BindReadVertex(METASOUND_GET_PARAM_NAME(InputSpreadAmount), SpreadAmount);
+			InOutVertexData.BindReadVertex(METASOUND_GET_PARAM_NAME(InputEqualPower), bEqualPower);
 		}
 
-		virtual FDataReferenceCollection GetOutputs() const override
+		virtual void BindOutputs(FOutputVertexInterfaceData& InOutVertexData) override
 		{
 			using namespace MidSideEncodeVertexNames;
-			
-			FDataReferenceCollection OutputDataReferences;
 
-			OutputDataReferences.AddDataReadReference(METASOUND_GET_PARAM_NAME(OutputAudioMid), AudioOutputMid);
-			OutputDataReferences.AddDataReadReference(METASOUND_GET_PARAM_NAME(OutputAudioSide), AudioOutputSide);
-			
-			return OutputDataReferences;
+			InOutVertexData.BindReadVertex(METASOUND_GET_PARAM_NAME(OutputAudioMid), AudioOutputMid);
+			InOutVertexData.BindReadVertex(METASOUND_GET_PARAM_NAME(OutputAudioSide), AudioOutputSide);
+		}
+
+		virtual FDataReferenceCollection GetInputs() const override
+		{
+			// This should never be called. Bind(...) is called instead. This method
+			// exists as a stop-gap until the API can be deprecated and removed.
+			checkNoEntry();
+			return {};
+		}
+
+		virtual FDataReferenceCollection GetOutputs() const override
+		{
+			// This should never be called. Bind(...) is called instead. This method
+			// exists as a stop-gap until the API can be deprecated and removed.
+			checkNoEntry();
+			return {};
 		}
 		
 		static TUniquePtr<IOperator> CreateOperator(const FCreateOperatorParams& InParams, FBuildErrorArray& OutErrors)
@@ -350,31 +358,38 @@ namespace Metasound
 			return Interface;
 		}
 
-		virtual FDataReferenceCollection GetInputs() const override
+		virtual void BindInputs(FInputVertexInterfaceData& InOutVertexData) override
 		{
 			using namespace MidSideDecodeVertexNames;
 
-			FDataReferenceCollection InputDataReferences;
-
-			InputDataReferences.AddDataReadReference(METASOUND_GET_PARAM_NAME(InputAudioMid), AudioInputMid);
-			InputDataReferences.AddDataReadReference(METASOUND_GET_PARAM_NAME(InputAudioSide), AudioInputSide);
-			InputDataReferences.AddDataReadReference(METASOUND_GET_PARAM_NAME(InputSpreadAmount), SpreadAmount);
-			InputDataReferences.AddDataReadReference(METASOUND_GET_PARAM_NAME(InputEqualPower), bEqualPower);
-
-			return InputDataReferences;
-
+			InOutVertexData.BindReadVertex(METASOUND_GET_PARAM_NAME(InputAudioMid), AudioInputMid);
+			InOutVertexData.BindReadVertex(METASOUND_GET_PARAM_NAME(InputAudioSide), AudioInputSide);
+			InOutVertexData.BindReadVertex(METASOUND_GET_PARAM_NAME(InputSpreadAmount), SpreadAmount);
+			InOutVertexData.BindReadVertex(METASOUND_GET_PARAM_NAME(InputEqualPower), bEqualPower);
 		}
 
-		virtual FDataReferenceCollection GetOutputs() const override
+		virtual void BindOutputs(FOutputVertexInterfaceData& InOutVertexData) override
 		{
 			using namespace MidSideDecodeVertexNames;
 
-			FDataReferenceCollection OutputDataReferences;
+			InOutVertexData.BindReadVertex(METASOUND_GET_PARAM_NAME(OutputAudioLeft), AudioOutputLeft);
+			InOutVertexData.BindReadVertex(METASOUND_GET_PARAM_NAME(OutputAudioRight), AudioOutputRight);
+		}
 
-			OutputDataReferences.AddDataReadReference(METASOUND_GET_PARAM_NAME(OutputAudioLeft), AudioOutputLeft);
-			OutputDataReferences.AddDataReadReference(METASOUND_GET_PARAM_NAME(OutputAudioRight), AudioOutputRight);
+		virtual FDataReferenceCollection GetInputs() const override
+		{
+			// This should never be called. Bind(...) is called instead. This method
+			// exists as a stop-gap until the API can be deprecated and removed.
+			checkNoEntry();
+			return {};
+		}
 
-			return OutputDataReferences;
+		virtual FDataReferenceCollection GetOutputs() const override
+		{
+			// This should never be called. Bind(...) is called instead. This method
+			// exists as a stop-gap until the API can be deprecated and removed.
+			checkNoEntry();
+			return {};
 		}
 
 		static TUniquePtr<IOperator> CreateOperator(const FCreateOperatorParams& InParams, FBuildErrorArray& OutErrors)
diff --git a/Engine/Plugins/Runtime/Metasound/Source/MetasoundStandardNodes/Private/MetasoundMidiNoteQuantizerNode.cpp b/Engine/Plugins/Runtime/Metasound/Source/MetasoundStandardNodes/Private/MetasoundMidiNoteQuantizerNode.cpp
--- a/Engine/Plugins/Runtime/Metasound/Source/MetasoundStandardNodes/Private/MetasoundMidiNoteQuantizerNode.cpp
+++ b/Engine/Plugins/Runtime/Metasound/Source/MetasoundStandardNodes/Private/MetasoundMidiNoteQuantizerNode.cpp
@@ -60,6 +60,9 @@ namespace Metasound
 		static const FNodeClassMetadata& GetNodeInfo();
 		static FVertexInterface DeclareVertexInterface();
 		static TUniquePtr<IOperator> CreateOperator(const FCreateOperatorParams& InParams, FBuildErrorArray& OutErrors);
+
+		virtual void BindInputs(FInputVertexInterfaceData& InOutVertexData) override;
+		virtual void BindOutputs(FOutputVertexInterfaceData& InOutVertexData) override;
 		virtual FDataReferenceCollection GetInputs() const override;
 		virtual FDataReferenceCollection GetOutputs() const override;
 		void Reset(const IOperator::FResetParams& InParams);
@@ -167,23 +170,32 @@ namespace Metasound
 			);
 	}
 
-	FDataReferenceCollection FMidiNoteQuantizerOperator::GetInputs() const
+	void FMidiNoteQuantizerOperator::BindInputs(FInputVertexInterfaceData& InOutVertexData)
 	{
-		FDataReferenceCollection InputDataReferences;
-		InputDataReferences.AddDataReadReference(METASOUND_GET_PARAM_NAME(ParamNoteIn), FFloatReadRef(MidiNoteIn));
-		InputDataReferences.AddDataReadReference(METASOUND_GET_PARAM_NAME(ParamRootNote), FFloatReadRef(RootNote));
-		InputDataReferences.AddDataReadReference(METASOUND_GET_PARAM_NAME(ParamScaleDegrees), FArrayScaleDegreeReadRef(Scale));
+		InOutVertexData.BindReadVertex(METASOUND_GET_PARAM_NAME(ParamNoteIn), MidiNoteIn);
+		InOutVertexData.BindReadVertex(METASOUND_GET_PARAM_NAME(ParamRootNote), RootNote);
+		InOutVertexData.BindReadVertex(METASOUND_GET_PARAM_NAME(ParamScaleDegrees), Scale);
+	}
 
-		return InputDataReferences;
+	void FMidiNoteQuantizerOperator::BindOutputs(FOutputVertexInterfaceData& InOutVertexData)
+	{
+		InOutVertexData.BindReadVertex(METASOUND_GET_PARAM_NAME(ParamNoteOutput), MidiNoteOut);
 	}
 
-	FDataReferenceCollection FMidiNoteQuantizerOperator::GetOutputs() const
+	FDataReferenceCollection FMidiNoteQuantizerOperator::GetInputs() const
 	{
-		// expose read access to our output buffer for other processors in the graph
-		FDataReferenceCollection OutputDataReferences;
-		OutputDataReferences.AddDataReadReference(METASOUND_GET_PARAM_NAME(ParamNoteOutput), FFloatReadRef(MidiNoteOut));
+		// This should never be called. Bind(...) is called instead. This method
+		// exists as a stop-gap until the API can be deprecated and removed.
+		checkNoEntry();
+		return {};
+	}
 
-		return OutputDataReferences;
+	FDataReferenceCollection FMidiNoteQuantizerOperator::GetOutputs() const
+	{
+		// This should never be called. Bind(...) is called instead. This method
+		// exists as a stop-gap until the API can be deprecated and removed.
+		checkNoEntry();
+		return {};
 	}
 
 	void FMidiNoteQuantizerOperator::Reset(const IOperator::FResetParams& InParams)
diff --git a/Engine/Plugins/Runtime/Metasound/Source/MetasoundStandardNodes/Private/MetasoundMidiScaleToArrayNode.cpp b/Engine/Plugins/Runtime/Metasound/Source/MetasoundStandardNodes/Private/MetasoundMidiScaleToArrayNode.cpp
--- a/Engine/Plugins/Runtime/Metasound/Source/MetasoundStandardNodes/Private/MetasoundMidiScaleToArrayNode.cpp
+++ b/Engine/Plugins/Runtime/Metasound/Source/MetasoundStandardNodes/Private/MetasoundMidiScaleToArrayNode.cpp
@@ -98,6 +98,8 @@ namespace Metasound
 		static const FNodeClassMetadata& GetNodeInfo();
 		static FVertexInterface DeclareVertexInterface();
 		static TUniquePtr<IOperator> CreateOperator(const FCreateOperatorParams& InParams, FBuildErrorArray& OutErrors);
+		virtual void BindInputs(FInputVertexInterfaceData& InOutVertexData) override;
+		virtual void BindOutputs(FOutputVertexInterfaceData& InOutVertexData) override;
 		virtual FDataReferenceCollection GetInputs() const override;
 		virtual FDataReferenceCollection GetOutputs() const override;
 		void Reset(const IOperator::FResetParams& InParams);
@@ -194,22 +196,31 @@ namespace Metasound
 			);
 	}
 
-	FDataReferenceCollection FMusicalScaleToNoteArrayOperator::GetInputs() const
+	void FMusicalScaleToNoteArrayOperator::BindInputs(FInputVertexInterfaceData& InOutVertexData)
 	{
-		FDataReferenceCollection InputDataReferences;
-		InputDataReferences.AddDataReadReference(METASOUND_GET_PARAM_NAME(ParamScaleDegreesPreset), FEnumMusicalScaleReadRef(Scale));
-		InputDataReferences.AddDataReadReference(METASOUND_GET_PARAM_NAME(ParamChordTonesOnly), FBoolReadRef(bChordTonesOnly));
+		InOutVertexData.BindReadVertex(METASOUND_GET_PARAM_NAME(ParamScaleDegreesPreset), FEnumMusicalScaleReadRef(Scale));
+		InOutVertexData.BindReadVertex(METASOUND_GET_PARAM_NAME(ParamChordTonesOnly), FBoolReadRef(bChordTonesOnly));
+	}
 
-		return InputDataReferences;
+	void FMusicalScaleToNoteArrayOperator::BindOutputs(FOutputVertexInterfaceData& InOutVertexData)
+	{
+		InOutVertexData.BindReadVertex(METASOUND_GET_PARAM_NAME(ParamNoteArrayOutput), FArrayScaleDegreeWriteRef(ScaleArrayOutput));
 	}
 
-	FDataReferenceCollection FMusicalScaleToNoteArrayOperator::GetOutputs() const
+	FDataReferenceCollection FMusicalScaleToNoteArrayOperator::GetInputs() const
 	{
-		// expose read access to our output buffer for other processors in the graph
-		FDataReferenceCollection OutputDataReferences;
-		OutputDataReferences.AddDataReadReference(METASOUND_GET_PARAM_NAME(ParamNoteArrayOutput), FArrayScaleDegreeWriteRef(ScaleArrayOutput));
+		// This should never be called. Bind(...) is called instead. This method
+		// exists as a stop-gap until the API can be deprecated and removed.
+		checkNoEntry();
+		return {};
+	}
 
-		return OutputDataReferences;
+	FDataReferenceCollection FMusicalScaleToNoteArrayOperator::GetOutputs() const
+	{
+		// This should never be called. Bind(...) is called instead. This method
+		// exists as a stop-gap until the API can be deprecated and removed.
+		checkNoEntry();
+		return {};
 	}
 
 	void FMusicalScaleToNoteArrayOperator::Reset(const IOperator::FResetParams& InParams)
diff --git a/Engine/Plugins/Runtime/Metasound/Source/MetasoundStandardNodes/Private/MetasoundMinNode.cpp b/Engine/Plugins/Runtime/Metasound/Source/MetasoundStandardNodes/Private/MetasoundMinNode.cpp
--- a/Engine/Plugins/Runtime/Metasound/Source/MetasoundStandardNodes/Private/MetasoundMinNode.cpp
+++ b/Engine/Plugins/Runtime/Metasound/Source/MetasoundStandardNodes/Private/MetasoundMinNode.cpp
@@ -187,25 +187,35 @@ namespace Metasound
 
 		virtual ~TMinNodeOperator() = default;
 
-		virtual FDataReferenceCollection GetInputs() const override
+		virtual void BindInputs(FInputVertexInterfaceData& InOutVertexData) override
 		{
 			using namespace MinVertexNames;
 
-			FDataReferenceCollection Inputs;
-			Inputs.AddDataReadReference(METASOUND_GET_PARAM_NAME(InputAValue), InputA);
-			Inputs.AddDataReadReference(METASOUND_GET_PARAM_NAME(InputBValue), InputB);
-
-			return Inputs;
+			InOutVertexData.BindReadVertex(METASOUND_GET_PARAM_NAME(InputAValue), InputA);
+			InOutVertexData.BindReadVertex(METASOUND_GET_PARAM_NAME(InputBValue), InputB);
 		}
 
-		virtual FDataReferenceCollection GetOutputs() const override
+		virtual void BindOutputs(FOutputVertexInterfaceData& InOutVertexData) override
 		{
 			using namespace MinVertexNames;
 
-			FDataReferenceCollection Outputs;
-			Outputs.AddDataReadReference(METASOUND_GET_PARAM_NAME(OutputValue), OutputValue);
+			InOutVertexData.BindReadVertex(METASOUND_GET_PARAM_NAME(OutputValue), OutputValue);
+		}
+
+		virtual FDataReferenceCollection GetInputs() const override
+		{
+			// This should never be called. Bind(...) is called instead. This method
+			// exists as a stop-gap until the API can be deprecated and removed.
+			checkNoEntry();
+			return {};
+		}
 
-			return Outputs;
+		virtual FDataReferenceCollection GetOutputs() const override
+		{
+			// This should never be called. Bind(...) is called instead. This method
+			// exists as a stop-gap until the API can be deprecated and removed.
+			checkNoEntry();
+			return {};
 		}
 
 		void GetMin()
diff --git a/Engine/Plugins/Runtime/Metasound/Source/MetasoundStandardNodes/Private/MetasoundMixerNode.cpp b/Engine/Plugins/Runtime/Metasound/Source/MetasoundStandardNodes/Private/MetasoundMixerNode.cpp
--- a/Engine/Plugins/Runtime/Metasound/Source/MetasoundStandardNodes/Private/MetasoundMixerNode.cpp
+++ b/Engine/Plugins/Runtime/Metasound/Source/MetasoundStandardNodes/Private/MetasoundMixerNode.cpp
@@ -167,32 +167,41 @@ namespace Metasound
 			return MakeUnique<TAudioMixerNodeOperator<NumInputs, NumChannels>>(InParams, MoveTemp(InputBuffers), MoveTemp(InputGains));
 		}
 
-		virtual FDataReferenceCollection GetInputs() const override
+		virtual void BindInputs(FInputVertexInterfaceData& InOutVertexData) override
 		{
-			FDataReferenceCollection InputPins;
 			for (uint32 i = 0; i < NumInputs; ++i)
 			{
 				for (uint32 Chan = 0; Chan < NumChannels; ++Chan)
 				{
-					InputPins.AddDataReadReference(GetAudioInputName(i, Chan), Inputs[i * NumChannels + Chan]);
+					InOutVertexData.BindReadVertex(GetAudioInputName(i, Chan), Inputs[i * NumChannels + Chan]);
 				}
 
-				InputPins.AddDataReadReference(GetGainInputName(i), Gains[i]);
+				InOutVertexData.BindReadVertex(GetGainInputName(i), Gains[i]);
 			}
-
-			return InputPins;
 		}
 
-		virtual FDataReferenceCollection GetOutputs() const override
+		virtual void BindOutputs(FOutputVertexInterfaceData& InOutVertexData) override
 		{
-			FDataReferenceCollection OutputPins;
-
 			for (uint32 i = 0; i < NumChannels; ++i)
 			{
-				OutputPins.AddDataReadReference(GetAudioOutputName(i), Outputs[i]);
+				InOutVertexData.BindReadVertex(GetAudioOutputName(i), Outputs[i]);
 			}
+		}
+
+		virtual FDataReferenceCollection GetInputs() const override
+		{
+			// This should never be called. Bind(...) is called instead. This method
+			// exists as a stop-gap until the API can be deprecated and removed.
+			checkNoEntry();
+			return {};
+		}
 
-			return OutputPins;
+		virtual FDataReferenceCollection GetOutputs() const override
+		{
+			// This should never be called. Bind(...) is called instead. This method
+			// exists as a stop-gap until the API can be deprecated and removed.
+			checkNoEntry();
+			return {};
 		}
 
 		void Reset(const IOperator::FResetParams& InParams)
@@ -358,7 +367,6 @@ namespace Metasound
 #pragma endregion
 
 
-
 #pragma region Node Definition
 	template<uint32 NumInputs, uint32 NumChannels>
 	class METASOUNDSTANDARDNODES_API TAudioMixerNode : public FNodeFacade
diff --git a/Engine/Plugins/Runtime/Metasound/Source/MetasoundStandardNodes/Private/MetasoundNoiseGenerator.cpp b/Engine/Plugins/Runtime/Metasound/Source/MetasoundStandardNodes/Private/MetasoundNoiseGenerator.cpp
--- a/Engine/Plugins/Runtime/Metasound/Source/MetasoundStandardNodes/Private/MetasoundNoiseGenerator.cpp
+++ b/Engine/Plugins/Runtime/Metasound/Source/MetasoundStandardNodes/Private/MetasoundNoiseGenerator.cpp
@@ -51,6 +51,8 @@ namespace Metasound
 		static FVertexInterface DeclareVertexInterface();
 
 		FNoiseOperator(const FOperatorSettings& InSettings, FInt32ReadRef&& InReadRef);
+		virtual void BindInputs(FInputVertexInterfaceData& InOutVertexData) override;
+		virtual void BindOutputs(FOutputVertexInterfaceData& InOutVertexData) override;
 		virtual FDataReferenceCollection GetInputs() const override;
 		virtual FDataReferenceCollection GetOutputs() const override;
 
@@ -177,22 +179,33 @@ namespace Metasound
 		, OldSeed(*Seed)
 	{}
 
-	FDataReferenceCollection FNoiseOperator::GetInputs() const
+
+	void FNoiseOperator::BindInputs(FInputVertexInterfaceData& InOutVertexData)
 	{
 		using namespace NoiseGeneratorVertexNames;
-
-		FDataReferenceCollection InputDataReferences;
-		InputDataReferences.AddDataReadReference(METASOUND_GET_PARAM_NAME(InputSeed), FInt32ReadRef(Seed));
-		return InputDataReferences;
+		InOutVertexData.BindReadVertex(METASOUND_GET_PARAM_NAME(InputSeed), Seed);
 	}
 
-	FDataReferenceCollection FNoiseOperator::GetOutputs() const
+	void FNoiseOperator::BindOutputs(FOutputVertexInterfaceData& InOutVertexData)
 	{
 		using namespace NoiseGeneratorVertexNames;
+		InOutVertexData.BindReadVertex(METASOUND_GET_PARAM_NAME(OutAudio), Out);
+	}
 
-		FDataReferenceCollection OutputDataReferences;
-		OutputDataReferences.AddDataReadReference(METASOUND_GET_PARAM_NAME(OutAudio), FAudioBufferReadRef(Out));
-		return OutputDataReferences;
+	FDataReferenceCollection FNoiseOperator::GetInputs() const
+	{
+		// This should never be called. Bind(...) is called instead. This method
+		// exists as a stop-gap until the API can be deprecated and removed.
+		checkNoEntry();
+		return {};
+	}
+
+	FDataReferenceCollection FNoiseOperator::GetOutputs() const
+	{
+		// This should never be called. Bind(...) is called instead. This method
+		// exists as a stop-gap until the API can be deprecated and removed.
+		checkNoEntry();
+		return {};
 	}
 
 	FVertexInterface FNoiseOperator::DeclareVertexInterface()
diff --git a/Engine/Plugins/Runtime/Metasound/Source/MetasoundStandardNodes/Private/MetasoundOscillatorNodes.cpp b/Engine/Plugins/Runtime/Metasound/Source/MetasoundStandardNodes/Private/MetasoundOscillatorNodes.cpp
--- a/Engine/Plugins/Runtime/Metasound/Source/MetasoundStandardNodes/Private/MetasoundOscillatorNodes.cpp
+++ b/Engine/Plugins/Runtime/Metasound/Source/MetasoundStandardNodes/Private/MetasoundOscillatorNodes.cpp
@@ -98,27 +98,40 @@ namespace Metasound
 			check(AudioBuffer->Num() == InConstructParams.Settings.GetNumFramesPerBlock());
 		}
 
-		FDataReferenceCollection GetInputs() const override
+
+		virtual void BindInputs(FInputVertexInterfaceData& InOutVertexData) override
 		{
-			using namespace OscillatorCommonVertexNames; 
+			using namespace OscillatorCommonVertexNames;
 
-			FDataReferenceCollection InputDataReferences;
-			InputDataReferences.AddDataReadReference(METASOUND_GET_PARAM_NAME(EnabledPin), Enabled);
-			InputDataReferences.AddDataReadReference(METASOUND_GET_PARAM_NAME(OscBaseFrequencyPin), BaseFrequency);
-			InputDataReferences.AddDataReadReference(METASOUND_GET_PARAM_NAME(PhaseOffsetPin), PhaseOffset);
-			InputDataReferences.AddDataReadReference(METASOUND_GET_PARAM_NAME(OscPhaseResetPin), PhaseReset);
-			InputDataReferences.AddDataReadReference(METASOUND_GET_PARAM_NAME(GlideFactorPin), GlideFactor);
-			InputDataReferences.AddDataReadReference(METASOUND_GET_PARAM_NAME(BiPolarPin), BiPolar);
-			return InputDataReferences;
+			InOutVertexData.BindReadVertex(METASOUND_GET_PARAM_NAME(EnabledPin), Enabled);
+			InOutVertexData.BindReadVertex(METASOUND_GET_PARAM_NAME(OscBaseFrequencyPin), BaseFrequency);
+			InOutVertexData.BindReadVertex(METASOUND_GET_PARAM_NAME(PhaseOffsetPin), PhaseOffset);
+			InOutVertexData.BindReadVertex(METASOUND_GET_PARAM_NAME(OscPhaseResetPin), PhaseReset);
+			InOutVertexData.BindReadVertex(METASOUND_GET_PARAM_NAME(GlideFactorPin), GlideFactor);
+			InOutVertexData.BindReadVertex(METASOUND_GET_PARAM_NAME(BiPolarPin), BiPolar);
 		}
 
-		FDataReferenceCollection GetOutputs() const override
+		virtual void BindOutputs(FOutputVertexInterfaceData& InOutVertexData) override
 		{
 			using namespace OscillatorCommonVertexNames;
-			
-			FDataReferenceCollection OutputDataReferences;
-			OutputDataReferences.AddDataReadReference(METASOUND_GET_PARAM_NAME(AudioOutPin), AudioBuffer);
-			return OutputDataReferences;
+
+			InOutVertexData.BindReadVertex(METASOUND_GET_PARAM_NAME(AudioOutPin), AudioBuffer);
+		}
+
+		virtual FDataReferenceCollection GetInputs() const override
+		{
+			// This should never be called. Bind(...) is called instead. This method
+			// exists as a stop-gap until the API can be deprecated and removed.
+			checkNoEntry();
+			return {};
+		}
+
+		virtual FDataReferenceCollection GetOutputs() const override
+		{
+			// This should never be called. Bind(...) is called instead. This method
+			// exists as a stop-gap until the API can be deprecated and removed.
+			checkNoEntry();
+			return {};
 		}
 
 		void ResetPhase(float InPhaseInDegrees)
@@ -211,13 +224,27 @@ namespace Metasound
 			: Super(InCommonParams), Fm(InFmData)
 		{}
 
-		FDataReferenceCollection GetInputs() const override
+
+		virtual void BindInputs(FInputVertexInterfaceData& InOutVertexData) override
 		{
-			using namespace OscillatorCommonVertexNames; 
+			using namespace OscillatorCommonVertexNames;
+			InOutVertexData.BindReadVertex(METASOUND_GET_PARAM_NAME(FrequencyModPin), Fm);
+		}
+
+		virtual FDataReferenceCollection GetInputs() const override
+		{
+			// This should never be called. Bind(...) is called instead. This method
+			// exists as a stop-gap until the API can be deprecated and removed.
+			checkNoEntry();
+			return {};
+		}
 
-			FDataReferenceCollection Inputs = Super::GetInputs();
-			Inputs.AddDataReadReference(METASOUND_GET_PARAM_NAME(FrequencyModPin), Fm);
-			return Inputs;
+		virtual FDataReferenceCollection GetOutputs() const override
+		{
+			// This should never be called. Bind(...) is called instead. This method
+			// exists as a stop-gap until the API can be deprecated and removed.
+			checkNoEntry();
+			return {};
 		}
 
 		void Generate(int32 InStartFrame, int32 InEndFrame, float InClampedFreq, float InClampedGlideEase)
@@ -937,28 +964,40 @@ namespace Metasound
 			ResetPhase();
 		}
 
-		FDataReferenceCollection GetInputs() const override
+		virtual void BindInputs(FInputVertexInterfaceData& InOutVertexData) override
 		{
 			using namespace LfoVertexNames;
 
-			FDataReferenceCollection InputDataReferences;
-			InputDataReferences.AddDataReadReference(METASOUND_GET_PARAM_NAME(LfoBaseFrequencyPin), Frequency);
-			InputDataReferences.AddDataReadReference(METASOUND_GET_PARAM_NAME(WaveshapePin), Waveshape);
-			InputDataReferences.AddDataReadReference(METASOUND_GET_PARAM_NAME(MinOutputValuePin), MinValue);
-			InputDataReferences.AddDataReadReference(METASOUND_GET_PARAM_NAME(MaxOutputValuePin), MaxValue);
-			InputDataReferences.AddDataReadReference(METASOUND_GET_PARAM_NAME(PhaseOffsetPin), PhaseOffset);
-			InputDataReferences.AddDataReadReference(METASOUND_GET_PARAM_NAME(LfoPhaseResetPin), PhaseReset);
-			InputDataReferences.AddDataReadReference(METASOUND_GET_PARAM_NAME(LfoPulseWidthPin), PulseWidth);
-			return InputDataReferences;
+			InOutVertexData.BindReadVertex(METASOUND_GET_PARAM_NAME(LfoBaseFrequencyPin), Frequency);
+			InOutVertexData.BindReadVertex(METASOUND_GET_PARAM_NAME(WaveshapePin), Waveshape);
+			InOutVertexData.BindReadVertex(METASOUND_GET_PARAM_NAME(MinOutputValuePin), MinValue);
+			InOutVertexData.BindReadVertex(METASOUND_GET_PARAM_NAME(MaxOutputValuePin), MaxValue);
+			InOutVertexData.BindReadVertex(METASOUND_GET_PARAM_NAME(PhaseOffsetPin), PhaseOffset);
+			InOutVertexData.BindReadVertex(METASOUND_GET_PARAM_NAME(LfoPhaseResetPin), PhaseReset);
+			InOutVertexData.BindReadVertex(METASOUND_GET_PARAM_NAME(LfoPulseWidthPin), PulseWidth);
 		}
 
-		FDataReferenceCollection GetOutputs() const override
+		virtual void BindOutputs(FOutputVertexInterfaceData& InOutVertexData) override
 		{
 			using namespace LfoVertexNames;
 
-			FDataReferenceCollection OutputDataReferences;
-			OutputDataReferences.AddDataReadReference(METASOUND_GET_PARAM_NAME(LfoOutPin), Output);
-			return OutputDataReferences;
+			InOutVertexData.BindReadVertex(METASOUND_GET_PARAM_NAME(LfoOutPin), Output);
+		}
+
+		virtual FDataReferenceCollection GetInputs() const override
+		{
+			// This should never be called. Bind(...) is called instead. This method
+			// exists as a stop-gap until the API can be deprecated and removed.
+			checkNoEntry();
+			return {};
+		}
+
+		virtual FDataReferenceCollection GetOutputs() const override
+		{
+			// This should never be called. Bind(...) is called instead. This method
+			// exists as a stop-gap until the API can be deprecated and removed.
+			checkNoEntry();
+			return {};
 		}
 
 		void ResetPhase()
diff --git a/Engine/Plugins/Runtime/Metasound/Source/MetasoundStandardNodes/Private/MetasoundPerlinNoiseNode.cpp b/Engine/Plugins/Runtime/Metasound/Source/MetasoundStandardNodes/Private/MetasoundPerlinNoiseNode.cpp
--- a/Engine/Plugins/Runtime/Metasound/Source/MetasoundStandardNodes/Private/MetasoundPerlinNoiseNode.cpp
+++ b/Engine/Plugins/Runtime/Metasound/Source/MetasoundStandardNodes/Private/MetasoundPerlinNoiseNode.cpp
@@ -148,21 +148,42 @@ namespace Metasound
 			TThisType::ResetOutputWriteRef(OutputScaled);
 			TThisType::ResetOutputWriteRef(OutputNormalized);
 		}
-
-		FDataReferenceCollection GetInputs() const override { return {}; }
-		FDataReferenceCollection GetOutputs() const override { return {}; }
 		
 		void Bind(FVertexInterfaceData& InVertexData) const override
 		{
-			using namespace PerlinNoiseVertexNames;
-			FInputVertexInterfaceData& Inputs = InVertexData.GetInputs();
-			Inputs.BindReadVertex(METASOUND_GET_PARAM_NAME(XPin), Pins.X);
-			Inputs.SetValue(METASOUND_GET_PARAM_NAME(OffsetPin), Pins.Offset);
-			Inputs.BindReadVertex(METASOUND_GET_PARAM_NAME(OctavesPin), Pins.Octaves);
-			
+
 			FOutputVertexInterfaceData& Outputs = InVertexData.GetOutputs();
-			Outputs.BindReadVertex(METASOUND_GET_PARAM_NAME(OutputScaledPin), OutputScaled);
-			Outputs.BindReadVertex(METASOUND_GET_PARAM_NAME(OutputNormalizedPin), OutputNormalized);
+		}
+
+		virtual void BindInputs(FInputVertexInterfaceData& InOutVertexData) override
+		{
+			using namespace PerlinNoiseVertexNames;
+			InOutVertexData.BindReadVertex(METASOUND_GET_PARAM_NAME(XPin), Pins.X);
+			InOutVertexData.SetValue(METASOUND_GET_PARAM_NAME(OffsetPin), Pins.Offset);
+			InOutVertexData.BindReadVertex(METASOUND_GET_PARAM_NAME(OctavesPin), Pins.Octaves);
+		}
+
+		virtual void BindOutputs(FOutputVertexInterfaceData& InOutVertexData) override
+		{
+			using namespace PerlinNoiseVertexNames;
+			InOutVertexData.BindReadVertex(METASOUND_GET_PARAM_NAME(OutputScaledPin), OutputScaled);
+			InOutVertexData.BindReadVertex(METASOUND_GET_PARAM_NAME(OutputNormalizedPin), OutputNormalized);
+		}
+
+		virtual FDataReferenceCollection GetInputs() const override
+		{
+			// This should never be called. Bind(...) is called instead. This method
+			// exists as a stop-gap until the API can be deprecated and removed.
+			checkNoEntry();
+			return {};
+		}
+
+		virtual FDataReferenceCollection GetOutputs() const override
+		{
+			// This should never be called. Bind(...) is called instead. This method
+			// exists as a stop-gap until the API can be deprecated and removed.
+			checkNoEntry();
+			return {};
 		}
 
 		void AccumulateTime()
diff --git a/Engine/Plugins/Runtime/Metasound/Source/MetasoundStandardNodes/Private/MetasoundRingModNode.cpp b/Engine/Plugins/Runtime/Metasound/Source/MetasoundStandardNodes/Private/MetasoundRingModNode.cpp
--- a/Engine/Plugins/Runtime/Metasound/Source/MetasoundStandardNodes/Private/MetasoundRingModNode.cpp
+++ b/Engine/Plugins/Runtime/Metasound/Source/MetasoundStandardNodes/Private/MetasoundRingModNode.cpp
@@ -84,27 +84,36 @@ namespace Metasound
 			return Interface;
 		}
 
-		virtual FDataReferenceCollection GetInputs() const override
+
+		virtual void BindInputs(FInputVertexInterfaceData& InOutVertexData) override
 		{
 			using namespace RingModVertexNames;
 
-			FDataReferenceCollection InputDataReferences;
-			
-			InputDataReferences.AddDataReadReference(METASOUND_GET_PARAM_NAME(InputAudioCarrier), AudioInputCarrier);
-			InputDataReferences.AddDataReadReference(METASOUND_GET_PARAM_NAME(InputAudioModulator), AudioInputModulator);
-			
-			return InputDataReferences;
+			InOutVertexData.BindReadVertex(METASOUND_GET_PARAM_NAME(InputAudioCarrier), AudioInputCarrier);
+			InOutVertexData.BindReadVertex(METASOUND_GET_PARAM_NAME(InputAudioModulator), AudioInputModulator);
 		}
 
-		virtual FDataReferenceCollection GetOutputs() const override
+		virtual void BindOutputs(FOutputVertexInterfaceData& InOutVertexData) override
 		{
 			using namespace RingModVertexNames;
-			
-			FDataReferenceCollection OutputDataReferences;
 
-			OutputDataReferences.AddDataReadReference(METASOUND_GET_PARAM_NAME(OutputAudio), AudioOutput);
-			
-			return OutputDataReferences;
+			InOutVertexData.BindReadVertex(METASOUND_GET_PARAM_NAME(OutputAudio), AudioOutput);
+		}
+
+		virtual FDataReferenceCollection GetInputs() const override
+		{
+			// This should never be called. Bind(...) is called instead. This method
+			// exists as a stop-gap until the API can be deprecated and removed.
+			checkNoEntry();
+			return {};
+		}
+
+		virtual FDataReferenceCollection GetOutputs() const override
+		{
+			// This should never be called. Bind(...) is called instead. This method
+			// exists as a stop-gap until the API can be deprecated and removed.
+			checkNoEntry();
+			return {};
 		}
 		
 		static TUniquePtr<IOperator> CreateOperator(const FCreateOperatorParams& InParams, FBuildErrorArray& OutErrors)
diff --git a/Engine/Plugins/Runtime/Metasound/Source/MetasoundStandardNodes/Private/MetasoundSampleAndHoldNode.cpp b/Engine/Plugins/Runtime/Metasound/Source/MetasoundStandardNodes/Private/MetasoundSampleAndHoldNode.cpp
--- a/Engine/Plugins/Runtime/Metasound/Source/MetasoundStandardNodes/Private/MetasoundSampleAndHoldNode.cpp
+++ b/Engine/Plugins/Runtime/Metasound/Source/MetasoundStandardNodes/Private/MetasoundSampleAndHoldNode.cpp
@@ -95,25 +95,36 @@ namespace Metasound
 			Reset(InParams);
 		}
 
-		virtual FDataReferenceCollection GetInputs() const override
+		virtual void BindInputs(FInputVertexInterfaceData& InOutVertexData) override
 		{
 			using namespace SampleAndHoldVertexNames;
 
-			FDataReferenceCollection InputDataReferences;
-			InputDataReferences.AddDataReadReference(METASOUND_GET_PARAM_NAME(InputAudio), AudioInput);
-			InputDataReferences.AddDataReadReference(METASOUND_GET_PARAM_NAME(InputTriggerSampleAndHold), TriggerSampleAndHold);
+			InOutVertexData.BindReadVertex(METASOUND_GET_PARAM_NAME(InputAudio), AudioInput);
+			InOutVertexData.BindReadVertex(METASOUND_GET_PARAM_NAME(InputTriggerSampleAndHold), TriggerSampleAndHold);
+		}
 
-			return InputDataReferences;
+		virtual void BindOutputs(FOutputVertexInterfaceData& InOutVertexData) override
+		{
+			using namespace SampleAndHoldVertexNames;
 
+			InOutVertexData.BindReadVertex(METASOUND_GET_PARAM_NAME(OutputOnSampleAndHold), TriggerSampleAndHold);
+			InOutVertexData.BindReadVertex(METASOUND_GET_PARAM_NAME(OutputAudio), AudioOutput);
+		}
+
+		virtual FDataReferenceCollection GetInputs() const override
+		{
+			// This should never be called. Bind(...) is called instead. This method
+			// exists as a stop-gap until the API can be deprecated and removed.
+			checkNoEntry();
+			return {};
 		}
 
 		virtual FDataReferenceCollection GetOutputs() const override
 		{
-			using namespace SampleAndHoldVertexNames;
-			FDataReferenceCollection OutputDataReferences;
-			OutputDataReferences.AddDataReadReference(METASOUND_GET_PARAM_NAME(OutputOnSampleAndHold), TriggerSampleAndHold);
-			OutputDataReferences.AddDataReadReference(METASOUND_GET_PARAM_NAME(OutputAudio), AudioOutput);
-			return OutputDataReferences;
+			// This should never be called. Bind(...) is called instead. This method
+			// exists as a stop-gap until the API can be deprecated and removed.
+			checkNoEntry();
+			return {};
 		}
 		
 		void Execute()
diff --git a/Engine/Plugins/Runtime/Metasound/Source/MetasoundStandardNodes/Private/MetasoundSemitoneToFreqMultiplierNode.cpp b/Engine/Plugins/Runtime/Metasound/Source/MetasoundStandardNodes/Private/MetasoundSemitoneToFreqMultiplierNode.cpp
--- a/Engine/Plugins/Runtime/Metasound/Source/MetasoundStandardNodes/Private/MetasoundSemitoneToFreqMultiplierNode.cpp
+++ b/Engine/Plugins/Runtime/Metasound/Source/MetasoundStandardNodes/Private/MetasoundSemitoneToFreqMultiplierNode.cpp
@@ -33,6 +33,8 @@ namespace Metasound
 
 		FSemitoneToFrequencyMultiplierOperator(const FOperatorSettings& InSettings, const FFloatReadRef& InSemitone);
 
+		virtual void BindInputs(FInputVertexInterfaceData& InOutVertexData) override;
+		virtual void BindOutputs(FOutputVertexInterfaceData& InOutVertexData) override;
 		virtual FDataReferenceCollection GetInputs() const override;
 		virtual FDataReferenceCollection GetOutputs() const override;
 		void Execute();
@@ -56,24 +58,35 @@ namespace Metasound
 	{
 	}
 
-	
-	FDataReferenceCollection FSemitoneToFrequencyMultiplierOperator::GetInputs() const
+
+	void FSemitoneToFrequencyMultiplierOperator::BindInputs(FInputVertexInterfaceData& InOutVertexData)
 	{
 		using namespace SemitoneToFrequencyMultiplierVertexNames;
 
-		FDataReferenceCollection InputDataReferences;
-		InputDataReferences.AddDataReadReference(METASOUND_GET_PARAM_NAME(InputSemitone), SemitoneInput);
-
-		return InputDataReferences;
+		InOutVertexData.BindReadVertex(METASOUND_GET_PARAM_NAME(InputSemitone), SemitoneInput);
 	}
 
-	FDataReferenceCollection FSemitoneToFrequencyMultiplierOperator::GetOutputs() const
+	void FSemitoneToFrequencyMultiplierOperator::BindOutputs(FOutputVertexInterfaceData& InOutVertexData)
 	{
 		using namespace SemitoneToFrequencyMultiplierVertexNames;
 
-		FDataReferenceCollection OutputDataReferences;
-		OutputDataReferences.AddDataReadReference(METASOUND_GET_PARAM_NAME(OutputFrequencyMultiplier), FrequencyMultiplierOutput);
-		return OutputDataReferences;
+		InOutVertexData.BindReadVertex(METASOUND_GET_PARAM_NAME(OutputFrequencyMultiplier), FrequencyMultiplierOutput);
+	}
+
+	FDataReferenceCollection FSemitoneToFrequencyMultiplierOperator::GetInputs() const
+	{
+		// This should never be called. Bind(...) is called instead. This method
+		// exists as a stop-gap until the API can be deprecated and removed.
+		checkNoEntry();
+		return {};
+	}
+
+	FDataReferenceCollection FSemitoneToFrequencyMultiplierOperator::GetOutputs() const
+	{
+		// This should never be called. Bind(...) is called instead. This method
+		// exists as a stop-gap until the API can be deprecated and removed.
+		checkNoEntry();
+		return {};
 	}
 
 	void FSemitoneToFrequencyMultiplierOperator::Execute()
diff --git a/Engine/Plugins/Runtime/Metasound/Source/MetasoundStandardNodes/Private/MetasoundStereoDelayNode.cpp b/Engine/Plugins/Runtime/Metasound/Source/MetasoundStandardNodes/Private/MetasoundStereoDelayNode.cpp
--- a/Engine/Plugins/Runtime/Metasound/Source/MetasoundStandardNodes/Private/MetasoundStereoDelayNode.cpp
+++ b/Engine/Plugins/Runtime/Metasound/Source/MetasoundStandardNodes/Private/MetasoundStereoDelayNode.cpp
@@ -69,6 +69,9 @@ namespace Metasound
 			const FFloatReadRef& InWetLevel, 
 			const FFloatReadRef& InFeedback);
 
+
+		virtual void BindInputs(FInputVertexInterfaceData& InOutVertexData) override;
+		virtual void BindOutputs(FOutputVertexInterfaceData& InOutVertexData) override;
 		virtual FDataReferenceCollection GetInputs() const override;
 		virtual FDataReferenceCollection GetOutputs() const override;
 		void Execute();
@@ -142,31 +145,43 @@ namespace Metasound
 		RightDelayBuffer.SetDelayMsec(PrevDelayTimeMsec * (1.0f - PrevDelayRatio));
 	}
 
-	FDataReferenceCollection FStereoDelayOperator::GetInputs() const
+
+	void FStereoDelayOperator::BindInputs(FInputVertexInterfaceData& InOutVertexData)
 	{
-		using namespace StereoDelay; 
+		using namespace StereoDelay;
+
+		InOutVertexData.BindReadVertex(METASOUND_GET_PARAM_NAME(InAudioInputLeft), FAudioBufferReadRef(LeftAudioInput));
+		InOutVertexData.BindReadVertex(METASOUND_GET_PARAM_NAME(InAudioInputRight), FAudioBufferReadRef(RightAudioInput));
+		InOutVertexData.BindReadVertex(METASOUND_GET_PARAM_NAME(InDelayMode), FStereoDelayModeReadRef(StereoDelayMode));
+		InOutVertexData.BindReadVertex(METASOUND_GET_PARAM_NAME(InDelayTime), FTimeReadRef(DelayTime));
+		InOutVertexData.BindReadVertex(METASOUND_GET_PARAM_NAME(InDelayRatio), FFloatReadRef(DelayRatio));
+		InOutVertexData.BindReadVertex(METASOUND_GET_PARAM_NAME(InDryLevel), FFloatReadRef(DryLevel));
+		InOutVertexData.BindReadVertex(METASOUND_GET_PARAM_NAME(InWetLevel), FFloatReadRef(WetLevel));
+		InOutVertexData.BindReadVertex(METASOUND_GET_PARAM_NAME(InFeedbackAmount), FFloatReadRef(Feedback));
+	}
+
+	void FStereoDelayOperator::BindOutputs(FOutputVertexInterfaceData& InOutVertexData)
+	{
+		using namespace StereoDelay;
 
-		FDataReferenceCollection InputDataReferences;
-		InputDataReferences.AddDataReadReference(METASOUND_GET_PARAM_NAME(InAudioInputLeft), FAudioBufferReadRef(LeftAudioInput));
-		InputDataReferences.AddDataReadReference(METASOUND_GET_PARAM_NAME(InAudioInputRight), FAudioBufferReadRef(RightAudioInput));
-		InputDataReferences.AddDataReadReference(METASOUND_GET_PARAM_NAME(InDelayMode), FStereoDelayModeReadRef(StereoDelayMode));
-		InputDataReferences.AddDataReadReference(METASOUND_GET_PARAM_NAME(InDelayTime), FTimeReadRef(DelayTime));
-		InputDataReferences.AddDataReadReference(METASOUND_GET_PARAM_NAME(InDelayRatio), FFloatReadRef(DelayRatio));
-		InputDataReferences.AddDataReadReference(METASOUND_GET_PARAM_NAME(InDryLevel), FFloatReadRef(DryLevel));
-		InputDataReferences.AddDataReadReference(METASOUND_GET_PARAM_NAME(InWetLevel), FFloatReadRef(WetLevel));
-		InputDataReferences.AddDataReadReference(METASOUND_GET_PARAM_NAME(InFeedbackAmount), FFloatReadRef(Feedback));
-
-		return InputDataReferences;
+		InOutVertexData.BindReadVertex(METASOUND_GET_PARAM_NAME(OutAudioLeft), FAudioBufferReadRef(LeftAudioOutput));
+		InOutVertexData.BindReadVertex(METASOUND_GET_PARAM_NAME(OutAudioRight), FAudioBufferReadRef(RightAudioOutput));
 	}
 
-	FDataReferenceCollection FStereoDelayOperator::GetOutputs() const
+	FDataReferenceCollection FStereoDelayOperator::GetInputs() const
 	{
-		using namespace StereoDelay; 
+		// This should never be called. Bind(...) is called instead. This method
+		// exists as a stop-gap until the API can be deprecated and removed.
+		checkNoEntry();
+		return {};
+	}
 
-		FDataReferenceCollection OutputDataReferences;
-		OutputDataReferences.AddDataReadReference(METASOUND_GET_PARAM_NAME(OutAudioLeft), FAudioBufferReadRef(LeftAudioOutput));
-		OutputDataReferences.AddDataReadReference(METASOUND_GET_PARAM_NAME(OutAudioRight), FAudioBufferReadRef(RightAudioOutput));
-		return OutputDataReferences;
+	FDataReferenceCollection FStereoDelayOperator::GetOutputs() const
+	{
+		// This should never be called. Bind(...) is called instead. This method
+		// exists as a stop-gap until the API can be deprecated and removed.
+		checkNoEntry();
+		return {};
 	}
 
 	float FStereoDelayOperator::GetInputDelayTimeMsecClamped() const
diff --git a/Engine/Plugins/Runtime/Metasound/Source/MetasoundStandardNodes/Private/MetasoundStereoPannerNode.cpp b/Engine/Plugins/Runtime/Metasound/Source/MetasoundStandardNodes/Private/MetasoundStereoPannerNode.cpp
--- a/Engine/Plugins/Runtime/Metasound/Source/MetasoundStandardNodes/Private/MetasoundStereoPannerNode.cpp
+++ b/Engine/Plugins/Runtime/Metasound/Source/MetasoundStandardNodes/Private/MetasoundStereoPannerNode.cpp
@@ -57,6 +57,9 @@ namespace Metasound
 			const FFloatReadRef& InPanningAmount,
 			const FPanningLawReadRef& InPanningLaw);
 
+
+		virtual void BindInputs(FInputVertexInterfaceData& InOutVertexData) override;
+		virtual void BindOutputs(FOutputVertexInterfaceData& InOutVertexData) override;
 		virtual FDataReferenceCollection GetInputs() const override;
 		virtual FDataReferenceCollection GetOutputs() const override;
 		void Execute();
@@ -99,27 +102,38 @@ namespace Metasound
 		ComputePanGains(PrevPanningAmount, PrevLeftPan, PrevRightPan);
 	}
 
-	FDataReferenceCollection FStereoPannerOperator::GetInputs() const
+
+	void FStereoPannerOperator::BindInputs(FInputVertexInterfaceData& InOutVertexData)
 	{
 		using namespace StereoPannerVertexNames;
 
-		FDataReferenceCollection InputDataReferences;
-		InputDataReferences.AddDataReadReference(METASOUND_GET_PARAM_NAME(InputAudio), AudioInput);
-		InputDataReferences.AddDataReadReference(METASOUND_GET_PARAM_NAME(InputPanAmount), PanningAmount);
-		InputDataReferences.AddDataReadReference(METASOUND_GET_PARAM_NAME(InputPanningLaw), PanningLaw);
-
-		return InputDataReferences;
+		InOutVertexData.BindReadVertex(METASOUND_GET_PARAM_NAME(InputAudio), AudioInput);
+		InOutVertexData.BindReadVertex(METASOUND_GET_PARAM_NAME(InputPanAmount), PanningAmount);
+		InOutVertexData.BindReadVertex(METASOUND_GET_PARAM_NAME(InputPanningLaw), PanningLaw);
 	}
 
-	FDataReferenceCollection FStereoPannerOperator::GetOutputs() const
+	void FStereoPannerOperator::BindOutputs(FOutputVertexInterfaceData& InOutVertexData)
 	{
 		using namespace StereoPannerVertexNames;
 
-		FDataReferenceCollection OutputDataReferences;
-		OutputDataReferences.AddDataReadReference(METASOUND_GET_PARAM_NAME(OutputAudioLeft), AudioLeftOutput);
-		OutputDataReferences.AddDataReadReference(METASOUND_GET_PARAM_NAME(OutputAudioRight), AudioRightOutput);
+		InOutVertexData.BindReadVertex(METASOUND_GET_PARAM_NAME(OutputAudioLeft), AudioLeftOutput);
+		InOutVertexData.BindReadVertex(METASOUND_GET_PARAM_NAME(OutputAudioRight), AudioRightOutput);
+	}
+
+	FDataReferenceCollection FStereoPannerOperator::GetInputs() const
+	{
+		// This should never be called. Bind(...) is called instead. This method
+		// exists as a stop-gap until the API can be deprecated and removed.
+		checkNoEntry();
+		return {};
+	}
 
-		return OutputDataReferences;
+	FDataReferenceCollection FStereoPannerOperator::GetOutputs() const
+	{
+		// This should never be called. Bind(...) is called instead. This method
+		// exists as a stop-gap until the API can be deprecated and removed.
+		checkNoEntry();
+		return {};
 	}
 
 	void FStereoPannerOperator::ComputePanGains(float InPanningAmmount, float& OutLeftGain, float& OutRightGain) const
diff --git a/Engine/Plugins/Runtime/Metasound/Source/MetasoundStandardNodes/Private/MetasoundSuperOscillatorNode.cpp b/Engine/Plugins/Runtime/Metasound/Source/MetasoundStandardNodes/Private/MetasoundSuperOscillatorNode.cpp
--- a/Engine/Plugins/Runtime/Metasound/Source/MetasoundStandardNodes/Private/MetasoundSuperOscillatorNode.cpp
+++ b/Engine/Plugins/Runtime/Metasound/Source/MetasoundStandardNodes/Private/MetasoundSuperOscillatorNode.cpp
@@ -217,24 +217,22 @@ namespace Metasound
 			
 			InitializeGenerators();
 		}
-		
-		virtual void Bind(FVertexInterfaceData& InVertexData) const override
+
+		virtual void BindInputs(FInputVertexInterfaceData& InOutVertexData) override 
 		{
-			using namespace SuperOscillatorVertexNames; 
+			using namespace SuperOscillatorVertexNames;
 
-			FInputVertexInterfaceData& Inputs = InVertexData.GetInputs();
-
-			Inputs.BindReadVertex(METASOUND_GET_PARAM_NAME(EnabledPin), bEnabled);
-			Inputs.BindReadVertex(METASOUND_GET_PARAM_NAME(LimitOutputPin), bLimit);
-			Inputs.BindReadVertex(METASOUND_GET_PARAM_NAME(BaseFrequencyPin), BaseFrequency);
-			Inputs.SetValue(METASOUND_GET_PARAM_NAME(NumVoicesPin), NumVoices);
-			Inputs.BindReadVertex(METASOUND_GET_PARAM_NAME(MaxDetunePin), MaxDetune);
-			Inputs.BindReadVertex(METASOUND_GET_PARAM_NAME(DetuneVolumePin), DetuneDb);
-			Inputs.BindReadVertex(METASOUND_GET_PARAM_NAME(GlideFactorPin), GlideFactor);
-			Inputs.BindReadVertex(METASOUND_GET_PARAM_NAME(PulseWidthPin), PulseWidth);
-			Inputs.BindReadVertex(METASOUND_GET_PARAM_NAME(FrequencyModPin), Fm);
-			Inputs.BindReadVertex(METASOUND_GET_PARAM_NAME(EntropyPin), Entropy);
-			Inputs.SetValue(METASOUND_GET_PARAM_NAME(WaveTypePin), OscType);
+			InOutVertexData.BindReadVertex(METASOUND_GET_PARAM_NAME(EnabledPin), bEnabled);
+			InOutVertexData.BindReadVertex(METASOUND_GET_PARAM_NAME(LimitOutputPin), bLimit);
+			InOutVertexData.BindReadVertex(METASOUND_GET_PARAM_NAME(BaseFrequencyPin), BaseFrequency);
+			InOutVertexData.SetValue(METASOUND_GET_PARAM_NAME(NumVoicesPin), NumVoices);
+			InOutVertexData.BindReadVertex(METASOUND_GET_PARAM_NAME(MaxDetunePin), MaxDetune);
+			InOutVertexData.BindReadVertex(METASOUND_GET_PARAM_NAME(DetuneVolumePin), DetuneDb);
+			InOutVertexData.BindReadVertex(METASOUND_GET_PARAM_NAME(GlideFactorPin), GlideFactor);
+			InOutVertexData.BindReadVertex(METASOUND_GET_PARAM_NAME(PulseWidthPin), PulseWidth);
+			InOutVertexData.BindReadVertex(METASOUND_GET_PARAM_NAME(FrequencyModPin), Fm);
+			InOutVertexData.BindReadVertex(METASOUND_GET_PARAM_NAME(EntropyPin), Entropy);
+			InOutVertexData.SetValue(METASOUND_GET_PARAM_NAME(WaveTypePin), OscType);
 		}
 		
 		virtual FDataReferenceCollection GetInputs() const override
@@ -530,13 +528,15 @@ namespace Metasound
 			SetLimiterSettings(Limiter);
 		}
 
-		virtual void Bind(FVertexInterfaceData& InVertexData) const override
+		virtual void BindInputs(FInputVertexInterfaceData& InOutVertexData) override
+		{
+			Super::BindInputs(InOutVertexData);
+		}
+
+		virtual void BindOutputs(FOutputVertexInterfaceData& InOutVertexData) override
 		{
 			using namespace SuperOscillatorVertexNames;
-			Super::Bind(InVertexData);
-			
-			FOutputVertexInterfaceData& Outputs = InVertexData.GetOutputs();
-			Outputs.BindReadVertex(METASOUND_GET_PARAM_NAME(AudioOutPin), AudioBuffer);
+			InOutVertexData.BindReadVertex(METASOUND_GET_PARAM_NAME(AudioOutPin), AudioBuffer);
 		}
 		
 		FAudioBufferWriteRef AudioBuffer;
@@ -668,17 +668,20 @@ namespace Metasound
 			LimiterRight.ProcessAudio(AudioRight->GetData(), AudioRight->Num(), AudioRight->GetData());
 		}
 
-		virtual void Bind(FVertexInterfaceData& InVertexData) const override
+		virtual void BindInputs(FInputVertexInterfaceData& InOutVertexData) override
 		{
+			Super::BindInputs(InOutVertexData);
+			
 			using namespace SuperOscillatorVertexNames;
-			Super::Bind(InVertexData);
+			InOutVertexData.BindReadVertex(METASOUND_GET_PARAM_NAME(StereoWidthPin), StereoWidth);
+		}
 
-			FInputVertexInterfaceData& Inputs = InVertexData.GetInputs();
-			Inputs.BindReadVertex(METASOUND_GET_PARAM_NAME(StereoWidthPin), StereoWidth);
+		virtual void BindOutputs(FOutputVertexInterfaceData& InOutVertexData) override
+		{
+			using namespace SuperOscillatorVertexNames;
 			
-			FOutputVertexInterfaceData& Outputs = InVertexData.GetOutputs();
-			Outputs.BindReadVertex(METASOUND_GET_PARAM_NAME(AudioLeftPin), AudioLeft);
-			Outputs.BindReadVertex(METASOUND_GET_PARAM_NAME(AudioRightPin), AudioRight);
+			InOutVertexData.BindReadVertex(METASOUND_GET_PARAM_NAME(AudioLeftPin), AudioLeft);
+			InOutVertexData.BindReadVertex(METASOUND_GET_PARAM_NAME(AudioRightPin), AudioRight);
 		}
 
 		FFloatReadRef StereoWidth;
diff --git a/Engine/Plugins/Runtime/Metasound/Source/MetasoundStandardNodes/Private/MetasoundSwitchAndRampNode.cpp b/Engine/Plugins/Runtime/Metasound/Source/MetasoundStandardNodes/Private/MetasoundSwitchAndRampNode.cpp
--- a/Engine/Plugins/Runtime/Metasound/Source/MetasoundStandardNodes/Private/MetasoundSwitchAndRampNode.cpp
+++ b/Engine/Plugins/Runtime/Metasound/Source/MetasoundStandardNodes/Private/MetasoundSwitchAndRampNode.cpp
@@ -98,25 +98,37 @@ namespace Metasound
 		{
 		}
 
-		FDataReferenceCollection GetInputs() const
+
+		virtual void BindInputs(FInputVertexInterfaceData& InOutVertexData) override
 		{
 			using namespace SwitchAndRampVertexNames;
 
-			FDataReferenceCollection InputDataReferences;
-			InputDataReferences.AddDataReadReference(METASOUND_GET_PARAM_NAME(InputTrigger), TriggerIn);
-			InputDataReferences.AddDataReadReference(METASOUND_GET_PARAM_NAME(InputAudio), AudioIn);
-			InputDataReferences.AddDataReadReference(METASOUND_GET_PARAM_NAME(InputSmoothTime), SmoothTime);
-
-			return InputDataReferences;
+			InOutVertexData.BindReadVertex(METASOUND_GET_PARAM_NAME(InputTrigger), TriggerIn);
+			InOutVertexData.BindReadVertex(METASOUND_GET_PARAM_NAME(InputAudio), AudioIn);
+			InOutVertexData.BindReadVertex(METASOUND_GET_PARAM_NAME(InputSmoothTime), SmoothTime);
 		}
 
-		virtual FDataReferenceCollection GetOutputs() const
+		virtual void BindOutputs(FOutputVertexInterfaceData& InOutVertexData) override
 		{
 			using namespace SwitchAndRampVertexNames;
 
-			FDataReferenceCollection OutputDataReferences;
-			OutputDataReferences.AddDataReadReference(METASOUND_GET_PARAM_NAME(OutputAudio), AudioOut);
-			return OutputDataReferences;
+			InOutVertexData.BindReadVertex(METASOUND_GET_PARAM_NAME(OutputAudio), AudioOut);
+		}
+
+		virtual FDataReferenceCollection GetInputs() const override
+		{
+			// This should never be called. Bind(...) is called instead. This method
+			// exists as a stop-gap until the API can be deprecated and removed.
+			checkNoEntry();
+			return {};
+		}
+
+		virtual FDataReferenceCollection GetOutputs() const override
+		{
+			// This should never be called. Bind(...) is called instead. This method
+			// exists as a stop-gap until the API can be deprecated and removed.
+			checkNoEntry();
+			return {};
 		}
 
 		void Execute()
diff --git a/Engine/Plugins/Runtime/Metasound/Source/MetasoundStandardNodes/Private/MetasoundTriggerAnyNode.cpp b/Engine/Plugins/Runtime/Metasound/Source/MetasoundStandardNodes/Private/MetasoundTriggerAnyNode.cpp
--- a/Engine/Plugins/Runtime/Metasound/Source/MetasoundStandardNodes/Private/MetasoundTriggerAnyNode.cpp
+++ b/Engine/Plugins/Runtime/Metasound/Source/MetasoundStandardNodes/Private/MetasoundTriggerAnyNode.cpp
@@ -137,27 +137,35 @@ namespace Metasound
 		virtual ~TTriggerAnyOperator() = default;
 
 
-		virtual FDataReferenceCollection GetInputs() const override
+		virtual void BindInputs(FInputVertexInterfaceData& InOutVertexData) override
 		{
 			using namespace TriggerAnyVertexNames;
-
-			FDataReferenceCollection Inputs;
 			for (uint32 i = 0; i < NumInputs; ++i)
 			{
-				Inputs.AddDataReadReference(GetInputTriggerName(i), InputTriggers[i]);
+				InOutVertexData.BindReadVertex(GetInputTriggerName(i), InputTriggers[i]);
 			}
-
-			return Inputs;
 		}
 
-		virtual FDataReferenceCollection GetOutputs() const override
+		virtual void BindOutputs(FOutputVertexInterfaceData& InOutVertexData) override
 		{
 			using namespace TriggerAnyVertexNames;
+			InOutVertexData.BindReadVertex(METASOUND_GET_PARAM_NAME(OutputTrigger), OutputTrigger);
+		}
 
-			FDataReferenceCollection Outputs;
-			Outputs.AddDataReadReference(METASOUND_GET_PARAM_NAME(OutputTrigger), OutputTrigger);
+		virtual FDataReferenceCollection GetInputs() const override
+		{
+			// This should never be called. Bind(...) is called instead. This method
+			// exists as a stop-gap until the API can be deprecated and removed.
+			checkNoEntry();
+			return {};
+		}
 
-			return Outputs;
+		virtual FDataReferenceCollection GetOutputs() const override
+		{
+			// This should never be called. Bind(...) is called instead. This method
+			// exists as a stop-gap until the API can be deprecated and removed.
+			checkNoEntry();
+			return {};
 		}
 
 		void Execute()
diff --git a/Engine/Plugins/Runtime/Metasound/Source/MetasoundStandardNodes/Private/MetasoundTriggerCoinNode.cpp b/Engine/Plugins/Runtime/Metasound/Source/MetasoundStandardNodes/Private/MetasoundTriggerCoinNode.cpp
--- a/Engine/Plugins/Runtime/Metasound/Source/MetasoundStandardNodes/Private/MetasoundTriggerCoinNode.cpp
+++ b/Engine/Plugins/Runtime/Metasound/Source/MetasoundStandardNodes/Private/MetasoundTriggerCoinNode.cpp
@@ -43,8 +43,9 @@ namespace Metasound
 			const FInt32ReadRef& InSeed, 
 			const FFloatReadRef& InProbability);
 
+		virtual void BindInputs(FInputVertexInterfaceData& InOutVertexData) override;
+		virtual void BindOutputs(FOutputVertexInterfaceData& InOutVertexData) override;
 		virtual FDataReferenceCollection GetInputs() const override;
-
 		virtual FDataReferenceCollection GetOutputs() const override;
 
 		void Execute();
@@ -93,27 +94,37 @@ namespace Metasound
 		RandomStream.Reset();
 	}
 
-	FDataReferenceCollection FTriggerCoinOperator::GetInputs() const
+
+	void FTriggerCoinOperator::BindInputs(FInputVertexInterfaceData& InOutVertexData)
 	{
 		using namespace TriggerCoinVertexNames;
-
-		FDataReferenceCollection InputDataReferences;
-		InputDataReferences.AddDataReadReference(METASOUND_GET_PARAM_NAME(InputTrigger), TriggerInput);
-		InputDataReferences.AddDataReadReference(METASOUND_GET_PARAM_NAME(InputReset), TriggerResetInput);
-		InputDataReferences.AddDataReadReference(METASOUND_GET_PARAM_NAME(InputSeed), SeedInput);
-		InputDataReferences.AddDataReadReference(METASOUND_GET_PARAM_NAME(InputProbability), ProbabilityInput);
-		return InputDataReferences;
+		InOutVertexData.BindReadVertex(METASOUND_GET_PARAM_NAME(InputTrigger), TriggerInput);
+		InOutVertexData.BindReadVertex(METASOUND_GET_PARAM_NAME(InputReset), TriggerResetInput);
+		InOutVertexData.BindReadVertex(METASOUND_GET_PARAM_NAME(InputSeed), SeedInput);
+		InOutVertexData.BindReadVertex(METASOUND_GET_PARAM_NAME(InputProbability), ProbabilityInput);
 	}
 
-	FDataReferenceCollection FTriggerCoinOperator::GetOutputs() const
+	void FTriggerCoinOperator::BindOutputs(FOutputVertexInterfaceData& InOutVertexData)
 	{
 		using namespace TriggerCoinVertexNames;
+		InOutVertexData.BindReadVertex(METASOUND_GET_PARAM_NAME(OutputTrueTrigger), TriggerTrueOutput);
+		InOutVertexData.BindReadVertex(METASOUND_GET_PARAM_NAME(OutputFalseTrigger), TriggerFalseOutput);
+	}
 
-		FDataReferenceCollection OutputDataReferences;
-		OutputDataReferences.AddDataReadReference(METASOUND_GET_PARAM_NAME(OutputTrueTrigger), TriggerTrueOutput);
-		OutputDataReferences.AddDataReadReference(METASOUND_GET_PARAM_NAME(OutputFalseTrigger), TriggerFalseOutput);
+	FDataReferenceCollection FTriggerCoinOperator::GetInputs() const
+	{
+		// This should never be called. Bind(...) is called instead. This method
+		// exists as a stop-gap until the API can be deprecated and removed.
+		checkNoEntry();
+		return {};
+	}
 
-		return OutputDataReferences;
+	FDataReferenceCollection FTriggerCoinOperator::GetOutputs() const
+	{
+		// This should never be called. Bind(...) is called instead. This method
+		// exists as a stop-gap until the API can be deprecated and removed.
+		checkNoEntry();
+		return {};
 	}
 
 	void FTriggerCoinOperator::Execute()
diff --git a/Engine/Plugins/Runtime/Metasound/Source/MetasoundStandardNodes/Private/MetasoundTriggerControlNode.cpp b/Engine/Plugins/Runtime/Metasound/Source/MetasoundStandardNodes/Private/MetasoundTriggerControlNode.cpp
--- a/Engine/Plugins/Runtime/Metasound/Source/MetasoundStandardNodes/Private/MetasoundTriggerControlNode.cpp
+++ b/Engine/Plugins/Runtime/Metasound/Source/MetasoundStandardNodes/Private/MetasoundTriggerControlNode.cpp
@@ -45,6 +45,10 @@ namespace Metasound
 			const FTriggerReadRef& InTriggerToggle,
 			const FBoolReadRef& InStartClosed);
 
+		virtual void BindInputs(FInputVertexInterfaceData& InOutVertexData) override;
+
+		virtual void BindOutputs(FOutputVertexInterfaceData& InOutVertexData) override;
+
 		virtual FDataReferenceCollection GetInputs() const override;
 
 		virtual FDataReferenceCollection GetOutputs() const override;
@@ -87,27 +91,36 @@ namespace Metasound
 	{
 	}
 
-	FDataReferenceCollection FTriggerControlOperator::GetInputs() const
+	void FTriggerControlOperator::BindInputs(FInputVertexInterfaceData& InOutVertexData)
 	{
 		using namespace TriggerControlVertexNames;
-
-		FDataReferenceCollection InputDataReferences;
-		InputDataReferences.AddDataReadReference(METASOUND_GET_PARAM_NAME(InputEnter), TriggerEnterInput);
-		InputDataReferences.AddDataReadReference(METASOUND_GET_PARAM_NAME(InputOpen), TriggerOpenInput);
-		InputDataReferences.AddDataReadReference(METASOUND_GET_PARAM_NAME(InputClose), TriggerCloseInput);
-		InputDataReferences.AddDataReadReference(METASOUND_GET_PARAM_NAME(InputToggle), TriggerToggleInput);
-		InputDataReferences.AddDataReadReference(METASOUND_GET_PARAM_NAME(InputStartClosed), bStartClosedInput);
-		return InputDataReferences;
+		InOutVertexData.BindReadVertex(METASOUND_GET_PARAM_NAME(InputEnter), TriggerEnterInput);
+		InOutVertexData.BindReadVertex(METASOUND_GET_PARAM_NAME(InputOpen), TriggerOpenInput);
+		InOutVertexData.BindReadVertex(METASOUND_GET_PARAM_NAME(InputClose), TriggerCloseInput);
+		InOutVertexData.BindReadVertex(METASOUND_GET_PARAM_NAME(InputToggle), TriggerToggleInput);
+		InOutVertexData.BindReadVertex(METASOUND_GET_PARAM_NAME(InputStartClosed), bStartClosedInput);
 	}
 
-	FDataReferenceCollection FTriggerControlOperator::GetOutputs() const
+	void FTriggerControlOperator::BindOutputs(FOutputVertexInterfaceData& InOutVertexData)
 	{
 		using namespace TriggerControlVertexNames;
+		InOutVertexData.BindReadVertex(METASOUND_GET_PARAM_NAME(OutputExit), TriggerExitOutput);
+	}
 
-		FDataReferenceCollection OutputDataReferences;
-		OutputDataReferences.AddDataReadReference(METASOUND_GET_PARAM_NAME(OutputExit), TriggerExitOutput);
+	FDataReferenceCollection FTriggerControlOperator::GetInputs() const
+	{
+		// This should never be called. Bind(...) is called instead. This method
+		// exists as a stop-gap until the API can be deprecated and removed.
+		checkNoEntry();
+		return {};
+	}
 
-		return OutputDataReferences;
+	FDataReferenceCollection FTriggerControlOperator::GetOutputs() const
+	{
+		// This should never be called. Bind(...) is called instead. This method
+		// exists as a stop-gap until the API can be deprecated and removed.
+		checkNoEntry();
+		return {};
 	}
 
 	void FTriggerControlOperator::Execute()
diff --git a/Engine/Plugins/Runtime/Metasound/Source/MetasoundStandardNodes/Private/MetasoundTriggerCounterNode.cpp b/Engine/Plugins/Runtime/Metasound/Source/MetasoundStandardNodes/Private/MetasoundTriggerCounterNode.cpp
--- a/Engine/Plugins/Runtime/Metasound/Source/MetasoundStandardNodes/Private/MetasoundTriggerCounterNode.cpp
+++ b/Engine/Plugins/Runtime/Metasound/Source/MetasoundStandardNodes/Private/MetasoundTriggerCounterNode.cpp
@@ -46,6 +46,8 @@ namespace Metasound
 				const FFloatReadRef& InStepSize,
 				const FInt32ReadRef& InAutoResetCount);
 
+			virtual void BindInputs(FInputVertexInterfaceData& InOutVertexData) override;
+			virtual void BindOutputs(FOutputVertexInterfaceData& InOutVertexData) override;
 			virtual FDataReferenceCollection GetInputs() const override;
 			virtual FDataReferenceCollection GetOutputs() const override;
 
@@ -90,31 +92,43 @@ namespace Metasound
 		CurrentAutoResetCount = FMath::Max(0, *AutoResetCount);
 	}
 
-	FDataReferenceCollection FTriggerCounterOperator::GetInputs() const
+	void FTriggerCounterOperator::BindInputs(FInputVertexInterfaceData& InOutVertexData)
 	{
 		using namespace TriggerCounterVertexNames;
 
 		FDataReferenceCollection InputDataReferences;
-		InputDataReferences.AddDataReadReference(METASOUND_GET_PARAM_NAME(InputInTrigger), TriggerIn);
-		InputDataReferences.AddDataReadReference(METASOUND_GET_PARAM_NAME(InputReset), TriggerReset);
-		InputDataReferences.AddDataReadReference(METASOUND_GET_PARAM_NAME(InputStartValue), StartValue);
-		InputDataReferences.AddDataReadReference(METASOUND_GET_PARAM_NAME(InputStepSize), StepSize);
-		InputDataReferences.AddDataReadReference(METASOUND_GET_PARAM_NAME(InputAutoResetCount), AutoResetCount);
-
-		return InputDataReferences;
+		InOutVertexData.BindReadVertex(METASOUND_GET_PARAM_NAME(InputInTrigger), TriggerIn);
+		InOutVertexData.BindReadVertex(METASOUND_GET_PARAM_NAME(InputReset), TriggerReset);
+		InOutVertexData.BindReadVertex(METASOUND_GET_PARAM_NAME(InputStartValue), StartValue);
+		InOutVertexData.BindReadVertex(METASOUND_GET_PARAM_NAME(InputStepSize), StepSize);
+		InOutVertexData.BindReadVertex(METASOUND_GET_PARAM_NAME(InputAutoResetCount), AutoResetCount);
 	}
 
-	FDataReferenceCollection FTriggerCounterOperator::GetOutputs() const
+	void FTriggerCounterOperator::BindOutputs(FOutputVertexInterfaceData& InOutVertexData)
 	{
 		using namespace TriggerCounterVertexNames;
 
 		FDataReferenceCollection OutputDataReferences;
-		OutputDataReferences.AddDataReadReference(METASOUND_GET_PARAM_NAME(OutputOnTrigger), TriggerOut);
-		OutputDataReferences.AddDataReadReference(METASOUND_GET_PARAM_NAME(OutputOnReset), TriggerOnReset);
-		OutputDataReferences.AddDataReadReference(METASOUND_GET_PARAM_NAME(OutputCount), OutCount);
-		OutputDataReferences.AddDataReadReference(METASOUND_GET_PARAM_NAME(OutputValue), OutValue);
+		InOutVertexData.BindReadVertex(METASOUND_GET_PARAM_NAME(OutputOnTrigger), TriggerOut);
+		InOutVertexData.BindReadVertex(METASOUND_GET_PARAM_NAME(OutputOnReset), TriggerOnReset);
+		InOutVertexData.BindReadVertex(METASOUND_GET_PARAM_NAME(OutputCount), OutCount);
+		InOutVertexData.BindReadVertex(METASOUND_GET_PARAM_NAME(OutputValue), OutValue);
+	}
+
+	FDataReferenceCollection FTriggerCounterOperator::GetInputs() const
+	{
+		// This should never be called. Bind(...) is called instead. This method
+		// exists as a stop-gap until the API can be deprecated and removed.
+		checkNoEntry();
+		return {};
+	}
 
-		return OutputDataReferences;
+	FDataReferenceCollection FTriggerCounterOperator::GetOutputs() const
+	{
+		// This should never be called. Bind(...) is called instead. This method
+		// exists as a stop-gap until the API can be deprecated and removed.
+		checkNoEntry();
+		return {};
 	}
 
 	void FTriggerCounterOperator::Execute()
diff --git a/Engine/Plugins/Runtime/Metasound/Source/MetasoundStandardNodes/Private/MetasoundTriggerDelayNode.cpp b/Engine/Plugins/Runtime/Metasound/Source/MetasoundStandardNodes/Private/MetasoundTriggerDelayNode.cpp
--- a/Engine/Plugins/Runtime/Metasound/Source/MetasoundStandardNodes/Private/MetasoundTriggerDelayNode.cpp
+++ b/Engine/Plugins/Runtime/Metasound/Source/MetasoundStandardNodes/Private/MetasoundTriggerDelayNode.cpp
@@ -35,6 +35,8 @@ namespace Metasound
 
 			FTriggerDelayOperator(const FOperatorSettings& InSettings, const FTriggerReadRef& InTriggerReset, const FTriggerReadRef& InTriggerIn, const FTimeReadRef& InDelay);
 
+			virtual void BindInputs(FInputVertexInterfaceData& InOutVertexData) override;
+			virtual void BindOutputs(FOutputVertexInterfaceData& InOutVertexData) override;
 			virtual FDataReferenceCollection GetInputs() const override;
 			virtual FDataReferenceCollection GetOutputs() const override;
 
@@ -59,25 +61,34 @@ namespace Metasound
 	{
 	}
 
-	FDataReferenceCollection FTriggerDelayOperator::GetInputs() const
+	void FTriggerDelayOperator::BindInputs(FInputVertexInterfaceData& InOutVertexData)
 	{
 		using namespace TriggerDelayVertexNames;
-
-		FDataReferenceCollection InputDataReferences;
-		InputDataReferences.AddDataReadReference(METASOUND_GET_PARAM_NAME(InputInTriggerDelay), TriggerIn);
-		InputDataReferences.AddDataReadReference(METASOUND_GET_PARAM_NAME(InputResetDelay), TriggerReset);
-		InputDataReferences.AddDataReadReference(METASOUND_GET_PARAM_NAME(InputDelayTime), DelayTime);
-		return InputDataReferences;
+		InOutVertexData.BindReadVertex(METASOUND_GET_PARAM_NAME(InputInTriggerDelay), TriggerIn);
+		InOutVertexData.BindReadVertex(METASOUND_GET_PARAM_NAME(InputResetDelay), TriggerReset);
+		InOutVertexData.BindReadVertex(METASOUND_GET_PARAM_NAME(InputDelayTime), DelayTime);
 	}
 
-	FDataReferenceCollection FTriggerDelayOperator::GetOutputs() const
+	void FTriggerDelayOperator::BindOutputs(FOutputVertexInterfaceData& InOutVertexData)
 	{
 		using namespace TriggerDelayVertexNames;
+		InOutVertexData.BindReadVertex(METASOUND_GET_PARAM_NAME(OutputOnTrigger), TriggerOut);
+	}
 
-		FDataReferenceCollection OutputDataReferences;
-		OutputDataReferences.AddDataReadReference(METASOUND_GET_PARAM_NAME(OutputOnTrigger), TriggerOut);
+	FDataReferenceCollection FTriggerDelayOperator::GetInputs() const
+	{
+		// This should never be called. Bind(...) is called instead. This method
+		// exists as a stop-gap until the API can be deprecated and removed.
+		checkNoEntry();
+		return {};
+	}
 
-		return OutputDataReferences;
+	FDataReferenceCollection FTriggerDelayOperator::GetOutputs() const
+	{
+		// This should never be called. Bind(...) is called instead. This method
+		// exists as a stop-gap until the API can be deprecated and removed.
+		checkNoEntry();
+		return {};
 	}
 
 	void FTriggerDelayOperator::Execute()
diff --git a/Engine/Plugins/Runtime/Metasound/Source/MetasoundStandardNodes/Private/MetasoundTriggerOnThresholdNode.cpp b/Engine/Plugins/Runtime/Metasound/Source/MetasoundStandardNodes/Private/MetasoundTriggerOnThresholdNode.cpp
--- a/Engine/Plugins/Runtime/Metasound/Source/MetasoundStandardNodes/Private/MetasoundTriggerOnThresholdNode.cpp
+++ b/Engine/Plugins/Runtime/Metasound/Source/MetasoundStandardNodes/Private/MetasoundTriggerOnThresholdNode.cpp
@@ -242,24 +242,34 @@ namespace Metasound
 			LastSample = 0;
 		}
 
-		FDataReferenceCollection GetInputs() const
+		virtual void BindInputs(FInputVertexInterfaceData& InOutVertexData) override
 		{
 			using namespace TriggerOnThresholdVertexNames;
-
-			FDataReferenceCollection InputDataReferences;
-			InputDataReferences.AddDataReadReference(METASOUND_GET_PARAM_NAME(InThresholdPin), Threshold);
-			InputDataReferences.AddDataReadReference(METASOUND_GET_PARAM_NAME(InPin), In);
-			InputDataReferences.AddDataReadReference(METASOUND_GET_PARAM_NAME(InTriggerType), TriggerType);
-			return InputDataReferences;
+			InOutVertexData.BindReadVertex(METASOUND_GET_PARAM_NAME(InThresholdPin), Threshold);
+			InOutVertexData.BindReadVertex(METASOUND_GET_PARAM_NAME(InPin), In);
+			InOutVertexData.BindReadVertex(METASOUND_GET_PARAM_NAME(InTriggerType), TriggerType);
 		}
 
-		FDataReferenceCollection GetOutputs() const
+		virtual void BindOutputs(FOutputVertexInterfaceData& InOutVertexData) override
 		{
 			using namespace TriggerOnThresholdVertexNames;
+			InOutVertexData.BindReadVertex(METASOUND_GET_PARAM_NAME(OutPin), FTriggerWriteRef(Out));
+		}
+
+		virtual FDataReferenceCollection GetInputs() const override
+		{
+			// This should never be called. Bind(...) is called instead. This method
+			// exists as a stop-gap until the API can be deprecated and removed.
+			checkNoEntry();
+			return {};
+		}
 
-			FDataReferenceCollection OutputDataReferences;
-			OutputDataReferences.AddDataReadReference(METASOUND_GET_PARAM_NAME(OutPin), FTriggerWriteRef(Out));
-			return OutputDataReferences;
+		virtual FDataReferenceCollection GetOutputs() const override
+		{
+			// This should never be called. Bind(...) is called instead. This method
+			// exists as a stop-gap until the API can be deprecated and removed.
+			checkNoEntry();
+			return {};
 		}
 
 		static const FVertexInterface DeclareVertexInterface()
diff --git a/Engine/Plugins/Runtime/Metasound/Source/MetasoundStandardNodes/Private/MetasoundTriggerOnValueChangeNode.cpp b/Engine/Plugins/Runtime/Metasound/Source/MetasoundStandardNodes/Private/MetasoundTriggerOnValueChangeNode.cpp
--- a/Engine/Plugins/Runtime/Metasound/Source/MetasoundStandardNodes/Private/MetasoundTriggerOnValueChangeNode.cpp
+++ b/Engine/Plugins/Runtime/Metasound/Source/MetasoundStandardNodes/Private/MetasoundTriggerOnValueChangeNode.cpp
@@ -51,6 +51,10 @@ namespace Metasound
 
 		TTriggerOnValueChangeOperator(const FOperatorSettings& InSettings, const TDataReadReference<ValueType>& InValue);
 
+		virtual void BindInputs(FInputVertexInterfaceData& InOutVertexData) override;
+
+		virtual void BindOutputs(FOutputVertexInterfaceData& InOutVertexData) override;
+
 		virtual FDataReferenceCollection GetInputs() const override;
 
 		virtual FDataReferenceCollection GetOutputs() const override;
@@ -78,24 +82,35 @@ namespace Metasound
 	}
 
 	template <typename ValueType>
-	FDataReferenceCollection TTriggerOnValueChangeOperator<ValueType>::GetInputs() const
+	void TTriggerOnValueChangeOperator<ValueType>::BindInputs(FInputVertexInterfaceData& InOutVertexData)
 	{
 		using namespace TriggerOnValueChangeVertexNames;
-
-		FDataReferenceCollection InputDataReferences;
-		InputDataReferences.AddDataReadReference(METASOUND_GET_PARAM_NAME(InputValue), ValueInput);
-		return InputDataReferences;
+		InOutVertexData.BindReadVertex(METASOUND_GET_PARAM_NAME(InputValue), ValueInput);
 	}
 
 	template <typename ValueType>
-	FDataReferenceCollection TTriggerOnValueChangeOperator<ValueType>::GetOutputs() const
+	void TTriggerOnValueChangeOperator<ValueType>::BindOutputs(FOutputVertexInterfaceData& InOutVertexData)
 	{
 		using namespace TriggerOnValueChangeVertexNames;
+		InOutVertexData.BindReadVertex(METASOUND_GET_PARAM_NAME(OutputOnChange), TriggerOnChangeOutput);
+	}
 
-		FDataReferenceCollection OutputDataReferences;
-		OutputDataReferences.AddDataReadReference(METASOUND_GET_PARAM_NAME(OutputOnChange), TriggerOnChangeOutput);
+	template <typename ValueType>
+	FDataReferenceCollection TTriggerOnValueChangeOperator<ValueType>::GetInputs() const
+	{
+		// This should never be called. Bind(...) is called instead. This method
+		// exists as a stop-gap until the API can be deprecated and removed.
+		checkNoEntry();
+		return {};
+	}
 
-		return OutputDataReferences;
+	template <typename ValueType>
+	FDataReferenceCollection TTriggerOnValueChangeOperator<ValueType>::GetOutputs() const
+	{
+		// This should never be called. Bind(...) is called instead. This method
+		// exists as a stop-gap until the API can be deprecated and removed.
+		checkNoEntry();
+		return {};
 	}
 
 	template <typename ValueType>
diff --git a/Engine/Plugins/Runtime/Metasound/Source/MetasoundStandardNodes/Private/MetasoundTriggerOnceNode.cpp b/Engine/Plugins/Runtime/Metasound/Source/MetasoundStandardNodes/Private/MetasoundTriggerOnceNode.cpp
--- a/Engine/Plugins/Runtime/Metasound/Source/MetasoundStandardNodes/Private/MetasoundTriggerOnceNode.cpp
+++ b/Engine/Plugins/Runtime/Metasound/Source/MetasoundStandardNodes/Private/MetasoundTriggerOnceNode.cpp
@@ -41,8 +41,9 @@ namespace Metasound
 			const FTriggerReadRef& InTriggerReset,
 			const FBoolReadRef& InStartClosed);
 
+		virtual void BindInputs(FInputVertexInterfaceData& InOutVertexData) override;
+		virtual void BindOutputs(FOutputVertexInterfaceData& InOutVertexData) override;
 		virtual FDataReferenceCollection GetInputs() const override;
-
 		virtual FDataReferenceCollection GetOutputs() const override;
 
 		void Execute();
@@ -75,25 +76,35 @@ namespace Metasound
 	{
 	}
 
-	FDataReferenceCollection FTriggerOnceOperator::GetInputs() const
+
+	void FTriggerOnceOperator::BindInputs(FInputVertexInterfaceData& InOutVertexData)
 	{
 		using namespace TriggerOnceVertexNames;
-
-		FDataReferenceCollection InputDataReferences;
-		InputDataReferences.AddDataReadReference(METASOUND_GET_PARAM_NAME(InputEnter), TriggerEnterInput);
-		InputDataReferences.AddDataReadReference(METASOUND_GET_PARAM_NAME(InputReset), TriggerResetInput);
-		InputDataReferences.AddDataReadReference(METASOUND_GET_PARAM_NAME(InputStartClosed), bStartClosedInput);
-		return InputDataReferences;
+		InOutVertexData.BindReadVertex(METASOUND_GET_PARAM_NAME(InputEnter), TriggerEnterInput);
+		InOutVertexData.BindReadVertex(METASOUND_GET_PARAM_NAME(InputReset), TriggerResetInput);
+		InOutVertexData.BindReadVertex(METASOUND_GET_PARAM_NAME(InputStartClosed), bStartClosedInput);
 	}
 
-	FDataReferenceCollection FTriggerOnceOperator::GetOutputs() const
+	void FTriggerOnceOperator::BindOutputs(FOutputVertexInterfaceData& InOutVertexData)
 	{
 		using namespace TriggerOnceVertexNames;
+		InOutVertexData.BindReadVertex(METASOUND_GET_PARAM_NAME(OutputExit), TriggerExitOutput);
+	}
 
-		FDataReferenceCollection OutputDataReferences;
-		OutputDataReferences.AddDataReadReference(METASOUND_GET_PARAM_NAME(OutputExit), TriggerExitOutput);
+	FDataReferenceCollection FTriggerOnceOperator::GetInputs() const
+	{
+		// This should never be called. Bind(...) is called instead. This method
+		// exists as a stop-gap until the API can be deprecated and removed.
+		checkNoEntry();
+		return {};
+	}
 
-		return OutputDataReferences;
+	FDataReferenceCollection FTriggerOnceOperator::GetOutputs() const
+	{
+		// This should never be called. Bind(...) is called instead. This method
+		// exists as a stop-gap until the API can be deprecated and removed.
+		checkNoEntry();
+		return {};
 	}
 
 	void FTriggerOnceOperator::Execute()
diff --git a/Engine/Plugins/Runtime/Metasound/Source/MetasoundStandardNodes/Private/MetasoundTriggerPipeNode.cpp b/Engine/Plugins/Runtime/Metasound/Source/MetasoundStandardNodes/Private/MetasoundTriggerPipeNode.cpp
--- a/Engine/Plugins/Runtime/Metasound/Source/MetasoundStandardNodes/Private/MetasoundTriggerPipeNode.cpp
+++ b/Engine/Plugins/Runtime/Metasound/Source/MetasoundStandardNodes/Private/MetasoundTriggerPipeNode.cpp
@@ -38,6 +38,8 @@ namespace Metasound
 
 		FTriggerPipeOperator(const FOperatorSettings& InSettings, const FTriggerReadRef& InTriggerReset, const FTriggerReadRef& InTriggerIn, const FTimeReadRef& InDelay);
 
+		virtual void BindInputs(FInputVertexInterfaceData& InOutVertexData) override;
+		virtual void BindOutputs(FOutputVertexInterfaceData& InOutVertexData) override;
 		virtual FDataReferenceCollection GetInputs() const override;
 		virtual FDataReferenceCollection GetOutputs() const override;
 		void Execute();
@@ -66,25 +68,34 @@ namespace Metasound
 	{
 	}
 
-	FDataReferenceCollection FTriggerPipeOperator::GetInputs() const
+	void FTriggerPipeOperator::BindInputs(FInputVertexInterfaceData& InOutVertexData)
 	{
 		using namespace TriggerPipeVertexNames;
-
-		FDataReferenceCollection InputDataReferences;
-		InputDataReferences.AddDataReadReference(METASOUND_GET_PARAM_NAME(InputInTrigger), TriggerIn);
-		InputDataReferences.AddDataReadReference(METASOUND_GET_PARAM_NAME(InputReset), TriggerReset);
-		InputDataReferences.AddDataReadReference(METASOUND_GET_PARAM_NAME(InputDelayTime), DelayTime);
-		return InputDataReferences;
+		InOutVertexData.BindReadVertex(METASOUND_GET_PARAM_NAME(InputInTrigger), TriggerIn);
+		InOutVertexData.BindReadVertex(METASOUND_GET_PARAM_NAME(InputReset), TriggerReset);
+		InOutVertexData.BindReadVertex(METASOUND_GET_PARAM_NAME(InputDelayTime), DelayTime);
 	}
 
-	FDataReferenceCollection FTriggerPipeOperator::GetOutputs() const
+	void FTriggerPipeOperator::BindOutputs(FOutputVertexInterfaceData& InOutVertexData)
 	{
 		using namespace TriggerPipeVertexNames;
+		InOutVertexData.BindReadVertex(METASOUND_GET_PARAM_NAME(OutputOutTrigger), TriggerOut);
+	}
 
-		FDataReferenceCollection OutputDataReferences;
-		OutputDataReferences.AddDataReadReference(METASOUND_GET_PARAM_NAME(OutputOutTrigger), TriggerOut);
+	FDataReferenceCollection FTriggerPipeOperator::GetInputs() const
+	{
+		// This should never be called. Bind(...) is called instead. This method
+		// exists as a stop-gap until the API can be deprecated and removed.
+		checkNoEntry();
+		return {};
+	}
 
-		return OutputDataReferences;
+	FDataReferenceCollection FTriggerPipeOperator::GetOutputs() const
+	{
+		// This should never be called. Bind(...) is called instead. This method
+		// exists as a stop-gap until the API can be deprecated and removed.
+		checkNoEntry();
+		return {};
 	}
 
 	void FTriggerPipeOperator::Execute()
diff --git a/Engine/Plugins/Runtime/Metasound/Source/MetasoundStandardNodes/Private/MetasoundTriggerRepeatNode.cpp b/Engine/Plugins/Runtime/Metasound/Source/MetasoundStandardNodes/Private/MetasoundTriggerRepeatNode.cpp
--- a/Engine/Plugins/Runtime/Metasound/Source/MetasoundStandardNodes/Private/MetasoundTriggerRepeatNode.cpp
+++ b/Engine/Plugins/Runtime/Metasound/Source/MetasoundStandardNodes/Private/MetasoundTriggerRepeatNode.cpp
@@ -41,6 +41,10 @@ namespace Metasound
 
 			FTriggerRepeatOperator(const FOperatorSettings& InSettings, const FTriggerReadRef& InTriggerEnable, const FTriggerReadRef& InTriggerDisable, const FTimeReadRef& InPeriod);
 
+			virtual void BindInputs(FInputVertexInterfaceData& InOutVertexData) override;
+
+			virtual void BindOutputs(FOutputVertexInterfaceData& InOutVertexData) override;
+
 			virtual FDataReferenceCollection GetInputs() const override;
 
 			virtual FDataReferenceCollection GetOutputs() const override;
@@ -73,25 +77,35 @@ namespace Metasound
 	{
 	}
 
-	FDataReferenceCollection FTriggerRepeatOperator::GetInputs() const
+
+	void FTriggerRepeatOperator::BindInputs(FInputVertexInterfaceData& InOutVertexData)
 	{
 		using namespace TriggerRepeatVertexNames;
-
-		FDataReferenceCollection InputDataReferences;
-		InputDataReferences.AddDataReadReference(METASOUND_GET_PARAM_NAME(InputStart), TriggerEnable);
-		InputDataReferences.AddDataReadReference(METASOUND_GET_PARAM_NAME(InputStop), TriggerDisable);
-		InputDataReferences.AddDataReadReference(METASOUND_GET_PARAM_NAME(InputPeriod), Period);
-		return InputDataReferences;
+		InOutVertexData.BindReadVertex(METASOUND_GET_PARAM_NAME(InputStart), TriggerEnable);
+		InOutVertexData.BindReadVertex(METASOUND_GET_PARAM_NAME(InputStop), TriggerDisable);
+		InOutVertexData.BindReadVertex(METASOUND_GET_PARAM_NAME(InputPeriod), Period);
 	}
 
-	FDataReferenceCollection FTriggerRepeatOperator::GetOutputs() const
+	void FTriggerRepeatOperator::BindOutputs(FOutputVertexInterfaceData& InOutVertexData)
 	{
 		using namespace TriggerRepeatVertexNames;
+		InOutVertexData.BindReadVertex(METASOUND_GET_PARAM_NAME(RepeatOutputOnTrigger), TriggerOut);
+	}
 
-		FDataReferenceCollection OutputDataReferences;
-		OutputDataReferences.AddDataReadReference(METASOUND_GET_PARAM_NAME(RepeatOutputOnTrigger), TriggerOut);
+	FDataReferenceCollection FTriggerRepeatOperator::GetInputs() const
+	{
+		// This should never be called. Bind(...) is called instead. This method
+		// exists as a stop-gap until the API can be deprecated and removed.
+		checkNoEntry();
+		return {};
+	}
 
-		return OutputDataReferences;
+	FDataReferenceCollection FTriggerRepeatOperator::GetOutputs() const
+	{
+		// This should never be called. Bind(...) is called instead. This method
+		// exists as a stop-gap until the API can be deprecated and removed.
+		checkNoEntry();
+		return {};
 	}
 
 	void FTriggerRepeatOperator::Execute()
diff --git a/Engine/Plugins/Runtime/Metasound/Source/MetasoundStandardNodes/Private/MetasoundTriggerRouteNode.cpp b/Engine/Plugins/Runtime/Metasound/Source/MetasoundStandardNodes/Private/MetasoundTriggerRouteNode.cpp
--- a/Engine/Plugins/Runtime/Metasound/Source/MetasoundStandardNodes/Private/MetasoundTriggerRouteNode.cpp
+++ b/Engine/Plugins/Runtime/Metasound/Source/MetasoundStandardNodes/Private/MetasoundTriggerRouteNode.cpp
@@ -178,30 +178,37 @@ namespace Metasound
 
 		virtual ~TTriggerRouteOperator() = default;
 
-
-		virtual FDataReferenceCollection GetInputs() const override
+		virtual void BindInputs(FInputVertexInterfaceData& InOutVertexData) override
 		{
 			using namespace TriggerRouteVertexNames;
-
-			FDataReferenceCollection Inputs;
 			for (uint32 i = 0; i < NumInputs; ++i)
 			{
-				Inputs.AddDataReadReference(GetInputTriggerName(i), InputTriggers[i]);
-				Inputs.AddDataReadReference(GetInputValueName(i), InputValues[i]);
+				InOutVertexData.BindReadVertex(GetInputTriggerName(i), InputTriggers[i]);
+				InOutVertexData.BindReadVertex(GetInputValueName(i), InputValues[i]);
 			}
-
-			return Inputs;
 		}
 
-		virtual FDataReferenceCollection GetOutputs() const override
+		virtual void BindOutputs(FOutputVertexInterfaceData& InOutVertexData) override
 		{
 			using namespace TriggerRouteVertexNames;
+			InOutVertexData.BindReadVertex(METASOUND_GET_PARAM_NAME(OutputTrigger), OutputTrigger);
+			InOutVertexData.BindReadVertex(METASOUND_GET_PARAM_NAME(OutputValue), OutputValue);
+		}
 
-			FDataReferenceCollection Outputs;
-			Outputs.AddDataReadReference(METASOUND_GET_PARAM_NAME(OutputTrigger), OutputTrigger);
-			Outputs.AddDataReadReference(METASOUND_GET_PARAM_NAME(OutputValue), OutputValue);
+		virtual FDataReferenceCollection GetInputs() const override
+		{
+			// This should never be called. Bind(...) is called instead. This method
+			// exists as a stop-gap until the API can be deprecated and removed.
+			checkNoEntry();
+			return {};
+		}
 
-			return Outputs;
+		virtual FDataReferenceCollection GetOutputs() const override
+		{
+			// This should never be called. Bind(...) is called instead. This method
+			// exists as a stop-gap until the API can be deprecated and removed.
+			checkNoEntry();
+			return {};
 		}
 
 		void Execute()
diff --git a/Engine/Plugins/Runtime/Metasound/Source/MetasoundStandardNodes/Private/MetasoundTriggerSelectNode.cpp b/Engine/Plugins/Runtime/Metasound/Source/MetasoundStandardNodes/Private/MetasoundTriggerSelectNode.cpp
--- a/Engine/Plugins/Runtime/Metasound/Source/MetasoundStandardNodes/Private/MetasoundTriggerSelectNode.cpp
+++ b/Engine/Plugins/Runtime/Metasound/Source/MetasoundStandardNodes/Private/MetasoundTriggerSelectNode.cpp
@@ -44,8 +44,9 @@ namespace Metasound
 			const FTriggerReadRef& InTrigger,
 			const FInt32ReadRef& InIndex);
 
+		virtual void BindInputs(FInputVertexInterfaceData& InOutVertexData) override;
+		virtual void BindOutputs(FOutputVertexInterfaceData& InOutVertexData) override;
 		virtual FDataReferenceCollection GetInputs() const override;
-
 		virtual FDataReferenceCollection GetOutputs() const override;
 
 		void Execute();
@@ -72,28 +73,40 @@ namespace Metasound
 	}
 
 	template<uint32 NumOutputs>
-	FDataReferenceCollection TTriggerSelectOperator<NumOutputs>::GetInputs() const
+	void TTriggerSelectOperator<NumOutputs>::BindInputs(FInputVertexInterfaceData& InOutVertexData)
 	{
 		using namespace TriggerSelectVertexNames;
-
-		FDataReferenceCollection InputDataReferences;
-		InputDataReferences.AddDataReadReference(METASOUND_GET_PARAM_NAME(InputTrigger), TriggerInput);
-		InputDataReferences.AddDataReadReference(METASOUND_GET_PARAM_NAME(InputIndex), IndexInput);
-		return InputDataReferences;
+		InOutVertexData.BindReadVertex(METASOUND_GET_PARAM_NAME(InputTrigger), TriggerInput);
+		InOutVertexData.BindReadVertex(METASOUND_GET_PARAM_NAME(InputIndex), IndexInput);
 	}
 
 	template<uint32 NumOutputs>
-	FDataReferenceCollection TTriggerSelectOperator<NumOutputs>::GetOutputs() const
+	void TTriggerSelectOperator<NumOutputs>::BindOutputs(FOutputVertexInterfaceData& InOutVertexData)
 	{
 		using namespace TriggerSelectVertexNames;
-
-		FDataReferenceCollection OutputDataReferences;
 		for (uint32 i = 0; i < NumOutputs; ++i)
 		{
-			OutputDataReferences.AddDataReadReference(METASOUND_GET_PARAM_NAME_WITH_INDEX(OutputTrigger, i), TriggerOutputs[i]);
+			InOutVertexData.BindReadVertex(METASOUND_GET_PARAM_NAME_WITH_INDEX(OutputTrigger, i), TriggerOutputs[i]);
 		}
+	}
+
+	template<uint32 NumOutputs>
+	FDataReferenceCollection TTriggerSelectOperator<NumOutputs>::GetInputs() const
+	{
+		// This should never be called. Bind(...) is called instead. This method
+		// exists as a stop-gap until the API can be deprecated and removed.
+		checkNoEntry();
+		return {};
+	}
+
+	template<uint32 NumOutputs>
+	FDataReferenceCollection TTriggerSelectOperator<NumOutputs>::GetOutputs() const
+	{
 
-		return OutputDataReferences;
+		// This should never be called. Bind(...) is called instead. This method
+		// exists as a stop-gap until the API can be deprecated and removed.
+		checkNoEntry();
+		return {};
 	}
 
 	template<uint32 NumOutputs>
diff --git a/Engine/Plugins/Runtime/Metasound/Source/MetasoundStandardNodes/Private/MetasoundTriggerSequenceNode.cpp b/Engine/Plugins/Runtime/Metasound/Source/MetasoundStandardNodes/Private/MetasoundTriggerSequenceNode.cpp
--- a/Engine/Plugins/Runtime/Metasound/Source/MetasoundStandardNodes/Private/MetasoundTriggerSequenceNode.cpp
+++ b/Engine/Plugins/Runtime/Metasound/Source/MetasoundStandardNodes/Private/MetasoundTriggerSequenceNode.cpp
@@ -47,8 +47,9 @@ namespace Metasound
 			const FTriggerReadRef& InReset,
 			const FBoolReadRef& bInLoop);
 
+		virtual void BindInputs(FInputVertexInterfaceData& InOutVertexData) override;
+		virtual void BindOutputs(FOutputVertexInterfaceData& InOutVertexData) override;
 		virtual FDataReferenceCollection GetInputs() const override;
-
 		virtual FDataReferenceCollection GetOutputs() const override;
 
 		void Execute();
@@ -74,7 +75,6 @@ namespace Metasound
 		, TriggerResetInput(InReset)
 		, bLoopInput(bInLoop)
 	{
-		
 		for (uint32 i = 0; i < NumOutputs; ++i)
 		{
 			TriggerOutputs.Add(FTriggerWriteRef::CreateNew(InSettings));
@@ -82,31 +82,45 @@ namespace Metasound
 	}
 
 	template<uint32 NumOutputs>
-	FDataReferenceCollection TTriggerSequenceOperator<NumOutputs>::GetInputs() const
+	void TTriggerSequenceOperator<NumOutputs>::BindInputs(FInputVertexInterfaceData& InOutVertexData)
 	{
 		using namespace TriggerSequenceVertexNames;
-
-		FDataReferenceCollection InputDataReferences;
-		InputDataReferences.AddDataReadReference(METASOUND_GET_PARAM_NAME(InputTrigger), TriggerInput);
-		InputDataReferences.AddDataReadReference(METASOUND_GET_PARAM_NAME(InputReset), TriggerResetInput);
-		InputDataReferences.AddDataReadReference(METASOUND_GET_PARAM_NAME(InputLoop), bLoopInput);
-		return InputDataReferences;
+		InOutVertexData.BindReadVertex(METASOUND_GET_PARAM_NAME(InputTrigger), TriggerInput);
+		InOutVertexData.BindReadVertex(METASOUND_GET_PARAM_NAME(InputReset), TriggerResetInput);
+		InOutVertexData.BindReadVertex(METASOUND_GET_PARAM_NAME(InputLoop), bLoopInput);
 	}
 
 	template<uint32 NumOutputs>
-	FDataReferenceCollection TTriggerSequenceOperator<NumOutputs>::GetOutputs() const
+	void TTriggerSequenceOperator<NumOutputs>::BindOutputs(FOutputVertexInterfaceData& InOutVertexData)
 	{
 		using namespace TriggerSequenceVertexNames;
-
-		FDataReferenceCollection OutputDataReferences;
 		for (uint32 i = 0; i < NumOutputs; ++i)
 		{
-			OutputDataReferences.AddDataReadReference(METASOUND_GET_PARAM_NAME_WITH_INDEX(OutputTrigger, i), TriggerOutputs[i]);
+			InOutVertexData.BindReadVertex(METASOUND_GET_PARAM_NAME_WITH_INDEX(OutputTrigger, i), TriggerOutputs[i]);
 		}
+	}
+
 
-		return OutputDataReferences;
+	template<uint32 NumOutputs>
+	FDataReferenceCollection TTriggerSequenceOperator<NumOutputs>::GetInputs() const
+	{
+		// This should never be called. Bind(...) is called instead. This method
+		// exists as a stop-gap until the API can be deprecated and removed.
+		checkNoEntry();
+		return {};
+	}
+
+	template<uint32 NumOutputs>
+	FDataReferenceCollection TTriggerSequenceOperator<NumOutputs>::GetOutputs() const
+	{
+		// This should never be called. Bind(...) is called instead. This method
+		// exists as a stop-gap until the API can be deprecated and removed.
+		checkNoEntry();
+		return {};
 	}
 
+
+
 	template<uint32 NumOutputs>
 	void TTriggerSequenceOperator<NumOutputs>::Execute()
 	{
diff --git a/Engine/Plugins/Runtime/Metasound/Source/MetasoundStandardNodes/Private/MetasoundTriggerToggleNode.cpp b/Engine/Plugins/Runtime/Metasound/Source/MetasoundStandardNodes/Private/MetasoundTriggerToggleNode.cpp
--- a/Engine/Plugins/Runtime/Metasound/Source/MetasoundStandardNodes/Private/MetasoundTriggerToggleNode.cpp
+++ b/Engine/Plugins/Runtime/Metasound/Source/MetasoundStandardNodes/Private/MetasoundTriggerToggleNode.cpp
@@ -39,6 +39,8 @@ namespace Metasound
 
 			FTriggerToggleOperator(const FOperatorSettings& InSettings, const FTriggerReadRef& InTriggerOn, const FTriggerReadRef& InTriggerOff, const FBoolReadRef& InInitValue);
 
+			virtual void BindInputs(FInputVertexInterfaceData& InOutVertexData) override;
+			virtual void BindOutputs(FOutputVertexInterfaceData& InOutVertexData) override;
 			virtual FDataReferenceCollection GetInputs() const override;
 			virtual FDataReferenceCollection GetOutputs() const override;
 
@@ -63,26 +65,35 @@ namespace Metasound
 	{
 	}
 
-	FDataReferenceCollection FTriggerToggleOperator::GetInputs() const
+	void FTriggerToggleOperator::BindInputs(FInputVertexInterfaceData& InOutVertexData)
 	{
 		using namespace TriggerToggle;
-
-		FDataReferenceCollection InputDataReferences;
-		InputDataReferences.AddDataReadReference(METASOUND_GET_PARAM_NAME(InputOnTrigger), TriggerOn);
-		InputDataReferences.AddDataReadReference(METASOUND_GET_PARAM_NAME(InputOffTrigger), TriggerOff);
-		InputDataReferences.AddDataReadReference(METASOUND_GET_PARAM_NAME(InputInit), InitValue);
-		return InputDataReferences;
+		InOutVertexData.BindReadVertex(METASOUND_GET_PARAM_NAME(InputOnTrigger), TriggerOn);
+		InOutVertexData.BindReadVertex(METASOUND_GET_PARAM_NAME(InputOffTrigger), TriggerOff);
+		InOutVertexData.BindReadVertex(METASOUND_GET_PARAM_NAME(InputInit), InitValue);
 	}
 
-	FDataReferenceCollection FTriggerToggleOperator::GetOutputs() const
+	void FTriggerToggleOperator::BindOutputs(FOutputVertexInterfaceData& InOutVertexData)
 	{
 		using namespace TriggerToggle;
+		InOutVertexData.BindReadVertex(METASOUND_GET_PARAM_NAME(OutputTrigger), TriggerOutput);
+		InOutVertexData.BindReadVertex(METASOUND_GET_PARAM_NAME(OutputValue), ValueOutput);
+	}
 
-		FDataReferenceCollection OutputDataReferences;
-		OutputDataReferences.AddDataReadReference(METASOUND_GET_PARAM_NAME(OutputTrigger), TriggerOutput);
-		OutputDataReferences.AddDataReadReference(METASOUND_GET_PARAM_NAME(OutputValue), ValueOutput);
+	FDataReferenceCollection FTriggerToggleOperator::GetInputs() const
+	{
+		// This should never be called. Bind(...) is called instead. This method
+		// exists as a stop-gap until the API can be deprecated and removed.
+		checkNoEntry();
+		return {};
+	}
 
-		return OutputDataReferences;
+	FDataReferenceCollection FTriggerToggleOperator::GetOutputs() const
+	{
+		// This should never be called. Bind(...) is called instead. This method
+		// exists as a stop-gap until the API can be deprecated and removed.
+		checkNoEntry();
+		return {};
 	}
 
 	void FTriggerToggleOperator::Execute()
diff --git a/Engine/Plugins/Runtime/Metasound/Source/MetasoundStandardNodes/Private/MetasoundWaveShaperNode.cpp b/Engine/Plugins/Runtime/Metasound/Source/MetasoundStandardNodes/Private/MetasoundWaveShaperNode.cpp
--- a/Engine/Plugins/Runtime/Metasound/Source/MetasoundStandardNodes/Private/MetasoundWaveShaperNode.cpp
+++ b/Engine/Plugins/Runtime/Metasound/Source/MetasoundStandardNodes/Private/MetasoundWaveShaperNode.cpp
@@ -52,6 +52,8 @@ namespace Metasound
 			const FFloatReadRef& InOutputGain,
 			const FEnumWaveShaperReadRef& InType);
 
+		virtual void BindInputs(FInputVertexInterfaceData& InOutVertexData) override;
+		virtual void BindOutputs(FOutputVertexInterfaceData& InOutVertexData) override;
 		virtual FDataReferenceCollection GetInputs() const override;
 		virtual FDataReferenceCollection GetOutputs() const override;
 		void Execute();
@@ -96,27 +98,34 @@ namespace Metasound
 		WaveShaper.Init(InSettings.GetSampleRate());
 	}
 
+	void FWaveShaperOperator::BindInputs(FInputVertexInterfaceData& InOutVertexData)
+	{
+		InOutVertexData.BindReadVertex(METASOUND_GET_PARAM_NAME(InAudioInput), AudioInput);
+		InOutVertexData.BindReadVertex(METASOUND_GET_PARAM_NAME(InWaveShapeAmount), Amount);
+		InOutVertexData.BindReadVertex(METASOUND_GET_PARAM_NAME(InWaveShapeBias), Bias);
+		InOutVertexData.BindReadVertex(METASOUND_GET_PARAM_NAME(InOutputGain), OutputGain);
+		InOutVertexData.BindReadVertex(METASOUND_GET_PARAM_NAME(InType), Type);
+	}
+
+	void FWaveShaperOperator::BindOutputs(FOutputVertexInterfaceData& InOutVertexData)
+	{
+		InOutVertexData.BindReadVertex(METASOUND_GET_PARAM_NAME(OutAudio), AudioOutput);
+	}
+
 	FDataReferenceCollection FWaveShaperOperator::GetInputs() const
 	{
-		using namespace WaveShaperNode;
-		
-		FDataReferenceCollection InputDataReferences;
-		InputDataReferences.AddDataReadReference(METASOUND_GET_PARAM_NAME(InAudioInput), AudioInput);
-		InputDataReferences.AddDataReadReference(METASOUND_GET_PARAM_NAME(InWaveShapeAmount), Amount);
-		InputDataReferences.AddDataReadReference(METASOUND_GET_PARAM_NAME(InWaveShapeBias), Bias);
-		InputDataReferences.AddDataReadReference(METASOUND_GET_PARAM_NAME(InOutputGain), OutputGain);
-		InputDataReferences.AddDataReadReference(METASOUND_GET_PARAM_NAME(InType), Type);
-
-		return InputDataReferences;
+		// This should never be called. Bind(...) is called instead. This method
+		// exists as a stop-gap until the API can be deprecated and removed.
+		checkNoEntry();
+		return {};
 	}
 
 	FDataReferenceCollection FWaveShaperOperator::GetOutputs() const
 	{
-		using namespace WaveShaperNode; 
-
-		FDataReferenceCollection OutputDataReferences;
-		OutputDataReferences.AddDataReadReference(METASOUND_GET_PARAM_NAME(OutAudio), AudioOutput);
-		return OutputDataReferences;
+		// This should never be called. Bind(...) is called instead. This method
+		// exists as a stop-gap until the API can be deprecated and removed.
+		checkNoEntry();
+		return {};
 	}
 
 	void FWaveShaperOperator::Execute()
diff --git a/Engine/Plugins/Runtime/Metasound/Source/MetasoundStandardNodes/Private/MetasoundWaveWriterNode.cpp b/Engine/Plugins/Runtime/Metasound/Source/MetasoundStandardNodes/Private/MetasoundWaveWriterNode.cpp
--- a/Engine/Plugins/Runtime/Metasound/Source/MetasoundStandardNodes/Private/MetasoundWaveWriterNode.cpp
+++ b/Engine/Plugins/Runtime/Metasound/Source/MetasoundStandardNodes/Private/MetasoundWaveWriterNode.cpp
@@ -247,19 +247,33 @@ namespace Metasound
 			}
 		}
 
-		FDataReferenceCollection GetInputs() const override
+
+		virtual void BindInputs(FInputVertexInterfaceData& InOutVertexData) override
 		{
-			FDataReferenceCollection InputPins;
 			for (int32 i = 0; i < NumInputChannels; ++i)
 			{
-				InputPins.AddDataReadReference(GetAudioInputName(i), AudioInputs[i]);
+				InOutVertexData.BindReadVertex(GetAudioInputName(i), AudioInputs[i]);
 			}
-			return InputPins;
 		}
-		FDataReferenceCollection GetOutputs() const override
+
+		virtual void BindOutputs(FOutputVertexInterfaceData&) override
+		{
+		}
+
+		virtual FDataReferenceCollection GetInputs() const override
+		{
+			// This should never be called. Bind(...) is called instead. This method
+			// exists as a stop-gap until the API can be deprecated and removed.
+			checkNoEntry();
+			return {};
+		}
+	
+		virtual FDataReferenceCollection GetOutputs() const override
 		{
-			FDataReferenceCollection OutputDataReferences;
-			return OutputDataReferences;
+			// This should never be called. Bind(...) is called instead. This method
+			// exists as a stop-gap until the API can be deprecated and removed.
+			checkNoEntry();
+			return {};
 		}
 
 		static const FVertexInterface& DeclareVertexInterface()
diff --git a/Engine/Plugins/Runtime/Metasound/Source/MetasoundStandardNodes/Public/MetasoundRandomNode.h b/Engine/Plugins/Runtime/Metasound/Source/MetasoundStandardNodes/Public/MetasoundRandomNode.h
--- a/Engine/Plugins/Runtime/Metasound/Source/MetasoundStandardNodes/Public/MetasoundRandomNode.h
+++ b/Engine/Plugins/Runtime/Metasound/Source/MetasoundStandardNodes/Public/MetasoundRandomNode.h
@@ -385,33 +385,45 @@ namespace Metasound
 		virtual ~TRandomNodeOperator() = default;
 
 
-		virtual FDataReferenceCollection GetInputs() const override
+		virtual void BindInputs(FInputVertexInterfaceData& InOutVertexData) override
 		{
 			using namespace RandomNodeNames;
 
 			FDataReferenceCollection Inputs;
-			Inputs.AddDataReadReference(METASOUND_GET_PARAM_NAME(InputNextTrigger), NextTrigger);
- 			Inputs.AddDataReadReference(METASOUND_GET_PARAM_NAME(InputResetTrigger), ResetTrigger);
- 			Inputs.AddDataReadReference(METASOUND_GET_PARAM_NAME(InputSeed), SeedValue);
+			InOutVertexData.BindReadVertex(METASOUND_GET_PARAM_NAME(InputNextTrigger), NextTrigger);
+ 			InOutVertexData.BindReadVertex(METASOUND_GET_PARAM_NAME(InputResetTrigger), ResetTrigger);
+ 			InOutVertexData.BindReadVertex(METASOUND_GET_PARAM_NAME(InputSeed), SeedValue);
 
 			// If the type doesn't have a range no need for input pins to define it
 			if (TRandomNodeSpecialization<ValueType>::HasRange())
 			{
-				Inputs.AddDataReadReference(METASOUND_GET_PARAM_NAME(InputMin), MinValue);
-				Inputs.AddDataReadReference(METASOUND_GET_PARAM_NAME(InputMax), MinValue);
+				InOutVertexData.BindReadVertex(METASOUND_GET_PARAM_NAME(InputMin), MinValue);
+				InOutVertexData.BindReadVertex(METASOUND_GET_PARAM_NAME(InputMax), MinValue);
 			}
-			return Inputs;
 		}
 
-		virtual FDataReferenceCollection GetOutputs() const override
+		virtual void BindOutputs(FOutputVertexInterfaceData& InOutVertexData) override
 		{
 			using namespace RandomNodeNames;
+			InOutVertexData.BindReadVertex(METASOUND_GET_PARAM_NAME(OutputOnNextTrigger), TriggerOutOnNext);
+			InOutVertexData.BindReadVertex(METASOUND_GET_PARAM_NAME(OutputOnResetTrigger), TriggerOutOnReset);
+			InOutVertexData.BindReadVertex(METASOUND_GET_PARAM_NAME(OutputValue), OutputValue);
+		}
 
-			FDataReferenceCollection Outputs;
-			Outputs.AddDataReadReference(METASOUND_GET_PARAM_NAME(OutputOnNextTrigger), TriggerOutOnNext);
-			Outputs.AddDataReadReference(METASOUND_GET_PARAM_NAME(OutputOnResetTrigger), TriggerOutOnReset);
-			Outputs.AddDataReadReference(METASOUND_GET_PARAM_NAME(OutputValue), OutputValue);
-			return Outputs;
+		virtual FDataReferenceCollection GetInputs() const override
+		{
+			// This should never be called. Bind(...) is called instead. This method
+			// exists as a stop-gap until the API can be deprecated and removed.
+			checkNoEntry();
+			return {};
+		}
+
+		virtual FDataReferenceCollection GetOutputs() const override
+		{
+			// This should never be called. Bind(...) is called instead. This method
+			// exists as a stop-gap until the API can be deprecated and removed.
+			checkNoEntry();
+			return {};
 		}
 
 		void Execute()
diff --git a/Engine/Plugins/Runtime/Metasound/Source/MetasoundStandardNodes/Public/MetasoundTriggerAccumulatorNode.h b/Engine/Plugins/Runtime/Metasound/Source/MetasoundStandardNodes/Public/MetasoundTriggerAccumulatorNode.h
--- a/Engine/Plugins/Runtime/Metasound/Source/MetasoundStandardNodes/Public/MetasoundTriggerAccumulatorNode.h
+++ b/Engine/Plugins/Runtime/Metasound/Source/MetasoundStandardNodes/Public/MetasoundTriggerAccumulatorNode.h
@@ -103,28 +103,38 @@ namespace Metasound
 		virtual ~TTriggerAccumulatorOperator() = default;
 
 
-		virtual FDataReferenceCollection GetInputs() const override
+		virtual void BindInputs(FInputVertexInterfaceData& InOutVertexData) override
 		{
 			using namespace TriggerAccumulatorVertexNames;
 
-			FDataReferenceCollection Inputs;
 			for (uint32 i = 0; i < NumInputs; ++i)
 			{
-				Inputs.AddDataReadReference(METASOUND_GET_PARAM_NAME_WITH_INDEX(InputTrigger, i), InputTriggers[i]);
+				InOutVertexData.BindReadVertex(METASOUND_GET_PARAM_NAME_WITH_INDEX(InputTrigger, i), InputTriggers[i]);
 			}
-			Inputs.AddDataReadReference(METASOUND_GET_PARAM_NAME(InputAutoReset), bAutoReset);
-
-			return Inputs;
+			InOutVertexData.BindReadVertex(METASOUND_GET_PARAM_NAME(InputAutoReset), bAutoReset);
 		}
 
-		virtual FDataReferenceCollection GetOutputs() const override
+		virtual void BindOutputs(FOutputVertexInterfaceData& InOutVertexData) override
 		{
 			using namespace TriggerAccumulatorVertexNames;
 
-			FDataReferenceCollection Outputs;
-			Outputs.AddDataReadReference(METASOUND_GET_PARAM_NAME(AccumulateOutputOnTrigger), OutputTrigger);
+			InOutVertexData.BindReadVertex(METASOUND_GET_PARAM_NAME(AccumulateOutputOnTrigger), OutputTrigger);
+		}
+
+		virtual FDataReferenceCollection GetInputs() const override
+		{
+			// This should never be called. Bind(...) is called instead. This method
+			// exists as a stop-gap until the API can be deprecated and removed.
+			checkNoEntry();
+			return {};
+		}
 
-			return Outputs;
+		virtual FDataReferenceCollection GetOutputs() const override
+		{
+			// This should never be called. Bind(...) is called instead. This method
+			// exists as a stop-gap until the API can be deprecated and removed.
+			checkNoEntry();
+			return {};
 		}
 
 		void TriggerOutputIfReady(int32 InStartFrame)
diff --git a/Engine/Plugins/Runtime/Metasound/Source/MetasoundStandardNodes/Public/MetasoundTriggerCompareNode.h b/Engine/Plugins/Runtime/Metasound/Source/MetasoundStandardNodes/Public/MetasoundTriggerCompareNode.h
--- a/Engine/Plugins/Runtime/Metasound/Source/MetasoundStandardNodes/Public/MetasoundTriggerCompareNode.h
+++ b/Engine/Plugins/Runtime/Metasound/Source/MetasoundStandardNodes/Public/MetasoundTriggerCompareNode.h
@@ -126,28 +126,38 @@ namespace Metasound
 
 		virtual ~TTriggerCompareNodeOperator() = default;
 
-		virtual FDataReferenceCollection GetInputs() const override
+		virtual void BindInputs(FInputVertexInterfaceData& InOutVertexData) override
 		{
 			using namespace TriggerCompareVertexNames;
 
-			FDataReferenceCollection Inputs;
-			Inputs.AddDataReadReference(METASOUND_GET_PARAM_NAME(InputCompare), OnCompareTrigger);
-			Inputs.AddDataReadReference(METASOUND_GET_PARAM_NAME(InputParamA), ValueA);
-			Inputs.AddDataReadReference(METASOUND_GET_PARAM_NAME(InputParamB), ValueB);
-			Inputs.AddDataReadReference(METASOUND_GET_PARAM_NAME(InputCompareType), TriggerComparisonType);
-
-			return Inputs;
+			InOutVertexData.BindReadVertex(METASOUND_GET_PARAM_NAME(InputCompare), OnCompareTrigger);
+			InOutVertexData.BindReadVertex(METASOUND_GET_PARAM_NAME(InputParamA), ValueA);
+			InOutVertexData.BindReadVertex(METASOUND_GET_PARAM_NAME(InputParamB), ValueB);
+			InOutVertexData.BindReadVertex(METASOUND_GET_PARAM_NAME(InputCompareType), TriggerComparisonType);
 		}
 
-		virtual FDataReferenceCollection GetOutputs() const override
+		virtual void BindOutputs(FOutputVertexInterfaceData& InOutVertexData) override
 		{
 			using namespace TriggerCompareVertexNames;
 
-			FDataReferenceCollection Outputs;
-			Outputs.AddDataReadReference(METASOUND_GET_PARAM_NAME(OutputOnTrue), TriggerOutOnTrue);
-			Outputs.AddDataReadReference(METASOUND_GET_PARAM_NAME(OutputOnFalse), TriggerOutOnFalse);
+			InOutVertexData.BindReadVertex(METASOUND_GET_PARAM_NAME(OutputOnTrue), TriggerOutOnTrue);
+			InOutVertexData.BindReadVertex(METASOUND_GET_PARAM_NAME(OutputOnFalse), TriggerOutOnFalse);
+		}
+
+		virtual FDataReferenceCollection GetInputs() const override
+		{
+			// This should never be called. Bind(...) is called instead. This method
+			// exists as a stop-gap until the API can be deprecated and removed.
+			checkNoEntry();
+			return {};
+		}
 
-			return Outputs;
+		virtual FDataReferenceCollection GetOutputs() const override
+		{
+			// This should never be called. Bind(...) is called instead. This method
+			// exists as a stop-gap until the API can be deprecated and removed.
+			checkNoEntry();
+			return {};
 		}
 
 		void Execute()
diff --git a/Engine/Plugins/Runtime/Metasound/Source/MetasoundStandardNodes/Public/MetasoundValueNode.h b/Engine/Plugins/Runtime/Metasound/Source/MetasoundStandardNodes/Public/MetasoundValueNode.h
--- a/Engine/Plugins/Runtime/Metasound/Source/MetasoundStandardNodes/Public/MetasoundValueNode.h
+++ b/Engine/Plugins/Runtime/Metasound/Source/MetasoundStandardNodes/Public/MetasoundValueNode.h
@@ -105,30 +105,40 @@ namespace Metasound
 		virtual ~TValueOperator() = default;
 
 
-		virtual FDataReferenceCollection GetInputs() const override
+		virtual void BindInputs(FInputVertexInterfaceData& InOutVertexData) override
 		{
 			using namespace ValueVertexNames;
+			InOutVertexData.BindReadVertex(METASOUND_GET_PARAM_NAME(InputSetTrigger), SetTrigger);
+			InOutVertexData.BindReadVertex(METASOUND_GET_PARAM_NAME(InputResetTrigger), ResetTrigger);
+			InOutVertexData.BindReadVertex(METASOUND_GET_PARAM_NAME(InputInitValue), InitValue);
+			InOutVertexData.BindReadVertex(METASOUND_GET_PARAM_NAME(InputTargetValue), TargetValue);
+		}
 
-			FDataReferenceCollection Inputs;
-			Inputs.AddDataReadReference(METASOUND_GET_PARAM_NAME(InputSetTrigger), SetTrigger);
-			Inputs.AddDataReadReference(METASOUND_GET_PARAM_NAME(InputResetTrigger), ResetTrigger);
-			Inputs.AddDataReadReference(METASOUND_GET_PARAM_NAME(InputInitValue), InitValue);
-			Inputs.AddDataReadReference(METASOUND_GET_PARAM_NAME(InputTargetValue), TargetValue);
+		virtual void BindOutputs(FOutputVertexInterfaceData& InOutVertexData) override
+		{
+			using namespace ValueVertexNames;
+			InOutVertexData.BindReadVertex(METASOUND_GET_PARAM_NAME(OutputOnSet), TriggerOnSet);
+			InOutVertexData.BindReadVertex(METASOUND_GET_PARAM_NAME(OutputOnReset), TriggerOnReset);
+			InOutVertexData.BindReadVertex(METASOUND_GET_PARAM_NAME(OutputValue), OutputValue);
+		}
 
-			return Inputs;
+		virtual FDataReferenceCollection GetInputs() const override
+		{
+			// This should never be called. Bind(...) is called instead. This method
+			// exists as a stop-gap until the API can be deprecated and removed.
+			checkNoEntry();
+			return {};
 		}
 
 		virtual FDataReferenceCollection GetOutputs() const override
 		{
-			using namespace ValueVertexNames;
+			// This should never be called. Bind(...) is called instead. This method
+			// exists as a stop-gap until the API can be deprecated and removed.
+			checkNoEntry();
+			return {};
+		}
 
-			FDataReferenceCollection Outputs;
-			Outputs.AddDataReadReference(METASOUND_GET_PARAM_NAME(OutputOnSet), TriggerOnSet);
-			Outputs.AddDataReadReference(METASOUND_GET_PARAM_NAME(OutputOnReset), TriggerOnReset);
-			Outputs.AddDataReadReference(METASOUND_GET_PARAM_NAME(OutputValue), OutputValue);
 
-			return Outputs;
-		}
 
 		void Execute()
 		{
