Fix ADSR bug where attack was bound to both attack and release

#rnx

diff --git a/Engine/Plugins/Runtime/Metasound/Source/MetasoundStandardNodes/Private/MetasoundADSREnvelopeNode.cpp b/Engine/Plugins/Runtime/Metasound/Source/MetasoundStandardNodes/Private/MetasoundADSREnvelopeNode.cpp
--- a/Engine/Plugins/Runtime/Metasound/Source/MetasoundStandardNodes/Private/MetasoundADSREnvelopeNode.cpp
+++ b/Engine/Plugins/Runtime/Metasound/Source/MetasoundStandardNodes/Private/MetasoundADSREnvelopeNode.cpp
@@ -610,7 +610,7 @@ namespace Metasound
 			using namespace ADSREnvelopeVertexNames;
 
 			InOutVertexData.BindReadVertex(METASOUND_GET_PARAM_NAME(InputAttackTrigger), TriggerAttackIn);
-			InOutVertexData.BindReadVertex(METASOUND_GET_PARAM_NAME(InputAttackTrigger), TriggerReleaseIn);
+			InOutVertexData.BindReadVertex(METASOUND_GET_PARAM_NAME(InputReleaseTrigger), TriggerReleaseIn);
 			InOutVertexData.BindReadVertex(METASOUND_GET_PARAM_NAME(InputAttackTime), AttackTime);
 			InOutVertexData.BindReadVertex(METASOUND_GET_PARAM_NAME(InputDecayTime), DecayTime);
 			InOutVertexData.BindReadVertex(METASOUND_GET_PARAM_NAME(InputSustainLevel), SustainLevel);
