Fixing up callsights to use GetDataReadReferenceOrConstructWithVertexDefault

diff --git a/Engine/Plugins/Runtime/Metasound/Source/MetasoundStandardNodes/Private/MetasoundCompressorNode.cpp b/Engine/Plugins/Runtime/Metasound/Source/MetasoundStandardNodes/Private/MetasoundCompressorNode.cpp
--- a/Engine/Plugins/Runtime/Metasound/Source/MetasoundStandardNodes/Private/MetasoundCompressorNode.cpp
+++ b/Engine/Plugins/Runtime/Metasound/Source/MetasoundStandardNodes/Private/MetasoundCompressorNode.cpp
@@ -157,7 +157,7 @@ namespace Metasound
 					TInputDataVertex<FTime>(METASOUND_GET_PARAM_NAME_AND_METADATA(InputLookaheadTime), 0.01f),
 					TInputDataVertex<float>(METASOUND_GET_PARAM_NAME_AND_METADATA(InputKnee), 10.0f),
 					TInputDataVertex<FAudioBuffer>(METASOUND_GET_PARAM_NAME_AND_METADATA(InputSidechain)),
-					TInputDataVertex<FEnumEnvelopePeakMode>(METASOUND_GET_PARAM_NAME_AND_METADATA(InputEnvelopeMode)),
+					TInputDataVertex<FEnumEnvelopePeakMode>(METASOUND_GET_PARAM_NAME_AND_METADATA(InputEnvelopeMode), (int32)EEnvelopePeakMode::Peak),
 					TInputDataVertex<bool>(METASOUND_GET_PARAM_NAME_AND_METADATA(InputIsAnalog), true),
 					TInputDataVertex<bool>(METASOUND_GET_PARAM_NAME_AND_METADATA(InputIsUpwards), false),
 					TInputDataVertex<float>(METASOUND_GET_PARAM_NAME_AND_METADATA(InputWetDryMix), 1.0f)
@@ -227,7 +227,7 @@ namespace Metasound
 			FTimeReadRef LookaheadTimeIn = Inputs.GetDataReadReferenceOrConstructWithVertexDefault<FTime>(InputInterface, METASOUND_GET_PARAM_NAME(InputLookaheadTime), InParams.OperatorSettings);
 			FFloatReadRef KneeIn = Inputs.GetDataReadReferenceOrConstructWithVertexDefault<float>(InputInterface, METASOUND_GET_PARAM_NAME(InputKnee), InParams.OperatorSettings);
 			FAudioBufferReadRef SidechainIn = Inputs.GetDataReadReferenceOrConstruct<FAudioBuffer>(METASOUND_GET_PARAM_NAME(InputSidechain), InParams.OperatorSettings);
-			FEnvelopePeakModeReadRef EnvelopeModeIn = Inputs.GetDataReadReferenceOrConstruct<FEnumEnvelopePeakMode>(METASOUND_GET_PARAM_NAME(InputEnvelopeMode));
+			FEnvelopePeakModeReadRef EnvelopeModeIn = Inputs.GetDataReadReferenceOrConstructWithVertexDefault<FEnumEnvelopePeakMode>(InputInterface, METASOUND_GET_PARAM_NAME(InputEnvelopeMode), InParams.OperatorSettings);
 			FBoolReadRef bIsAnalogIn = Inputs.GetDataReadReferenceOrConstructWithVertexDefault<bool>(InputInterface, METASOUND_GET_PARAM_NAME(InputIsAnalog), InParams.OperatorSettings);
 			FBoolReadRef bIsUpwardsIn = Inputs.GetDataReadReferenceOrConstructWithVertexDefault<bool>(InputInterface, METASOUND_GET_PARAM_NAME(InputIsUpwards), InParams.OperatorSettings);
 			FFloatReadRef WetDryMixIn = Inputs.GetDataReadReferenceOrConstructWithVertexDefault<float>(InputInterface, METASOUND_GET_PARAM_NAME(InputWetDryMix), InParams.OperatorSettings);
diff --git a/Engine/Plugins/Runtime/Metasound/Source/MetasoundStandardNodes/Private/MetasoundDynamicFilterNode.cpp b/Engine/Plugins/Runtime/Metasound/Source/MetasoundStandardNodes/Private/MetasoundDynamicFilterNode.cpp
--- a/Engine/Plugins/Runtime/Metasound/Source/MetasoundStandardNodes/Private/MetasoundDynamicFilterNode.cpp
+++ b/Engine/Plugins/Runtime/Metasound/Source/MetasoundStandardNodes/Private/MetasoundDynamicFilterNode.cpp
@@ -189,7 +189,7 @@ namespace Metasound
 					}, 0.0f),
 					TInputDataVertex<FTime>(METASOUND_GET_PARAM_NAME_AND_METADATA(InputAttackTime), 0.01f),
 					TInputDataVertex<FTime>(METASOUND_GET_PARAM_NAME_AND_METADATA(InputReleaseTime), 0.1f),
-					TInputDataVertex<FEnumEnvelopePeakMode>(METASOUND_GET_PARAM_NAME_AND_METADATA(InputEnvelopeMode), 1.0f),
+					TInputDataVertex<FEnumEnvelopePeakMode>(METASOUND_GET_PARAM_NAME_AND_METADATA(InputEnvelopeMode), (int32)EEnvelopePeakMode::MeanSquared),
 					TInputDataVertex<bool>(METASOUND_GET_PARAM_NAME_AND_METADATA(InputAnalogMode), true)
 				),
 				FOutputVertexInterface(
@@ -219,7 +219,7 @@ namespace Metasound
 			FFloatReadRef GainIn = Inputs.GetDataReadReferenceOrConstructWithVertexDefault<float>(InputInterface, METASOUND_GET_PARAM_NAME(InputGain), InParams.OperatorSettings);
 			FTimeReadRef AttackTimeIn = Inputs.GetDataReadReferenceOrConstructWithVertexDefault<FTime>(InputInterface, METASOUND_GET_PARAM_NAME(InputAttackTime), InParams.OperatorSettings);
 			FTimeReadRef ReleaseTimeIn = Inputs.GetDataReadReferenceOrConstructWithVertexDefault<FTime>(InputInterface, METASOUND_GET_PARAM_NAME(InputReleaseTime), InParams.OperatorSettings);
-			FEnvelopePeakModeReadRef EnvelopeModeIn = Inputs.GetDataReadReferenceOrConstruct<FEnumEnvelopePeakMode>(METASOUND_GET_PARAM_NAME(InputEnvelopeMode));
+			FEnvelopePeakModeReadRef EnvelopeModeIn = Inputs.GetDataReadReferenceOrConstructWithVertexDefault<FEnumEnvelopePeakMode>(InputInterface, METASOUND_GET_PARAM_NAME(InputEnvelopeMode), InParams.OperatorSettings);
 			FBoolReadRef bAnalogIn = Inputs.GetDataReadReferenceOrConstructWithVertexDefault<bool>(InputInterface, METASOUND_GET_PARAM_NAME(InputAnalogMode), InParams.OperatorSettings);
 
 			bool bIsSidechainConnected = Inputs.ContainsDataReadReference<FAudioBuffer>(METASOUND_GET_PARAM_NAME(InputSidechain));
diff --git a/Engine/Plugins/Runtime/Metasound/Source/MetasoundStandardNodes/Private/MetasoundEnvelopeFollowerNode.cpp b/Engine/Plugins/Runtime/Metasound/Source/MetasoundStandardNodes/Private/MetasoundEnvelopeFollowerNode.cpp
--- a/Engine/Plugins/Runtime/Metasound/Source/MetasoundStandardNodes/Private/MetasoundEnvelopeFollowerNode.cpp
+++ b/Engine/Plugins/Runtime/Metasound/Source/MetasoundStandardNodes/Private/MetasoundEnvelopeFollowerNode.cpp
@@ -202,7 +202,7 @@ namespace Metasound
 				TInputDataVertex<FAudioBuffer>(METASOUND_GET_PARAM_NAME_AND_METADATA(InParamAudioInput)),
 				TInputDataVertex<FTime>(METASOUND_GET_PARAM_NAME_AND_METADATA(InParamAttackTime), 0.01f),
 				TInputDataVertex<FTime>(METASOUND_GET_PARAM_NAME_AND_METADATA(InParamReleaseTime), 0.1f),
-				TInputDataVertex<FEnumEnvelopePeakMode>(METASOUND_GET_PARAM_NAME_AND_METADATA(InParamFollowMode))
+				TInputDataVertex<FEnumEnvelopePeakMode>(METASOUND_GET_PARAM_NAME_AND_METADATA(InParamFollowMode), (int32)EEnvelopePeakMode::Peak)
 			),
 			FOutputVertexInterface(
 				TOutputDataVertex<float>(METASOUND_GET_PARAM_NAME_AND_METADATA(OutParamEnvelope)),
@@ -246,8 +246,7 @@ namespace Metasound
 		FAudioBufferReadRef AudioIn = InputCollection.GetDataReadReferenceOrConstruct<FAudioBuffer>(METASOUND_GET_PARAM_NAME(InParamAudioInput), InParams.OperatorSettings);
 		FTimeReadRef AttackTime = InputCollection.GetDataReadReferenceOrConstructWithVertexDefault<FTime>(InputInterface, METASOUND_GET_PARAM_NAME(InParamAttackTime), InParams.OperatorSettings);
 		FTimeReadRef ReleaseTime = InputCollection.GetDataReadReferenceOrConstructWithVertexDefault<FTime>(InputInterface, METASOUND_GET_PARAM_NAME(InParamReleaseTime), InParams.OperatorSettings);
-		FEnvelopePeakModeReadRef EnvelopeModeIn = InputCollection.GetDataReadReferenceOrConstruct<FEnumEnvelopePeakMode>(METASOUND_GET_PARAM_NAME(InParamFollowMode));
-
+		FEnvelopePeakModeReadRef EnvelopeModeIn = InputCollection.GetDataReadReferenceOrConstructWithVertexDefault<FEnumEnvelopePeakMode>(InputInterface, METASOUND_GET_PARAM_NAME(InParamFollowMode), InParams.OperatorSettings);
 
 		return MakeUnique<FEnvelopeFollowerOperator>(InParams, AudioIn, AttackTime, ReleaseTime, EnvelopeModeIn);
 	}
diff --git a/Engine/Plugins/Runtime/Metasound/Source/MetasoundStandardNodes/Private/MetasoundFilterQToBandwidthNode.cpp b/Engine/Plugins/Runtime/Metasound/Source/MetasoundStandardNodes/Private/MetasoundFilterQToBandwidthNode.cpp
--- a/Engine/Plugins/Runtime/Metasound/Source/MetasoundStandardNodes/Private/MetasoundFilterQToBandwidthNode.cpp
+++ b/Engine/Plugins/Runtime/Metasound/Source/MetasoundStandardNodes/Private/MetasoundFilterQToBandwidthNode.cpp
@@ -71,9 +71,10 @@ namespace Metasound
 		{
 			using namespace QToBandwidth;
 
-			const FDataReferenceCollection& Inputs = InParams.InputDataReferences;
+			const FDataReferenceCollection& InputDataRefs = InParams.InputDataReferences;
+			const FInputVertexInterface& InputInterface = InParams.Node.GetVertexInterface().GetInputInterface();
 
-			FFloatReadRef FloatInput = Inputs.GetDataReadReferenceOrConstruct<float>(METASOUND_GET_PARAM_NAME(InputQ));
+			FFloatReadRef FloatInput = InputDataRefs.GetDataReadReferenceOrConstructWithVertexDefault<float>(InputInterface, METASOUND_GET_PARAM_NAME(InputQ), InParams.OperatorSettings);
 
 			return MakeUnique<FConvertQToBandwidthNodeOperator>(InParams.OperatorSettings, FloatInput);
 		}
diff --git a/Engine/Plugins/Runtime/Metasound/Source/MetasoundStandardNodes/Private/MetasoundGrainDelayNode.cpp b/Engine/Plugins/Runtime/Metasound/Source/MetasoundStandardNodes/Private/MetasoundGrainDelayNode.cpp
--- a/Engine/Plugins/Runtime/Metasound/Source/MetasoundStandardNodes/Private/MetasoundGrainDelayNode.cpp
+++ b/Engine/Plugins/Runtime/Metasound/Source/MetasoundStandardNodes/Private/MetasoundGrainDelayNode.cpp
@@ -174,7 +174,7 @@ namespace Metasound
 			InputInterface.Add(TInputDataVertex<float>(METASOUND_GET_PARAM_NAME_AND_METADATA(InGrainDurationDelta), 0.0f));
 			InputInterface.Add(TInputDataVertex<float>(METASOUND_GET_PARAM_NAME_AND_METADATA(InGrainPitchShift), 0.0f));
 			InputInterface.Add(TInputDataVertex<float>(METASOUND_GET_PARAM_NAME_AND_METADATA(InGrainPitchShiftDelta), 0.0f));
-			InputInterface.Add(TInputDataVertex<FEnumGrainDelayEnvelope>(METASOUND_GET_PARAM_NAME_AND_METADATA(InGrainDelayEnvelope), 0));
+			InputInterface.Add(TInputDataVertex<FEnumGrainDelayEnvelope>(METASOUND_GET_PARAM_NAME_AND_METADATA(InGrainDelayEnvelope), (int32)Audio::Grain::EEnvelope::Gaussian));
 			InputInterface.Add(TInputDataVertex<int32>(METASOUND_GET_PARAM_NAME_AND_METADATA(InGrainMaxCount), 16));
 			InputInterface.Add(TInputDataVertex<float>(METASOUND_GET_PARAM_NAME_AND_METADATA(InGrainDelayFeedbackAmount), 0.0f));
 
@@ -203,7 +203,7 @@ namespace Metasound
 		FFloatReadRef GrainDurationDelta = InputCollection.GetDataReadReferenceOrConstructWithVertexDefault<float>(InputInterface, METASOUND_GET_PARAM_NAME(InGrainDurationDelta), InParams.OperatorSettings);
 		FFloatReadRef GrainPitchShift = InputCollection.GetDataReadReferenceOrConstructWithVertexDefault<float>(InputInterface, METASOUND_GET_PARAM_NAME(InGrainPitchShift), InParams.OperatorSettings);
 		FFloatReadRef GrainPitchShiftDelta = InputCollection.GetDataReadReferenceOrConstructWithVertexDefault<float>(InputInterface, METASOUND_GET_PARAM_NAME(InGrainPitchShiftDelta), InParams.OperatorSettings);
-		FGrainDelayEnvelopeReadRef GrainDelayEnvelope = InputCollection.GetDataReadReferenceOrConstruct<FEnumGrainDelayEnvelope>(METASOUND_GET_PARAM_NAME(InGrainDelayEnvelope));
+		FGrainDelayEnvelopeReadRef GrainDelayEnvelope = InputCollection.GetDataReadReferenceOrConstructWithVertexDefault<FEnumGrainDelayEnvelope>(InputInterface, METASOUND_GET_PARAM_NAME(InGrainDelayEnvelope), InParams.OperatorSettings);
 		FInt32ReadRef GrainMaxCount = InputCollection.GetDataReadReferenceOrConstructWithVertexDefault<int32>(InputInterface, METASOUND_GET_PARAM_NAME(InGrainMaxCount), InParams.OperatorSettings);
 		FFloatReadRef GrainDelayFeedbackAmount = InputCollection.GetDataReadReferenceOrConstructWithVertexDefault<float>(InputInterface, METASOUND_GET_PARAM_NAME(InGrainDelayFeedbackAmount), InParams.OperatorSettings);
 
diff --git a/Engine/Plugins/Runtime/Metasound/Source/MetasoundStandardNodes/Private/MetasoundLimiterNode.cpp b/Engine/Plugins/Runtime/Metasound/Source/MetasoundStandardNodes/Private/MetasoundLimiterNode.cpp
--- a/Engine/Plugins/Runtime/Metasound/Source/MetasoundStandardNodes/Private/MetasoundLimiterNode.cpp
+++ b/Engine/Plugins/Runtime/Metasound/Source/MetasoundStandardNodes/Private/MetasoundLimiterNode.cpp
@@ -110,7 +110,7 @@ namespace Metasound
 					TInputDataVertex<float>(METASOUND_GET_PARAM_NAME_AND_METADATA(InputInGainDb), 0.0f),
 					TInputDataVertex<float>(METASOUND_GET_PARAM_NAME_AND_METADATA(InputThresholdDb), 0.0f),
 					TInputDataVertex<FTime>(METASOUND_GET_PARAM_NAME_AND_METADATA(InputReleaseTime), 0.1f),
-					TInputDataVertex<FEnumKneeMode>(METASOUND_GET_PARAM_NAME_AND_METADATA(InputKneeMode))
+					TInputDataVertex<FEnumKneeMode>(METASOUND_GET_PARAM_NAME_AND_METADATA(InputKneeMode), (int32)EKneeMode::Hard)
 				),
 				FOutputVertexInterface(
 					TOutputDataVertex<FAudioBuffer>(METASOUND_GET_PARAM_NAME_AND_METADATA(OutputAudio))
@@ -164,7 +164,7 @@ namespace Metasound
 			FFloatReadRef InGainDbIn = Inputs.GetDataReadReferenceOrConstructWithVertexDefault<float>(InputInterface, METASOUND_GET_PARAM_NAME(InputInGainDb), InParams.OperatorSettings);
 			FFloatReadRef ThresholdDbIn = Inputs.GetDataReadReferenceOrConstructWithVertexDefault<float>(InputInterface, METASOUND_GET_PARAM_NAME(InputThresholdDb), InParams.OperatorSettings);
 			FTimeReadRef ReleaseTimeIn = Inputs.GetDataReadReferenceOrConstructWithVertexDefault<FTime>(InputInterface, METASOUND_GET_PARAM_NAME(InputReleaseTime), InParams.OperatorSettings);
-			FKneeModeReadRef KneeModeIn = Inputs.GetDataReadReferenceOrConstruct<FEnumKneeMode>(METASOUND_GET_PARAM_NAME(InputKneeMode));
+			FKneeModeReadRef KneeModeIn = Inputs.GetDataReadReferenceOrConstructWithVertexDefault<FEnumKneeMode>(InputInterface, METASOUND_GET_PARAM_NAME(InputKneeMode), InParams.OperatorSettings);
 
 			return MakeUnique<FLimiterOperator>(InParams, AudioIn, InGainDbIn, ThresholdDbIn, ReleaseTimeIn, KneeModeIn);
 		}
diff --git a/Engine/Plugins/Runtime/Metasound/Source/MetasoundStandardNodes/Private/MetasoundNoiseGenerator.cpp b/Engine/Plugins/Runtime/Metasound/Source/MetasoundStandardNodes/Private/MetasoundNoiseGenerator.cpp
--- a/Engine/Plugins/Runtime/Metasound/Source/MetasoundStandardNodes/Private/MetasoundNoiseGenerator.cpp
+++ b/Engine/Plugins/Runtime/Metasound/Source/MetasoundStandardNodes/Private/MetasoundNoiseGenerator.cpp
@@ -215,7 +215,7 @@ namespace Metasound
 		static const FVertexInterface Interface(
 			FInputVertexInterface(
 				TInputDataVertex<int32>(METASOUND_GET_PARAM_NAME_AND_METADATA(InputSeed), FNoiseOperator::DefaultSeed),
-				TInputDataVertex<FEnumNoiseType>(METASOUND_GET_PARAM_NAME_AND_METADATA(InputType))
+				TInputDataVertex<FEnumNoiseType>(METASOUND_GET_PARAM_NAME_AND_METADATA(InputType), (int32)ENoiseType::Pink)
 			),
 			FOutputVertexInterface(
 				TOutputDataVertex<FAudioBuffer>(METASOUND_GET_PARAM_NAME_AND_METADATA(OutAudio))
@@ -260,12 +260,13 @@ namespace Metasound
 		const FNoiseNode& Node = static_cast<const FNoiseNode&>(InParams.Node);
 		const FDataReferenceCollection& InputCol = InParams.InputDataReferences;
 		const FOperatorSettings& Settings = InParams.OperatorSettings;
+		const FInputVertexInterface& InputInterface = InParams.Node.GetVertexInterface().GetInputInterface();
 
 		// Static property pin, only used for factory.
-		FEnumNoiseTypeReadRef Type = InputCol.GetDataReadReferenceOrConstruct<FEnumNoiseType>(METASOUND_GET_PARAM_NAME(InputType));
+		FEnumNoiseTypeReadRef Type = InputCol.GetDataReadReferenceOrConstructWithVertexDefault<FEnumNoiseType>(InputInterface, METASOUND_GET_PARAM_NAME(InputType), Settings);
 
 		// Seed.
-		FInt32ReadRef Seed = InputCol.GetDataReadReferenceOrConstruct<int32>(METASOUND_GET_PARAM_NAME(InputSeed), Node.GetDefaultSeed());
+		FInt32ReadRef Seed = InputCol.GetDataReadReferenceOrConstructWithVertexDefault<int32>(InputInterface, METASOUND_GET_PARAM_NAME(InputSeed), Settings);
 
 		switch (*Type)
 		{
diff --git a/Engine/Plugins/Runtime/Metasound/Source/MetasoundStandardNodes/Private/MetasoundOscillatorNodes.cpp b/Engine/Plugins/Runtime/Metasound/Source/MetasoundStandardNodes/Private/MetasoundOscillatorNodes.cpp
--- a/Engine/Plugins/Runtime/Metasound/Source/MetasoundStandardNodes/Private/MetasoundOscillatorNodes.cpp
+++ b/Engine/Plugins/Runtime/Metasound/Source/MetasoundStandardNodes/Private/MetasoundOscillatorNodes.cpp
@@ -892,7 +892,7 @@ namespace Metasound
 			{
 				FInputVertexInterface{
 					TInputDataVertex<float>(METASOUND_GET_PARAM_NAME_AND_METADATA(LfoBaseFrequencyPin), 5.f),
-					TInputDataVertex<FEnumLfoWaveshapeType>(METASOUND_GET_PARAM_NAME_AND_METADATA(WaveshapePin)),
+					TInputDataVertex<FEnumLfoWaveshapeType>(METASOUND_GET_PARAM_NAME_AND_METADATA(WaveshapePin), (int32)ELfoWaveshapeType::Sine),
 					TInputDataVertex<float>(METASOUND_GET_PARAM_NAME_AND_METADATA(MinOutputValuePin), -1.0f),
 					TInputDataVertex<float>(METASOUND_GET_PARAM_NAME_AND_METADATA(MaxOutputValuePin), 1.0f),
 					TInputDataVertex<FTrigger>(METASOUND_GET_PARAM_NAME_AND_METADATA(LfoPhaseResetPin)),
@@ -939,7 +939,7 @@ namespace Metasound
 			return MakeUnique<FLfoOperator>(
 				  Settings
 				, InputCol.GetDataReadReferenceOrConstructWithVertexDefault<float>(InputInterface, METASOUND_GET_PARAM_NAME(LfoBaseFrequencyPin), Settings)
-				, InputCol.GetDataReadReferenceOrConstruct<FEnumLfoWaveshapeType>(METASOUND_GET_PARAM_NAME(WaveshapePin))
+				, InputCol.GetDataReadReferenceOrConstructWithVertexDefault<FEnumLfoWaveshapeType>(InputInterface, METASOUND_GET_PARAM_NAME(WaveshapePin), Settings)
 				, InputCol.GetDataReadReferenceOrConstructWithVertexDefault<float>(InputInterface, METASOUND_GET_PARAM_NAME(MinOutputValuePin), Settings)
 				, InputCol.GetDataReadReferenceOrConstructWithVertexDefault<float>(InputInterface, METASOUND_GET_PARAM_NAME(MaxOutputValuePin), Settings)
 				, InputCol.GetDataReadReferenceOrConstruct<FTrigger>(METASOUND_GET_PARAM_NAME(LfoPhaseResetPin), Settings)
diff --git a/Engine/Plugins/Runtime/Metasound/Source/MetasoundStandardNodes/Private/MetasoundStereoDelayNode.cpp b/Engine/Plugins/Runtime/Metasound/Source/MetasoundStandardNodes/Private/MetasoundStereoDelayNode.cpp
--- a/Engine/Plugins/Runtime/Metasound/Source/MetasoundStandardNodes/Private/MetasoundStereoDelayNode.cpp
+++ b/Engine/Plugins/Runtime/Metasound/Source/MetasoundStandardNodes/Private/MetasoundStereoDelayNode.cpp
@@ -323,9 +323,9 @@ namespace Metasound
 			FInputVertexInterface(
 				TInputDataVertex<FAudioBuffer>(METASOUND_GET_PARAM_NAME_AND_METADATA(InAudioInputLeft)),
 				TInputDataVertex<FAudioBuffer>(METASOUND_GET_PARAM_NAME_AND_METADATA(InAudioInputRight)),
-				TInputDataVertex<FEnumStereoDelayMode>(METASOUND_GET_PARAM_NAME_AND_METADATA(InDelayMode)),
+				TInputDataVertex<FEnumStereoDelayMode>(METASOUND_GET_PARAM_NAME_AND_METADATA(InDelayMode), (int32)EStereoDelayMode::PingPong),
 				TInputDataVertex<FTime>(METASOUND_GET_PARAM_NAME_AND_METADATA(InDelayTime), 1.0f),
-				TInputDataVertex<float>(METASOUND_GET_PARAM_NAME_AND_METADATA(InDelayRatio), 0.0f),
+				TInputDataVertex<float>(METASOUND_GET_PARAM_NAME_AND_METADATA(InDelayRatio), 0.1f),
 				TInputDataVertex<float>(METASOUND_GET_PARAM_NAME_AND_METADATA(InDryLevel), 0.0f),
 				TInputDataVertex<float>(METASOUND_GET_PARAM_NAME_AND_METADATA(InWetLevel), 1.0f),
 				TInputDataVertex<float>(METASOUND_GET_PARAM_NAME_AND_METADATA(InFeedbackAmount), 0.0f)
@@ -370,7 +370,7 @@ namespace Metasound
 
 		FAudioBufferReadRef LeftAudioIn = InputCollection.GetDataReadReferenceOrConstruct<FAudioBuffer>(METASOUND_GET_PARAM_NAME(InAudioInputLeft), InParams.OperatorSettings);
 		FAudioBufferReadRef RightAudioIn = InputCollection.GetDataReadReferenceOrConstruct<FAudioBuffer>(METASOUND_GET_PARAM_NAME(InAudioInputRight), InParams.OperatorSettings);
-		FStereoDelayModeReadRef StereoDelayMode = InputCollection.GetDataReadReferenceOrConstruct<FEnumStereoDelayMode>(METASOUND_GET_PARAM_NAME(InDelayMode));
+		FStereoDelayModeReadRef StereoDelayMode = InputCollection.GetDataReadReferenceOrConstructWithVertexDefault<FEnumStereoDelayMode>(InputInterface, METASOUND_GET_PARAM_NAME(InDelayMode), InParams.OperatorSettings);
 		FTimeReadRef DelayTime = InputCollection.GetDataReadReferenceOrConstructWithVertexDefault<FTime>(InputInterface, METASOUND_GET_PARAM_NAME(InDelayTime), InParams.OperatorSettings);
 		FFloatReadRef DelayRatio = InputCollection.GetDataReadReferenceOrConstructWithVertexDefault<float>(InputInterface, METASOUND_GET_PARAM_NAME(InDelayRatio), InParams.OperatorSettings);
 		FFloatReadRef DryLevel = InputCollection.GetDataReadReferenceOrConstructWithVertexDefault<float>(InputInterface, METASOUND_GET_PARAM_NAME(InDryLevel), InParams.OperatorSettings);
diff --git a/Engine/Plugins/Runtime/Metasound/Source/MetasoundStandardNodes/Private/MetasoundStereoPannerNode.cpp b/Engine/Plugins/Runtime/Metasound/Source/MetasoundStandardNodes/Private/MetasoundStereoPannerNode.cpp
--- a/Engine/Plugins/Runtime/Metasound/Source/MetasoundStandardNodes/Private/MetasoundStereoPannerNode.cpp
+++ b/Engine/Plugins/Runtime/Metasound/Source/MetasoundStandardNodes/Private/MetasoundStereoPannerNode.cpp
@@ -211,7 +211,7 @@ namespace Metasound
 			FInputVertexInterface(
 				TInputDataVertex<FAudioBuffer>(METASOUND_GET_PARAM_NAME_AND_METADATA(InputAudio)),
 				TInputDataVertex<float>(METASOUND_GET_PARAM_NAME_AND_METADATA(InputPanAmount), 0.0f),
-				TInputDataVertex<FEnumPanningLaw>(METASOUND_GET_PARAM_NAME_AND_METADATA(InputPanningLaw))
+				TInputDataVertex<FEnumPanningLaw>(METASOUND_GET_PARAM_NAME_AND_METADATA(InputPanningLaw), (int32)EPanningLaw::EqualPower)
 			),
 			FOutputVertexInterface(
 				TOutputDataVertex<FAudioBuffer>(METASOUND_GET_PARAM_NAME_AND_METADATA(OutputAudioLeft)),
@@ -253,7 +253,7 @@ namespace Metasound
 
 		FAudioBufferReadRef AudioIn = InputCollection.GetDataReadReferenceOrConstruct<FAudioBuffer>(METASOUND_GET_PARAM_NAME(InputAudio), InParams.OperatorSettings);
 		FFloatReadRef PanningAmount = InputCollection.GetDataReadReferenceOrConstructWithVertexDefault<float>(InputInterface, METASOUND_GET_PARAM_NAME(InputPanAmount), InParams.OperatorSettings);
-		FPanningLawReadRef PanningLaw = InputCollection.GetDataReadReferenceOrConstruct<FEnumPanningLaw>(METASOUND_GET_PARAM_NAME(InputPanningLaw));
+		FPanningLawReadRef PanningLaw = InputCollection.GetDataReadReferenceOrConstructWithVertexDefault<FEnumPanningLaw>(InputInterface, METASOUND_GET_PARAM_NAME(InputPanningLaw), InParams.OperatorSettings);
 
 		return MakeUnique<FStereoPannerOperator>(InParams.OperatorSettings, AudioIn, PanningAmount, PanningLaw);
 	}
diff --git a/Engine/Plugins/Runtime/Metasound/Source/MetasoundStandardNodes/Private/MetasoundTriggerOnThresholdNode.cpp b/Engine/Plugins/Runtime/Metasound/Source/MetasoundStandardNodes/Private/MetasoundTriggerOnThresholdNode.cpp
--- a/Engine/Plugins/Runtime/Metasound/Source/MetasoundStandardNodes/Private/MetasoundTriggerOnThresholdNode.cpp
+++ b/Engine/Plugins/Runtime/Metasound/Source/MetasoundStandardNodes/Private/MetasoundTriggerOnThresholdNode.cpp
@@ -94,7 +94,7 @@ namespace Metasound
 				FInputVertexInterface(
 					TInputDataVertex<ValueType>(METASOUND_GET_PARAM_NAME_AND_METADATA(InPin), 0),
 					TInputDataVertex<ThresholdType>(METASOUND_GET_PARAM_NAME_AND_METADATA(InThresholdPin), DefaultThreshold),
-					TInputDataVertex<FEnumBufferTriggerType>(METASOUND_GET_PARAM_NAME_AND_METADATA(InTriggerType))
+					TInputDataVertex<FEnumBufferTriggerType>(METASOUND_GET_PARAM_NAME_AND_METADATA(InTriggerType), (int32)EBufferTriggerType::RisingEdge)
 				),
 				FOutputVertexInterface(
 					TOutputDataVertex<FTrigger>(METASOUND_GET_PARAM_NAME_AND_METADATA(OutPin))
@@ -308,7 +308,7 @@ namespace Metasound
 			const FOperatorSettings& Settings = InParams.OperatorSettings;
 			const FInputVertexInterface& InputInterface = InParams.Node.GetVertexInterface().GetInputInterface();
 
-			FBufferTriggerTypeReadRef Type = InputCol.GetDataReadReferenceOrConstruct<FEnumBufferTriggerType>(METASOUND_GET_PARAM_NAME(InTriggerType));
+			FBufferTriggerTypeReadRef Type = InputCol.GetDataReadReferenceOrConstructWithVertexDefault<FEnumBufferTriggerType>(InputInterface, METASOUND_GET_PARAM_NAME(InTriggerType), Settings);
 			TDataReadReference<ValueType> Input = TTriggerOnThresholdHelper<ValueType>::CreateInput(InParams);
 			TDataReadReference<ThresholdType> Threshold = InputCol.GetDataReadReferenceOrConstructWithVertexDefault<ThresholdType>(InputInterface, METASOUND_GET_PARAM_NAME(InThresholdPin), Settings);
 
diff --git a/Engine/Plugins/Runtime/Metasound/Source/MetasoundStandardNodes/Private/MetasoundWaveShaperNode.cpp b/Engine/Plugins/Runtime/Metasound/Source/MetasoundStandardNodes/Private/MetasoundWaveShaperNode.cpp
--- a/Engine/Plugins/Runtime/Metasound/Source/MetasoundStandardNodes/Private/MetasoundWaveShaperNode.cpp
+++ b/Engine/Plugins/Runtime/Metasound/Source/MetasoundStandardNodes/Private/MetasoundWaveShaperNode.cpp
@@ -157,7 +157,7 @@ namespace Metasound
 				TInputDataVertex<float>(METASOUND_GET_PARAM_NAME_AND_METADATA(InWaveShapeAmount), 1.0f),
 				TInputDataVertex<float>(METASOUND_GET_PARAM_NAME_AND_METADATA(InWaveShapeBias), 0.0f),
 				TInputDataVertex<float>(METASOUND_GET_PARAM_NAME_AND_METADATA(InOutputGain), 1.0f),
-				TInputDataVertex<FEnumEWaveShaperType>(METASOUND_GET_PARAM_NAME_AND_METADATA(InType))
+				TInputDataVertex<FEnumEWaveShaperType>(METASOUND_GET_PARAM_NAME_AND_METADATA(InType), (int32)Audio::EWaveShaperType::Sin)
 			),
 			FOutputVertexInterface(
 				TOutputDataVertex<FAudioBuffer>(METASOUND_GET_PARAM_NAME_AND_METADATA(OutAudio))
@@ -198,7 +198,7 @@ namespace Metasound
 		FFloatReadRef InAmount = InputCollection.GetDataReadReferenceOrConstructWithVertexDefault<float>(InputInterface, METASOUND_GET_PARAM_NAME(InWaveShapeAmount), InParams.OperatorSettings);
 		FFloatReadRef InBias = InputCollection.GetDataReadReferenceOrConstructWithVertexDefault<float>(InputInterface, METASOUND_GET_PARAM_NAME(InWaveShapeBias), InParams.OperatorSettings);
 		FFloatReadRef OutGain = InputCollection.GetDataReadReferenceOrConstructWithVertexDefault<float>(InputInterface, METASOUND_GET_PARAM_NAME(InOutputGain), InParams.OperatorSettings);
-		FEnumWaveShaperReadRef InType = InputCollection.GetDataReadReferenceOrConstruct<FEnumEWaveShaperType>(METASOUND_GET_PARAM_NAME(InType));
+		FEnumWaveShaperReadRef InType = InputCollection.GetDataReadReferenceOrConstructWithVertexDefault<FEnumEWaveShaperType>(InputInterface, METASOUND_GET_PARAM_NAME(InType), InParams.OperatorSettings);
 
 		return MakeUnique<FWaveShaperOperator>(InParams.OperatorSettings, AudioIn, InAmount, InBias, OutGain, InType);
 	}
diff --git a/Engine/Plugins/Runtime/Metasound/Source/MetasoundStandardNodes/Public/MetasoundTriggerCompareNode.h b/Engine/Plugins/Runtime/Metasound/Source/MetasoundStandardNodes/Public/MetasoundTriggerCompareNode.h
--- a/Engine/Plugins/Runtime/Metasound/Source/MetasoundStandardNodes/Public/MetasoundTriggerCompareNode.h
+++ b/Engine/Plugins/Runtime/Metasound/Source/MetasoundStandardNodes/Public/MetasoundTriggerCompareNode.h
@@ -64,7 +64,7 @@ namespace Metasound
 					TInputDataVertex<FTrigger>(METASOUND_GET_PARAM_NAME_AND_METADATA(InputCompare)),
 					TInputDataVertex<ValueType>(METASOUND_GET_PARAM_NAME_AND_METADATA(InputParamA)),
 					TInputDataVertex<ValueType>(METASOUND_GET_PARAM_NAME_AND_METADATA(InputParamB)),
-					TInputDataVertex<FEnumTriggerComparisonType>(METASOUND_GET_PARAM_NAME_AND_METADATA(InputCompareType))
+					TInputDataVertex<FEnumTriggerComparisonType>(METASOUND_GET_PARAM_NAME_AND_METADATA(InputCompareType), (int32)ETriggerComparisonType::Equals)
 				),
 				FOutputVertexInterface(
 					TOutputDataVertex<FTrigger>(METASOUND_GET_PARAM_NAME_AND_METADATA(OutputOnTrue)),
@@ -102,7 +102,7 @@ namespace Metasound
 			FTriggerReadRef InOnTriggerCompare = InputCollection.GetDataReadReferenceOrConstruct<FTrigger>(METASOUND_GET_PARAM_NAME(InputCompare), InParams.OperatorSettings);
 			TDataReadReference<ValueType> InValueA = InputCollection.GetDataReadReferenceOrConstructWithVertexDefault<ValueType>(InputInterface, METASOUND_GET_PARAM_NAME(InputParamA), InParams.OperatorSettings);
 			TDataReadReference<ValueType> InValueB = InputCollection.GetDataReadReferenceOrConstructWithVertexDefault<ValueType>(InputInterface, METASOUND_GET_PARAM_NAME(InputParamB), InParams.OperatorSettings);
-			FEnumTriggerComparisonTypeReadRef InComparison = InputCollection.GetDataReadReferenceOrConstruct<FEnumTriggerComparisonType>(METASOUND_GET_PARAM_NAME(InputCompareType));
+			FEnumTriggerComparisonTypeReadRef InComparison = InputCollection.GetDataReadReferenceOrConstructWithVertexDefault<FEnumTriggerComparisonType>(InputInterface, METASOUND_GET_PARAM_NAME(InputCompareType), InParams.OperatorSettings);
 
 			return MakeUnique<TTriggerCompareNodeOperator<ValueType>>(InParams.OperatorSettings, InOnTriggerCompare, InValueA, InValueB, InComparison);
 
