Fix reset to default arrows appearing on MetaSound object inputs and reroute nodes when copy/pasted

diff --git a/Engine/Plugins/Runtime/Metasound/Source/MetasoundEditor/Private/MetasoundEditorGraphBuilder.cpp b/Engine/Plugins/Runtime/Metasound/Source/MetasoundEditor/Private/MetasoundEditorGraphBuilder.cpp
--- a/Engine/Plugins/Runtime/Metasound/Source/MetasoundEditor/Private/MetasoundEditorGraphBuilder.cpp
+++ b/Engine/Plugins/Runtime/Metasound/Source/MetasoundEditor/Private/MetasoundEditorGraphBuilder.cpp
@@ -1092,7 +1092,7 @@ namespace Metasound
 					
 					if (!bObjectFound)
 					{
-						OutDefaultLiteral.Set(static_cast<UObject*>(nullptr));
+						OutDefaultLiteral.Clear();
 					}
 				}
 				break;
diff --git a/Engine/Plugins/Runtime/Metasound/Source/MetasoundEditor/Private/SMetasoundGraphPin.h b/Engine/Plugins/Runtime/Metasound/Source/MetasoundEditor/Private/SMetasoundGraphPin.h
--- a/Engine/Plugins/Runtime/Metasound/Source/MetasoundEditor/Private/SMetasoundGraphPin.h
+++ b/Engine/Plugins/Runtime/Metasound/Source/MetasoundEditor/Private/SMetasoundGraphPin.h
@@ -13,6 +13,7 @@
 #include "MetasoundEditorGraphMemberDefaults.h"
 #include "MetasoundEditorGraphNode.h"
 #include "MetasoundFrontendController.h"
+#include "NodeTemplates/MetasoundFrontendNodeTemplateReroute.h"
 #include "SGraphPin.h"
 #include "SGraphNodeKnot.h"
 #include "SMetasoundPinValueInspector.h"
@@ -262,7 +263,19 @@ namespace Metasound
 								{
 									const bool bIsDefaultConstructed = Literal->GetType() == EMetasoundFrontendLiteralType::None;
 									const bool bIsTriggerDataType = InputHandle->GetDataType() == GetMetasoundDataTypeName<FTrigger>();
-									if (!bIsDefaultConstructed && !bIsTriggerDataType)
+									bool bIsRerouteNode = false;
+									if (UEdGraphPin* Pin = ParentPinType::GetPinObj())
+									{
+										if (UMetasoundEditorGraphExternalNode* Node = Cast<UMetasoundEditorGraphExternalNode>(Pin->GetOwningNode()))
+										{
+											if (Node->GetClassName() == FRerouteNodeTemplate::ClassName)
+											{
+												bIsRerouteNode = true;
+											}
+										}
+									}
+									
+									if (!bIsDefaultConstructed && !bIsTriggerDataType && !bIsRerouteNode)
 									{
 										return EVisibility::Visible;
 									}
