Fix copypaste in Metasound Random nodes that broke min/max value assignments

#rnx

diff --git a/Engine/Plugins/Runtime/Metasound/Source/MetasoundStandardNodes/Public/MetasoundRandomNode.h b/Engine/Plugins/Runtime/Metasound/Source/MetasoundStandardNodes/Public/MetasoundRandomNode.h
--- a/Engine/Plugins/Runtime/Metasound/Source/MetasoundStandardNodes/Public/MetasoundRandomNode.h
+++ b/Engine/Plugins/Runtime/Metasound/Source/MetasoundStandardNodes/Public/MetasoundRandomNode.h
@@ -398,7 +398,7 @@ namespace Metasound
 			if (TRandomNodeSpecialization<ValueType>::HasRange())
 			{
 				InOutVertexData.BindReadVertex(METASOUND_GET_PARAM_NAME(InputMin), MinValue);
-				InOutVertexData.BindReadVertex(METASOUND_GET_PARAM_NAME(InputMax), MinValue);
+				InOutVertexData.BindReadVertex(METASOUND_GET_PARAM_NAME(InputMax), MaxValue);
 			}
 		}
 
