MetaSound Frontend allows derived classes to attach to base class pins at runtime.

FDataTypeRegistry::GetDataTypeInfo walks up the class hierarchy searching for a registered class.
Adds FDataTypeRegistryInfo::bIsExplicit to replace the editor's IsExplicitProxyClass and enable the same functionality of requiring exact types at runtime.

diff --git a/Engine/Plugins/Runtime/Metasound/Source/MetasoundEditor/Private/MetasoundEditorGraphNode.cpp b/Engine/Plugins/Runtime/Metasound/Source/MetasoundEditor/Private/MetasoundEditorGraphNode.cpp
--- a/Engine/Plugins/Runtime/Metasound/Source/MetasoundEditor/Private/MetasoundEditorGraphNode.cpp
+++ b/Engine/Plugins/Runtime/Metasound/Source/MetasoundEditor/Private/MetasoundEditorGraphNode.cpp
@@ -436,10 +436,17 @@ TSet<FString> UMetasoundEditorGraphNode::GetDisallowedPinClassNames(const UEdGra
 			continue;
 		}
 
+PRAGMA_DISABLE_DEPRECATION_WARNINGS
 		if (EditorModule.IsExplicitProxyClass(*ProxyGenClass) && Class->IsChildOf(ProxyGenClass))
 		{
 			DisallowedClasses.Add(ClassIt->GetClassPathName().ToString());
 		}
+PRAGMA_ENABLE_DEPRECATION_WARNINGS
+
+		if (DataTypeInfo.bIsExplicit && Class->IsChildOf(ProxyGenClass))
+		{
+			DisallowedClasses.Add(ClassIt->GetClassPathName().ToString());
+		}
 	}
 
 	return DisallowedClasses;
diff --git a/Engine/Plugins/Runtime/Metasound/Source/MetasoundEditor/Private/MetasoundEditorGraphSchema.cpp b/Engine/Plugins/Runtime/Metasound/Source/MetasoundEditor/Private/MetasoundEditorGraphSchema.cpp
--- a/Engine/Plugins/Runtime/Metasound/Source/MetasoundEditor/Private/MetasoundEditorGraphSchema.cpp
+++ b/Engine/Plugins/Runtime/Metasound/Source/MetasoundEditor/Private/MetasoundEditorGraphSchema.cpp
@@ -126,10 +126,16 @@ namespace Metasound
 					{
 						if (const UClass* Class = Asset.GetClass())
 						{
+PRAGMA_DISABLE_DEPRECATION_WARNINGS
 							if (EditorModule.IsExplicitProxyClass(*InRegistryInfo.ProxyGeneratorClass))
 							{
 								return Class == InRegistryInfo.ProxyGeneratorClass;
 							}
+PRAGMA_ENABLE_DEPRECATION_WARNINGS
+							if (InRegistryInfo.bIsExplicit)
+							{
+								return Class == InRegistryInfo.ProxyGeneratorClass;
+							}
 							else
 							{
 								return Class->IsChildOf(InRegistryInfo.ProxyGeneratorClass);
diff --git a/Engine/Plugins/Runtime/Metasound/Source/MetasoundEditor/Private/MetasoundEditorModule.cpp b/Engine/Plugins/Runtime/Metasound/Source/MetasoundEditor/Private/MetasoundEditorModule.cpp
--- a/Engine/Plugins/Runtime/Metasound/Source/MetasoundEditor/Private/MetasoundEditorModule.cpp
+++ b/Engine/Plugins/Runtime/Metasound/Source/MetasoundEditor/Private/MetasoundEditorModule.cpp
@@ -815,8 +815,6 @@ namespace Metasound
 				// the Metasound Editor registration of a USoundWave.
 				FModuleManager::LoadModuleChecked<IModuleInterface>("MetasoundEngine");
 
-				RegisterExplicitProxyClass(*USoundWave::StaticClass());
-
 				// Required to ensure logic to order nodes for presets exclusive to
 				// editor is propagated to transform instances while editing in editor.
 				Frontend::DocumentTransform::RegisterNodeDisplayNameProjection([](const Frontend::FNodeHandle& NodeHandle)
diff --git a/Engine/Plugins/Runtime/Metasound/Source/MetasoundEditor/Private/MetasoundNodeDetailCustomization.cpp b/Engine/Plugins/Runtime/Metasound/Source/MetasoundEditor/Private/MetasoundNodeDetailCustomization.cpp
--- a/Engine/Plugins/Runtime/Metasound/Source/MetasoundEditor/Private/MetasoundNodeDetailCustomization.cpp
+++ b/Engine/Plugins/Runtime/Metasound/Source/MetasoundEditor/Private/MetasoundNodeDetailCustomization.cpp
@@ -251,7 +251,14 @@ namespace Metasound
 						{
 							if (UClass* Class = AssetData.GetClass())
 							{
+PRAGMA_DISABLE_DEPRECATION_WARNINGS
 								if (EditorModule.IsExplicitProxyClass(*DataTypeInfo.ProxyGeneratorClass))
+								{
+									bCanDrop &= Class == DataTypeInfo.ProxyGeneratorClass;
+									continue;
+								}
+PRAGMA_ENABLE_DEPRECATION_WARNINGS
+								if (DataTypeInfo.bIsExplicit)
 								{
 									bCanDrop &= Class == DataTypeInfo.ProxyGeneratorClass;
 								}
@@ -437,18 +444,24 @@ namespace Metasound
 			TSharedPtr<IPropertyHandle> PropertyHandle = StructPropertyHandle->GetChildHandle(GET_MEMBER_NAME_CHECKED(FMetasoundEditorGraphMemberDefaultObjectRef, Object));
 
 			const IMetasoundEditorModule& EditorModule = FModuleManager::GetModuleChecked<IMetasoundEditorModule>("MetaSoundEditor");
-			auto FilterAsset = [InEditorModule = &EditorModule, InProxyGenClass = DataTypeInfo.ProxyGeneratorClass](const FAssetData& InAsset)
+			auto FilterAsset = [InEditorModule = &EditorModule, &InDataTypeInfo = DataTypeInfo](const FAssetData& InAsset)
 			{
-				if (InProxyGenClass)
+				if (InDataTypeInfo.ProxyGeneratorClass)
 				{
 					if (UClass* Class = InAsset.GetClass())
 					{
-						if (InEditorModule->IsExplicitProxyClass(*InProxyGenClass))
+PRAGMA_DISABLE_DEPRECATION_WARNINGS
+						if (InEditorModule->IsExplicitProxyClass(*InDataTypeInfo.ProxyGeneratorClass))
+						{
+							return Class != InDataTypeInfo.ProxyGeneratorClass;
+						}
+PRAGMA_ENABLE_DEPRECATION_WARNINGS
+						if (InDataTypeInfo.bIsExplicit)
 						{
-							return Class != InProxyGenClass;
+							return Class != InDataTypeInfo.ProxyGeneratorClass;
 						}
 
-						return !Class->IsChildOf(InProxyGenClass);
+						return !Class->IsChildOf(InDataTypeInfo.ProxyGeneratorClass);
 					}
 				}
 
diff --git a/Engine/Plugins/Runtime/Metasound/Source/MetasoundEditor/Public/MetasoundEditorModule.h b/Engine/Plugins/Runtime/Metasound/Source/MetasoundEditor/Public/MetasoundEditorModule.h
--- a/Engine/Plugins/Runtime/Metasound/Source/MetasoundEditor/Public/MetasoundEditorModule.h
+++ b/Engine/Plugins/Runtime/Metasound/Source/MetasoundEditor/Public/MetasoundEditorModule.h
@@ -79,13 +79,10 @@ namespace Metasound
 		class METASOUNDEDITOR_API IMetasoundEditorModule : public IModuleInterface
 		{
 		public:
-			// Whether or not the given proxy class has to be explicit (i.e.
-			// selectors do not support inherited types). By default, proxy
-			// classes support child classes & inheritance.
+			UE_DEPRECATED(5.3, "IsExplicitProxyClass is deprecated, use Metasound::Frontend::FDataTypeRegistryInfo::bIsExplicit")
 			virtual bool IsExplicitProxyClass(const UClass& InClass) const = 0;
 
-			// Register proxy class as explicitly selectable.
-			// By default, proxy classes support child classes & inheritance.
+			UE_DEPRECATED(5.3, "RegisterExplicitProxyClass is deprecated, use Metasound::TIsExplicit<>")
 			virtual void RegisterExplicitProxyClass(const UClass& InClass) = 0;
 
 			UE_DEPRECATED(5.3, "IsMetaSoundAssetClass is deprecated, use IMetasoundUObjectRegistry::IsRegisteredClass")
diff --git a/Engine/Plugins/Runtime/Metasound/Source/MetasoundEngine/Public/MetasoundWave.h b/Engine/Plugins/Runtime/Metasound/Source/MetasoundEngine/Public/MetasoundWave.h
--- a/Engine/Plugins/Runtime/Metasound/Source/MetasoundEngine/Public/MetasoundWave.h
+++ b/Engine/Plugins/Runtime/Metasound/Source/MetasoundEngine/Public/MetasoundWave.h
@@ -17,6 +17,12 @@ namespace Metasound
 	class FWaveAsset;
 	typedef TDataReadReference<FWaveAsset> FWaveAssetReadRef;
 
+	// Helper utility to test if exact types are required for a datatype.
+	template <>
+	struct TIsExplicit<FWaveAsset>
+	{
+		static constexpr bool Value = true;
+	};
 
 	// Metasound data type that holds onto a weak ptr. Mostly used as a placeholder until we have a proper proxy type.
 	class METASOUNDENGINE_API FWaveAsset
diff --git a/Engine/Plugins/Runtime/Metasound/Source/MetasoundFrontend/Private/MetasoundFrontendDataTypeRegistry.cpp b/Engine/Plugins/Runtime/Metasound/Source/MetasoundFrontend/Private/MetasoundFrontendDataTypeRegistry.cpp
--- a/Engine/Plugins/Runtime/Metasound/Source/MetasoundFrontend/Private/MetasoundFrontendDataTypeRegistry.cpp
+++ b/Engine/Plugins/Runtime/Metasound/Source/MetasoundFrontend/Private/MetasoundFrontendDataTypeRegistry.cpp
@@ -639,12 +639,20 @@ namespace Metasound
 				if (InObject)
 				{
 					FScopeLock Lock(&RegistryObjectMapMutex);
-					if (const FName* DataTypeName = RegisteredObjectClasses.Find(InObject->GetClass()))
+					for (UClass* Class = InObject->GetClass(); Class != UObject::StaticClass(); Class = Class->GetSuperClass())
 					{
-						if (const IDataTypeRegistryEntry* Entry = FindDataTypeEntry(*DataTypeName))
+						if (const FName* DataTypeName = RegisteredObjectClasses.Find(Class))
 						{
-							OutInfo = Entry->GetDataTypeInfo();
-							return true;
+							if (const IDataTypeRegistryEntry* Entry = FindDataTypeEntry(*DataTypeName))
+							{
+								const FDataTypeRegistryInfo& Info = Entry->GetDataTypeInfo();
+								if (Info.bIsExplicit && Class != InObject->GetClass())
+								{
+									return false;
+								}
+								OutInfo = Info;
+								return true;
+							}
 						}
 					}
 				}
@@ -823,9 +831,13 @@ namespace Metasound
 				UClass* ObjectClass = InObject->GetClass();
 				while (ObjectClass != UObject::StaticClass())
 				{
-					if (RegisteredObjectClasses.Contains(ObjectClass))
+					if (const FName* DataTypeName = RegisteredObjectClasses.Find(ObjectClass))
 					{
-						return true;
+						if (const IDataTypeRegistryEntry* Entry = FindDataTypeEntry(*DataTypeName))
+						{
+							const FDataTypeRegistryInfo& Info = Entry->GetDataTypeInfo();
+							return !Info.bIsExplicit || ObjectClass == InObject->GetClass();
+						}
 					}
 
 					ObjectClass = ObjectClass->GetSuperClass();
diff --git a/Engine/Plugins/Runtime/Metasound/Source/MetasoundFrontend/Public/MetasoundDataTypeRegistrationMacro.h b/Engine/Plugins/Runtime/Metasound/Source/MetasoundFrontend/Public/MetasoundDataTypeRegistrationMacro.h
--- a/Engine/Plugins/Runtime/Metasound/Source/MetasoundFrontend/Public/MetasoundDataTypeRegistrationMacro.h
+++ b/Engine/Plugins/Runtime/Metasound/Source/MetasoundFrontend/Public/MetasoundDataTypeRegistrationMacro.h
@@ -217,6 +217,7 @@ namespace Metasound
 			RegistryInfo.bIsUniquePtrProxyArrayParsable_DEPRECATED = TIsParsable<TDataType, const TArray<TUniquePtr<Audio::IProxyData>>& >::Value;
 
 			RegistryInfo.bIsEnum = TEnumTraits<TDataType>::bIsEnum;
+			RegistryInfo.bIsExplicit = TIsExplicit<TDataType>::Value;
 			RegistryInfo.bIsVariable = TIsVariable<TDataType>::Value;
 			RegistryInfo.bIsTransmittable = TIsTransmittable<TDataType>::Value;
 			RegistryInfo.bIsConstructorType = TIsConstructorVertexSupported<TDataType>::Value;
diff --git a/Engine/Plugins/Runtime/Metasound/Source/MetasoundFrontend/Public/MetasoundFrontendDataTypeRegistry.h b/Engine/Plugins/Runtime/Metasound/Source/MetasoundFrontend/Public/MetasoundFrontendDataTypeRegistry.h
--- a/Engine/Plugins/Runtime/Metasound/Source/MetasoundFrontend/Public/MetasoundFrontendDataTypeRegistry.h
+++ b/Engine/Plugins/Runtime/Metasound/Source/MetasoundFrontend/Public/MetasoundFrontendDataTypeRegistry.h
@@ -56,6 +56,9 @@ namespace Metasound
 			// Is a TEnum wrapped enum
 			bool bIsEnum = false;
 
+			// Whether exact types are required
+			bool bIsExplicit = false;
+
 			// Determines whether the type can be used with send/receive transmitters
 			bool bIsTransmittable = false;
 
diff --git a/Engine/Plugins/Runtime/Metasound/Source/MetasoundFrontend/Public/MetasoundFrontendDataTypeTraits.h b/Engine/Plugins/Runtime/Metasound/Source/MetasoundFrontend/Public/MetasoundFrontendDataTypeTraits.h
--- a/Engine/Plugins/Runtime/Metasound/Source/MetasoundFrontend/Public/MetasoundFrontendDataTypeTraits.h
+++ b/Engine/Plugins/Runtime/Metasound/Source/MetasoundFrontend/Public/MetasoundFrontendDataTypeTraits.h
@@ -69,6 +69,13 @@ namespace Metasound
 		static constexpr bool Value = bIsParsable && bIsCopyable && bIsEnabled;
 	};
 
+	// Helper utility to test if exact types are required for a datatype.
+	template <typename TDataType>
+	struct TIsExplicit
+	{
+		static constexpr bool Value = false;
+	};
+
 	// Helper utility to test if we can transmit a datatype between a send and a receive node.
 	template <typename TDataType>
 	struct TIsTransmittable
