Add NumRepeats field to TriggerRepeat

diff --git a/Engine/Plugins/Runtime/Metasound/Source/MetasoundStandardNodes/Private/MetasoundTriggerRepeatNode.cpp b/Engine/Plugins/Runtime/Metasound/Source/MetasoundStandardNodes/Private/MetasoundTriggerRepeatNode.cpp
--- a/Engine/Plugins/Runtime/Metasound/Source/MetasoundStandardNodes/Private/MetasoundTriggerRepeatNode.cpp
+++ b/Engine/Plugins/Runtime/Metasound/Source/MetasoundStandardNodes/Private/MetasoundTriggerRepeatNode.cpp
@@ -29,6 +29,7 @@ namespace Metasound
 		METASOUND_PARAM(InputStart, "Start", "Starts executing periodic output triggers.");
 		METASOUND_PARAM(InputStop, "Stop", "Stops executing periodic output triggers.");
 		METASOUND_PARAM(InputPeriod, "Period", "The period to trigger in seconds.");
+		METASOUND_PARAM(InputNumRepeats, "Num Repeats", "The number of times to repeat. When set to 0, will repeat indefinitely until stopped.")
 		METASOUND_PARAM(RepeatOutputOnTrigger, "RepeatOut", "The periodically generated output trigger.");
 	}
 
@@ -39,7 +40,8 @@ namespace Metasound
 			static const FVertexInterface& GetVertexInterface();
 			static TUniquePtr<IOperator> CreateOperator(const FCreateOperatorParams& InParams, FBuildErrorArray& OutErrors);
 
-			FTriggerRepeatOperator(const FOperatorSettings& InSettings, const FTriggerReadRef& InTriggerEnable, const FTriggerReadRef& InTriggerDisable, const FTimeReadRef& InPeriod);
+			FTriggerRepeatOperator(const FOperatorSettings& InSettings, const FTriggerReadRef& InTriggerEnable, const FTriggerReadRef& InTriggerDisable, 
+				const FTimeReadRef& InPeriod, const FInt32ReadRef& InNumRepeats);
 
 			virtual void BindInputs(FInputVertexInterfaceData& InOutVertexData) override;
 
@@ -62,16 +64,22 @@ namespace Metasound
 
 			FTimeReadRef Period;
 
+			FInt32ReadRef NumRepeats;
+
 			FSampleCount BlockSize;
 			FSampleCounter SampleCounter;
+
+			//The number of repeats the node has already performed
+			int CurrentRepeats = 0;
 	};
 
-	FTriggerRepeatOperator::FTriggerRepeatOperator(const FOperatorSettings& InSettings, const FTriggerReadRef& InTriggerEnable, const FTriggerReadRef& InTriggerDisable, const FTimeReadRef& InPeriod)
+	FTriggerRepeatOperator::FTriggerRepeatOperator(const FOperatorSettings& InSettings, const FTriggerReadRef& InTriggerEnable, const FTriggerReadRef& InTriggerDisable, const FTimeReadRef& InPeriod, const FInt32ReadRef& InNumRepeats)
 		: bEnabled(false)
 		, TriggerOut(FTriggerWriteRef::CreateNew(InSettings))
 		, TriggerEnable(InTriggerEnable)
 		, TriggerDisable(InTriggerDisable)
 		, Period(InPeriod)
+		, NumRepeats(InNumRepeats)
 		, BlockSize(InSettings.GetNumFramesPerBlock())
 		, SampleCounter(0, InSettings.GetSampleRate())
 	{
@@ -84,6 +92,8 @@ namespace Metasound
 		InOutVertexData.BindReadVertex(METASOUND_GET_PARAM_NAME(InputStart), TriggerEnable);
 		InOutVertexData.BindReadVertex(METASOUND_GET_PARAM_NAME(InputStop), TriggerDisable);
 		InOutVertexData.BindReadVertex(METASOUND_GET_PARAM_NAME(InputPeriod), Period);
+		InOutVertexData.BindReadVertex(METASOUND_GET_PARAM_NAME(InputNumRepeats), NumRepeats);
+
 	}
 
 	void FTriggerRepeatOperator::BindOutputs(FOutputVertexInterfaceData& InOutVertexData)
@@ -115,6 +125,9 @@ namespace Metasound
 			{
 				bEnabled = true;
 				SampleCounter.SetNumSamples(0);
+
+				//Reset current number of repeats if we get another play command
+				CurrentRepeats = 0;
 			}
 		);
 
@@ -137,9 +150,19 @@ namespace Metasound
 
 			while ((SampleCounter - BlockSize).GetNumSamples() <= 0)
 			{
+				//Whether or not the node has any repeats it should perform
+				bool bShouldRepeat = (*NumRepeats == 0) || (CurrentRepeats < *NumRepeats);
+
+				if (!bShouldRepeat)
+				{
+					break;
+				}
+
 				const int32 StartOffset = static_cast<int32>(SampleCounter.GetNumSamples());
 				TriggerOut->TriggerFrame(StartOffset);
 				SampleCounter += PeriodInSamples;
+
+				++CurrentRepeats;
 			}
 
 			SampleCounter -= BlockSize;
@@ -152,6 +175,7 @@ namespace Metasound
 
 		TriggerOut->Reset();
 		SampleCounter.SetNumSamples(0);
+		CurrentRepeats = 0;
 	}
 
 	TUniquePtr<IOperator> FTriggerRepeatOperator::CreateOperator(const FCreateOperatorParams& InParams, FBuildErrorArray& OutErrors)
@@ -164,8 +188,9 @@ namespace Metasound
 		FTriggerReadRef TriggerEnable = InParams.InputDataReferences.GetDataReadReferenceOrConstruct<FTrigger>(METASOUND_GET_PARAM_NAME(InputStart), InParams.OperatorSettings);
 		FTriggerReadRef TriggerDisable = InParams.InputDataReferences.GetDataReadReferenceOrConstruct<FTrigger>(METASOUND_GET_PARAM_NAME(InputStop), InParams.OperatorSettings);
 		FTimeReadRef Period = InParams.InputDataReferences.GetDataReadReferenceOrConstructWithVertexDefault<FTime>(InputInterface, METASOUND_GET_PARAM_NAME(InputPeriod), InParams.OperatorSettings);
+		FInt32ReadRef NumRepeats = InParams.InputDataReferences.GetDataReadReferenceOrConstructWithVertexDefault<int32>(InputInterface, METASOUND_GET_PARAM_NAME(InputNumRepeats), InParams.OperatorSettings);
 
-		return MakeUnique<FTriggerRepeatOperator>(InParams.OperatorSettings, TriggerEnable, TriggerDisable, Period);
+		return MakeUnique<FTriggerRepeatOperator>(InParams.OperatorSettings, TriggerEnable, TriggerDisable, Period, NumRepeats);
 	}
 
 	const FVertexInterface& FTriggerRepeatOperator::GetVertexInterface()
@@ -176,7 +201,8 @@ namespace Metasound
 			FInputVertexInterface(
 				TInputDataVertex<FTrigger>(METASOUND_GET_PARAM_NAME_AND_METADATA(InputStart)),
 				TInputDataVertex<FTrigger>(METASOUND_GET_PARAM_NAME_AND_METADATA(InputStop)),
-				TInputDataVertex<FTime>(METASOUND_GET_PARAM_NAME_AND_METADATA(InputPeriod), 0.2f)
+				TInputDataVertex<FTime>(METASOUND_GET_PARAM_NAME_AND_METADATA(InputPeriod), 0.2f),
+				TInputDataVertex<int32>(METASOUND_GET_PARAM_NAME_AND_METADATA(InputNumRepeats), 0)
 			),
 			FOutputVertexInterface(
 				TOutputDataVertex<FTrigger>(METASOUND_GET_PARAM_NAME_AND_METADATA(RepeatOutputOnTrigger))
