Move Graph builder helper functions to NodeTestGraphBuilder

Fix up some test checks in MetronomeNodeTests
- more comprehensive tests in MetronomedNodeTests
- Test a range of Tempos and Time Signatures

diff --git a/Engine/Plugins/Runtime/Metasound/Source/MetasoundEngineTest/Private/NodeTestGraphBuilder.cpp b/Engine/Plugins/Runtime/Metasound/Source/MetasoundEngineTest/Private/NodeTestGraphBuilder.cpp
--- a/Engine/Plugins/Runtime/Metasound/Source/MetasoundEngineTest/Private/NodeTestGraphBuilder.cpp
+++ b/Engine/Plugins/Runtime/Metasound/Source/MetasoundEngineTest/Private/NodeTestGraphBuilder.cpp
@@ -34,6 +34,17 @@ namespace Metasound::Test
 		return Node;
 	}
 
+
+	FNodeHandle FNodeTestGraphBuilder::GetInput(const FName& InName) const
+	{
+		return RootGraph->GetInputNodeWithName(InName);
+	}
+
+	FNodeHandle FNodeTestGraphBuilder::GetOutput(const FName& InName) const
+	{
+		return RootGraph->GetOutputNodeWithName(InName);
+	}
+
 	FNodeHandle FNodeTestGraphBuilder::AddInput(
 		const FName& InputName,
 		const FName& TypeName,
@@ -171,4 +182,70 @@ namespace Metasound::Test
 		// build the graph
 		return Builder.BuildGenerator(SampleRate, SamplesPerBlock);
 	}
+
+	bool FNodeTestGraphBuilder::ConnectNodes(const Frontend::FNodeHandle& LeftNode, const FName& OutputName, const Frontend::FNodeHandle& RightNode, const FName& InputName)
+	{
+		if (!LeftNode->IsValid() || !RightNode->IsValid())
+		{
+			return false;
+		}
+
+		FOutputHandle OutputToConnect = LeftNode->GetOutputWithVertexName(OutputName);
+		FInputHandle InputToConnect = RightNode->GetInputWithVertexName(InputName);
+
+		if (!OutputToConnect->IsValid() || !InputToConnect->IsValid())
+		{
+			return false;
+		}
+		return InputToConnect->Connect(*OutputToConnect);
+	}
+
+	bool FNodeTestGraphBuilder::ConnectNodes(const Frontend::FNodeHandle& LeftNode, const Frontend::FNodeHandle& RightNode, const FName& InputOutputName)
+	{
+		return ConnectNodes(LeftNode, InputOutputName, RightNode, InputOutputName);
+	}
+
+	bool FNodeTestGraphBuilder::AddAndConnectDataReferenceInputs(const Frontend::FNodeHandle& NodeHandle)
+	{
+		// add the inputs and connect them
+		for (FInputHandle Input : NodeHandle->GetInputs())
+		{
+			if (Input->GetVertexAccessType() != EMetasoundFrontendVertexAccessType::Reference)
+			{
+				continue;
+			}
+
+			FNodeHandle InputNode = AddInput(Input->GetName(), Input->GetDataType(), Input->GetVertexAccessType());
+
+			if (!InputNode->IsValid())
+			{
+				return false;
+			}
+
+			FOutputHandle OutputToConnect = InputNode->GetOutputWithVertexName(Input->GetName());
+			FInputHandle InputToConnect = NodeHandle->GetInputWithVertexName(Input->GetName());
+
+			if (!InputToConnect->Connect(*OutputToConnect))
+			{
+				return false;
+			}
+		}
+
+		return true;
+	}
+
+	FNodeHandle FNodeTestGraphBuilder::AddAndConnectDataReferenceInput(const FNodeHandle& NodeToConnect, const FName& InputName, const FName& TypeName) const
+	{
+		FNodeHandle InputNode = AddInput(InputName, TypeName);
+		check(ConnectNodes(InputNode, NodeToConnect, InputName));
+		return InputNode;
+	}
+
+	FNodeHandle FNodeTestGraphBuilder::AddAndConnectDataReferenceOutput(const FNodeHandle& NodeToConnect, const FName& OutputName, const FName& TypeName)
+	{
+		FNodeHandle OutputNode = AddOutput(OutputName, TypeName);
+		check(ConnectNodes(NodeToConnect, OutputNode, OutputName));
+		return OutputNode;
+	}
+
 }
diff --git a/Engine/Plugins/Runtime/Metasound/Source/MetasoundEngineTest/Public/NodeTestGraphBuilder.h b/Engine/Plugins/Runtime/Metasound/Source/MetasoundEngineTest/Public/NodeTestGraphBuilder.h
--- a/Engine/Plugins/Runtime/Metasound/Source/MetasoundEngineTest/Public/NodeTestGraphBuilder.h
+++ b/Engine/Plugins/Runtime/Metasound/Source/MetasoundEngineTest/Public/NodeTestGraphBuilder.h
@@ -15,16 +15,51 @@ namespace Metasound::Test
 	public:
 		FNodeTestGraphBuilder();
 
-		/** Add a node to the graph */
+		/** 
+		* Add a node to the graph 
+		*
+		* @param ClassName		Class name of the node to add
+		* @param MajorVersion	Major version of the node to add
+		* @return				NodeHandle of the added node
+		*/
 		Frontend::FNodeHandle AddNode(const FNodeClassName& ClassName, int32 MajorVersion) const;
 
-		/** Add an input node to the graph */
+		/**
+		* Add an input node from the graph by name
+		*
+		* @param Name			The name of the input node to get
+		* @return				NodeHandle, which can be invalid if an input with the given name doesn't exist in the graph
+		*/
+		Frontend::FNodeHandle GetInput(const FName& Name) const;
+
+		/**
+		* Add an output node from the graph by name
+		*
+		* @param Name			The name of the output node to get
+		* @return				NodeHandle, which can be invalid if an input with the given name doesn't exist in the graph
+		*/
+		Frontend::FNodeHandle GetOutput(const FName& Name) const;
+
+		/**
+		* Add an input node from the graph by name, type, and access type
+		*
+		* @param InputName		The name of the output node to get
+		* @param TypeName		The type of the node to add (use: GetMetasoundDataTypeName<DataType>());
+		* @param AccessType		Whether the node should be a Value, Referenc, or Unset type
+		* @return				NodeHandle of the added node, can be invalid if failed to add
+		*/
 		Frontend::FNodeHandle AddInput(
 			const FName& InputName,
 			const FName& TypeName,
 			EMetasoundFrontendVertexAccessType AccessType = EMetasoundFrontendVertexAccessType::Reference) const;
 
-		/** Add a constructor input to the graph */
+		/**
+		* Add a constructor input node from the graph by name and type
+		*
+		* @param InputName		The name of the output node to get
+		* @param Value			The value for the node to hold. The type of the value will determine the type for the node.
+		* @return				NodeHandle of the added node, can be invalid if failed to add
+		*/
 		template<typename DataType>
 		Frontend::FNodeHandle AddConstructorInput(const FName& InputName, const DataType& Value) const
 		{
@@ -39,17 +74,102 @@ namespace Metasound::Test
 			return RootGraph->AddInputVertex(Input);
 		}
 
-		/** Add an output node to the graph */
+		/**
+		* Add a data reference output node from the graph by name and type
+		*
+		* @param InputName		The name of the output node to get
+		* @param TypeName		The type for the node to reference
+		* @return				NodeHandle of the added node, can be invalid if failed to add
+		*/
 		Frontend::FNodeHandle AddOutput(const FName& OutputName, const FName& TypeName);
 
+		/** 
+		* Adds an constructor input node to the graph and connects it as an input to the NodeToConnect
+		* 
+		* @param NodeToConnect	The node to connect the newly added constructor input node to
+		* @param InputName		The name of the input pin to connect the node to, and the name to use for the new node.
+		* @param Value			The value for the node to hold. The type of the value will determine the type for the node.
+		* @return				NodeHandle of the added node, can be invalid if failed to add
+		*/
+		template <typename DataType>
+		Frontend::FNodeHandle AddAndConnectConstructorInput(const Frontend::FNodeHandle& NodeToConnect, const FName& InputName, const DataType& Value) const
+		{
+			Frontend::FNodeHandle InputNode = AddConstructorInput<DataType>(InputName, Value);
+			check(ConnectNodes(InputNode, NodeToConnect, InputName));
+			return InputNode;
+		}
+
+		/** 
+		* Adds data reference input nodes for each input on the NodeHandle and connects them to the inputs of the NodeHandle
+		*
+		* @param NodeHandle		The node to connect the input reference nodes to
+		* @return				whether all the nodes were successfully created and connected
+		*/
+		bool AddAndConnectDataReferenceInputs(const Frontend::FNodeHandle& NodeHandle);
+
+		/** 
+		* Adds a data reference input node with InputName and connects it to the node's input with the same name
+		*
+		* @param NodeToConnect	The node to connect the newly added constructor input node to
+		* @param InputName		The name of the input pin to connect the node to, and the name to use for the new node.
+		* @param TypeName		The type name of the data reference
+		* @return				NodeHandle of the added node, can be invalid if failed to add
+		*/
+		Frontend::FNodeHandle AddAndConnectDataReferenceInput(const Frontend::FNodeHandle& NodeToConnect, const FName& InputName, const FName& TypeName) const;
+
+		/** Adds a data reference output node with OutputName and connects it to the node's output with the same name
+		* 	
+		* @param NodeToConnect	The node to connect the newly added constructor input node to
+		* @param InputName		The name of the output pin to connect the node to, and the name to use for the new node.
+		* @param TypeName		The type name of the data reference
+		* @return				NodeHandle of the added node, can be invalid if failed to add
+		*/
+		Frontend::FNodeHandle AddAndConnectDataReferenceOutput(const Frontend::FNodeHandle& NodeToConnect, const FName& OutputName, const FName& TypeName);
+
+		/** 
+		* Call last to create a generator from the graph constructed using the other method. 
+		* 
+		* @param SampleRate			The Sample Rate the generator should render audio at
+		* @param SamplesPerBlock	The number of samples to render per block
+		* @return					A unique pointer to the Metasound Generator created
+		*/
 		TUniquePtr<FMetasoundGenerator> BuildGenerator(FSampleRate SampleRate = 48000, int32 SamplesPerBlock = 256) const;
 
-		/** Helper that will add a single node, wire up the node's inputs and outputs, and hand back the graph's operator */
+		/** 
+		* Helper that will add a single node, wire up the node's inputs and outputs, and hand back the graph's operator 
+		*
+		* @param ClassName			The class name of the node to add
+		* @param MajorVersion		The version of the node to use
+		* @param SampleRate			The sample rate the generator should render audio at
+		* @param SamplesPerBlock	The number of samples to render per block
+		* @return					A unique pointer to the Metasound Generator created
+		*/
 		static TUniquePtr<FMetasoundGenerator> MakeSingleNodeGraph(
 			const FNodeClassName& ClassName,
 			int32 MajorVersion,
 			FSampleRate SampleRate = 48000,
 			int32 SamplesPerBlock = 256);
+		
+		/** 
+		* Connects an Input (Left) Node's output with name OutputName to another (Right) node's input with name InputName 
+		*
+		* @param LeftNode		The "left hand" side node to connect from
+		* @param OutputName		The name of the output pin to connect from the left node
+		* @param RightNode		The "right hand" side node to connect to
+		* @param InputName		The name of the input pin to connect to the right node
+		* @bool					true if able to connect the nodes (ie. all nodes were valid, all input/outputs were valid)
+		*/
+		static bool ConnectNodes(const Frontend::FNodeHandle& LeftNode, const FName& OutputName, const Frontend::FNodeHandle& RightNode, const FName& InputName);
+
+		/**
+		* Connects an Input (Left) Node's output with name OutputName to another (Right) node's input with name InputName
+		*
+		* @param LeftNode			The "left hand" side node to connect from
+		* @param RightNode			The "right hand" side node to connect to
+		* @param InputOutputName	The name of the input pin to connect to the right node
+		* @bool						true if able to connect the nodes (ie. all nodes were valid, all input/outputs were valid)
+		*/
+		static bool ConnectNodes(const Frontend::FNodeHandle& LeftNode, const Frontend::FNodeHandle& RightNode, const FName& InputOutputName);
 
 	private:
 		FMetasoundFrontendDocument Document;
