Expand and consolidate MetaSound builder preset asset functionality
- Added document builder and BP functions for SetGraphInputDefault, ConvertFromPreset, ConvertToPreset, CreatePatch/SourcePresetBuilder, IsPreset
- Refactor MetaSound asset and preset initialization to use builder architecture more

diff --git a/Engine/Plugins/Runtime/Metasound/Source/MetasoundEditor/Private/MetasoundEditorGraphBuilder.h b/Engine/Plugins/Runtime/Metasound/Source/MetasoundEditor/Private/MetasoundEditorGraphBuilder.h
--- a/Engine/Plugins/Runtime/Metasound/Source/MetasoundEditor/Private/MetasoundEditorGraphBuilder.h
+++ b/Engine/Plugins/Runtime/Metasound/Source/MetasoundEditor/Private/MetasoundEditorGraphBuilder.h
@@ -233,11 +233,12 @@ namespace Metasound
 			static FLinearColor GetPinCategoryColor(const FEdGraphPinType& PinType);
 
 			// Initializes MetaSound with default inputs & outputs.
-			UE_DEPRECATED(5.3, "Use FMetaSoundFrontendBuilder::InitDocument instead.")
+			UE_DEPRECATED(5.3, "Use FMetaSoundFrontendDocumentBuilder::InitDocument instead.")
 			static void InitMetaSound(UObject& InMetaSound, const FString& InAuthor);
 
 			// Initializes a MetaSound Preset using the provided ReferencedMetaSound asset's
 			// root graph as the sole, encapsulated topology.
+			UE_DEPRECATED(5.3, "Use FMetaSoundFrontendDocumentBuilder::ConvertToPreset instead.")
 			static void InitMetaSoundPreset(UObject& InMetaSoundReferenced, UObject& InMetaSoundPreset);
 
 			// Rebuilds all editor node pins based on the provided node handle's class definition.
diff --git a/Engine/Plugins/Runtime/Metasound/Source/MetasoundEditor/Private/MetasoundEditorSubsystem.cpp b/Engine/Plugins/Runtime/Metasound/Source/MetasoundEditor/Private/MetasoundEditorSubsystem.cpp
--- a/Engine/Plugins/Runtime/Metasound/Source/MetasoundEditor/Private/MetasoundEditorSubsystem.cpp
+++ b/Engine/Plugins/Runtime/Metasound/Source/MetasoundEditor/Private/MetasoundEditorSubsystem.cpp
@@ -2,9 +2,12 @@
 #include "MetasoundEditorSubsystem.h"
 
 #include "IAssetTools.h"
+#include "Kismet/KismetSystemLibrary.h"
 #include "MetasoundBuilderSubsystem.h"
 #include "MetasoundEditorGraph.h"
+#include "MetasoundEditorGraphBuilder.h"
 #include "MetasoundEditorGraphSchema.h"
+#include "MetasoundEditorSettings.h"
 #include "MetasoundFactory.h"
 #include "MetasoundUObjectRegistry.h"
 
@@ -34,7 +37,14 @@ TScriptInterface<IMetaSoundDocumentInterface> UMetaSoundEditorSubsystem::BuildTo
 			BuilderOptions.ExistingMetaSound = NewMetaSound;
 			TScriptInterface<IMetaSoundDocumentInterface> DocInterface = InBuilder->Build(Parent, BuilderOptions);
 
-			UMetaSoundFactory::InitEdGraph(*NewMetaSound);
+			if (InBuilder->IsPreset())
+			{
+				FMetasoundAssetBase* PresetAsset = IMetasoundUObjectRegistry::Get().GetObjectAsAssetBase(NewMetaSound);
+				check(PresetAsset);
+				PresetAsset->ConformObjectDataToInterfaces();
+			}
+
+			InitEdGraph(*NewMetaSound);
 			OutResult = EMetaSoundBuilderResult::Succeeded;
 			return NewMetaSound;
 		}
@@ -43,4 +53,60 @@ TScriptInterface<IMetaSoundDocumentInterface> UMetaSoundEditorSubsystem::BuildTo
 	OutResult = EMetaSoundBuilderResult::Failed;
 	return nullptr;
 }
+
+void UMetaSoundEditorSubsystem::InitEdGraph(UObject& InMetaSound)
+{
+	using namespace Metasound;
+	using namespace Metasound::Editor;
+
+	FMetasoundAssetBase* MetaSoundAsset = IMetasoundUObjectRegistry::Get().GetObjectAsAssetBase(&InMetaSound);
+	checkf(MetaSoundAsset, TEXT("EdGraph can only be initialized on registered MetaSoundAsset type"));
+
+	UMetasoundEditorGraph* Graph = Cast<UMetasoundEditorGraph>(MetaSoundAsset->GetGraph());
+	if (!Graph)
+	{
+		Graph = NewObject<UMetasoundEditorGraph>(&InMetaSound, FName(), RF_Transactional);
+		Graph->Schema = UMetasoundEditorGraphSchema::StaticClass();
+		MetaSoundAsset->SetGraph(Graph);
+
+		// Has to be done inline to have valid graph initially when opening editor for the first
+		// time (as opposed to being applied on tick when the document's modify context has updates)
+		FGraphBuilder::SynchronizeGraph(InMetaSound);
+	}
+}
+
+void UMetaSoundEditorSubsystem::RegisterGraphWithFrontend(UObject& InMetaSound, bool bInForceViewSynchronization)
+{
+	Metasound::Editor::FGraphBuilder::RegisterGraphWithFrontend(InMetaSound, bInForceViewSynchronization);
+}
+
+const FString UMetaSoundEditorSubsystem::GetDefaultAuthor()
+{
+	FString Author = UKismetSystemLibrary::GetPlatformUserName();
+	if (const UMetasoundEditorSettings* EditorSettings = GetDefault<UMetasoundEditorSettings>())
+	{
+		if (!EditorSettings->DefaultAuthor.IsEmpty())
+		{
+			Author = EditorSettings->DefaultAuthor;
+		}
+	}
+	return Author;
+}
+
+UMetaSoundEditorSubsystem& UMetaSoundEditorSubsystem::GetChecked()
+{
+	checkf(GEditor, TEXT("Cannot access UMetaSoundEditorSubsystem without editor loaded"));
+	UMetaSoundEditorSubsystem* EditorSubsystem = GEditor->GetEditorSubsystem<UMetaSoundEditorSubsystem>();
+	checkf(EditorSubsystem, TEXT("Failed to find initialized 'UMetaSoundEditorSubsystem"));
+	return *EditorSubsystem;
+}
+
+const UMetaSoundEditorSubsystem& UMetaSoundEditorSubsystem::GetConstChecked()
+{
+	checkf(GEditor, TEXT("Cannot access UMetaSoundEditorSubsystem without editor loaded"));
+	UMetaSoundEditorSubsystem* EditorSubsystem = GEditor->GetEditorSubsystem<UMetaSoundEditorSubsystem>();
+	checkf(EditorSubsystem, TEXT("Failed to find initialized 'UMetaSoundEditorSubsystem"));
+	return *EditorSubsystem;
+}
+
 #undef LOCTEXT_NAMESPACE // "MetaSoundEditor"
diff --git a/Engine/Plugins/Runtime/Metasound/Source/MetasoundEditor/Private/MetasoundFactory.cpp b/Engine/Plugins/Runtime/Metasound/Source/MetasoundEditor/Private/MetasoundFactory.cpp
--- a/Engine/Plugins/Runtime/Metasound/Source/MetasoundEditor/Private/MetasoundFactory.cpp
+++ b/Engine/Plugins/Runtime/Metasound/Source/MetasoundEditor/Private/MetasoundFactory.cpp
@@ -1,12 +1,11 @@
 // Copyright Epic Games, Inc. All Rights Reserved.
 #include "MetasoundFactory.h"
 
-#include "Kismet/KismetSystemLibrary.h"
 #include "Metasound.h"
 #include "MetasoundEditorGraph.h"
 #include "MetasoundEditorGraphBuilder.h"
 #include "MetasoundEditorGraphSchema.h"
-#include "MetasoundEditorSettings.h"
+#include "MetasoundEditorSubsystem.h"
 #include "MetasoundFrontendDocumentBuilder.h"
 #include "MetasoundFrontendTransform.h"
 #include "MetasoundSource.h"
@@ -47,47 +46,34 @@ void UMetaSoundBaseFactory::InitAsset(UObject& InNewMetaSound, UObject* InRefere
 	TScriptInterface<IMetaSoundDocumentInterface> DocInterface = &InNewMetaSound;
 	FMetaSoundFrontendDocumentBuilder Builder(DocInterface);
 	Builder.InitDocument();
+#if WITH_EDITORONLY_DATA
 	Builder.InitNodeLocations();
+#endif // WITH_EDITORONLY_DATA
 
-	FString Author = UKismetSystemLibrary::GetPlatformUserName();
-	if (const UMetasoundEditorSettings* EditorSettings = GetDefault<UMetasoundEditorSettings>())
-	{
-		if (!EditorSettings->DefaultAuthor.IsEmpty())
-		{
-			Author = EditorSettings->DefaultAuthor;
-		}
-	}
+	const FString& Author = UMetaSoundEditorSubsystem::GetChecked().GetDefaultAuthor();
 	Builder.SetAuthor(Author);
 
+	// Initialize asset as a preset
 	if (InReferencedMetaSound)
 	{
-		FGraphBuilder::InitMetaSoundPreset(*InReferencedMetaSound, InNewMetaSound);
+		// Ensure the referenced MetaSound is registered already
+		UMetaSoundEditorSubsystem::GetChecked().RegisterGraphWithFrontend(*InReferencedMetaSound);
+		
+		// Initialize preset with referenced Metasound 
+		TScriptInterface<IMetaSoundDocumentInterface> ReferencedDocInterface = InReferencedMetaSound;
+		const IMetaSoundDocumentInterface* ReferencedInterface = ReferencedDocInterface.GetInterface();
+		check(ReferencedInterface);
+		Builder.ConvertToPreset(ReferencedInterface->GetDocument());
+		
+		// Update asset object data from interfaces 
+		FMetasoundAssetBase* PresetAsset = IMetasoundUObjectRegistry::Get().GetObjectAsAssetBase(&InNewMetaSound);
+		check(PresetAsset);
+		PresetAsset->ConformObjectDataToInterfaces();
 	}
 
 	// Initial graph generation is not something to be managed by the transaction
 	// stack, so don't track dirty state until after initial setup if necessary.
-	InitEdGraph(InNewMetaSound);
-}
-
-void UMetaSoundBaseFactory::InitEdGraph(UObject& InMetaSound)
-{
-	using namespace Metasound;
-	using namespace Metasound::Editor;
-
-	FMetasoundAssetBase* MetaSoundAsset = IMetasoundUObjectRegistry::Get().GetObjectAsAssetBase(&InMetaSound);
-	checkf(MetaSoundAsset, TEXT("EdGraph can only be initialized on registered MetaSoundAsset type"));
-
-	UMetasoundEditorGraph* Graph = Cast<UMetasoundEditorGraph>(MetaSoundAsset->GetGraph());
-	if (!Graph)
-	{
-		Graph = NewObject<UMetasoundEditorGraph>(&InMetaSound, FName(), RF_Transactional);
-		Graph->Schema = UMetasoundEditorGraphSchema::StaticClass();
-		MetaSoundAsset->SetGraph(Graph);
-
-		// Has to be done inline to have valid graph initially when opening editor for the first
-		// time (as opposed to being applied on tick when the document's modify context has updates)
-		FGraphBuilder::SynchronizeGraph(InMetaSound);
-	}
+	UMetaSoundEditorSubsystem::GetChecked().InitEdGraph(InNewMetaSound);
 }
 
 UMetaSoundFactory::UMetaSoundFactory(const FObjectInitializer& ObjectInitializer)
diff --git a/Engine/Plugins/Runtime/Metasound/Source/MetasoundEditor/Public/MetasoundEditorSubsystem.h b/Engine/Plugins/Runtime/Metasound/Source/MetasoundEditor/Public/MetasoundEditorSubsystem.h
--- a/Engine/Plugins/Runtime/Metasound/Source/MetasoundEditor/Public/MetasoundEditorSubsystem.h
+++ b/Engine/Plugins/Runtime/Metasound/Source/MetasoundEditor/Public/MetasoundEditorSubsystem.h
@@ -21,4 +21,19 @@ class METASOUNDEDITOR_API UMetaSoundEditorSubsystem : public UEditorSubsystem
 public:
 	UFUNCTION(BlueprintCallable, Category = "Audio|MetaSound|Builder|Editor", meta = (WorldContext = "Parent", ExpandEnumAsExecs = "OutResult"))
 	UPARAM(DisplayName = "MetaSound Asset") TScriptInterface<IMetaSoundDocumentInterface> BuildToAsset(UMetaSoundBuilderBase* InBuilder, const FString& Author, const FString& AssetName, const FString& PackagePath, EMetaSoundBuilderResult& OutResult);
+
+	// Initialize UMetasoundEditorGraph for a given MetaSound object
+	void InitEdGraph(UObject& InMetaSound);
+
+	// Wraps RegisterGraphWithFrontend logic in Frontend with any additional logic required to refresh editor & respective editor object state.
+	// @param InMetaSound - MetaSound to register
+	// @param bInForceSynchronize - Forces the synchronize flag for all open graphs being registered by this call (all referenced graphs and
+	// referencing graphs open in editors)
+	void RegisterGraphWithFrontend(UObject& InMetaSound, bool bInForceViewSynchronization = false);
+
+	// Get the default author for a MetaSound asset
+	const FString GetDefaultAuthor();
+
+	static UMetaSoundEditorSubsystem& GetChecked();
+	static const UMetaSoundEditorSubsystem& GetConstChecked(); 
 };
diff --git a/Engine/Plugins/Runtime/Metasound/Source/MetasoundEditor/Public/MetasoundFactory.h b/Engine/Plugins/Runtime/Metasound/Source/MetasoundEditor/Public/MetasoundFactory.h
--- a/Engine/Plugins/Runtime/Metasound/Source/MetasoundEditor/Public/MetasoundFactory.h
+++ b/Engine/Plugins/Runtime/Metasound/Source/MetasoundEditor/Public/MetasoundFactory.h
@@ -22,7 +22,6 @@ class UMetaSoundBaseFactory : public UFactory
 	TObjectPtr<UObject> ReferencedMetaSoundObject;
 
 	static void InitAsset(UObject& InNewMetaSound, UObject* InReferencedMetaSound = nullptr);
-	static void InitEdGraph(UObject& InMetaSound);
 };
 
 UCLASS(hidecategories=Object, MinimalAPI)
diff --git a/Engine/Plugins/Runtime/Metasound/Source/MetasoundEngine/Private/MetasoundBuilderSubsystem.cpp b/Engine/Plugins/Runtime/Metasound/Source/MetasoundEngine/Private/MetasoundBuilderSubsystem.cpp
--- a/Engine/Plugins/Runtime/Metasound/Source/MetasoundEngine/Private/MetasoundBuilderSubsystem.cpp
+++ b/Engine/Plugins/Runtime/Metasound/Source/MetasoundEngine/Private/MetasoundBuilderSubsystem.cpp
@@ -510,6 +510,17 @@ bool UMetaSoundBuilderBase::NodeOutputIsConnected(const FMetaSoundBuilderNodeOut
 	return Builder.IsNodeOutputConnected(OutputHandle.NodeID, OutputHandle.VertexID);
 }
 
+bool UMetaSoundBuilderBase::IsPreset() const
+{
+	return Builder.IsPreset();
+}
+
+void UMetaSoundBuilderBase::ConvertFromPreset(EMetaSoundBuilderResult& OutResult)
+{
+	const bool bSuccess = Builder.ConvertFromPreset();
+	OutResult = bSuccess ? EMetaSoundBuilderResult::Succeeded : EMetaSoundBuilderResult::Failed;
+}
+
 void UMetaSoundBuilderBase::RemoveGraphInput(FName Name, EMetaSoundBuilderResult& OutResult)
 {
 	const bool bRemoved = Builder.RemoveGraphInput(Name);
@@ -550,6 +561,12 @@ void UMetaSoundBuilderBase::SetNodeInputDefault(const FMetaSoundBuilderNodeInput
 	OutResult = bInputDefaultSet ? EMetaSoundBuilderResult::Succeeded : EMetaSoundBuilderResult::Failed;
 }
 
+void UMetaSoundBuilderBase::SetGraphInputDefault(FName InputName, const FMetasoundFrontendLiteral& Literal, EMetaSoundBuilderResult& OutResult)
+{
+	const bool bInputDefaultSet = Builder.SetGraphInputDefault(InputName, Literal);
+	OutResult = bInputDefaultSet ? EMetaSoundBuilderResult::Succeeded : EMetaSoundBuilderResult::Failed;
+}
+
 FMetaSoundNodeHandle UMetaSoundBuilderBase::FindNodeInputParent(const FMetaSoundBuilderNodeInputHandle& InputHandle, EMetaSoundBuilderResult& OutResult)
 {
 	if (Builder.ContainsNode(InputHandle.NodeID))
@@ -621,6 +638,46 @@ void UMetaSoundBuilderBase::GetNodeOutputData(const FMetaSoundBuilderNodeOutputH
 	}
 }
 
+void UMetaSoundBuilderBase::ConvertToPreset(const TScriptInterface<IMetaSoundDocumentInterface>& ReferencedNodeClass, EMetaSoundBuilderResult& OutResult)
+{
+	const IMetaSoundDocumentInterface* ReferencedInterface = ReferencedNodeClass.GetInterface();
+	check(ReferencedInterface);
+	
+	// Ensure the referenced node class isn't transient 
+	if (Cast<UMetaSoundBuilderDocument>(ReferencedInterface))
+	{
+		UE_LOG(LogMetaSound, Warning, TEXT("Transient document builders cannot be referenced when converting builder '%s' to a preset. Build the referenced node class an asset first or use an existing asset instead"), *GetName());
+		OutResult = EMetaSoundBuilderResult::Failed;
+		return;
+	}
+
+	// Ensure the referenced node class is a matching object type 
+	const UClass& BaseMetaSoundClass = ReferencedInterface->GetBaseMetaSoundUClass();
+	UObject* ReferencedObject = ReferencedNodeClass.GetObject();
+	if (!ReferencedObject || (ReferencedObject && !ReferencedObject->IsA(&BaseMetaSoundClass)))
+	{
+		UE_LOG(LogMetaSound, Warning, TEXT("The referenced node type must match the base MetaSound class when converting builder '%s' to a preset (ex. source preset must reference another source)"), *GetName());
+		OutResult = EMetaSoundBuilderResult::Failed;
+		return;
+	}
+
+	// Ensure the referenced node is registered
+	if (FMetasoundAssetBase* ReferencedMetaSoundAsset = Metasound::IMetasoundUObjectRegistry::Get().GetObjectAsAssetBase(ReferencedObject))
+	{
+		ReferencedMetaSoundAsset->RegisterGraphWithFrontend();
+	}
+
+	const FMetasoundFrontendDocument& ReferencedDocument = ReferencedInterface->GetDocument();
+	if (Builder.ConvertToPreset(ReferencedDocument))
+	{
+		OutResult = EMetaSoundBuilderResult::Succeeded;
+	}
+	else
+	{
+		OutResult = EMetaSoundBuilderResult::Failed;
+	}
+}
+
 TScriptInterface<IMetaSoundDocumentInterface> UMetaSoundPatchBuilder::Build(UObject* Parent, const FMetaSoundBuilderOptions& InBuilderOptions) const
 {
 	return BuildInternal<UMetaSoundPatch>(Parent, InBuilderOptions);
@@ -862,6 +919,20 @@ UMetaSoundSourceBuilder* UMetaSoundBuilderSubsystem::CreateSourceBuilder(
 	return &NewBuilder;
 }
 
+UMetaSoundPatchBuilder* UMetaSoundBuilderSubsystem::CreatePatchPresetBuilder(FName BuilderName, const TScriptInterface<IMetaSoundDocumentInterface>& ReferencedNodeClass, EMetaSoundBuilderResult& OutResult)
+{
+	UMetaSoundPatchBuilder& Builder = Metasound::Engine::BuilderSubsystemPrivate::CreateTransientBuilder<UMetaSoundPatchBuilder>();
+	Builder.ConvertToPreset(ReferencedNodeClass, OutResult);
+	return &Builder;
+}
+
+UMetaSoundSourceBuilder* UMetaSoundBuilderSubsystem::CreateSourcePresetBuilder(FName BuilderName, const TScriptInterface<IMetaSoundDocumentInterface>& ReferencedNodeClass, EMetaSoundBuilderResult& OutResult)
+{
+	UMetaSoundSourceBuilder& Builder = Metasound::Engine::BuilderSubsystemPrivate::CreateTransientBuilder<UMetaSoundSourceBuilder>();
+	Builder.ConvertToPreset(ReferencedNodeClass, OutResult);
+	return &Builder;
+}
+
 UMetaSoundBuilderSubsystem& UMetaSoundBuilderSubsystem::GetChecked()
 {
 	checkf(GEngine, TEXT("Cannot access UMetaSoundBuilderSubsystem without engine loaded"));
diff --git a/Engine/Plugins/Runtime/Metasound/Source/MetasoundEngine/Public/MetasoundBuilderSubsystem.h b/Engine/Plugins/Runtime/Metasound/Source/MetasoundEngine/Public/MetasoundBuilderSubsystem.h
--- a/Engine/Plugins/Runtime/Metasound/Source/MetasoundEngine/Public/MetasoundBuilderSubsystem.h
+++ b/Engine/Plugins/Runtime/Metasound/Source/MetasoundEngine/Public/MetasoundBuilderSubsystem.h
@@ -267,13 +267,25 @@ public:
 	// Returns if a given node output is connected.
 	UFUNCTION(BlueprintCallable, Category = "Audio|MetaSound|Builder")
 	UPARAM(DisplayName = "Connected") bool NodeOutputIsConnected(const FMetaSoundBuilderNodeOutputHandle& OutputHandle) const;
+		
+	// Returns whether this is a preset.
+	UFUNCTION(BlueprintCallable, Category = "Audio|MetaSound|Builder")
+	bool IsPreset() const;
+
+	// Converts this preset to a fully accessible MetaSound; sets result to succeeded if it was converted successfully and failed if it was not.
+	UFUNCTION(BlueprintCallable, Category = "Audio|MetaSound|Builder", meta = (ExpandEnumAsExecs = "OutResult"))
+	void ConvertFromPreset(EMetaSoundBuilderResult& OutResult);
+
+	// Convert this builder to a MetaSound source preset with the given referenced source builder 
+	UFUNCTION(BlueprintCallable, Category = "Audio|MetaSound|Builder", meta = (ExpandEnumAsExecs = "OutResult"))
+	void ConvertToPreset(const TScriptInterface<IMetaSoundDocumentInterface>& ReferencedNodeClass, EMetaSoundBuilderResult& OutResult);
 
 	// Removes graph input if it exists; sets result to succeeded if it was removed and failed if it was not.
-	UFUNCTION(BlueprintCallable, Category = "Audio|MetaSound|Builder")
+	UFUNCTION(BlueprintCallable, Category = "Audio|MetaSound|Builder", meta = (ExpandEnumAsExecs = "OutResult"))
 	void RemoveGraphInput(FName Name, EMetaSoundBuilderResult& OutResult);
 
 	// Removes graph output if it exists; sets result to succeeded if it was removed and failed if it was not.
-	UFUNCTION(BlueprintCallable, Category = "Audio|MetaSound|Builder")
+	UFUNCTION(BlueprintCallable, Category = "Audio|MetaSound|Builder", meta = (ExpandEnumAsExecs = "OutResult"))
 	void RemoveGraphOutput(FName Name, EMetaSoundBuilderResult& OutResult);
 
 	// Removes the interface with the given name from the builder's MetaSound. Removes any graph inputs
@@ -294,6 +306,10 @@ public:
 	UFUNCTION(BlueprintCallable, Category = "Audio|MetaSound|Builder", meta = (ExpandEnumAsExecs = "OutResult"))
 	void SetNodeInputDefault(const FMetaSoundBuilderNodeInputHandle& NodeInputHandle, const FMetasoundFrontendLiteral& Literal, EMetaSoundBuilderResult& OutResult);
 
+	// Sets the input node's default value, overriding the default if the graph is a preset. 
+	UFUNCTION(BlueprintCallable, Category = "Audio|MetaSound|Builder", meta = (ExpandEnumAsExecs = "OutResult"))
+	void SetGraphInputDefault(FName InputName, const FMetasoundFrontendLiteral& Literal, EMetaSoundBuilderResult& OutResult);
+
 	virtual TScriptInterface<IMetaSoundDocumentInterface> Build(UObject* Parent, const FMetaSoundBuilderOptions& Options) const PURE_VIRTUAL(UMetaSoundBuilderBase::Build, return { }; );
 
 	// Constructs a transient UMetaSoundBuilderDocument & attaches builder to it
@@ -381,7 +397,7 @@ protected:
 	friend class UMetaSoundBuilderSubsystem;
 };
 
-/** Builder in charge of building a MetaSound */
+/** Builder in charge of building a MetaSound Patch */
 UCLASS(Transient, BlueprintType)
 class METASOUNDENGINE_API UMetaSoundPatchBuilder : public UMetaSoundBuilderBase
 {
@@ -451,6 +467,12 @@ public:
 		EMetaSoundBuilderResult& OutResult,
 		EMetaSoundOutputAudioFormat OutputFormat = EMetaSoundOutputAudioFormat::Mono,
 		bool bIsOneShot = true);
+	
+	UFUNCTION(BlueprintCallable, Category = "Audio|MetaSound|Builder",  meta = (ExpandEnumAsExecs = "OutResult"))
+	UPARAM(DisplayName = "Patch Preset Builder") UMetaSoundPatchBuilder* CreatePatchPresetBuilder(FName BuilderName, const TScriptInterface<IMetaSoundDocumentInterface>& ReferencedPatchClass, EMetaSoundBuilderResult& OutResult);
+
+	UFUNCTION(BlueprintCallable, Category = "Audio|MetaSound|Builder", meta = (ExpandEnumAsExecs = "OutResult"))
+	UPARAM(DisplayName = "Source Preset Builder") UMetaSoundSourceBuilder* CreateSourcePresetBuilder(FName BuilderName, const TScriptInterface<IMetaSoundDocumentInterface>& ReferencedSourceClass, EMetaSoundBuilderResult& OutResult);
 
 	UFUNCTION(BlueprintCallable, Category = "Audio|MetaSound|Builder", meta = (DisplayName = "Create MetaSound Bool Literal"))
 	UPARAM(DisplayName = "Bool Literal") FMetasoundFrontendLiteral CreateBoolMetaSoundLiteral(bool Value, FName& DataType);
diff --git a/Engine/Plugins/Runtime/Metasound/Source/MetasoundFrontend/Private/MetasoundFrontendDocumentBuilder.cpp b/Engine/Plugins/Runtime/Metasound/Source/MetasoundFrontend/Private/MetasoundFrontendDocumentBuilder.cpp
--- a/Engine/Plugins/Runtime/Metasound/Source/MetasoundFrontend/Private/MetasoundFrontendDocumentBuilder.cpp
+++ b/Engine/Plugins/Runtime/Metasound/Source/MetasoundFrontend/Private/MetasoundFrontendDocumentBuilder.cpp
@@ -3,6 +3,7 @@
 
 #include "Algo/AnyOf.h"
 #include "Algo/Find.h"
+#include "Algo/ForEach.h"
 #include "Algo/NoneOf.h"
 #include "Algo/Sort.h"
 #include "Algo/Transform.h"
@@ -11,6 +12,7 @@
 #include "MetasoundAssetBase.h"
 #include "MetasoundAssetManager.h"
 #include "MetasoundDocumentInterface.h"
+#include "MetasoundFrontendController.h"
 #include "MetasoundFrontendDocumentCache.h"
 #include "MetasoundFrontendDocument.h"
 #include "MetasoundFrontendRegistries.h"
@@ -1043,6 +1045,19 @@ bool FMetaSoundFrontendDocumentBuilder::RemoveUnusedDependencies()
 	return bDidEdit;
 }
 
+void FMetaSoundFrontendDocumentBuilder::ClearGraph()
+{
+	FMetasoundFrontendGraphClass& GraphClass = GetDocument().RootGraph;
+	GraphClass.Graph.Nodes.Reset();
+	GraphClass.Graph.Edges.Reset();
+	GraphClass.Interface.Inputs.Reset();
+	GraphClass.Interface.Outputs.Reset();
+	GraphClass.PresetOptions.InputsInheritingDefault.Reset();
+	GetDocument().Interfaces.Reset();
+	RemoveUnusedDependencies();
+	ReloadCache();
+}
+
 bool FMetaSoundFrontendDocumentBuilder::ContainsEdge(const FMetasoundFrontendEdge& InEdge) const
 {
 	using namespace Metasound::Frontend;
@@ -1061,17 +1076,56 @@ bool FMetaSoundFrontendDocumentBuilder::ConvertFromPreset()
 {
 	using namespace Metasound::Frontend;
 
-	FMetasoundFrontendDocument& Document = GetDocument();
-	FMetasoundFrontendGraphClass& RootGraphClass = Document.RootGraph;
-	FMetasoundFrontendGraphClassPresetOptions& PresetOptions = RootGraphClass.PresetOptions;
-	PresetOptions.bIsPreset = false;
+	if (IsPreset())
+	{
+		FMetasoundFrontendDocument& Document = GetDocument();
+		FMetasoundFrontendGraphClass& RootGraphClass = Document.RootGraph;
+		FMetasoundFrontendGraphClassPresetOptions& PresetOptions = RootGraphClass.PresetOptions;
+		PresetOptions.bIsPreset = false;
 
 #if WITH_EDITOR
-	FMetasoundFrontendGraphStyle& Style = Document.RootGraph.Graph.Style;
-	Style.bIsGraphEditable = true;
+		FMetasoundFrontendGraphStyle& Style = Document.RootGraph.Graph.Style;
+		Style.bIsGraphEditable = true;
 #endif // WITH_EDITOR
 
-	return true;
+		return true;
+	}
+
+	return false;
+}
+
+bool FMetaSoundFrontendDocumentBuilder::ConvertToPreset(const FMetasoundFrontendDocument& InReferencedDocument)
+{
+	using namespace Metasound;
+	using namespace Metasound::Frontend;
+
+	ClearGraph();
+
+	FMetasoundFrontendGraphClass& PresetAssetRootGraph = GetDocument().RootGraph;
+	FMetasoundFrontendGraph& PresetAssetGraph = PresetAssetRootGraph.Graph;
+	// Mark preset as auto-update and non-editable
+#if WITH_EDITORONLY_DATA
+	PresetAssetGraph.Style.bIsGraphEditable = false;
+#endif // WITH_EDITORONLY_DATA
+
+	// Mark all inputs as inherited by default
+	TArray<FName> InputsInheritingDefault;
+	Algo::Transform(PresetAssetRootGraph.Interface.Inputs, InputsInheritingDefault, [](const FMetasoundFrontendClassInput& Input)
+	{
+		return Input.Name;
+	});
+
+	PresetAssetRootGraph.PresetOptions.bIsPreset = true;
+	PresetAssetRootGraph.PresetOptions.InputsInheritingDefault = TSet<FName>(InputsInheritingDefault);
+
+	// Apply root graph transform 
+	FRebuildPresetRootGraph RebuildPresetRootGraph(InReferencedDocument);
+	if (RebuildPresetRootGraph.Transform(GetDocument()))
+	{
+		ReloadCache();
+		return true;
+	}
+	return false;
 }
 
 bool FMetaSoundFrontendDocumentBuilder::FindDeclaredInterfaces(TArray<const Metasound::Frontend::IInterfaceRegistryEntry*>& OutInterfaces) const
@@ -1466,6 +1520,11 @@ bool FMetaSoundFrontendDocumentBuilder::IsInterfaceDeclared(const FMetasoundFron
 	return GetDocument().Interfaces.Contains(InInterfaceVersion);
 }
 
+bool FMetaSoundFrontendDocumentBuilder::IsPreset() const
+{
+	return GetDocument().RootGraph.PresetOptions.bIsPreset;
+}
+
 bool FMetaSoundFrontendDocumentBuilder::ModifyInterfaces(Metasound::Frontend::FModifyInterfaceOptions&& InOptions)
 {
 	using namespace Metasound::Frontend;
@@ -1872,6 +1931,31 @@ void FMetaSoundFrontendDocumentBuilder::SetAuthor(const FString& InAuthor)
 }
 #endif // WITH_EDITOR
 
+bool FMetaSoundFrontendDocumentBuilder::SetGraphInputDefault(FName InputName, const FMetasoundFrontendLiteral& InDefaultLiteral)
+{
+	using namespace Metasound::Frontend;
+
+	auto NameMatchesInput = [&InputName](const FMetasoundFrontendClassInput& Input) { return Input.Name == InputName; };
+	TArray<FMetasoundFrontendClassInput>& Inputs = GetDocument().RootGraph.Interface.Inputs;
+	if (FMetasoundFrontendClassInput* Input = Inputs.FindByPredicate(NameMatchesInput))
+	{
+		if (IDataTypeRegistry::Get().IsLiteralTypeSupported(Input->TypeName, InDefaultLiteral.GetType()))
+		{
+			Input->DefaultLiteral = InDefaultLiteral;
+
+			// Set the input as no longer inheriting default for presets
+			if (IsPreset())
+			{
+				return SetInputInheritsDefault(InputName, /*bInputInheritsDefault=*/false);
+			}
+			return true;
+		}
+		UE_LOG(LogMetaSound, Error, TEXT("Attempting to set graph input of type '%s' with unsupported literal type"), *Input->TypeName.ToString());
+	}
+
+	return false;
+}
+
 bool FMetaSoundFrontendDocumentBuilder::SetNodeInputDefault(const FGuid& InNodeID, const FGuid& InVertexID, const FMetasoundFrontendLiteral& InLiteral)
 {
 	using namespace Metasound::Frontend;
@@ -1906,6 +1990,28 @@ bool FMetaSoundFrontendDocumentBuilder::SetNodeInputDefault(const FGuid& InNodeI
 	return false;
 }
 
+
+bool FMetaSoundFrontendDocumentBuilder::SetInputInheritsDefault(FName InName, bool bInputInheritsDefault)
+{
+	FMetasoundFrontendGraphClassPresetOptions& PresetOptions = GetDocument().RootGraph.PresetOptions;
+	if (bInputInheritsDefault)
+	{
+		if (PresetOptions.bIsPreset)
+		{
+			return PresetOptions.InputsInheritingDefault.Add(InName).IsValidId();
+		}
+	}
+	else
+	{
+		if (PresetOptions.bIsPreset)
+		{
+			return PresetOptions.InputsInheritingDefault.Remove(InName) > 0;
+		}
+	}
+
+	return false;
+}
+
 bool FMetaSoundFrontendDocumentBuilder::SwapGraphInput(const FMetasoundFrontendClassVertex& InExistingInputVertex, const FMetasoundFrontendClassVertex& InNewInputVertex)
 {
 	using namespace Metasound::Frontend;
diff --git a/Engine/Plugins/Runtime/Metasound/Source/MetasoundFrontend/Private/MetasoundFrontendTransform.cpp b/Engine/Plugins/Runtime/Metasound/Source/MetasoundFrontend/Private/MetasoundFrontendTransform.cpp
--- a/Engine/Plugins/Runtime/Metasound/Source/MetasoundFrontend/Private/MetasoundFrontendTransform.cpp
+++ b/Engine/Plugins/Runtime/Metasound/Source/MetasoundFrontend/Private/MetasoundFrontendTransform.cpp
@@ -775,6 +775,18 @@ namespace Metasound
 
 			return true;
 		}
+	
+		FRebuildPresetRootGraph::FRebuildPresetRootGraph(const FMetasoundFrontendDocument& InReferencedDocument)
+		{
+			// TODO: Swap implementation to not use access pointers/controllers
+			ReferencedDocument = IDocumentController::CreateDocumentHandle(InReferencedDocument);
+		}
+
+		bool FRebuildPresetRootGraph::Transform(FMetasoundFrontendDocument& InDocument) const
+		{
+			// TODO: Swap implementation to not use access pointers/controllers
+			return Transform(IDocumentController::CreateDocumentHandle(InDocument));
+		}
 
 		void FRebuildPresetRootGraph::AddAndConnectInputs(const TArray<FMetasoundFrontendClassInput>& InClassInputs, FGraphHandle& InPresetGraphHandle, FNodeHandle& InReferencedNode) const
 		{
diff --git a/Engine/Plugins/Runtime/Metasound/Source/MetasoundFrontend/Public/MetasoundFrontendDocumentBuilder.h b/Engine/Plugins/Runtime/Metasound/Source/MetasoundFrontend/Public/MetasoundFrontendDocumentBuilder.h
--- a/Engine/Plugins/Runtime/Metasound/Source/MetasoundFrontend/Public/MetasoundFrontendDocumentBuilder.h
+++ b/Engine/Plugins/Runtime/Metasound/Source/MetasoundFrontend/Public/MetasoundFrontendDocumentBuilder.h
@@ -151,9 +151,13 @@ public:
 
 	bool CanAddEdge(const FMetasoundFrontendEdge& InEdge) const;
 
+	void ClearGraph();
+
 	bool ContainsEdge(const FMetasoundFrontendEdge& InEdge) const;
 	bool ContainsNode(const FGuid& InNodeID) const;
+
 	bool ConvertFromPreset();
+	bool ConvertToPreset(const FMetasoundFrontendDocument& InReferencedDocument);
 
 	static bool FindDeclaredInterfaces(const FMetasoundFrontendDocument& InDocument, TArray<const Metasound::Frontend::IInterfaceRegistryEntry*>& OutInterfaces);
 	bool FindDeclaredInterfaces(TArray<const Metasound::Frontend::IInterfaceRegistryEntry*>& OutInterfaces) const;
@@ -191,6 +195,7 @@ public:
 
 	bool IsInterfaceDeclared(FName InInterfaceName) const;
 	bool IsInterfaceDeclared(const FMetasoundFrontendVersion& InInterfaceVersion) const;
+	bool IsPreset() const;
 
 	bool ModifyInterfaces(Metasound::Frontend::FModifyInterfaceOptions&& InOptions);
 
@@ -211,7 +216,9 @@ public:
 	void SetAuthor(const FString& InAuthor);
 #endif // WITH_EDITOR
 
+	bool SetGraphInputDefault(FName InputName, const FMetasoundFrontendLiteral& InDefaultLiteral);
 	bool SetNodeInputDefault(const FGuid& InNodeID, const FGuid& InVertexID, const FMetasoundFrontendLiteral& InLiteral);
+
 	bool SwapGraphInput(const FMetasoundFrontendClassVertex& InExistingInputVertex, const FMetasoundFrontendClassVertex& NewInputVertex);
 	bool SwapGraphOutput(const FMetasoundFrontendClassVertex& InExistingOutputVertex, const FMetasoundFrontendClassVertex& NewOutputVertex);
 
@@ -233,4 +240,6 @@ private:
 	const FMetasoundFrontendClassOutput* FindNodeOutputClassOutput(const FGuid& InNodeID, const FGuid& InVertexID) const;
 
 	void ReloadCache();
+
+	bool SetInputInheritsDefault(FName InName, bool bInputInheritsDefault);
 };
diff --git a/Engine/Plugins/Runtime/Metasound/Source/MetasoundFrontend/Public/MetasoundFrontendTransform.h b/Engine/Plugins/Runtime/Metasound/Source/MetasoundFrontend/Public/MetasoundFrontendTransform.h
--- a/Engine/Plugins/Runtime/Metasound/Source/MetasoundFrontend/Public/MetasoundFrontendTransform.h
+++ b/Engine/Plugins/Runtime/Metasound/Source/MetasoundFrontend/Public/MetasoundFrontendTransform.h
@@ -166,7 +166,10 @@ namespace Metasound
 			{
 			}
 
-			bool Transform(FDocumentHandle InDocument) const override;
+			FRebuildPresetRootGraph(const FMetasoundFrontendDocument& InReferencedDocument);
+
+			virtual bool Transform(FDocumentHandle InDocument) const override;
+			virtual bool Transform(FMetasoundFrontendDocument& InOutDocument) const override;
 
 		private:
 
@@ -184,7 +187,7 @@ namespace Metasound
 			// Add outputs to parent graph and connect to wrapped graph node.
 			void AddAndConnectOutputs(const TArray<FMetasoundFrontendClassOutput>& InClassOutputs, FGraphHandle& InParentGraphHandle, FNodeHandle& InReferencedNode) const;
 
-			FConstDocumentHandle ReferencedDocument;
+			FConstDocumentHandle ReferencedDocument = IDocumentController::GetInvalidHandle();
 		};
 
 		/** Automatically updates all nodes and respective dependencies in graph where
