Fix .Add/.Remove copy pasta
Minor member name/comment fix ups
#rnx

diff --git a/Engine/Plugins/Runtime/Metasound/Source/MetasoundEngine/Private/MetasoundBuilderSubsystem.cpp b/Engine/Plugins/Runtime/Metasound/Source/MetasoundEngine/Private/MetasoundBuilderSubsystem.cpp
--- a/Engine/Plugins/Runtime/Metasound/Source/MetasoundEngine/Private/MetasoundBuilderSubsystem.cpp
+++ b/Engine/Plugins/Runtime/Metasound/Source/MetasoundEngine/Private/MetasoundBuilderSubsystem.cpp
@@ -86,14 +86,14 @@ FMetaSoundBuilderNodeOutputHandle UMetaSoundBuilderBase::AddGraphInputNode(FName
 	return NewHandle;
 }
 
-FMetaSoundBuilderNodeInputHandle UMetaSoundBuilderBase::AddGraphOutputNode(FName Name, FName DataType, FMetasoundFrontendLiteral DefaultValue, EMetaSoundBuilderResult& OutResult, bool bIsConstructorInput)
+FMetaSoundBuilderNodeInputHandle UMetaSoundBuilderBase::AddGraphOutputNode(FName Name, FName DataType, FMetasoundFrontendLiteral DefaultValue, EMetaSoundBuilderResult& OutResult, bool bIsConstructorOutput)
 {
 	FMetasoundFrontendClassOutput Description;
 	Description.Name = Name;
 	Description.TypeName = DataType;
 	Description.NodeID = FGuid::NewGuid();
 	Description.VertexID = FGuid::NewGuid();
-	Description.AccessType = bIsConstructorInput ? EMetasoundFrontendVertexAccessType::Value : EMetasoundFrontendVertexAccessType::Reference;
+	Description.AccessType = bIsConstructorOutput ? EMetasoundFrontendVertexAccessType::Value : EMetasoundFrontendVertexAccessType::Reference;
 
 	FMetaSoundBuilderNodeInputHandle NewHandle;
 	const FMetasoundFrontendNode* Node = Builder.FindGraphOutputNode(Name);
@@ -145,6 +145,7 @@ FMetaSoundNodeHandle UMetaSoundBuilderBase::AddNode(const TScriptInterface<IMeta
 			check(Interface);
 			const FMetasoundFrontendDocument& NodeClassDoc = Interface->GetDocument();
 			const FMetasoundFrontendGraphClass& NodeClassGraph = NodeClassDoc.RootGraph;
+
 			if (const FMetasoundFrontendNode* NewNode = Builder.AddGraphNode(NodeClassGraph))
 			{
 				FMetasoundFrontendClassMetadata NodeClassMetadata = NodeClassGraph.Metadata;
diff --git a/Engine/Plugins/Runtime/Metasound/Source/MetasoundEngine/Public/MetasoundBuilderSubsystem.h b/Engine/Plugins/Runtime/Metasound/Source/MetasoundEngine/Public/MetasoundBuilderSubsystem.h
--- a/Engine/Plugins/Runtime/Metasound/Source/MetasoundEngine/Public/MetasoundBuilderSubsystem.h
+++ b/Engine/Plugins/Runtime/Metasound/Source/MetasoundEngine/Public/MetasoundBuilderSubsystem.h
@@ -126,14 +126,14 @@ public:
 	// Adds a graph output node with the given name, DataType, and sets output node's input to default value.
 	// Returns the new output node's input handle if it was successfully created, or an invalid handle if it failed.
 	UFUNCTION(BlueprintCallable, Category = "Audio|MetaSound|Builder", meta = (ExpandEnumAsExecs = "OutResult"))
-	UPARAM(DisplayName = "Input Handle") FMetaSoundBuilderNodeInputHandle AddGraphOutputNode(FName Name, FName DataType, FMetasoundFrontendLiteral DefaultValue, EMetaSoundBuilderResult& OutResult, bool bIsConstructorInput = false);
+	UPARAM(DisplayName = "Input Handle") FMetaSoundBuilderNodeInputHandle AddGraphOutputNode(FName Name, FName DataType, FMetasoundFrontendLiteral DefaultValue, EMetaSoundBuilderResult& OutResult, bool bIsConstructorOutput = false);
 
 	// Adds an interface registered with the given name to the graph, adding associated input and output nodes.
 	UFUNCTION(BlueprintCallable, Category = "Audio|MetaSound|Builder", meta = (ExpandEnumAsExecs = "OutResult"))
 	void AddInterface(FName InterfaceName, EMetaSoundBuilderResult& OutResult);
 
 	// Adds a node to the graph using the provided MetaSound asset as its defining NodeClass.
-	// Returns a node handle to the created node if successful, or an invliad handle if it failed.
+	// Returns a node handle to the created node if successful, or an invalid handle if it failed.
 	UFUNCTION(BlueprintCallable, Category = "Audio|MetaSound|Builder", DisplayName = "Add MetaSound Node From Asset Class", meta = (ExpandEnumAsExecs = "OutResult"))
 	UPARAM(DisplayName = "Node Handle") FMetaSoundNodeHandle AddNode(const TScriptInterface<IMetaSoundDocumentInterface>& NodeClass, EMetaSoundBuilderResult& OutResult);
 
diff --git a/Engine/Plugins/Runtime/Metasound/Source/MetasoundFrontend/Private/MetasoundFrontendDocumentCache.cpp b/Engine/Plugins/Runtime/Metasound/Source/MetasoundFrontend/Private/MetasoundFrontendDocumentCache.cpp
--- a/Engine/Plugins/Runtime/Metasound/Source/MetasoundFrontend/Private/MetasoundFrontendDocumentCache.cpp
+++ b/Engine/Plugins/Runtime/Metasound/Source/MetasoundFrontend/Private/MetasoundFrontendDocumentCache.cpp
@@ -336,7 +336,7 @@ namespace Metasound::Frontend
 		{
 			case EMetasoundFrontendClassType::Input:
 			{
-				InputNameToIndex.Add(NodeBeingRemoved.Name);
+				InputNameToIndex.Remove(NodeBeingRemoved.Name);
 			}
 			break;
 
