Automated MetaSound node tests for rebind.

diff --git a/Engine/Plugins/Runtime/Metasound/Source/MetasoundEngineTest/Private/EngineTestAutomatedNodeTests.cpp b/Engine/Plugins/Runtime/Metasound/Source/MetasoundEngineTest/Private/EngineTestAutomatedNodeTests.cpp
deleted file mode 100644
diff --git a/Engine/Plugins/Runtime/Metasound/Source/MetasoundEngineTest/Private/EngineTestMetaSoundAutomatedNodeTest.cpp b/Engine/Plugins/Runtime/Metasound/Source/MetasoundEngineTest/Private/EngineTestMetaSoundAutomatedNodeTest.cpp
new file mode 100644
--- /dev/null
+++ b/Engine/Plugins/Runtime/Metasound/Source/MetasoundEngineTest/Private/EngineTestMetaSoundAutomatedNodeTest.cpp
@@ -0,0 +1,840 @@
+// Copyright Epic Games, Inc. All Rights Reserved.
+
+#include "EngineTestMetaSoundAutomatedNodeTest.h"
+
+#include "Algo/AllOf.h"
+#include "AudioDevice.h"
+#include "AudioDeviceManager.h"
+#include "AudioMixerDevice.h"
+#include "Containers/Array.h"
+#include "Containers/UnrealString.h"
+#include "Interfaces/MetasoundFrontendSourceInterface.h"
+#include "Math/NumericLimits.h"
+#include "Misc/AutomationTest.h"
+#include "MetasoundAudioBuffer.h"
+#include "MetasoundFrontendDataTypeRegistry.h"
+#include "MetasoundFrontendRegistries.h"
+#include "MetasoundFrontendSearchEngine.h"
+#include "MetasoundLog.h"
+#include "MetasoundPrimitives.h"
+#include "MetasoundVariableNodes.h"
+#include "MetasoundVertex.h"
+#include "MetasoundVertexData.h"
+#include "Templates/UniquePtr.h"
+#include "Tests/AutomationCommon.h"
+
+#if WITH_DEV_AUTOMATION_TESTS
+
+#if WITH_EDITORONLY_DATA
+
+namespace Metasound::EngineTest{
+
+	// Return audio mixer device if one is available
+	Audio::FMixerDevice* GetAudioMixerDevice()
+	{
+		if (FAudioDeviceManager* DeviceManager = FAudioDeviceManager::Get())
+		{
+			if (FAudioDevice* AudioDevice = DeviceManager->GetMainAudioDeviceRaw())
+			{
+				if (AudioDevice->IsAudioMixerEnabled())
+				{
+					return static_cast<Audio::FMixerDevice*>(AudioDevice);
+				}
+			}
+		}
+		return nullptr;
+	}
+
+	// Create an example environment that generally exists for a UMetaSoundSoruce
+	FMetasoundEnvironment GetSourceEnvironmentForTest()
+	{
+		using namespace Frontend;
+
+		FMetasoundEnvironment Environment;
+
+		Environment.SetValue<uint32>(SourceInterface::Environment::SoundUniqueID, 0);
+		Environment.SetValue<bool>(SourceInterface::Environment::IsPreview, false);
+		Environment.SetValue<uint64>(SourceInterface::Environment::TransmitterID, 0);
+		Environment.SetValue<FString>(SourceInterface::Environment::GraphName, TEXT("ENGINE_TEST_REGISTERED_NODES"));
+
+		if (Audio::FMixerDevice* MixerDevice = GetAudioMixerDevice())
+		{
+			Environment.SetValue<Audio::FDeviceId>(SourceInterface::Environment::DeviceID, MixerDevice->DeviceID);
+			Environment.SetValue<int32>(SourceInterface::Environment::AudioMixerNumOutputFrames, MixerDevice->GetNumOutputFrames());
+		}
+
+		return Environment;
+	}
+
+	FString GetPrettyName(const Frontend::FNodeRegistryKey& InRegistryKey)
+	{
+		Frontend::IMetaSoundAssetManager* AssetManager = Frontend::IMetaSoundAssetManager::Get();
+		if (ensure(AssetManager))
+		{
+			if (const FSoftObjectPath* ObjectPath = AssetManager->FindObjectPathFromKey(InRegistryKey))
+			{
+				return ObjectPath->ToString();
+			}
+		}
+
+		FMetasoundFrontendRegistryContainer* NodeRegistry = FMetasoundFrontendRegistryContainer::Get();
+		if (ensure(NodeRegistry))
+		{
+			FMetasoundFrontendClass NodeClass;
+			if (NodeRegistry->FindFrontendClassFromRegistered(InRegistryKey, NodeClass))
+			{
+				return FString::Printf(TEXT("%s %s"), *NodeClass.Metadata.GetClassName().ToString(), *NodeClass.Metadata.GetVersion().ToString());
+			}
+		}
+
+		return TEXT("");
+	}
+
+	// TTestTypeInfo converts test types to strings
+	template<typename DataType>
+	struct TTestTypeInfo
+	{
+		static FString ToString(const DataType& InData)
+		{
+			return ::LexToString(InData);
+		}
+	};
+
+	template<typename ElementType>
+	struct TTestTypeInfo<TArray<ElementType>>
+	{
+		static FString ToString(TArrayView<const ElementType> InData)
+		{
+			return FString::Printf(TEXT("[%s]"), *FString::JoinBy(InData, TEXT(","), &TTestTypeInfo<ElementType>::ToString));
+		}
+	};
+
+	template<>
+	struct TTestTypeInfo<FAudioBuffer> : TTestTypeInfo<TArray<float>>
+	{
+	};
+	
+
+	template<>
+	struct TTestTypeInfo<FTime>
+	{
+		static FString ToString(const FTime& InData)
+		{
+			return ::LexToString(InData.GetSeconds());
+		}
+	};
+
+	// TTestTypeAnalysis performs analysis operator on data types.
+	template<typename DataType>
+	struct TTestTypeAnalysis
+	{
+		static bool IsValid(const DataType& InArray)
+		{
+			return true;
+		}
+
+		static bool IsEqual(const DataType& InLHS, const DataType& InRHS)
+		{
+			return InLHS == InRHS;
+		}
+	};
+
+	// Specialize type analysis for TArrays
+	template<typename ElementType>
+	struct TTestTypeAnalysis<TArray<ElementType>>
+	{
+		static bool IsValid(TArrayView<const ElementType> InArray)
+		{
+			return Algo::AllOf(InArray, &TTestTypeAnalysis<ElementType>::IsValid);
+		}
+
+		static bool IsEqual(TArrayView<const ElementType> InLHS, TArrayView<const ElementType> InRHS)
+		{
+			const int32 Num = InLHS.Num();
+			if (Num == InRHS.Num())
+			{
+				for (int32 i = 0; i < Num; i++)
+				{
+					if (InLHS[i] != InRHS[i])
+					{
+						return false;
+					}
+				}
+				return true;
+			}
+			return false;
+		}
+	};
+
+	// Specialize type analysis for audio buffers.
+	template<>
+	struct TTestTypeAnalysis<FAudioBuffer> : TTestTypeAnalysis<TArray<float>>
+	{
+	};
+
+	// Specialize type analysis for float
+	template<>
+	struct TTestTypeAnalysis<float>
+	{
+		static bool IsValid(float InValue)
+		{
+			return !FMath::IsNaN(InValue);
+		}
+
+		static bool IsEqual(float InLHS, float InRHS)
+		{
+			return InLHS == InRHS;
+		}
+	};
+
+	// Specialize type analysis for time
+	template<>
+	struct TTestTypeAnalysis<FTime>
+	{
+		static bool IsValid(const FTime& InValue)
+		{
+			return FMath::IsFinite(InValue.GetSeconds());
+		}
+
+		static bool IsEqual(const FTime& InLHS, const FTime& InRHS)
+		{
+			return InLHS == InRHS;
+		}
+	};
+
+	// TDataReferenceAnalyzer fulfils the IDataReferenceAnalyzer interface and uses
+	// various templates to implement methods.
+	template<typename DataType>
+	struct TDataReferenceAnalyzer : public IDataReferenceAnalyzer
+	{
+		virtual FAnyDataReference Copy(const FAnyDataReference& InDataRef) const override
+		{
+			if (const DataType* InValue = GetDataOrLogError(InDataRef))
+			{
+				return FAnyDataReference{TDataValueReference<DataType>::CreateNew(*InValue)};
+			}
+			return InDataRef;
+		}
+
+		virtual bool IsEqual(const FAnyDataReference& InLHS, const FAnyDataReference& InRHS) const override
+		{
+			if (const DataType* InLHSValue = GetDataOrLogError(InLHS))
+			{
+				if (const DataType* InRHSValue = GetDataOrLogError(InRHS))
+				{
+					return TTestTypeAnalysis<DataType>::IsEqual(*InLHSValue, *InRHSValue);
+				}
+			}
+			return false;
+		}
+
+		virtual bool IsValid(const FAnyDataReference& InDataRef) const override
+		{
+			if (const DataType* InValue = GetDataOrLogError(InDataRef))
+			{
+				return TTestTypeAnalysis<DataType>::IsValid(*InValue);
+			}
+			return false;
+		}
+
+		virtual FString ToString(const FAnyDataReference& InDataRef) const override
+		{
+			if (const DataType* InValue = GetDataOrLogError(InDataRef))
+			{
+				return TTestTypeInfo<DataType>::ToString(*InValue);
+			}
+			return TEXT("");
+		}
+	private:
+		const DataType* GetDataOrLogError(const FAnyDataReference& InDataRef) const
+		{
+			const DataType* Value = InDataRef.GetValue<DataType>();
+			if (nullptr == Value)	
+			{
+				// Data references should never be null
+				UE_LOG(LogMetaSound, Error, TEXT("Failed to get data type value of type %s"), *GetMetasoundDataTypeString<DataType>());
+			}
+			return Value;
+		}
+	};
+
+	// Register a data reference analyzer to support analyzing output of metasound nodes.
+	template<typename DataType>
+	void AddDataReferenceAnalyzerToMap(TMap<FName, TSharedPtr<const IDataReferenceAnalyzer>>& InMap)
+	{
+		InMap.Add(GetMetasoundDataTypeName<DataType>(), MakeShared<const TDataReferenceAnalyzer<DataType>>());
+	}
+
+	// A static registry of IDataReferenceAnalyzers
+	const TMap<FName, TSharedPtr<const IDataReferenceAnalyzer>>& GetDataTypeAnalyzerMap()
+	{
+		static TMap<FName, TSharedPtr<const IDataReferenceAnalyzer>> Map;
+
+		AddDataReferenceAnalyzerToMap<bool>(Map);
+		AddDataReferenceAnalyzerToMap<int32>(Map);
+		AddDataReferenceAnalyzerToMap<float>(Map);
+		AddDataReferenceAnalyzerToMap<FString>(Map);
+		AddDataReferenceAnalyzerToMap<FTime>(Map);
+		AddDataReferenceAnalyzerToMap<FAudioBuffer>(Map);
+		AddDataReferenceAnalyzerToMap<FTrigger>(Map);
+		AddDataReferenceAnalyzerToMap<TArray<bool>>(Map);
+		AddDataReferenceAnalyzerToMap<TArray<int32>>(Map);
+		AddDataReferenceAnalyzerToMap<TArray<float>>(Map);
+		AddDataReferenceAnalyzerToMap<TArray<FString>>(Map);
+		AddDataReferenceAnalyzerToMap<TArray<FTime>>(Map);
+
+		return Map;
+	}
+
+	void FOutputVertexDataTestController::FAnalyzableOutput::CaptureValue()
+	{
+		CapturedValue = DataReferenceAnalyzer->Copy(DataReference);
+	}
+
+	bool FOutputVertexDataTestController::FAnalyzableOutput::IsDataReferenceValid() const
+	{
+		return DataReferenceAnalyzer->IsValid(DataReference);
+	}
+
+	bool FOutputVertexDataTestController::FAnalyzableOutput::IsDataReferenceEqualToCapturedValue() const
+	{
+		return DataReferenceAnalyzer->IsEqual(CapturedValue, DataReference);
+	}
+
+	FString FOutputVertexDataTestController::FAnalyzableOutput::DataReferenceToString() const
+	{
+		return DataReferenceAnalyzer->ToString(DataReference);
+	}
+
+	FString FOutputVertexDataTestController::FAnalyzableOutput::CapturedValueToString() const
+	{
+		return DataReferenceAnalyzer->ToString(CapturedValue);
+	}
+
+	FOutputVertexDataTestController::FOutputVertexDataTestController( const FOutputVertexInterface& InOutputInterface, const FOutputVertexInterfaceData& InOutputData)
+	{
+		const TMap<FName, TSharedPtr<const IDataReferenceAnalyzer>>& AnalyzerMap = GetDataTypeAnalyzerMap();
+
+		for (const FOutputDataVertex& Vertex : InOutputInterface)
+		{
+			if (AnalyzerMap.Contains(Vertex.DataTypeName))
+			{
+				if (const FAnyDataReference* Ref = InOutputData.FindDataReference(Vertex.VertexName))
+				{
+					AnalyzableOutputs.Add(FAnalyzableOutput{*Ref, *Ref, Vertex.VertexName, AnalyzerMap[Vertex.DataTypeName]});
+				}
+			}
+		}
+	}
+
+	int32 FOutputVertexDataTestController::GetNumAnalyzableOutputs() const
+	{
+		return AnalyzableOutputs.Num();
+	}
+
+	bool FOutputVertexDataTestController::AreAllAnalyzableOutputsValid() const
+	{
+		bool bAllAreValid = true;
+		for (const FAnalyzableOutput& AnalyzableOutput : AnalyzableOutputs)
+		{
+			if (!AnalyzableOutput.IsDataReferenceValid())
+			{
+				UE_LOG(LogMetaSound, Warning, TEXT("Invalid output encountered %s %s"), *AnalyzableOutput.VertexName.ToString(), *AnalyzableOutput.DataReferenceToString());
+				bAllAreValid = false;
+			}
+		}
+		return bAllAreValid;
+	}
+
+	void FOutputVertexDataTestController::CaptureCurrentOutputValues()
+	{
+		for (FAnalyzableOutput& AnalyzableOutput : AnalyzableOutputs)
+		{
+			AnalyzableOutput.CaptureValue();
+		}
+	}
+
+	bool FOutputVertexDataTestController::AreAllOutputValuesEqualToCapturedValues() const
+	{
+		bool bAllAreEqual = true;
+		for (const FAnalyzableOutput& AnalyzableOutput : AnalyzableOutputs)
+		{
+			if (!AnalyzableOutput.IsDataReferenceEqualToCapturedValue())
+			{
+				UE_LOG(LogMetaSound, Warning, TEXT("Unequal output encountered %s found: %s expected: %s"), *AnalyzableOutput.VertexName.ToString(), *AnalyzableOutput.DataReferenceToString(), *AnalyzableOutput.CapturedValueToString());
+				bAllAreEqual = false;
+			}
+		}
+		return bAllAreEqual;
+	}
+
+
+	// TTestTypeValues should return basic bounds for tested input data types.
+	// Similar to TNumericLimits<>
+	template<typename DataType>
+	struct TTestTypeValues
+	{};
+
+	// TArray specialization to defer to a single element array with array element's values
+	template<typename ElementType>
+	struct TTestTypeValues<TArray<ElementType>>
+	{
+		static TArray<ElementType> Min(const FOperatorSettings& InSettings) 
+		{
+			return TArray<ElementType>({TTestTypeValues<ElementType>::Min(InSettings)}); 
+		}
+
+		static TArray<ElementType> Max(const FOperatorSettings& InSettings) 
+		{ 
+			return TArray<ElementType>({TTestTypeValues<ElementType>::Max(InSettings)}); 
+		}
+
+		static TArray<ElementType> Default(const FOperatorSettings& InSettings) 
+		{ 
+			return TArray<ElementType>({TTestTypeValues<ElementType>::Default(InSettings)}); 
+		}
+
+		static TArray<ElementType> Random(const FOperatorSettings& InSettings)
+		{ 
+			return TArray<ElementType>({TTestTypeValues<ElementType>::Random(InSettings)}); 
+		}
+	};
+
+	template<>
+	struct TTestTypeValues<bool>
+	{
+		static bool Min(const FOperatorSettings& InSettings) { return false; }
+		static bool Max(const FOperatorSettings& InSettings) { return true; }
+		static bool Default(const FOperatorSettings& InSettings) { return true; }
+		static bool Random(const FOperatorSettings& InSettings) { return FMath::RandRange(0.f, 1.f) > 0.5; }
+	};
+
+	template<>
+	struct TTestTypeValues<int32>
+	{
+		static int32 Min(const FOperatorSettings& InSettings) { return TNumericLimits<int32>::Min(); }
+		static int32 Max(const FOperatorSettings& InSettings) { return TNumericLimits<int32>::Max(); }
+		static int32 Default(const FOperatorSettings& InSettings) { return 0; }
+		static int32 Random(const FOperatorSettings& InSettings) { return FMath::RandRange(TNumericLimits<int32>::Min(), TNumericLimits<int32>::Max()); }
+	};
+
+	template<>
+	struct TTestTypeValues<float>
+	{
+		static float Min(const FOperatorSettings& InSettings) { return TNumericLimits<float>::Min(); }
+		static float Max(const FOperatorSettings& InSettings) { return TNumericLimits<float>::Max(); }
+		static float Default(const FOperatorSettings& InSettings) { return 0.f; }
+		static float Random(const FOperatorSettings& InSettings) { return FMath::RandRange(TNumericLimits<float>::Min(), TNumericLimits<float>::Max()); }
+	};
+
+	template<>
+	struct TTestTypeValues<FTime>
+	{
+		static FTime Min(const FOperatorSettings& InSettings) { return FTime{TNumericLimits<float>::Min()}; }
+		static FTime Max(const FOperatorSettings& InSettings) { return FTime{TNumericLimits<float>::Max()}; }
+		static FTime Default(const FOperatorSettings& InSettings) { return FTime{0.f}; }
+		static FTime Random(const FOperatorSettings& InSettings) { return FTime{FMath::RandRange(TNumericLimits<float>::Min(), TNumericLimits<float>::Max())}; }
+	};
+
+	template<>
+	struct TTestTypeValues<FTrigger>
+	{
+		static FTrigger Min(const FOperatorSettings& InSettings) { return FTrigger{InSettings, false}; }
+		static FTrigger Max(const FOperatorSettings& InSettings) 
+		{ 
+			FTrigger Trigger{InSettings, false};
+			for (int32 i = 0; i < InSettings.GetNumFramesPerBlock(); i++)
+			{
+				Trigger.TriggerFrame(i);
+			}
+			return Trigger;
+		}
+
+		static FTrigger Default(const FOperatorSettings& InSettings) { return FTrigger{InSettings, true}; }
+		static FTrigger Random(const FOperatorSettings& InSettings) 
+		{ 
+			FTrigger Trigger{InSettings, false};
+			int32 NumTriggers = FMath::RandRange(0, InSettings.GetNumFramesPerBlock());
+			while (NumTriggers > 0)
+			{
+				Trigger.TriggerFrame(FMath::RandRange(0, InSettings.GetNumFramesPerBlock()));
+				NumTriggers--;
+			}
+			return Trigger;
+		}
+	};
+
+	template<>
+	struct TTestTypeValues<FString>
+	{
+		static FString Min(const FOperatorSettings& InSettings) { return TEXT(""); }
+		static FString Max(const FOperatorSettings& InSettings) { return TEXT("THIS IS SUPPOSED TO REPRESENT A MAXIMUM STRING BUT THERE IS NO SUCH THING SO?"); }
+		static FString Default(const FOperatorSettings& InSettings) { return TEXT("TestString"); }
+		static FString Random(const FOperatorSettings& InSettings) { return TEXT("We should probably implement a random string."); }
+	};
+
+	template<typename DataType>
+	struct TDataReferenceMutator : IDataReferenceMutator
+	{
+		virtual void SetDefault(const FOperatorSettings& InSettings, const FAnyDataReference& InDataRef) const override
+		{
+			*InDataRef.GetDataWriteReference<DataType>() = TTestTypeValues<DataType>::Default(InSettings);
+		}
+
+		virtual void SetMax(const FOperatorSettings& InSettings, const FAnyDataReference& InDataRef) const override
+		{
+			*InDataRef.GetDataWriteReference<DataType>() = TTestTypeValues<DataType>::Max(InSettings);
+		}
+
+		virtual void SetMin(const FOperatorSettings& InSettings, const FAnyDataReference& InDataRef) const override
+		{
+			*InDataRef.GetDataWriteReference<DataType>() = TTestTypeValues<DataType>::Min(InSettings);
+		}
+
+		virtual void SetRandom(const FOperatorSettings& InSettings, const FAnyDataReference& InDataRef) const override
+		{
+			*InDataRef.GetDataWriteReference<DataType>() = TTestTypeValues<DataType>::Random(InSettings);
+		}
+
+		virtual FString ToString(const FAnyDataReference& InDataRef) const override
+		{
+			if (const DataType* Data = InDataRef.GetValue<DataType>())
+			{
+				return FString::Printf(TEXT("%s:%s"), *GetMetasoundDataTypeString<DataType>(), *TTestTypeInfo<DataType>::ToString(*Data));
+			}
+			else
+			{
+				// Data references should never be null
+				UE_LOG(LogMetaSound, Error, TEXT("Failed to get data type value of type %s"), *GetMetasoundDataTypeString<DataType>());
+				return TEXT("");
+			}
+		}
+
+		virtual void SetValue(const FAnyDataReference& InSrcDataRef, const FAnyDataReference& InDstDataRef) const override
+		{
+			*InDstDataRef.GetDataWriteReference<DataType>() = *InSrcDataRef.GetValue<DataType>();
+		}
+
+		virtual FAnyDataReference Copy(const FAnyDataReference& InDataRef) const override
+		{
+			if (const DataType* Data = InDataRef.GetValue<DataType>())
+			{
+				return FAnyDataReference{TDataValueReference<DataType>::CreateNew(*Data)};
+			}
+			else
+			{
+				// Data references should never be null
+				UE_LOG(LogMetaSound, Error, TEXT("Failed to get data type value of type %s"), *GetMetasoundDataTypeString<DataType>());
+				return FAnyDataReference{TDataValueReference<int32>::CreateNew()}; // we are going to crash soon.
+			}
+		}
+	};
+
+	template<typename DataType>
+	void AddDataReferenceMutatorEntryToMap(TMap<FName, TSharedPtr<const IDataReferenceMutator>>& InMap)
+	{
+		InMap.Add(GetMetasoundDataTypeName<DataType>(), MakeShared<TDataReferenceMutator<DataType>>());
+	}
+
+	// Returns map of mutable input types
+	const TMap<FName, TSharedPtr<const IDataReferenceMutator>>& GetDataTypeGeneratorMap()
+	{
+		static TMap<FName, TSharedPtr<const IDataReferenceMutator>> Map;
+
+		AddDataReferenceMutatorEntryToMap<bool>(Map);
+		AddDataReferenceMutatorEntryToMap<int32>(Map);
+		AddDataReferenceMutatorEntryToMap<float>(Map);
+		AddDataReferenceMutatorEntryToMap<FString>(Map);
+		AddDataReferenceMutatorEntryToMap<FTime>(Map);
+		AddDataReferenceMutatorEntryToMap<FTrigger>(Map);
+		AddDataReferenceMutatorEntryToMap<TArray<bool>>(Map);
+		AddDataReferenceMutatorEntryToMap<TArray<int32>>(Map);
+		AddDataReferenceMutatorEntryToMap<TArray<float>>(Map);
+		AddDataReferenceMutatorEntryToMap<TArray<FString>>(Map);
+		AddDataReferenceMutatorEntryToMap<TArray<FTime>>(Map);
+		AddDataReferenceMutatorEntryToMap<TArray<FTrigger>>(Map);
+
+		return Map;
+	}
+
+	// Convenience class for setting node input data reference values to default, min, max or random values. 
+	FInputVertexDataTestController::FInputVertexDataTestController(const FOperatorSettings& InSettings, const FInputVertexInterface& InInputInterface, const FInputVertexInterfaceData& InInputData)
+	: Settings(InSettings)
+	{
+		const TMap<FName, TSharedPtr<const IDataReferenceMutator>>& GeneratorMap = GetDataTypeGeneratorMap();
+
+		for (const FInputDataVertex& Vertex : InInputInterface)
+		{
+			if (GeneratorMap.Contains(Vertex.DataTypeName))
+			{
+				if (const FAnyDataReference* Ref = InInputData.FindDataReference(Vertex.VertexName))
+				{
+					if (EDataReferenceAccessType::Write == Ref->GetAccessType())
+					{
+						MutableInputs.Add(FMutableInput{*Ref, Vertex.VertexName, GeneratorMap[Vertex.DataTypeName]});
+					}
+				}
+			}
+		}
+	}
+
+	int32 FInputVertexDataTestController::GetNumMutableInputs() const
+	{
+		return MutableInputs.Num();
+	}
+
+	FInterfaceState FInputVertexDataTestController::GetInterfaceState() const
+	{
+		FInterfaceState State;
+		for (const FMutableInput& MutableInput : MutableInputs)
+		{
+			State.Add(MutableInput.VertexName, MutableInput.DataReferenceMutator->Copy(MutableInput.DataReference));
+		}
+		return State;
+	}
+
+	void FInputVertexDataTestController::SetMutableInputsToState(const FInterfaceState& InState) 
+	{
+		for (const FMutableInput& MutableInput : MutableInputs)
+		{
+			if (const FAnyDataReference* Value = InState.Find(MutableInput.VertexName))
+			{
+				MutableInput.DataReferenceMutator->SetValue(*Value, MutableInput.DataReference);
+			}
+		}
+		UE_LOG(LogMetaSound, Verbose, TEXT("Setting operator input values:%s%s"), LINE_TERMINATOR, *FString::Join(GetInputValueStrings(), LINE_TERMINATOR));
+	}
+
+	void FInputVertexDataTestController::SetMutableInputsToMin()
+	{
+		for (const FMutableInput& MutableInput : MutableInputs)
+		{
+			MutableInput.DataReferenceMutator->SetMin(Settings, MutableInput.DataReference);
+		}
+		UE_LOG(LogMetaSound, Verbose, TEXT("Setting operator input values:%s%s"), LINE_TERMINATOR, *FString::Join(GetInputValueStrings(), LINE_TERMINATOR));
+	}
+
+	void FInputVertexDataTestController::SetMutableInputsToMax()
+	{
+		for (const FMutableInput& MutableInput : MutableInputs)
+		{
+			MutableInput.DataReferenceMutator->SetMax(Settings, MutableInput.DataReference);
+		}
+		UE_LOG(LogMetaSound, Verbose, TEXT("Setting operator input values:%s%s"), LINE_TERMINATOR, *FString::Join(GetInputValueStrings(), LINE_TERMINATOR));
+	}
+
+	void FInputVertexDataTestController::SetMutableInputsToDefault()
+	{
+		for (const FMutableInput& MutableInput : MutableInputs)
+		{
+			MutableInput.DataReferenceMutator->SetDefault(Settings, MutableInput.DataReference);
+		}
+		UE_LOG(LogMetaSound, Verbose, TEXT("Setting operator input values:%s%s"), LINE_TERMINATOR, *FString::Join(GetInputValueStrings(), LINE_TERMINATOR));
+	}
+
+	void FInputVertexDataTestController::SetMutableInputsToRandom()
+	{
+		for (const FMutableInput& MutableInput : MutableInputs)
+		{
+			MutableInput.DataReferenceMutator->SetRandom(Settings, MutableInput.DataReference);
+		}
+		UE_LOG(LogMetaSound, Verbose, TEXT("Setting operator input values:%s%s"), LINE_TERMINATOR, *FString::Join(GetInputValueStrings(), LINE_TERMINATOR));
+	}
+
+	TArray<FString> FInputVertexDataTestController::GetInputValueStrings() const
+	{
+		TArray<FString> ValueStrings;
+		for (const FMutableInput& MutableInput : MutableInputs)
+		{
+			ValueStrings.Add(FString::Printf(TEXT("%s %s"), *MutableInput.VertexName.ToString(), *MutableInput.DataReferenceMutator->ToString(MutableInput.DataReference)));
+		}
+		return ValueStrings;
+	}
+
+
+	static const FLazyName TestNodeName{"TEST_NODE"};
+	static const FLazyName TestVertexName{"TEXT_VERTEX"};
+	static const FGuid TestNodeID{0xA5A5A5A5, 0xA5A5A5A5, 0xA5A5A5A5, 0xA5A5A5A5};
+
+	// Create a node from a node registry key
+	TUniquePtr<INode> CreateNodeFromRegistry(const Frontend::FNodeRegistryKey& InNodeRegistryKey)
+	{
+		using namespace Frontend;
+
+		TUniquePtr<INode> Node;
+
+		FMetasoundFrontendRegistryContainer* NodeRegistry = FMetasoundFrontendRegistryContainer::Get();
+		check(nullptr != NodeRegistry);
+		IDataTypeRegistry& DataTypeRegistry = IDataTypeRegistry::Get();
+
+		// Lookup node class metadata to determine how to create this node.
+		FMetasoundFrontendClass NodeClass;
+		if (!NodeRegistry->FindFrontendClassFromRegistered(InNodeRegistryKey, NodeClass))
+		{
+			UE_LOG(LogMetaSound, Error, TEXT("Failed to find registered class with registry key %s"), *InNodeRegistryKey);
+			return MoveTemp(Node);
+		}
+
+		// Build node differently dependent upon node type
+		switch (NodeClass.Metadata.GetType())
+		{
+			case EMetasoundFrontendClassType::VariableDeferredAccessor:
+			case EMetasoundFrontendClassType::VariableAccessor:
+			case EMetasoundFrontendClassType::VariableMutator:
+			case EMetasoundFrontendClassType::External:
+			case EMetasoundFrontendClassType::Graph:
+			{
+				FNodeInitData NodeInitData{TestNodeName, TestNodeID};
+				Node = NodeRegistry->CreateNode(InNodeRegistryKey, NodeInitData);
+			}
+			break;
+
+			case EMetasoundFrontendClassType::Input:
+			{
+				FName DataTypeName = NodeClass.Metadata.GetClassName().Name;
+				FInputNodeConstructorParams NodeInitData
+				{
+					TestNodeName, 
+					TestNodeID, 
+					TestVertexName,
+					NodeClass.Interface.Inputs[0].DefaultLiteral.ToLiteral(DataTypeName)
+				};
+
+				Node = DataTypeRegistry.CreateInputNode(DataTypeName, MoveTemp(NodeInitData));
+			}
+			break;
+
+			case EMetasoundFrontendClassType::Variable:
+			{
+				FName DataTypeName = NodeClass.Metadata.GetClassName().Name;
+				FDefaultLiteralNodeConstructorParams NodeInitData{TestNodeName, TestNodeID, DataTypeRegistry.CreateDefaultLiteral(DataTypeName)};
+				Node = DataTypeRegistry.CreateVariableNode(DataTypeName, MoveTemp(NodeInitData));
+			}
+			break;
+
+			case EMetasoundFrontendClassType::Literal:
+			{
+				FName DataTypeName = NodeClass.Metadata.GetClassName().Name;
+				FDefaultLiteralNodeConstructorParams NodeInitData{TestNodeName, TestNodeID, DataTypeRegistry.CreateDefaultLiteral(DataTypeName)};
+				Node = DataTypeRegistry.CreateLiteralNode(DataTypeName, MoveTemp(NodeInitData));
+			}
+			break;
+
+			case EMetasoundFrontendClassType::Output:
+			{
+				FName DataTypeName = NodeClass.Metadata.GetClassName().Name;
+				FDefaultNamedVertexNodeConstructorParams NodeInitData{TestNodeName, TestNodeID, TestVertexName};
+				Node = DataTypeRegistry.CreateOutputNode(DataTypeName, MoveTemp(NodeInitData));
+			}
+			break;
+
+			case EMetasoundFrontendClassType::Template:
+			default:
+			static_assert(static_cast<int32>(EMetasoundFrontendClassType::Invalid) == 10, "Possible missed EMetasoundFrontendClassType case coverage");
+		}
+
+		return MoveTemp(Node);
+	}
+
+
+	void GetAllRegisteredNodes(TArray<FString>& OutBeautifiedNames, TArray <FString>& OutTestCommands)
+	{
+		using namespace Metasound;
+
+		// Get all the classes that have been registered
+		Frontend::ISearchEngine& NodeSearchEngine = Frontend::ISearchEngine::Get();
+		TArray<FMetasoundFrontendClass> AllClasses = NodeSearchEngine.FindAllClasses(true /* IncludeAllVersions */);
+
+		FMetasoundFrontendRegistryContainer* NodeRegistry = FMetasoundFrontendRegistryContainer::Get();
+		check(nullptr != NodeRegistry);
+
+		for (const FMetasoundFrontendClass& NodeClass : AllClasses)
+		{
+			// Exclude template classes because they cannot be created directly from the node registry
+			if (NodeClass.Metadata.GetType() == EMetasoundFrontendClassType::Template)
+			{
+				continue;
+			}
+
+			Frontend::FNodeRegistryKey NodeRegistryKey = NodeRegistry->GetRegistryKey(NodeClass.Metadata);
+			
+			OutBeautifiedNames.Add(FString::Printf(TEXT("%s %s"), *NodeClass.Metadata.GetClassName().ToString(), *NodeClass.Metadata.GetVersion().ToString()));
+			// Test commands are node registry keys
+			OutTestCommands.Add(NodeRegistryKey);
+		}
+	}
+
+
+	void CreateVariables(const FOperatorSettings& InOperatorSettings, FInputVertexInterfaceData& OutVertexData)
+	{
+		using namespace MetasoundVertexDataPrivate;
+
+		// There is currently no easy way to instantiate variables. The only way
+		// they are instantiated is within the TVariableNode's IOperator. In order
+		// to create Variable inputs to nodes, we 
+		// - Create a TVariableNode
+		// - Create a IOperator from the TVariableNode
+		// - Access the variable from the outputs of the IOperator
+
+		Frontend::IDataTypeRegistry& DataTypeRegistry = Frontend::IDataTypeRegistry::Get();
+
+		for (const FInputBinding& Binding : OutVertexData)
+		{
+			const FInputDataVertex& InputVertex = Binding.GetVertex();
+
+			if (const Frontend::IDataTypeRegistryEntry* VariableEntry = DataTypeRegistry.FindDataTypeRegistryEntry(InputVertex.DataTypeName))
+			{
+				if (VariableEntry->GetDataTypeInfo().bIsVariable)
+				{
+					// Find the data type name of the data type wrapped by the variable
+					FName UnderlyingDataTypeName = *InputVertex.DataTypeName.ToString().Replace(TEXT(METASOUND_DATA_TYPE_NAME_VARIABLE_TYPE_SPECIFIER), TEXT(""));
+
+					if (const Frontend::IDataTypeRegistryEntry* DataTypeEntry = DataTypeRegistry.FindDataTypeRegistryEntry(UnderlyingDataTypeName))
+					{
+						// Create a variable node for the underlying data type
+						FVariableNodeConstructorParams Params;
+						Params.Literal = DataTypeRegistry.CreateDefaultLiteral(UnderlyingDataTypeName);
+						TUniquePtr<INode> VariableNode = DataTypeEntry->CreateVariableNode(MoveTemp(Params));
+
+						if (VariableNode.IsValid())
+						{
+							// Create the variable operator for the underlying data type
+							FInputVertexInterfaceData InputData(VariableNode->GetVertexInterface().GetInputInterface());
+							FBuildOperatorParams BuildParams
+							{
+								*VariableNode,
+								InOperatorSettings,
+								InputData,
+								FMetasoundEnvironment{}
+							};
+							FBuildResults OutResults;
+							TUniquePtr<IOperator> VariableOperator = VariableNode->GetDefaultOperatorFactory()->CreateOperator(BuildParams, OutResults);
+
+							if (VariableOperator.IsValid())
+							{
+								// Access the TVariable data referenced created within the operator.
+								FOutputVertexInterfaceData OutOperatorData;
+								VariableOperator->BindOutputs(OutOperatorData);
+								if (const FAnyDataReference* VariableRef = OutOperatorData.FindDataReference(METASOUND_GET_PARAM_NAME(VariableNames::OutputVariable)))
+								{
+									OutVertexData.SetVertex(InputVertex.VertexName, *VariableRef);
+								}
+							}
+						}
+					}
+				}
+			}
+		}
+	}
+
+	void CreateDefaultsAndVariables(const FOperatorSettings& InOperatorSettings, FInputVertexInterfaceData& OutVertexData)
+	{
+		Frontend::CreateDefaults(InOperatorSettings, OutVertexData);
+		CreateVariables(InOperatorSettings, OutVertexData);
+	}
+}
+
+#endif // WITH_EDITORONLY_DATA
+
+#endif //WITH_DEV_AUTOMATION_TESTS
+
diff --git a/Engine/Plugins/Runtime/Metasound/Source/MetasoundEngineTest/Private/EngineTestMetaSoundAutomatedNodeTest.h b/Engine/Plugins/Runtime/Metasound/Source/MetasoundEngineTest/Private/EngineTestMetaSoundAutomatedNodeTest.h
new file mode 100644
--- /dev/null
+++ b/Engine/Plugins/Runtime/Metasound/Source/MetasoundEngineTest/Private/EngineTestMetaSoundAutomatedNodeTest.h
@@ -0,0 +1,159 @@
+// Copyright Epic Games, Inc. All Rights Reserved.
+
+#pragma once
+
+#include "Containers/Array.h"
+#include "Containers/Map.h"
+#include "Containers/UnrealString.h"
+#include "CoreFwd.h"
+#include "MetasoundDataReference.h"
+#include "MetasoundLog.h"
+#include "MetasoundOperatorSettings.h"
+#include "MetasoundVertex.h"
+#include "MetasoundVertexData.h"
+#include "Templates/SharedPointer.h"
+#include "Templates/UniquePtr.h"
+
+#if WITH_DEV_AUTOMATION_TESTS
+
+#if WITH_EDITORONLY_DATA
+
+// Forward declare
+namespace Audio
+{
+	class FMixerDevice;
+}
+
+namespace Metasound
+{
+	class FMetasoundEnvironment;
+	class INode;
+	class FOutputVertexInterfaceData;
+	class FInputVertexInterfaceData;
+
+	namespace Frontend
+	{
+		typedef FString FNodeRegistryKey;
+	}
+}
+
+namespace Metasound::EngineTest 
+{
+	// Return audio mixer device if one is available
+	Audio::FMixerDevice* GetAudioMixerDevice();
+
+	// Create an example environment that generally exists for a UMetaSoundSoruce
+	FMetasoundEnvironment GetSourceEnvironmentForTest();
+
+	// Return string for printing test errors
+	FString GetPrettyName(const Frontend::FNodeRegistryKey& InRegistryKey);
+
+	// Interface for data reference analyzers.
+	struct IDataReferenceAnalyzer
+	{
+		virtual ~IDataReferenceAnalyzer() = default;
+		virtual FAnyDataReference Copy(const FAnyDataReference& InDataRef) const = 0;
+		virtual bool IsEqual(const FAnyDataReference& InLHS, const FAnyDataReference& InRHS) const = 0;
+		virtual bool IsValid(const FAnyDataReference& InDataRef) const = 0;
+		virtual FString ToString(const FAnyDataReference& InDataRef) const = 0;
+	};
+
+	// Convenience class for setting node input data reference values to default, min, max or random values. 
+	struct FOutputVertexDataTestController
+	{
+		struct FAnalyzableOutput
+		{
+			FAnyDataReference CapturedValue;
+			FAnyDataReference DataReference;
+			FVertexName VertexName;
+			TSharedPtr<const IDataReferenceAnalyzer> DataReferenceAnalyzer;
+
+			void CaptureValue();
+
+			bool IsDataReferenceValid() const;
+
+			bool IsDataReferenceEqualToCapturedValue() const;
+
+			FString DataReferenceToString() const;
+
+			FString CapturedValueToString() const;
+		};
+
+		FOutputVertexDataTestController( const FOutputVertexInterface& InOutputInterface, const FOutputVertexInterfaceData& InOutputData);
+
+		int32 GetNumAnalyzableOutputs() const;
+
+		bool AreAllAnalyzableOutputsValid() const;
+
+		void CaptureCurrentOutputValues();
+
+		bool AreAllOutputValuesEqualToCapturedValues() const;
+
+	private:
+
+		TArray<FAnalyzableOutput> AnalyzableOutputs;
+	};
+
+	// Interface for mutating data references
+	struct IDataReferenceMutator
+	{
+		virtual ~IDataReferenceMutator() = default;
+		virtual void SetDefault(const FOperatorSettings& InSettings, const FAnyDataReference& InDataRef) const = 0;
+		virtual void SetMax(const FOperatorSettings& InSettings, const FAnyDataReference& InDataRef) const = 0;
+		virtual void SetMin(const FOperatorSettings& InSettings, const FAnyDataReference& InDataRef) const = 0; 
+		virtual void SetRandom(const FOperatorSettings& InSettings, const FAnyDataReference& InDataRef) const = 0;
+		virtual FAnyDataReference Copy(const FAnyDataReference& InDataRef) const = 0;
+		virtual void SetValue(const FAnyDataReference& InSrcDataRef, const FAnyDataReference& InDstDataRef) const = 0;
+		virtual FString ToString(const FAnyDataReference& InDataRef) const = 0;
+	};
+
+	// Captures current data references on interface
+	using FInterfaceState = TSortedVertexNameMap<FAnyDataReference>;
+
+	// Convenience class for setting node input data reference values to default, min, max or random values. 
+	struct FInputVertexDataTestController
+	{
+		struct FMutableInput
+		{
+			FAnyDataReference DataReference;
+			FVertexName VertexName;
+			TSharedPtr<const IDataReferenceMutator> DataReferenceMutator;
+		};
+
+		FInputVertexDataTestController(const FOperatorSettings& InSettings, const FInputVertexInterface& InInputInterface, const FInputVertexInterfaceData& InInputData);
+
+		int32 GetNumMutableInputs() const;
+		FInterfaceState GetInterfaceState() const;
+
+		void SetMutableInputsToState(const FInterfaceState& InState);
+		void SetMutableInputsToMin();
+		void SetMutableInputsToMax();
+		void SetMutableInputsToDefault();
+		void SetMutableInputsToRandom();
+		TArray<FString> GetInputValueStrings() const;
+
+	private:
+
+		FOperatorSettings Settings;
+		TArray<FMutableInput> MutableInputs;
+	};
+
+
+	// Create a node from a node registry key
+	TUniquePtr<INode> CreateNodeFromRegistry(const Frontend::FNodeRegistryKey& InNodeRegistryKey);
+
+	// Return an array of all currently registered MetaSound nodes. 
+	void GetAllRegisteredNodes(TArray<FString>& OutBeautifiedNames, TArray<FString>& OutNodeRegistryKeys);
+
+	// Create any variables that exist on the vertex interface.
+	void CreateVariables(const FOperatorSettings& InOperatorSettings, FInputVertexInterfaceData& OutVertexData);
+
+	// Create default references and variables for the vertex interface data.
+	void CreateDefaultsAndVariables(const FOperatorSettings& InOperatorSettings, FInputVertexInterfaceData& OutVertexData);
+
+} // namespace Metasound::EngineTest
+
+#endif // WITH_EDITORONLY_DATA 
+
+#endif // WITH_DEV_AUTOMATION_TESTS
+
diff --git a/Engine/Plugins/Runtime/Metasound/Source/MetasoundEngineTest/Private/EngineTestMetaSoundAutomatedNodeTest_Bind.cpp b/Engine/Plugins/Runtime/Metasound/Source/MetasoundEngineTest/Private/EngineTestMetaSoundAutomatedNodeTest_Bind.cpp
new file mode 100644
--- /dev/null
+++ b/Engine/Plugins/Runtime/Metasound/Source/MetasoundEngineTest/Private/EngineTestMetaSoundAutomatedNodeTest_Bind.cpp
@@ -0,0 +1,176 @@
+// Copyright Epic Games, Inc. All Rights Reserved.
+
+#include "Algo/AllOf.h"
+#include "Containers/Array.h"
+#include "Containers/UnrealString.h"
+#include "EngineTestMetaSoundAutomatedNodeTest.h"
+#include "Misc/AutomationTest.h"
+#include "MetasoundBuilderInterface.h"
+#include "MetasoundFrontendDataTypeRegistry.h"
+#include "MetasoundOperatorInterface.h"
+#include "MetasoundLog.h"
+#include "MetasoundVertex.h"
+#include "MetasoundVertexData.h"
+#include "Templates/UniquePtr.h"
+#include "Tests/AutomationCommon.h"
+
+#if WITH_DEV_AUTOMATION_TESTS
+
+#if WITH_EDITORONLY_DATA
+
+IMPLEMENT_COMPLEX_AUTOMATION_TEST(FMetasoundAutomatedNodeTest_Bind, "Audio.Metasound.AutomatedNodeTest.Bind", EAutomationTestFlags::EditorContext | EAutomationTestFlags::StressFilter)
+void FMetasoundAutomatedNodeTest_Bind::GetTests(TArray<FString>& OutBeautifiedNames, TArray<FString>& OutTestCommands) const
+{
+	using namespace Metasound::EngineTest;
+	GetAllRegisteredNodes(OutBeautifiedNames, OutTestCommands);
+
+	UE_LOG(LogMetaSound, Verbose, TEXT("Found %d metasound nodes to test"), OutTestCommands.Num());
+}
+
+bool FMetasoundAutomatedNodeTest_Bind::RunTest(const FString& InRegistryKey)
+{
+	using namespace Metasound;
+	using namespace Metasound::EngineTest;
+
+	static const FOperatorSettings OperatorSettings{48000  /* samplerate */, 100.f /* block rate */};
+	static const FMetasoundEnvironment SourceEnvironment = GetSourceEnvironmentForTest();
+
+	TUniquePtr<INode> Node = CreateNodeFromRegistry(InRegistryKey);
+	if (!Node.IsValid())
+	{
+		AddError(FString::Printf(TEXT("Failed to create node %s from registry"), *InRegistryKey));
+		return false;
+	}
+
+	// Populate inputs to node. 
+	const FInputVertexInterface InputInterface = Node->GetVertexInterface().GetInputInterface();
+	FInputVertexInterfaceData BuildInputVertexData(InputInterface);
+	CreateDefaultsAndVariables(OperatorSettings, BuildInputVertexData);
+
+	TArray<FVertexDataState> BuildInputVertexDataState;
+	GetVertexInterfaceDataState(BuildInputVertexData, BuildInputVertexDataState);
+
+	// Create operator 
+	FBuildOperatorParams BuildParams
+	{
+		*Node,
+		OperatorSettings,
+		BuildInputVertexData,
+		SourceEnvironment
+	};
+
+	// Create the operator from the node factory
+	FBuildResults BuildResults;
+	TUniquePtr<IOperator> Operator = Node->GetDefaultOperatorFactory()->CreateOperator(BuildParams, BuildResults);
+	if (!Operator.IsValid())
+	{
+		AddError(FString::Printf(TEXT("Failed to create operator from node %s - %s."), *InRegistryKey, *GetPrettyName(InRegistryKey)));
+	}
+
+	// Query the data on the operator
+	FVertexInterfaceData AfterBuildVertexData;
+	Operator->BindInputs(AfterBuildVertexData.GetInputs());
+	Operator->BindOutputs(AfterBuildVertexData.GetOutputs());
+
+	// Check that build and queried data is the same.
+	for (const FInputDataVertex& InputVertex : InputInterface)
+	{
+		const bool bIsVertexBoundDuringBuild = BuildInputVertexData.IsVertexBound(InputVertex.VertexName);
+		const bool bIsVertexBoundAfterBuild = AfterBuildVertexData.GetInputs().IsVertexBound(InputVertex.VertexName);
+		
+		if (bIsVertexBoundDuringBuild  && !bIsVertexBoundAfterBuild)
+		{
+			// This means that a data reference was passed into the CreateOperator(...)
+			// function, but when IOperator::BindInputs(...) is called, the vertex 
+			// is not bound with any data. If data is sent to the CreateOperator(...)
+			// InputData, it should be reflected in the inputs of the IOperator through
+			// the BindInputs(...) method.
+			AddError(FString::Printf(TEXT("Input vertex (%s) of node (%s) does not bind supplied input data. Check that CreateOperator(...) and BindInputs(...) are correctly implemented for the node's operator."), *InputVertex.VertexName.ToString(), *GetPrettyName(InRegistryKey)));
+		}
+		else if (bIsVertexBoundDuringBuild && bIsVertexBoundAfterBuild)
+		{
+			// Input value vertices should bind to equal values. In order to detect
+			// correct functionality, we would need to check whether the values bound
+			// to the vertices are equal. Unfortunately we don't support a method
+			// for doing equality comparison based upon bound vertex data. 
+			if (InputVertex.AccessType != EVertexAccessType::Value)
+			{
+				const FAnyDataReference* BuildRef = BuildInputVertexData.FindDataReference(InputVertex.VertexName);
+				check(BuildRef);
+				const FDataReferenceID BuildDataRefID = GetDataReferenceID(*BuildRef);
+
+				const FAnyDataReference* AfterBuildRef = AfterBuildVertexData.GetInputs().FindDataReference(InputVertex.VertexName);
+				check(AfterBuildRef);
+				const FDataReferenceID AfterBuildDataRefID = GetDataReferenceID(*AfterBuildRef);
+
+				if (BuildDataRefID != AfterBuildDataRefID)
+				{
+					// This means that a data reference was passed into the CreateOperator(...)
+					// function, but when IOperator::BindInputs(...) is called, the vertex 
+					// is bound to different data. If data is sent to the CreateOperator(...)
+					// InputData, it should be reflected in the inputs of the IOperator through
+					// the BindInputs(...) method and be unchanged in this scenario.
+					AddError(FString::Printf(TEXT("Input vertex (%s) of node (%s) binds different data than the supplied at input data. Check that CreateOperator(...) and BindInputs(...) are correctly implemented for the node's operator."), *InputVertex.VertexName.ToString(), *GetPrettyName(InRegistryKey)));
+				}
+			}
+		}
+	}
+
+	// Cache current state of IOperator bindings
+	TArray<FVertexDataState> AfterBuildInputVertexDataState;
+	GetVertexInterfaceDataState(AfterBuildVertexData.GetInputs(), AfterBuildInputVertexDataState);
+
+	// Create new input data to override existing bindings.
+	FInputVertexInterfaceData UpdateToInputVertexData(Node->GetVertexInterface().GetInputInterface());
+	Frontend::CreateDefaults(OperatorSettings, UpdateToInputVertexData); // Note: Do not rebind new TVariables. This is not allowed. 
+
+	// Compare what is currently bound to what will be overridden
+	TSortedVertexNameMap<FAnyDataReference> ExpectedUpdates;
+	CompareVertexInterfaceDataToPriorState(UpdateToInputVertexData, AfterBuildInputVertexDataState, ExpectedUpdates);
+
+	// Perform the rebind which will override the current inputs
+	Operator->BindInputs(UpdateToInputVertexData);
+	
+	// Query the operator to find what the current inputs are.
+	FVertexInterfaceData NewVertexData;
+	Operator->BindInputs(NewVertexData.GetInputs());
+	Operator->BindOutputs(NewVertexData.GetOutputs());
+
+	// Compare the prior queried inputs to the new queried inputs to see what data
+	// references were actually updated.
+	TSortedVertexNameMap<FAnyDataReference> ActualUpdates;
+	CompareVertexInterfaceDataToPriorState(NewVertexData.GetInputs(), AfterBuildInputVertexDataState, ActualUpdates);
+
+	// Check that the expected updates were reflected in the IOperator::BindInputs call.
+	for (const FInputDataVertex& InputVertex : InputInterface)
+	{
+		const bool bIsUpdateExpected = ExpectedUpdates.Contains(InputVertex.VertexName);
+		const bool bIsUpdateActualized = ActualUpdates.Contains(InputVertex.VertexName);
+
+		if (bIsUpdateExpected && !bIsUpdateActualized)
+		{
+			// This means that the IOperator was given new data references during
+			// the call to IOperator::BindInputs(...), but there was no data bound
+			// on subsequent calls to IOperator::BindInputs(...).  If new data is 
+			// bound to the operator inputs, it should reflect that new data on subsequent
+			// calls to IOperator::BindInputs(...)
+			AddWarning(FString::Printf(TEXT("Input vertex (%s) of node (%s) does not reflect updated data reference when rebound. This node will not operator as expected in a dynamic graph. Check that CreateOperator(...) and BindInputs(...) are correctly implemented for the node's operator."), *InputVertex.VertexName.ToString(), *GetPrettyName(InRegistryKey)));
+		}
+		else if (!bIsUpdateExpected && bIsUpdateActualized)
+		{
+			// This means that the IOperator was given new data references during
+			// the call to IOperator::BindInputs(...), but there was no data bound
+			// on subsequent calls to IOperator::BindInputs(...).  If new data is 
+			// bound to the operator inputs, it should reflect that new data on subsequent
+			// calls to IOperator::BindInputs(...)
+			AddWarning(FString::Printf(TEXT("Input vertex (%s) of node (%s) reflects unintended updated of data reference when rebound. This node will not operator as expected in a dynamic graph. Check that CreateOperator(...) and BindInputs(...) are correctly implemented for the node's operator."), *InputVertex.VertexName.ToString(), *GetPrettyName(InRegistryKey)));
+		}
+	}
+
+	return true;
+}
+
+#endif // WITH_EDITORONLY_DATA
+
+#endif //WITH_DEV_AUTOMATION_TESTS
+
diff --git a/Engine/Plugins/Runtime/Metasound/Source/MetasoundEngineTest/Private/EngineTestMetaSoundAutomatedNodeTest_Reset.cpp b/Engine/Plugins/Runtime/Metasound/Source/MetasoundEngineTest/Private/EngineTestMetaSoundAutomatedNodeTest_Reset.cpp
new file mode 100644
--- /dev/null
+++ b/Engine/Plugins/Runtime/Metasound/Source/MetasoundEngineTest/Private/EngineTestMetaSoundAutomatedNodeTest_Reset.cpp
@@ -0,0 +1,158 @@
+// Copyright Epic Games, Inc. All Rights Reserved.
+
+#include "Algo/AllOf.h"
+#include "Containers/Array.h"
+#include "Containers/UnrealString.h"
+#include "EngineTestMetaSoundAutomatedNodeTest.h"
+#include "Misc/AutomationTest.h"
+#include "MetasoundBuilderInterface.h"
+#include "MetasoundFrontendDataTypeRegistry.h"
+#include "MetasoundOperatorInterface.h"
+#include "MetasoundLog.h"
+#include "MetasoundVertexData.h"
+#include "Templates/UniquePtr.h"
+#include "Tests/AutomationCommon.h"
+
+#if WITH_DEV_AUTOMATION_TESTS
+
+#if WITH_EDITORONLY_DATA
+
+
+IMPLEMENT_COMPLEX_AUTOMATION_TEST(FMetasoundAutomatedNodeTest_Reset, "Audio.Metasound.AutomatedNodeTest.Reset", EAutomationTestFlags::EditorContext | EAutomationTestFlags::StressFilter)
+void FMetasoundAutomatedNodeTest_Reset::GetTests(TArray<FString>& OutBeautifiedNames, TArray<FString>& OutTestCommands) const
+{
+	using namespace Metasound::EngineTest;
+	GetAllRegisteredNodes(OutBeautifiedNames, OutTestCommands);
+
+	UE_LOG(LogMetaSound, Verbose, TEXT("Found %d metasound nodes to test"), OutTestCommands.Num());
+}
+
+bool FMetasoundAutomatedNodeTest_Reset::RunTest(const FString& InRegistryKey)
+{
+	using namespace Metasound;
+	using namespace Metasound::EngineTest;
+
+	static const FOperatorSettings OperatorSettings{48000  /* samplerate */, 100.f /* block rate */};
+	static const FMetasoundEnvironment SourceEnvironment = GetSourceEnvironmentForTest();
+
+	TUniquePtr<INode> Node = CreateNodeFromRegistry(InRegistryKey);
+	if (!Node.IsValid())
+	{
+		AddError(FString::Printf(TEXT("Failed to create node %s from registry"), *InRegistryKey));
+		return false;
+	}
+
+	// Populate inputs to node. 
+	FInputVertexInterface InputInterface = Node->GetVertexInterface().GetInputInterface();
+	FInputVertexInterfaceData InputVertexData(InputInterface);
+	CreateDefaultsAndVariables(OperatorSettings, InputVertexData);
+
+	FInputVertexDataTestController InputTester(OperatorSettings, InputInterface, InputVertexData);
+
+	// Create operator 
+	FBuildOperatorParams BuildParams
+	{
+		*Node,
+		OperatorSettings,
+		InputVertexData,
+		SourceEnvironment
+	};
+
+	IOperator::FResetParams ResetParams
+	{
+		OperatorSettings,
+		SourceEnvironment
+	};
+
+	// Convenience function for testing entire lifecycle of an individual operator 
+	// with a variety of inputs.
+	auto RunTestIteration = [&]()
+	{
+		// Create the operator from the node factory
+		FBuildResults BuildResults;
+		TUniquePtr<IOperator> Operator = Node->GetDefaultOperatorFactory()->CreateOperator(BuildParams, BuildResults);
+
+		if (!Operator.IsValid())
+		{
+			AddError(FString::Printf(TEXT("Failed to create operator from node %s - %s."), *InRegistryKey, *GetPrettyName(InRegistryKey)));
+		}
+
+		// Store a copy of the input values data values so the inputs
+		// can be reset to this value later during the test.
+		FInterfaceState InitialInputState = InputTester.GetInterfaceState();
+
+		// Bind to inputs and output data of operator.
+		FVertexInterface VertexInterface = Node->GetVertexInterface();
+		FVertexInterfaceData VertexInterfaceData{VertexInterface};
+		Operator->BindInputs(VertexInterfaceData.GetInputs());
+		Operator->BindOutputs(VertexInterfaceData.GetOutputs());
+
+		// Create output tester which will analyzer outputs of operator.
+		FOutputVertexDataTestController OutputTester{VertexInterface.GetOutputInterface(), VertexInterfaceData.GetOutputs()};
+		
+		// Capture current state of outputs so they can be referenced at a later time.
+		// The captured values are held within the OutputTester.
+		OutputTester.CaptureCurrentOutputValues();
+
+		IOperator::FExecuteFunction OpExecFunc = Operator->GetExecuteFunction();
+		IOperator::FResetFunction OpResetFunc = Operator->GetResetFunction();
+		// Test execute function with input variations
+		if (OpExecFunc)
+		{
+			OpExecFunc(Operator.Get());
+
+			if (InputTester.GetNumMutableInputs() > 0)
+			{
+				InputTester.SetMutableInputsToDefault();
+				OpExecFunc(Operator.Get());
+				InputTester.SetMutableInputsToMin();
+				OpExecFunc(Operator.Get());
+				InputTester.SetMutableInputsToMax();
+				OpExecFunc(Operator.Get());
+				InputTester.SetMutableInputsToRandom();
+				OpExecFunc(Operator.Get());
+			}
+		}
+
+		// Return inputs to initial state. 
+		InputTester.SetMutableInputsToState(InitialInputState);
+
+		if (OpResetFunc)
+		{
+			OpResetFunc(Operator.Get(), ResetParams);
+		}
+		else if (OpExecFunc)
+		{
+			AddError(FString::Printf(TEXT("Missing initialize function when execute function exists for node %s - %s"), *InRegistryKey, *GetPrettyName(InRegistryKey)));
+		}
+
+		// Check that after returning all inputs to their original state and calling
+		// reset on the operator, that all output values have returned to their initail state. 
+		if (!OutputTester.AreAllOutputValuesEqualToCapturedValues())
+		{
+			AddError(FString::Printf(TEXT("Reset function resulted in different starting conditions for node %s - %s"), *InRegistryKey, *GetPrettyName(InRegistryKey)));
+		}
+	};
+
+	// Test entire operator lifecycle with different starting conditions if
+	// any of the inputs are mutable
+	InputTester.SetMutableInputsToDefault();
+	RunTestIteration();
+
+	if (InputTester.GetNumMutableInputs() > 0)
+	{
+		InputTester.SetMutableInputsToMin();
+		RunTestIteration();
+		InputTester.SetMutableInputsToMax();
+		RunTestIteration();
+		InputTester.SetMutableInputsToRandom();
+		RunTestIteration();
+	}
+
+	return true;
+}
+
+#endif // WITH_EDITORONLY_DATA
+
+#endif //WITH_DEV_AUTOMATION_TESTS
+
diff --git a/Engine/Plugins/Runtime/Metasound/Source/MetasoundEngineTest/Private/EngineTestMetaSoundAutomatedNodeTest_ValidOutputs.cpp b/Engine/Plugins/Runtime/Metasound/Source/MetasoundEngineTest/Private/EngineTestMetaSoundAutomatedNodeTest_ValidOutputs.cpp
new file mode 100644
--- /dev/null
+++ b/Engine/Plugins/Runtime/Metasound/Source/MetasoundEngineTest/Private/EngineTestMetaSoundAutomatedNodeTest_ValidOutputs.cpp
@@ -0,0 +1,167 @@
+// Copyright Epic Games, Inc. All Rights Reserved.
+
+#include "Algo/AllOf.h"
+#include "Containers/Array.h"
+#include "Containers/UnrealString.h"
+#include "EngineTestMetaSoundAutomatedNodeTest.h"
+#include "Misc/AutomationTest.h"
+#include "MetasoundBuilderInterface.h"
+#include "MetasoundFrontendDataTypeRegistry.h"
+#include "MetasoundOperatorInterface.h"
+#include "MetasoundLog.h"
+#include "MetasoundVertexData.h"
+#include "Templates/UniquePtr.h"
+#include "Tests/AutomationCommon.h"
+
+#if WITH_DEV_AUTOMATION_TESTS
+
+#if WITH_EDITORONLY_DATA
+
+
+IMPLEMENT_COMPLEX_AUTOMATION_TEST(FMetasoundAutomatedNodeTest_ValidOutputs, "Audio.Metasound.AutomatedNodeTest.ValidOutputs", EAutomationTestFlags::EditorContext | EAutomationTestFlags::StressFilter)
+void FMetasoundAutomatedNodeTest_ValidOutputs::GetTests(TArray<FString>& OutBeautifiedNames, TArray<FString>& OutTestCommands) const
+{
+	using namespace Metasound::EngineTest;
+	GetAllRegisteredNodes(OutBeautifiedNames, OutTestCommands);
+
+	UE_LOG(LogMetaSound, Verbose, TEXT("Found %d metasound nodes to test"), OutTestCommands.Num());
+}
+
+bool FMetasoundAutomatedNodeTest_ValidOutputs::RunTest(const FString& InRegistryKey)
+{
+	using namespace Metasound;
+	using namespace Metasound::EngineTest;
+
+	static const FOperatorSettings OperatorSettings{48000  /* samplerate */, 100.f /* block rate */};
+	static const FMetasoundEnvironment SourceEnvironment = GetSourceEnvironmentForTest();
+
+	TUniquePtr<INode> Node = CreateNodeFromRegistry(InRegistryKey);
+	if (!Node.IsValid())
+	{
+		AddError(FString::Printf(TEXT("Failed to create node %s from registry"), *InRegistryKey));
+		return false;
+	}
+
+	// Populate inputs to node. 
+	FInputVertexInterface InputInterface = Node->GetVertexInterface().GetInputInterface();
+	FInputVertexInterfaceData InputVertexData(InputInterface);
+	CreateDefaultsAndVariables(OperatorSettings, InputVertexData);
+
+	FInputVertexDataTestController InputTester(OperatorSettings, InputInterface, InputVertexData);
+
+	// Create operator 
+	FBuildOperatorParams BuildParams
+	{
+		*Node,
+		OperatorSettings,
+		InputVertexData,
+		SourceEnvironment
+	};
+
+	IOperator::FResetParams ResetParams
+	{
+		OperatorSettings,
+		SourceEnvironment
+	};
+
+	// Convenience function for testing entire lifecycle of an individual operator 
+	// with a variety of inputs.
+	auto RunTestIteration = [&]()
+	{
+		// Create the operator from the node factory
+		FBuildResults BuildResults;
+		TUniquePtr<IOperator> Operator = Node->GetDefaultOperatorFactory()->CreateOperator(BuildParams, BuildResults);
+
+		if (!Operator.IsValid())
+		{
+			AddError(FString::Printf(TEXT("Failed to create operator from node %s - %s."), *InRegistryKey, *GetPrettyName(InRegistryKey)));
+		}
+
+
+		// Store a copy of the input values data values so the inputs
+		// can be reset to this value later during the test.
+		FInterfaceState InitialInputState = InputTester.GetInterfaceState();
+
+		// Bind to inputs and output data of operator.
+		FVertexInterface VertexInterface = Node->GetVertexInterface();
+		FVertexInterfaceData VertexInterfaceData{VertexInterface};
+		Operator->BindInputs(VertexInterfaceData.GetInputs());
+		Operator->BindOutputs(VertexInterfaceData.GetOutputs());
+
+		// Create output tester which will analyzer outputs of operator.
+		FOutputVertexDataTestController OutputTester{VertexInterface.GetOutputInterface(), VertexInterfaceData.GetOutputs()};
+		// Convenience method for printing errors if the OutputTester finds an error.
+		auto CheckOutputValuesAreValid = [&]()
+		{
+			if (!OutputTester.AreAllAnalyzableOutputsValid())
+			{
+				AddError(FString::Printf(TEXT("Invalid output value encountered from node %s - %s."), *InRegistryKey, *GetPrettyName(InRegistryKey)));
+			}
+		};
+		// Initial values should all be valid.
+		CheckOutputValuesAreValid();
+		
+		// Capture current state of outputs so they can be referenced at a later time.
+		// The captured values are held within the OutputTester.
+		OutputTester.CaptureCurrentOutputValues();
+
+		IOperator::FExecuteFunction OpExecFunc = Operator->GetExecuteFunction();
+		IOperator::FResetFunction OpResetFunc = Operator->GetResetFunction();
+		// Test execute function with input variations
+		if (OpExecFunc)
+		{
+			OpExecFunc(Operator.Get());
+			CheckOutputValuesAreValid();
+
+			if (InputTester.GetNumMutableInputs() > 0)
+			{
+				InputTester.SetMutableInputsToDefault();
+				OpExecFunc(Operator.Get());
+				CheckOutputValuesAreValid();
+
+				InputTester.SetMutableInputsToMin();
+				OpExecFunc(Operator.Get());
+				CheckOutputValuesAreValid();
+
+				InputTester.SetMutableInputsToMax();
+				OpExecFunc(Operator.Get());
+				CheckOutputValuesAreValid();
+
+				InputTester.SetMutableInputsToRandom();
+				OpExecFunc(Operator.Get());
+				CheckOutputValuesAreValid();
+			}
+		}
+
+		// Return inputs to initial state. 
+		InputTester.SetMutableInputsToState(InitialInputState);
+
+		if (OpResetFunc)
+		{
+			OpResetFunc(Operator.Get(), ResetParams);
+			CheckOutputValuesAreValid();
+		}
+	};
+
+	// Test entire operator lifecycle with different starting conditions if
+	// any of the inputs are mutable
+	InputTester.SetMutableInputsToDefault();
+	RunTestIteration();
+
+	if (InputTester.GetNumMutableInputs() > 0)
+	{
+		InputTester.SetMutableInputsToMin();
+		RunTestIteration();
+		InputTester.SetMutableInputsToMax();
+		RunTestIteration();
+		InputTester.SetMutableInputsToRandom();
+		RunTestIteration();
+	}
+
+	return true;
+}
+
+#endif // WITH_EDITORONLY_DATA
+
+#endif //WITH_DEV_AUTOMATION_TESTS
+
diff --git a/Engine/Plugins/Runtime/Metasound/Source/MetasoundFrontend/Public/MetasoundArrayNodes.h b/Engine/Plugins/Runtime/Metasound/Source/MetasoundFrontend/Public/MetasoundArrayNodes.h
--- a/Engine/Plugins/Runtime/Metasound/Source/MetasoundFrontend/Public/MetasoundArrayNodes.h
+++ b/Engine/Plugins/Runtime/Metasound/Source/MetasoundFrontend/Public/MetasoundArrayNodes.h
@@ -140,24 +140,27 @@ namespace Metasound
 
 		virtual FDataReferenceCollection GetInputs() const override
 		{
-			using namespace ArrayNodeVertexNames;
-
-			FDataReferenceCollection Inputs;
-
-			Inputs.AddDataReadReference(METASOUND_GET_PARAM_NAME(InputArray), Array);
-
-			return Inputs;
+			checkNoEntry();
+			return {};
 		}
 
 		virtual FDataReferenceCollection GetOutputs() const override
 		{
-			using namespace ArrayNodeVertexNames;
+			checkNoEntry();
+			return {};
+		}
 
-			FDataReferenceCollection Outputs;
+		virtual void BindInputs(FInputVertexInterfaceData& InVertexData) override
+		{
+			using namespace ArrayNodeVertexNames;
+			InVertexData.BindReadVertex(METASOUND_GET_PARAM_NAME(InputArray), Array);
+		}
 
-			Outputs.AddDataReadReference(METASOUND_GET_PARAM_NAME(OutputNum), Num);
+		virtual void BindOutputs(FOutputVertexInterfaceData& InVertexData) override
+		{
+			using namespace ArrayNodeVertexNames;
 
-			return Outputs;
+			InVertexData.BindReadVertex(METASOUND_GET_PARAM_NAME(OutputNum), Num);
 		}
 
 		void Execute()
@@ -304,26 +307,30 @@ namespace Metasound
 
 		virtual FDataReferenceCollection GetInputs() const override
 		{
-			using namespace ArrayNodeVertexNames;
-
-			FDataReferenceCollection Inputs;
-
-			Inputs.AddDataReadReference(METASOUND_GET_PARAM_NAME(InputTriggerGet), Trigger);
-			Inputs.AddDataReadReference(METASOUND_GET_PARAM_NAME(InputArray), Array);
-			Inputs.AddDataReadReference(METASOUND_GET_PARAM_NAME(InputIndex), Index);
-
-			return Inputs;
+			checkNoEntry();
+			return {};
 		}
 
 		virtual FDataReferenceCollection GetOutputs() const override
+		{
+			checkNoEntry();
+			return {};
+		}
+
+		virtual void BindInputs(FInputVertexInterfaceData& InVertexData) override
 		{
 			using namespace ArrayNodeVertexNames;
 
-			FDataReferenceCollection Outputs;
+			InVertexData.BindReadVertex(METASOUND_GET_PARAM_NAME(InputTriggerGet), Trigger);
+			InVertexData.BindReadVertex(METASOUND_GET_PARAM_NAME(InputArray), Array);
+			InVertexData.BindReadVertex(METASOUND_GET_PARAM_NAME(InputIndex), Index);
+		}
 
-			Outputs.AddDataReadReference(METASOUND_GET_PARAM_NAME(OutputValue), Value);
+		virtual void BindOutputs(FOutputVertexInterfaceData& InVertexData) override
+		{
+			using namespace ArrayNodeVertexNames;
 
-			return Outputs;
+			InVertexData.BindReadVertex(METASOUND_GET_PARAM_NAME(OutputValue), Value);
 		}
 
 		void Execute()
@@ -499,25 +506,31 @@ namespace Metasound
 
 		virtual FDataReferenceCollection GetInputs() const override
 		{
-			using namespace ArrayNodeVertexNames;
-			FDataReferenceCollection Inputs;
-
-			Inputs.AddDataReadReference(METASOUND_GET_PARAM_NAME(InputTriggerSet), Trigger);
-			Inputs.AddDataReadReference(METASOUND_GET_PARAM_NAME(InputArray), InitArray);
-			Inputs.AddDataReadReference(METASOUND_GET_PARAM_NAME(InputIndex), Index);
-			Inputs.AddDataReadReference(METASOUND_GET_PARAM_NAME(InputValue), Value);
-
-			return Inputs;
+			checkNoEntry();
+			return {};
 		}
 
 		virtual FDataReferenceCollection GetOutputs() const override
+		{
+			checkNoEntry();
+			return {};
+		}
+
+		virtual void BindInputs(FInputVertexInterfaceData& InVertexData) override
 		{
 			using namespace ArrayNodeVertexNames;
-			FDataReferenceCollection Outputs;
 
-			Outputs.AddDataReadReference(METASOUND_GET_PARAM_NAME(OutputArraySet), Array);
+			InVertexData.BindReadVertex(METASOUND_GET_PARAM_NAME(InputTriggerSet), Trigger);
+			InVertexData.BindReadVertex(METASOUND_GET_PARAM_NAME(InputArray), InitArray);
+			InVertexData.BindReadVertex(METASOUND_GET_PARAM_NAME(InputIndex), Index);
+			InVertexData.BindReadVertex(METASOUND_GET_PARAM_NAME(InputValue), Value);
+		}
 
-			return Outputs;
+		virtual void BindOutputs(FOutputVertexInterfaceData& InVertexData) override
+		{
+			using namespace ArrayNodeVertexNames;
+
+			InVertexData.BindReadVertex(METASOUND_GET_PARAM_NAME(OutputArraySet), Array);
 		}
 
 		void Execute()
@@ -645,24 +658,30 @@ namespace Metasound
 
 		virtual FDataReferenceCollection GetInputs() const override
 		{
-			using namespace ArrayNodeVertexNames;
-			FDataReferenceCollection Inputs;
-
-			Inputs.AddDataReadReference(METASOUND_GET_PARAM_NAME(InputTriggerGet), Trigger);
-			Inputs.AddDataReadReference(METASOUND_GET_PARAM_NAME(InputLeftArray), LeftArray);
-			Inputs.AddDataReadReference(METASOUND_GET_PARAM_NAME(InputRightArray), RightArray);
-
-			return Inputs;
+			checkNoEntry();
+			return {};
 		}
 
 		virtual FDataReferenceCollection GetOutputs() const override
+		{
+			checkNoEntry();
+			return {};
+		}
+
+		virtual void BindInputs(FInputVertexInterfaceData& InVertexData) override
 		{
 			using namespace ArrayNodeVertexNames;
-			FDataReferenceCollection Outputs;
 
-			Outputs.AddDataReadReference(METASOUND_GET_PARAM_NAME(OutputArrayConcat), OutArray);
+			InVertexData.BindReadVertex(METASOUND_GET_PARAM_NAME(InputTriggerGet), Trigger);
+			InVertexData.BindReadVertex(METASOUND_GET_PARAM_NAME(InputLeftArray), LeftArray);
+			InVertexData.BindReadVertex(METASOUND_GET_PARAM_NAME(InputRightArray), RightArray);
+		}
 
-			return Outputs;
+		virtual void BindOutputs(FOutputVertexInterfaceData& InVertexData) override
+		{
+			using namespace ArrayNodeVertexNames;
+
+			InVertexData.BindReadVertex(METASOUND_GET_PARAM_NAME(OutputArrayConcat), OutArray);
 		}
 
 		void Execute()
@@ -779,27 +798,31 @@ namespace Metasound
 
 		virtual FDataReferenceCollection GetInputs() const override
 		{
-			using namespace ArrayNodeVertexNames;
-
-			FDataReferenceCollection Inputs;
-
-			Inputs.AddDataReadReference(METASOUND_GET_PARAM_NAME(InputTriggerGet), Trigger);
-			Inputs.AddDataReadReference(METASOUND_GET_PARAM_NAME(InputArray), InputArray);
-			Inputs.AddDataReadReference(METASOUND_GET_PARAM_NAME(InputStartIndex), StartIndex);
-			Inputs.AddDataReadReference(METASOUND_GET_PARAM_NAME(InputEndIndex), EndIndex);
-
-			return Inputs;
+			checkNoEntry();
+			return {};
 		}
 
 		virtual FDataReferenceCollection GetOutputs() const override
+		{
+			checkNoEntry();
+			return {};
+		}
+
+		virtual void BindInputs(FInputVertexInterfaceData& InVertexData) override
 		{
 			using namespace ArrayNodeVertexNames;
 
-			FDataReferenceCollection Outputs;
+			InVertexData.BindReadVertex(METASOUND_GET_PARAM_NAME(InputTriggerGet), Trigger);
+			InVertexData.BindReadVertex(METASOUND_GET_PARAM_NAME(InputArray), InputArray);
+			InVertexData.BindReadVertex(METASOUND_GET_PARAM_NAME(InputStartIndex), StartIndex);
+			InVertexData.BindReadVertex(METASOUND_GET_PARAM_NAME(InputEndIndex), EndIndex);
+		}
 
-			Outputs.AddDataReadReference(METASOUND_GET_PARAM_NAME(OutputArraySubset), OutputArray);
+		virtual void BindOutputs(FOutputVertexInterfaceData& InVertexData) override
+		{
+			using namespace ArrayNodeVertexNames;
 
-			return Outputs;
+			InVertexData.BindReadVertex(METASOUND_GET_PARAM_NAME(OutputArraySubset), OutputArray);
 		}
 
 		void Execute()
diff --git a/Engine/Plugins/Runtime/Metasound/Source/MetasoundFrontend/Public/MetasoundAutoConverterNode.h b/Engine/Plugins/Runtime/Metasound/Source/MetasoundFrontend/Public/MetasoundAutoConverterNode.h
--- a/Engine/Plugins/Runtime/Metasound/Source/MetasoundFrontend/Public/MetasoundAutoConverterNode.h
+++ b/Engine/Plugins/Runtime/Metasound/Source/MetasoundFrontend/Public/MetasoundAutoConverterNode.h
@@ -136,16 +136,24 @@ namespace Metasound
 
 			virtual FDataReferenceCollection GetInputs() const override
 			{
-				FDataReferenceCollection Inputs;
-				Inputs.AddDataReadReference<FromDataType>(GetInputName(), FromData);
-				return Inputs;
+				checkNoEntry();
+				return {};
 			}
 
 			virtual FDataReferenceCollection GetOutputs() const override
 			{
-				FDataReferenceCollection Outputs;
-				Outputs.AddDataReadReference<ToDataType>(GetOutputName(), ToData);
-				return Outputs;
+				checkNoEntry();
+				return {};
+			}
+
+			virtual void BindInputs(FInputVertexInterfaceData& InVertexData) override
+			{
+				InVertexData.BindReadVertex(GetInputName(), FromData);
+			}
+
+			virtual void BindOutputs(FOutputVertexInterfaceData& InVertexData) override
+			{
+				InVertexData.BindReadVertex(GetOutputName(), ToData);
 			}
 
 			void Execute()
diff --git a/Engine/Plugins/Runtime/Metasound/Source/MetasoundStandardNodes/Private/MetasoundMidiNoteQuantizerNode.cpp b/Engine/Plugins/Runtime/Metasound/Source/MetasoundStandardNodes/Private/MetasoundMidiNoteQuantizerNode.cpp
--- a/Engine/Plugins/Runtime/Metasound/Source/MetasoundStandardNodes/Private/MetasoundMidiNoteQuantizerNode.cpp
+++ b/Engine/Plugins/Runtime/Metasound/Source/MetasoundStandardNodes/Private/MetasoundMidiNoteQuantizerNode.cpp
@@ -216,8 +216,15 @@ namespace Metasound
 		PreviousScale = *Scale;
 		PreviousScaleRange = FMath::Max(*ScaleRange, 1.0f);
 
-		PreviousScale.Sort();
-		PreviousNoteOut = Audio::FMidiNoteQuantizer::QuantizeMidiNote(PreviousNoteIn, PreviousRoot, PreviousScale, PreviousScaleRange);
+		if (PreviousScale.IsEmpty())
+		{
+			PreviousNoteOut = PreviousNoteIn;
+		}
+		else
+		{
+			PreviousScale.Sort();
+			PreviousNoteOut = Audio::FMidiNoteQuantizer::QuantizeMidiNote(PreviousNoteIn, PreviousRoot, PreviousScale, PreviousScaleRange);
+		}
 		*MidiNoteOut = PreviousNoteOut;
 	}
 
diff --git a/Engine/Plugins/Runtime/Metasound/Source/MetasoundStandardNodes/Private/MetasoundPerlinNoiseNode.cpp b/Engine/Plugins/Runtime/Metasound/Source/MetasoundStandardNodes/Private/MetasoundPerlinNoiseNode.cpp
--- a/Engine/Plugins/Runtime/Metasound/Source/MetasoundStandardNodes/Private/MetasoundPerlinNoiseNode.cpp
+++ b/Engine/Plugins/Runtime/Metasound/Source/MetasoundStandardNodes/Private/MetasoundPerlinNoiseNode.cpp
@@ -58,7 +58,7 @@ namespace Metasound
 			{
 				FInputVertexInterface
 				{
-					TInputDataVertex<TDataClass>(METASOUND_GET_PARAM_NAME_AND_METADATA(XPin), 0),
+					TInputDataVertex<TDataClass>(METASOUND_GET_PARAM_NAME_AND_METADATA(XPin)),
 					TInputDataVertex<int32>(METASOUND_GET_PARAM_NAME_AND_METADATA(OctavesPin), 1),
 					TInputConstructorVertex<float>(METASOUND_GET_PARAM_NAME_AND_METADATA(OffsetPin), -1.0f),
 					TInputDataVertex<float>(METASOUND_GET_PARAM_NAME_AND_METADATA(MinValuePin), -1.0f),
