CIS Fix for IsAudioMixerEnabled in metasound engine test

diff --git a/Engine/Plugins/Runtime/Metasound/Source/MetasoundEngineTest/Private/EngineTestMetaSoundAutomatedNodeTest.cpp b/Engine/Plugins/Runtime/Metasound/Source/MetasoundEngineTest/Private/EngineTestMetaSoundAutomatedNodeTest.cpp
--- a/Engine/Plugins/Runtime/Metasound/Source/MetasoundEngineTest/Private/EngineTestMetaSoundAutomatedNodeTest.cpp
+++ b/Engine/Plugins/Runtime/Metasound/Source/MetasoundEngineTest/Private/EngineTestMetaSoundAutomatedNodeTest.cpp
@@ -36,10 +36,7 @@ namespace Metasound::EngineTest{
 		{
 			if (FAudioDevice* AudioDevice = DeviceManager->GetMainAudioDeviceRaw())
 			{
-				if (AudioDevice->IsAudioMixerEnabled())
-				{
-					return static_cast<Audio::FMixerDevice*>(AudioDevice);
-				}
+				return static_cast<Audio::FMixerDevice*>(AudioDevice);
 			}
 		}
 		return nullptr;
