Fix non-matching PRAGMA_ENABLE_DEPRECATION_WARNINGS pair.

diff --git a/Engine/Plugins/Runtime/Metasound/Source/MetasoundEngine/Private/MetasoundWaveProxyReader.cpp b/Engine/Plugins/Runtime/Metasound/Source/MetasoundEngine/Private/MetasoundWaveProxyReader.cpp
--- a/Engine/Plugins/Runtime/Metasound/Source/MetasoundEngine/Private/MetasoundWaveProxyReader.cpp
+++ b/Engine/Plugins/Runtime/Metasound/Source/MetasoundEngine/Private/MetasoundWaveProxyReader.cpp
@@ -30,6 +30,8 @@ FAutoConsoleVariableRef CVarMetaSoundWaveProxyReaderSimulateSeekOnNonSeekable(
 
 namespace Metasound
 {
+	PRAGMA_DISABLE_DEPRECATION_WARNINGS
+
 	namespace MetasoundWaveProxyReaderPrivate
 	{
 		/** Construct a FDecoderOutput
@@ -546,4 +548,6 @@ namespace Metasound
 			LoopEndFrameIndex = NumFramesInWave;
 		}
 	}
+
+	PRAGMA_ENABLE_DEPRECATION_WARNINGS
 }
diff --git a/Engine/Source/Runtime/AudioCodecEngine/Public/DecoderInputFactory.h b/Engine/Source/Runtime/AudioCodecEngine/Public/DecoderInputFactory.h
--- a/Engine/Source/Runtime/AudioCodecEngine/Public/DecoderInputFactory.h
+++ b/Engine/Source/Runtime/AudioCodecEngine/Public/DecoderInputFactory.h
@@ -20,5 +20,5 @@ namespace Audio
 	// Just loose for now.
 	PRAGMA_DISABLE_DEPRECATION_WARNINGS
 	AUDIOCODECENGINE_API TUniquePtr<IDecoderInput> CreateBackCompatDecoderInput(FName InOldFormatName, const FSoundWaveProxyPtr& InSoundWave);
-	PRAGMA_DISABLE_DEPRECATION_WARNINGS
+	PRAGMA_ENABLE_DEPRECATION_WARNINGS
 }
