Add ability to define a custom name for connected input/output nodes

cleanup some java-docs

diff --git a/Engine/Plugins/Runtime/Metasound/Source/MetasoundEngineTest/Private/NodeTestGraphBuilder.cpp b/Engine/Plugins/Runtime/Metasound/Source/MetasoundEngineTest/Private/NodeTestGraphBuilder.cpp
--- a/Engine/Plugins/Runtime/Metasound/Source/MetasoundEngineTest/Private/NodeTestGraphBuilder.cpp
+++ b/Engine/Plugins/Runtime/Metasound/Source/MetasoundEngineTest/Private/NodeTestGraphBuilder.cpp
@@ -234,17 +234,19 @@ namespace Metasound::Test
 		return true;
 	}
 
-	FNodeHandle FNodeTestGraphBuilder::AddAndConnectDataReferenceInput(const FNodeHandle& NodeToConnect, const FName& InputName, const FName& TypeName) const
+	FNodeHandle FNodeTestGraphBuilder::AddAndConnectDataReferenceInput(const FNodeHandle& NodeToConnect, const FName& InputName, const FName& TypeName, const FName& NodeName) const
 	{
-		FNodeHandle InputNode = AddInput(InputName, TypeName);
-		check(ConnectNodes(InputNode, NodeToConnect, InputName));
+		FName NameToUse = NodeName.IsNone() ? InputName : NodeName;
+		FNodeHandle InputNode = AddInput(NameToUse, TypeName);
+		ConnectNodes(InputNode, NameToUse, NodeToConnect, InputName);
 		return InputNode;
 	}
 
-	FNodeHandle FNodeTestGraphBuilder::AddAndConnectDataReferenceOutput(const FNodeHandle& NodeToConnect, const FName& OutputName, const FName& TypeName)
+	FNodeHandle FNodeTestGraphBuilder::AddAndConnectDataReferenceOutput(const FNodeHandle& NodeToConnect, const FName& OutputName, const FName& TypeName, const FName& NodeName)
 	{
-		FNodeHandle OutputNode = AddOutput(OutputName, TypeName);
-		check(ConnectNodes(NodeToConnect, OutputNode, OutputName));
+		FName NameToUse = NodeName.IsNone() ? OutputName : NodeName;
+		FNodeHandle OutputNode = AddOutput(NameToUse, TypeName);
+		ConnectNodes(NodeToConnect, OutputName, OutputNode, NameToUse);
 		return OutputNode;
 	}
 
diff --git a/Engine/Plugins/Runtime/Metasound/Source/MetasoundEngineTest/Public/NodeTestGraphBuilder.h b/Engine/Plugins/Runtime/Metasound/Source/MetasoundEngineTest/Public/NodeTestGraphBuilder.h
--- a/Engine/Plugins/Runtime/Metasound/Source/MetasoundEngineTest/Public/NodeTestGraphBuilder.h
+++ b/Engine/Plugins/Runtime/Metasound/Source/MetasoundEngineTest/Public/NodeTestGraphBuilder.h
@@ -43,9 +43,9 @@ namespace Metasound::Test
 		/**
 		* Add an input node from the graph by name, type, and access type
 		*
-		* @param InputName		The name of the output node to get
+		* @param InputName		The name of the input node to add
 		* @param TypeName		The type of the node to add (use: GetMetasoundDataTypeName<DataType>());
-		* @param AccessType		Whether the node should be a Value, Referenc, or Unset type
+		* @param AccessType		Whether the node should be a Value, Reference, or Unset type
 		* @return				NodeHandle of the added node, can be invalid if failed to add
 		*/
 		Frontend::FNodeHandle AddInput(
@@ -56,7 +56,7 @@ namespace Metasound::Test
 		/**
 		* Add a constructor input node from the graph by name and type
 		*
-		* @param InputName		The name of the output node to get
+		* @param InputName		The name of the input node to add
 		* @param Value			The value for the node to hold. The type of the value will determine the type for the node.
 		* @return				NodeHandle of the added node, can be invalid if failed to add
 		*/
@@ -77,25 +77,27 @@ namespace Metasound::Test
 		/**
 		* Add a data reference output node from the graph by name and type
 		*
-		* @param InputName		The name of the output node to get
+		* @param InputName		The name of the output node to add
 		* @param TypeName		The type for the node to reference
 		* @return				NodeHandle of the added node, can be invalid if failed to add
 		*/
 		Frontend::FNodeHandle AddOutput(const FName& OutputName, const FName& TypeName);
 
 		/** 
-		* Adds an constructor input node to the graph and connects it as an input to the NodeToConnect
+		* Adds a constructor input node to the graph and connects it as an input to the NodeToConnect
 		* 
 		* @param NodeToConnect	The node to connect the newly added constructor input node to
 		* @param InputName		The name of the input pin to connect the node to, and the name to use for the new node.
 		* @param Value			The value for the node to hold. The type of the value will determine the type for the node.
+		* @param NodeName		Optional name to use for the created constructor input node
 		* @return				NodeHandle of the added node, can be invalid if failed to add
 		*/
 		template <typename DataType>
-		Frontend::FNodeHandle AddAndConnectConstructorInput(const Frontend::FNodeHandle& NodeToConnect, const FName& InputName, const DataType& Value) const
+		Frontend::FNodeHandle AddAndConnectConstructorInput(const Frontend::FNodeHandle& NodeToConnect, const FName& InputName, const DataType& Value, const FName& NodeName = NAME_None) const
 		{
-			Frontend::FNodeHandle InputNode = AddConstructorInput<DataType>(InputName, Value);
-			check(ConnectNodes(InputNode, NodeToConnect, InputName));
+			FName NameToUse = NodeName.IsNone() ? InputName : NodeName;
+			Frontend::FNodeHandle InputNode = AddConstructorInput<DataType>(NameToUse, Value);
+			ConnectNodes(InputNode, NameToUse, NodeToConnect, InputName);
 			return InputNode;
 		}
 
@@ -113,18 +115,20 @@ namespace Metasound::Test
 		* @param NodeToConnect	The node to connect the newly added constructor input node to
 		* @param InputName		The name of the input pin to connect the node to, and the name to use for the new node.
 		* @param TypeName		The type name of the data reference
+		* @param NodeName		Optional name to use for the created data reference input node
 		* @return				NodeHandle of the added node, can be invalid if failed to add
 		*/
-		Frontend::FNodeHandle AddAndConnectDataReferenceInput(const Frontend::FNodeHandle& NodeToConnect, const FName& InputName, const FName& TypeName) const;
+		Frontend::FNodeHandle AddAndConnectDataReferenceInput(const Frontend::FNodeHandle& NodeToConnect, const FName& InputName, const FName& TypeName, const FName& NodeName = NAME_None) const;
 
 		/** Adds a data reference output node with OutputName and connects it to the node's output with the same name
 		* 	
 		* @param NodeToConnect	The node to connect the newly added constructor input node to
 		* @param InputName		The name of the output pin to connect the node to, and the name to use for the new node.
 		* @param TypeName		The type name of the data reference
+		* @param NodeName		Optoinal name to use for the created data reference output node
 		* @return				NodeHandle of the added node, can be invalid if failed to add
 		*/
-		Frontend::FNodeHandle AddAndConnectDataReferenceOutput(const Frontend::FNodeHandle& NodeToConnect, const FName& OutputName, const FName& TypeName);
+		Frontend::FNodeHandle AddAndConnectDataReferenceOutput(const Frontend::FNodeHandle& NodeToConnect, const FName& OutputName, const FName& TypeName, const FName& NodeName = NAME_None);
 
 		/** 
 		* Call last to create a generator from the graph constructed using the other method. 
