Nuke incorrect, redundant comment for self-documenting call
#rnx

diff --git a/Engine/Plugins/Runtime/Metasound/Source/MetasoundEngine/Private/Interfaces/MetasoundDeprecatedInterfaces.h b/Engine/Plugins/Runtime/Metasound/Source/MetasoundEngine/Private/Interfaces/MetasoundDeprecatedInterfaces.h
--- a/Engine/Plugins/Runtime/Metasound/Source/MetasoundEngine/Private/Interfaces/MetasoundDeprecatedInterfaces.h
+++ b/Engine/Plugins/Runtime/Metasound/Source/MetasoundEngine/Private/Interfaces/MetasoundDeprecatedInterfaces.h
@@ -92,7 +92,6 @@ namespace Metasound::Engine
 		const FMetasoundFrontendVersion& GetVersion();
 		FMetasoundFrontendInterface GetInterface();
 
-		// Update from MetasoundOutputFormatStereoV1_0 to MetasoundOutputFormatStereoV1_1
 		class FUpdateInterface : public Frontend::IDocumentTransform
 		{
 		public:
