Allow Metasound Bus Reader and Writer nodes to set up patch inputs after construction if their busses are set later

diff --git a/Engine/Plugins/Runtime/Metasound/Source/MetasoundEngine/Private/MetasoundAudioBusReaderNode.cpp b/Engine/Plugins/Runtime/Metasound/Source/MetasoundEngine/Private/MetasoundAudioBusReaderNode.cpp
--- a/Engine/Plugins/Runtime/Metasound/Source/MetasoundEngine/Private/MetasoundAudioBusReaderNode.cpp
+++ b/Engine/Plugins/Runtime/Metasound/Source/MetasoundEngine/Private/MetasoundAudioBusReaderNode.cpp
@@ -4,7 +4,6 @@
 #include "AudioBusSubsystem.h"
 #include "DSP/ConvertDeinterleave.h"
 #include "Internationalization/Text.h"
-#include "MediaPacket.h"
 #include "MetasoundAudioBuffer.h"
 #include "MetasoundAudioBus.h"
 #include "MetasoundEngineNodesNames.h"
@@ -145,9 +144,27 @@ namespace Metasound
 			const int32 BlockSizeFrames = AudioOutputs[0]->Num();
 			const int32 NumSamplesToPop = BlockSizeFrames * AudioBusChannels;
 
+			const FAudioBusProxyPtr& BusProxy = AudioBusAsset->GetAudioBusProxy();
+			if (BusProxy.IsValid() && BusProxy->AudioBusId != AudioBusId)
+			{
+				AudioBusPatchOutput.Reset();
+			}
+			
 			if (!AudioBusPatchOutput.IsValid())
 			{
-				return;
+				// if environment vars & a valid audio bus have been set since starting, try to create the patch now
+				if (SampleRate > 0.f && BusProxy.IsValid())
+				{
+					CreatePatchOutput(BlockSizeFrames);
+					if (!AudioBusPatchOutput.IsValid())
+					{
+						return;
+					}
+				}
+				else
+				{
+					return;
+				}
 			}
 
 			bool bPerformPop = false;
@@ -199,6 +216,35 @@ namespace Metasound
 			}
 		}
 
+		void CreatePatchOutput(const int32 BlockSizeFrames)
+		{
+			const FAudioBusProxyPtr& AudioBusProxy = AudioBusAsset->GetAudioBusProxy();
+			if (AudioBusProxy.IsValid())
+			{
+				if (FAudioDeviceManager* ADM = FAudioDeviceManager::Get())
+				{
+					if (FAudioDevice* AudioDevice = ADM->GetAudioDeviceRaw(AudioDeviceId))
+					{
+						// Start the audio bus in case it's not already started
+						AudioBusChannels = AudioBusProxy->NumChannels;
+						AudioBusId = AudioBusProxy->AudioBusId;
+						
+						const Audio::FAudioBusKey AudioBusKey = Audio::FAudioBusKey(AudioBusId);
+						UAudioBusSubsystem* AudioBusSubsystem = AudioDevice->GetSubsystem<UAudioBusSubsystem>();
+						
+						check(AudioBusSubsystem);
+						
+						AudioBusSubsystem->StartAudioBus(AudioBusKey, AudioBusChannels, false);
+
+						// Create a bus patch output with enough room for the number of samples we expect and some buffering
+						AudioBusPatchOutput = AudioBusSubsystem->AddPatchOutputForAudioBus(AudioBusKey, BlockSizeFrames, AudioBusChannels);
+					}
+				}
+			}
+
+			bFirstBlock = true;
+		}
+		
 		void Reset(const IOperator::FResetParams& InParams)
 		{
 			using namespace Frontend;
@@ -223,29 +269,9 @@ namespace Metasound
 			{
 				Buffer->Zero();
 			}
-
-			const FAudioBusProxyPtr& AudioBusProxy = AudioBusAsset->GetAudioBusProxy();
-			if (AudioBusProxy.IsValid())
-			{
-				if (FAudioDeviceManager* ADM = FAudioDeviceManager::Get())
-				{
-					if (FAudioDevice* AudioDevice = ADM->GetAudioDeviceRaw(AudioDeviceId))
-					{
-						// Start the audio bus in case it's not already started
-						AudioBusChannels = AudioBusProxy->NumChannels;
-						const Audio::FAudioBusKey AudioBusKey(AudioBusProxy->AudioBusId);
-						UAudioBusSubsystem* AudioBusSubsystem = AudioDevice->GetSubsystem<UAudioBusSubsystem>();
-						check(AudioBusSubsystem);
-						AudioBusSubsystem->StartAudioBus(AudioBusKey, AudioBusChannels, false);
-
-						// Create a bus patch output with enough room for the number of samples we expect and some buffering
-						const int32 BlockSizeFrames = InParams.OperatorSettings.GetNumFramesPerBlock();
-						AudioBusPatchOutput = AudioBusSubsystem->AddPatchOutputForAudioBus(AudioBusKey, BlockSizeFrames, AudioBusChannels);
-					}
-				}
-			}
-
-			bFirstBlock = true;
+			
+			const int32 BlockSizeFrames = InParams.OperatorSettings.GetNumFramesPerBlock();
+			CreatePatchOutput(BlockSizeFrames);
 		}
 
 	private:
@@ -255,6 +281,7 @@ namespace Metasound
 		TArray<float> InterleavedBuffer;
 		int32 AudioMixerOutputFrames = INDEX_NONE;
 		Audio::FDeviceId AudioDeviceId = INDEX_NONE;
+		uint32 AudioBusId = 0;
 		float SampleRate = 0.0f;
 		Audio::FPatchOutputStrongPtr AudioBusPatchOutput;
 		TUniquePtr<Audio::IConvertDeinterleave> ConvertDeinterleave;
diff --git a/Engine/Plugins/Runtime/Metasound/Source/MetasoundEngine/Private/MetasoundAudioBusWriterNode.cpp b/Engine/Plugins/Runtime/Metasound/Source/MetasoundEngine/Private/MetasoundAudioBusWriterNode.cpp
--- a/Engine/Plugins/Runtime/Metasound/Source/MetasoundEngine/Private/MetasoundAudioBusWriterNode.cpp
+++ b/Engine/Plugins/Runtime/Metasound/Source/MetasoundEngine/Private/MetasoundAudioBusWriterNode.cpp
@@ -112,26 +112,8 @@ namespace Metasound
 			Reset(InParams);
 		}
 
-		void Reset(const IOperator::FResetParams& InParams)
+		void CreatePatchInput()
 		{
-			using namespace Frontend;
-			using namespace AudioBusWriterNode;
-
-			bool bHasEnvironmentVars = InParams.Environment.Contains<Audio::FDeviceId>(SourceInterface::Environment::DeviceID);
-			bHasEnvironmentVars &= InParams.Environment.Contains<int32>(SourceInterface::Environment::AudioMixerNumOutputFrames);
-
-			if (bHasEnvironmentVars)
-			{
-				SampleRate = InParams.OperatorSettings.GetSampleRate();
-				AudioDeviceId = InParams.Environment.GetValue<Audio::FDeviceId>(SourceInterface::Environment::DeviceID);
-				AudioMixerOutputFrames = InParams.Environment.GetValue<int32>(SourceInterface::Environment::AudioMixerNumOutputFrames);
-			}
-			else
-			{
-				UE_LOG(LogMetaSound, Warning, TEXT("Audio bus writer node requires audio device ID '%s' and audio mixer num output frames '%s' environment variables")
-					, *SourceInterface::Environment::DeviceID.ToString(), *SourceInterface::Environment::AudioMixerNumOutputFrames.ToString());
-			}
-			
 			const FAudioBusProxyPtr& AudioBusProxy = AudioBusAsset->GetAudioBusProxy();
 			if (AudioBusProxy.IsValid())
 			{
@@ -141,21 +123,47 @@ namespace Metasound
 					{
 						UAudioBusSubsystem* AudioBusSubsystem = AudioDevice->GetSubsystem<UAudioBusSubsystem>();
 						check(AudioBusSubsystem);
-						Audio::FAudioBusKey AudioBusKey = Audio::FAudioBusKey(AudioBusProxy->AudioBusId);	
+						
+						AudioBusId = AudioBusProxy->AudioBusId;
+						const Audio::FAudioBusKey AudioBusKey = Audio::FAudioBusKey(AudioBusId);
+						
 						// Start the audio bus in case it's not already started					
 						AudioBusChannels = AudioBusProxy->NumChannels;
 						AudioBusSubsystem->StartAudioBus(AudioBusKey, AudioBusChannels, false);
 
-						BlockSizeFrames = InParams.OperatorSettings.GetNumFramesPerBlock();
+						InterleavedBuffer.Reset();
 						InterleavedBuffer.AddZeroed(BlockSizeFrames * AudioBusChannels);
 
 						// Create a bus patch input with enough room for the number of samples we expect and some buffering
 						AudioBusPatchInput = AudioBusSubsystem->AddPatchInputForAudioBus(AudioBusKey, BlockSizeFrames, AudioBusChannels);
-
 					}
 				}
 			}
+		}
+		
+		void Reset(const IOperator::FResetParams& InParams)
+		{
+			using namespace Frontend;
+			using namespace AudioBusWriterNode;
 
+			bool bHasEnvironmentVars = InParams.Environment.Contains<Audio::FDeviceId>(SourceInterface::Environment::DeviceID);
+			bHasEnvironmentVars &= InParams.Environment.Contains<int32>(SourceInterface::Environment::AudioMixerNumOutputFrames);
+
+			if (bHasEnvironmentVars)
+			{
+				SampleRate = InParams.OperatorSettings.GetSampleRate();
+				AudioDeviceId = InParams.Environment.GetValue<Audio::FDeviceId>(SourceInterface::Environment::DeviceID);
+				AudioMixerOutputFrames = InParams.Environment.GetValue<int32>(SourceInterface::Environment::AudioMixerNumOutputFrames);
+			}
+			else
+			{
+				UE_LOG(LogMetaSound, Warning, TEXT("Audio bus writer node requires audio device ID '%s' and audio mixer num output frames '%s' environment variables")
+					, *SourceInterface::Environment::DeviceID.ToString(), *SourceInterface::Environment::AudioMixerNumOutputFrames.ToString());
+			}
+			
+			BlockSizeFrames = InParams.OperatorSettings.GetNumFramesPerBlock();
+
+			CreatePatchInput();
 			bFirstBlock = true;
 		}
 
@@ -185,6 +193,25 @@ namespace Metasound
 
 		void Execute()
 		{
+			const FAudioBusProxyPtr& BusProxy = AudioBusAsset->GetAudioBusProxy();
+			if (BusProxy.IsValid() && BusProxy->AudioBusId != AudioBusId)
+			{
+				AudioBusPatchInput.Reset();
+			}
+			
+			if (!AudioBusPatchInput.IsValid())
+			{
+				// if environment vars & a valid audio bus have been set since starting, try to create the patch now
+				if (SampleRate > 0.f && BusProxy.IsValid())
+				{
+					CreatePatchInput();
+				}
+				else
+				{
+					return;
+				}
+			}
+
 			if (!AudioBusPatchInput.IsOutputStillActive())
 			{
 				return;
@@ -248,6 +275,7 @@ namespace Metasound
 		float SampleRate = 0.0f;
 		Audio::FPatchInput AudioBusPatchInput;
 		uint32 AudioBusChannels = INDEX_NONE;
+		uint32 AudioBusId = 0;
 		int32 BlockSizeFrames = 0;
 		bool bFirstBlock = true;
 	};
diff --git a/Engine/Source/Runtime/SignalProcessing/Private/MultithreadedPatching.cpp b/Engine/Source/Runtime/SignalProcessing/Private/MultithreadedPatching.cpp
--- a/Engine/Source/Runtime/SignalProcessing/Private/MultithreadedPatching.cpp
+++ b/Engine/Source/Runtime/SignalProcessing/Private/MultithreadedPatching.cpp
@@ -308,11 +308,21 @@ namespace Audio
 		OutputHandle->TargetGain = InGain;
 	}
 
-	bool FPatchInput::IsOutputStillActive()
+	bool FPatchInput::IsOutputStillActive() const
 	{
 		return OutputHandle.IsValid() && !OutputHandle.IsUnique();
 	}
 
+	bool FPatchInput::IsValid() const
+	{
+		return OutputHandle.IsValid();
+	}
+
+	void FPatchInput::Reset()
+	{
+		OutputHandle.Reset();
+	}
+
 	FPatchInput FPatchMixer::AddNewInput(int32 InMaxLatencyInSamples, float InGain)
 	{
 		FScopeLock ScopeLock(&PendingNewInputsCriticalSection);
diff --git a/Engine/Source/Runtime/SignalProcessing/Public/DSP/MultithreadedPatching.h b/Engine/Source/Runtime/SignalProcessing/Public/DSP/MultithreadedPatching.h
--- a/Engine/Source/Runtime/SignalProcessing/Public/DSP/MultithreadedPatching.h
+++ b/Engine/Source/Runtime/SignalProcessing/Public/DSP/MultithreadedPatching.h
@@ -114,8 +114,13 @@ namespace Audio
 		void SetGain(float InGain);
 
 		/** Returns false if this output was removed, either because someone called FPatchMixer::RemoveTap with this FPatchInput, or the FPatchMixer was destroyed. */
-		bool IsOutputStillActive();
+		bool IsOutputStillActive() const;
 
+		/** Returns false if this output was not initialized properly. */
+		bool IsValid() const;
+
+		void Reset();
+		
 		friend class FPatchMixer;
 		friend class FPatchSplitter;
 
