MetaSound Dynamic Operator and Dynamic Transactor

diff --git a/Engine/Plugins/Runtime/Metasound/Source/MetasoundEngine/Private/MetasoundSource.cpp b/Engine/Plugins/Runtime/Metasound/Source/MetasoundEngine/Private/MetasoundSource.cpp
--- a/Engine/Plugins/Runtime/Metasound/Source/MetasoundEngine/Private/MetasoundSource.cpp
+++ b/Engine/Plugins/Runtime/Metasound/Source/MetasoundEngine/Private/MetasoundSource.cpp
@@ -13,6 +13,7 @@
 #include "MetasoundAssetManager.h"
 #include "MetasoundAudioFormats.h"
 #include "MetasoundBuilderSubsystem.h"
+#include "MetasoundDynamicOperatorTransactor.h"
 #include "MetasoundEngineAsset.h"
 #include "MetasoundEngineEnvironment.h"
 #include "MetasoundEnvironment.h"
@@ -1014,5 +1015,34 @@ Metasound::SourcePrivate::FParameterRouter& UMetaSoundSource::GetParameterRouter
 	return Router;
 }
 
+TSharedPtr<Metasound::DynamicGraph::FDynamicOperatorTransactor> UMetaSoundSource::EnableDynamicGenerators(bool bInIsEnabled)
+{
+	using namespace Metasound;
+	using namespace Metasound::DynamicGraph;
+
+	if (bInIsEnabled && (!DynamicTransactor.IsValid()))
+	{
+		TSharedPtr<IGraph> CurrentGraph = GetRuntimeData().Graph;
+		if (CurrentGraph.IsValid())
+		{
+			DynamicTransactor = MakeShared<FDynamicOperatorTransactor>(*CurrentGraph);
+		}
+		else
+		{
+			DynamicTransactor = MakeShared<FDynamicOperatorTransactor>();
+		}
+	}
+	else
+	{
+		DynamicTransactor.Reset();
+	}
+
+	return DynamicTransactor;
+}
+
+TSharedPtr<Metasound::DynamicGraph::FDynamicOperatorTransactor> UMetaSoundSource::GetDynamicGeneratorTransactor() const
+{
+	return DynamicTransactor;
+}
 
 #undef LOCTEXT_NAMESPACE // MetaSound
diff --git a/Engine/Plugins/Runtime/Metasound/Source/MetasoundEngine/Public/MetasoundSource.h b/Engine/Plugins/Runtime/Metasound/Source/MetasoundEngine/Public/MetasoundSource.h
--- a/Engine/Plugins/Runtime/Metasound/Source/MetasoundEngine/Public/MetasoundSource.h
+++ b/Engine/Plugins/Runtime/Metasound/Source/MetasoundEngine/Public/MetasoundSource.h
@@ -28,6 +28,11 @@ namespace Metasound
 	{
 		class FParameterRouter;
 	}
+
+	namespace DynamicGraph
+	{
+		class FDynamicOperatorTransactor;
+	}
 } // namespace Metasound
 
 namespace Audio
@@ -230,6 +235,7 @@ private:
 		return RootMetasoundDocument;
 	}
 
+
 	bool IsParameterValid(const FAudioParameter& InParameter, const TMap<FName, FMetasoundFrontendVertex>& InInputNameVertexMap) const;
 
 	static Metasound::SourcePrivate::FParameterRouter& GetParameterRouter();
@@ -244,4 +250,23 @@ private:
 	TSortedMap<uint64, TWeakPtr<Metasound::FMetasoundGenerator>> Generators;
 	void TrackGenerator(uint64 Id, TSharedPtr<Metasound::FMetasoundGenerator> Generator);
 	void ForgetGenerator(ISoundGeneratorPtr Generator);
+
+	/** Enable dynamic generator. 
+	 *
+	 * Once a dynamic generator is enabled, all changes to the MetaSound should be applied to the
+	 * FDynamicOperatorTransactor in order to keep parity between the document and active graph.
+	 *
+	 * Note: Disabling the dynamic generator will sever the communication between any active generators
+	 * even if the dynamic generator is re-enabled. 
+	 */
+	TSharedPtr<Metasound::DynamicGraph::FDynamicOperatorTransactor> EnableDynamicGenerators(bool bInIsEnabled);
+
+	/** Get dynamic transactor
+	 *
+	 * If dynamic generators are enabled, this will return a valid pointer to a dynamic transactor. 
+	 * Changes to this transactor will be forwarded to any active Dynamic MetaSound Generators.
+	 */
+	TSharedPtr<Metasound::DynamicGraph::FDynamicOperatorTransactor> GetDynamicGeneratorTransactor() const;
+
+	TSharedPtr<Metasound::DynamicGraph::FDynamicOperatorTransactor> DynamicTransactor;
 };
diff --git a/Engine/Plugins/Runtime/Metasound/Source/MetasoundGenerator/Private/MetasoundGenerator.cpp b/Engine/Plugins/Runtime/Metasound/Source/MetasoundGenerator/Private/MetasoundGenerator.cpp
--- a/Engine/Plugins/Runtime/Metasound/Source/MetasoundGenerator/Private/MetasoundGenerator.cpp
+++ b/Engine/Plugins/Runtime/Metasound/Source/MetasoundGenerator/Private/MetasoundGenerator.cpp
@@ -210,6 +210,8 @@ namespace Metasound
 		MetaSoundName = {};
 		AudioOutputNames = {};
 		DefaultParameters = {};
+		DataChannel.Reset();
+		DynamicOperatorTransactor.Reset();
 	}
 
 	FMetasoundGenerator::FMetasoundGenerator(FMetasoundGeneratorInitParams&& InParams)
@@ -235,6 +237,8 @@ namespace Metasound
 		NumFramesPerExecute = InParams.OperatorSettings.GetNumFramesPerBlock();
 		NumSamplesPerExecute = NumChannels * NumFramesPerExecute;
 
+		const bool bIsDynamic = InParams.DynamicOperatorTransactor.IsValid();
+
 		// Create the routing for parameter packs
 PRAGMA_DISABLE_DEPRECATION_WARNINGS
 		ParameterPackSendAddress = UMetasoundParameterPack::CreateSendAddressFromEnvironment(InParams.Environment);
@@ -243,7 +247,10 @@ PRAGMA_ENABLE_DEPRECATION_WARNINGS
 
 		// attempt to use operator cache instead of building a new operator.
 		bool bDidUseCachedOperator = false;
-		if (ConsoleVariables::bEnableExperimentalOneShotOperatorCache || ConsoleVariables::bEnableExperimentalOperatorCache)
+		const bool bIsOperatorCacheEnabled = ConsoleVariables::bEnableExperimentalOneShotOperatorCache || ConsoleVariables::bEnableExperimentalOperatorCache;
+		// Dynamic operators cannot use the operator cache because they can change their internal structure. 
+		// The operator cache assumes that the operator is unchanged from it's original structure. 
+		if (bIsOperatorCacheEnabled && !bIsDynamic)
 		{
 			bUseOperatorCache = ConsoleVariables::bEnableExperimentalOperatorCache || MetasoundGeneratorPrivate::HasOneShotInterface(InParams.Graph->GetVertexInterface());
 
diff --git a/Engine/Plugins/Runtime/Metasound/Source/MetasoundGenerator/Private/MetasoundGeneratorBuilder.cpp b/Engine/Plugins/Runtime/Metasound/Source/MetasoundGenerator/Private/MetasoundGeneratorBuilder.cpp
--- a/Engine/Plugins/Runtime/Metasound/Source/MetasoundGenerator/Private/MetasoundGeneratorBuilder.cpp
+++ b/Engine/Plugins/Runtime/Metasound/Source/MetasoundGenerator/Private/MetasoundGeneratorBuilder.cpp
@@ -177,7 +177,16 @@ namespace Metasound
 			// Create an instance of the new graph
 			FBuildGraphOperatorParams BuildParams { *InInitParams.Graph, InInitParams.OperatorSettings, InputData, InInitParams.Environment };
 			FOperatorAndInputs OpAndInputs;
-			OpAndInputs.Operator = FOperatorBuilder(InInitParams.BuilderSettings).BuildGraphOperator(BuildParams, OutBuildResults);
+			FOperatorBuilder Builder(InInitParams.BuilderSettings);
+
+			if (InInitParams.DynamicOperatorTransactor.IsValid())
+			{
+				OpAndInputs.Operator = Builder.BuildDynamicGraphOperator(BuildParams, *InInitParams.DynamicOperatorTransactor, OutBuildResults);
+			}
+			else
+			{
+				OpAndInputs.Operator = Builder.BuildGraphOperator(BuildParams, OutBuildResults);
+			}
 			OpAndInputs.Inputs = InputData;
 
 			return OpAndInputs;
diff --git a/Engine/Plugins/Runtime/Metasound/Source/MetasoundGenerator/Public/MetasoundGenerator.h b/Engine/Plugins/Runtime/Metasound/Source/MetasoundGenerator/Public/MetasoundGenerator.h
--- a/Engine/Plugins/Runtime/Metasound/Source/MetasoundGenerator/Public/MetasoundGenerator.h
+++ b/Engine/Plugins/Runtime/Metasound/Source/MetasoundGenerator/Public/MetasoundGenerator.h
@@ -84,6 +84,7 @@ namespace Metasound
 		TArray<FAudioParameter> DefaultParameters;
 		bool bBuildSynchronous = false;
 		TSharedPtr<TSpscQueue<FMetaSoundParameterTransmitter::FParameter>> DataChannel;
+		TSharedPtr<Metasound::DynamicGraph::FDynamicOperatorTransactor> DynamicOperatorTransactor;
 
 		void Release();
 	};
diff --git a/Engine/Plugins/Runtime/Metasound/Source/MetasoundGraphCore/Private/MetasoundDynamicGraphAlgo.cpp b/Engine/Plugins/Runtime/Metasound/Source/MetasoundGraphCore/Private/MetasoundDynamicGraphAlgo.cpp
new file mode 100644
--- /dev/null
+++ b/Engine/Plugins/Runtime/Metasound/Source/MetasoundGraphCore/Private/MetasoundDynamicGraphAlgo.cpp
@@ -0,0 +1,262 @@
+// Copyright Epic Games, Inc. All Rights Reserved.
+
+#include "MetasoundDynamicGraphAlgo.h"
+
+#include "Containers/Array.h"
+#include "Containers/Map.h"
+#include "MetasoundDataReference.h"
+#include "MetasoundGraphAlgoPrivate.h"
+#include "MetasoundLog.h"
+#include "MetasoundVertex.h"
+#include "MetasoundVertexData.h"
+
+namespace Metasound
+{
+	namespace DynamicGraph
+	{
+		namespace DynamicGraphAlgoPrivate
+		{
+			template<typename EntryType, typename FunctionType>
+			void UpdateTableEntry(const TArray<FOperatorID>& InOperatorOrder, const FOperatorID& InOperatorID, IOperator& InOperator, TArray<EntryType>& InOutTable, FunctionType InFunction)
+			{
+				using namespace DirectedGraphAlgo;
+
+				int32 EntryIndex = InOutTable.IndexOfByPredicate([&](const EntryType& InEntry) { return InEntry.OperatorID == InOperatorID; });
+
+				bool bEntryShouldExist = (nullptr != InFunction);
+				bool bEntryCurrentlyExists = (EntryIndex != INDEX_NONE);
+
+				if (bEntryCurrentlyExists && bEntryShouldExist)
+				{
+					// Update the existing entry
+					InOutTable[EntryIndex].Function = InFunction;
+				}
+				else if (bEntryCurrentlyExists)
+				{
+					// Remove the existing entry
+					InOutTable.RemoveAt(EntryIndex);
+				}
+				else if (bEntryShouldExist)
+				{
+					// Add new entry
+					int32 OperatorIndex = InOperatorOrder.Find(InOperatorID);
+					if (INDEX_NONE == OperatorIndex)
+					{
+						UE_LOG(LogMetaSound, Error, TEXT("Failed to update operator stack. OperatorID not found in operator order array."));
+						return;
+					}
+
+					int32 PreceedingOperatorIndex = OperatorIndex - 1;
+					int32 InsertLocation = 0; // Default to inserting at the beginning
+					while (PreceedingOperatorIndex >= 0)
+					{
+						FOperatorID PreceedingOperatorID = InOperatorOrder[PreceedingOperatorIndex];
+						PreceedingOperatorIndex--;
+
+						int32 PreceedingEntryIndex = InOutTable.IndexOfByPredicate([&](const EntryType& InEntry) { return InEntry.OperatorID == PreceedingOperatorID; });
+						if (PreceedingEntryIndex != INDEX_NONE)
+						{
+							InsertLocation = PreceedingEntryIndex + 1;
+							break;
+						}
+					}
+
+					InOutTable.Insert(EntryType{InOperatorID, InOperator, InFunction}, InsertLocation);
+				}
+			}
+		} // namespace DynamicGraphAlgoPrivate
+
+		// Apply updates to data references through all the operators by following connections described in the FOperatorInfo map.
+		void PropagateBindUpdate(FOperatorID InInitialOperatorID, const FVertexName& InVertexName, const FAnyDataReference& InNewReference, FDynamicGraphOperatorData& InOutGraphOperatorData)
+		{
+			using namespace DirectedGraphAlgo;
+			using namespace DynamicGraphAlgoPrivate;
+
+			struct FInputToUpdate
+			{
+				FOperatorID OperatorID;
+				FVertexName VertexName;
+				FAnyDataReference DataReference;
+			};
+
+			TArray<FInputToUpdate> PropagateStack;
+			PropagateStack.Emplace(FInputToUpdate{InInitialOperatorID, InVertexName, InNewReference});
+
+			TArray<FVertexDataState> InitialOutputState;
+			TSortedVertexNameMap<FAnyDataReference> OutputUpdates;
+			while (PropagateStack.Num())
+			{
+				FInputToUpdate Current = PropagateStack.Pop();
+				if (FOperatorInfo* OpInfo = InOutGraphOperatorData.OperatorMap.Find(Current.OperatorID))
+				{
+					IOperator& Operator = *(OpInfo->Operator);
+					// Get current outputs
+					InitialOutputState.Reset();
+					GetVertexInterfaceDataState(OpInfo->VertexData.GetOutputs(), InitialOutputState);
+
+					// Set new input.
+					OpInfo->VertexData.GetInputs().SetVertex(Current.VertexName, MoveTemp(Current.DataReference));
+
+					// Bind inputs and outputs
+					Operator.BindInputs(OpInfo->VertexData.GetInputs());
+					Operator.BindOutputs(OpInfo->VertexData.GetOutputs());
+
+					// Update execute/postexecute/reset tables in case those have changed after rebinding.
+					UpdateTableEntry(InOutGraphOperatorData.OperatorOrder, Current.OperatorID, Operator, InOutGraphOperatorData.ExecuteTable, Operator.GetExecuteFunction());
+					UpdateTableEntry(InOutGraphOperatorData.OperatorOrder, Current.OperatorID, Operator, InOutGraphOperatorData.PostExecuteTable, Operator.GetPostExecuteFunction());
+					UpdateTableEntry(InOutGraphOperatorData.OperatorOrder, Current.OperatorID, Operator, InOutGraphOperatorData.ResetTable, Operator.GetResetFunction());
+
+					// See if binding altered the outputs. 
+					OutputUpdates.Reset();
+					CompareVertexInterfaceDataToPriorState(OpInfo->VertexData.GetOutputs(), InitialOutputState, OutputUpdates);
+
+					// Any updates to the outputs need to be propagated through the graph.
+					for (const TPair<FVertexName, FAnyDataReference>& OutputUpdate : OutputUpdates)
+					{
+						const FVertexName& OutputVertexName = OutputUpdate.Get<0>();
+						const FAnyDataReference& OutputDataReference = OutputUpdate.Get<1>();
+
+						if (const TArray<FGraphOperatorData::FVertexDestination>* Destinations = OpInfo->OutputConnections.Find(OutputVertexName))
+						{
+							for (const FGraphOperatorData::FVertexDestination& Destination : *Destinations)
+							{
+								PropagateStack.Push({Destination.OperatorID, Destination.VertexName, OutputDataReference});
+							}
+						}
+					}
+				}
+				else
+				{
+					UE_LOG(LogMetaSound, Error, TEXT("Failed to rebind graph operator state. Could not find operator info with ID %s"), *LexToString(Current.OperatorID));
+				}
+			}
+		}
+
+		FExecuteEntry::FExecuteEntry(DirectedGraphAlgo::FOperatorID InOperatorID, IOperator& InOperator, IOperator::FExecuteFunction InFunc)
+		: OperatorID(InOperatorID)
+		, Operator(&InOperator)
+		, Function(InFunc)
+		{
+		}
+
+		FPostExecuteEntry::FPostExecuteEntry(DirectedGraphAlgo::FOperatorID InOperatorID, IOperator& InOperator, IOperator::FPostExecuteFunction InFunc)
+		: OperatorID(InOperatorID)
+		, Operator(&InOperator)
+		, Function(InFunc)
+		{
+		}
+
+		FResetEntry::FResetEntry(DirectedGraphAlgo::FOperatorID InOperatorID, IOperator& InOperator, IOperator::FResetFunction InFunc)
+		: OperatorID(InOperatorID)
+		, Operator(&InOperator)
+		, Function(InFunc)
+		{
+		}
+
+		FDynamicGraphOperatorData::FDynamicGraphOperatorData(DirectedGraphAlgo::FGraphOperatorData&& InGraphOperatorData)
+		: DirectedGraphAlgo::FGraphOperatorData(MoveTemp(InGraphOperatorData))
+		{
+			// Populate execute/postexecute/reset stacks
+			for (const FOperatorID& OperatorID : OperatorOrder)
+			{
+				if (FOperatorInfo* OperatorInfo = OperatorMap.Find(OperatorID))
+				{
+					check(OperatorInfo->Operator.IsValid());
+					IOperator& Operator = *(OperatorInfo->Operator);
+
+					if (IOperator::FExecuteFunction ExecuteFunc = Operator.GetExecuteFunction())
+					{
+						ExecuteTable.Emplace(FExecuteEntry{OperatorID, Operator, ExecuteFunc});
+					}
+
+					if (IOperator::FPostExecuteFunction PostExecuteFunc = Operator.GetPostExecuteFunction())
+					{
+						PostExecuteTable.Emplace(FPostExecuteEntry{OperatorID, Operator, PostExecuteFunc});
+					}
+
+					if (IOperator::FResetFunction ResetFunc = Operator.GetResetFunction())
+					{
+						ResetTable.Emplace(FResetEntry{OperatorID, Operator, ResetFunc});
+					}
+				}
+			}
+		}
+
+		void UpdateOutputVertexData(FDynamicGraphOperatorData& InOutGraphOperatorData)
+		{
+			using namespace DirectedGraphAlgo;
+
+			// Iterate through the output operators and force their output data references
+			// to be reflected in the graph's FOutputVertexInterfaceData
+			for (const TPair<FVertexName, FOperatorID>& OutputVertexInfo : InOutGraphOperatorData.OutputVertexMap)
+			{
+				const FVertexName& VertexName = OutputVertexInfo.Get<0>();
+				const FOperatorID& OperatorID = OutputVertexInfo.Get<1>();
+
+				if (const FGraphOperatorData::FOperatorInfo* OperatorInfo = InOutGraphOperatorData.OperatorMap.Find(OperatorID))
+				{
+					if (const FAnyDataReference* Ref = OperatorInfo->VertexData.GetOutputs().FindDataReference(OutputVertexInfo.Get<0>()))
+					{
+						InOutGraphOperatorData.VertexData.GetOutputs().SetVertex(VertexName, *Ref);
+					}
+					else if (InOutGraphOperatorData.VertexData.GetOutputs().IsVertexBound(VertexName))
+					{
+						UE_LOG(LogMetaSound, Error, TEXT("Output vertex (%s) lost data reference after rebinding graph"), *VertexName.ToString());
+					}
+				}
+				else
+				{
+					UE_LOG(LogMetaSound, Error, TEXT("Failed to update graph operator outputs. Could not find output operator info with ID %s for vertex %s"), *LexToString(OperatorID), *VertexName.ToString());
+				}
+			}
+		}
+		
+		void RebindGraphInputs(FInputVertexInterfaceData& InOutVertexData, FDynamicGraphOperatorData& InOutGraphOperatorData)
+		{
+			using namespace DirectedGraphAlgo;
+
+			// Bind and diff the graph's interface to determine if there is an update to any vertices
+			FInputVertexInterfaceData& InputVertexData = InOutGraphOperatorData.VertexData.GetInputs();
+			TArray<FVertexDataState> InitialVertexDataState;
+			GetVertexInterfaceDataState(InputVertexData, InitialVertexDataState);
+
+			// Binding an input vertex interface may update `InputVertexData`
+			InOutVertexData.Bind(InputVertexData);
+
+			// Determine if there have been changes to `InputVertexData`. 
+			TSortedVertexNameMap<FAnyDataReference> GraphInputsToUpdate;
+			CompareVertexInterfaceDataToPriorState(InputVertexData, InitialVertexDataState, GraphInputsToUpdate);
+
+			// If there have been any changes to `InputVertexData`, then these need to be propagated
+			// through the graph to route them to operator inputs and to handle any knock-on updates
+			// to other data references. 
+			if (GraphInputsToUpdate.Num() > 0)
+			{
+				// Update the graph data by propagating the updates through the inputs nodes. 
+				for (const TPair<FVertexName, FAnyDataReference>& InputToUpdate : GraphInputsToUpdate)
+				{
+					const FVertexName& VertexName = InputToUpdate.Get<0>();
+					if (const FOperatorID* OperatorID = InOutGraphOperatorData.InputVertexMap.Find(VertexName))
+					{
+						PropagateBindUpdate(*OperatorID, VertexName, InputToUpdate.Get<1>(), InOutGraphOperatorData);
+					}
+					else
+					{
+						UE_LOG(LogMetaSound, Error, TEXT("No input operator exists for input vertex %s"), *VertexName.ToString());
+					}
+				}
+
+				// Refresh output vertex interface data in case any output nodes were updated
+				// when bind updates were propagated through the graph.
+				UpdateOutputVertexData(InOutGraphOperatorData);
+			}
+		}
+
+		void RebindGraphOutputs(FOutputVertexInterfaceData& InOutVertexData, FDynamicGraphOperatorData& InOutGraphOperatorData)
+		{
+			// Output rebinding does not alter data references in an operator. Here we can get away with
+			// simply reading the latest values.
+			InOutVertexData.Bind(InOutGraphOperatorData.VertexData.GetOutputs());
+		}
+	}
+}
diff --git a/Engine/Plugins/Runtime/Metasound/Source/MetasoundGraphCore/Private/MetasoundDynamicGraphAlgo.h b/Engine/Plugins/Runtime/Metasound/Source/MetasoundGraphCore/Private/MetasoundDynamicGraphAlgo.h
new file mode 100644
--- /dev/null
+++ b/Engine/Plugins/Runtime/Metasound/Source/MetasoundGraphCore/Private/MetasoundDynamicGraphAlgo.h
@@ -0,0 +1,105 @@
+// Copyright Epic Games, Inc. All Rights Reserved.
+
+#pragma once
+
+#include "Containers/Array.h"
+#include "MetasoundGraphAlgoPrivate.h"
+#include "MetasoundOperatorInterface.h"
+#include "MetasoundVertexData.h"
+
+namespace Metasound
+{
+	namespace DynamicGraph
+	{
+		struct FDynamicGraphOperatorData;
+
+		using FOperatorInfo = DirectedGraphAlgo::FGraphOperatorData::FOperatorInfo;
+		using FOperatorID = DirectedGraphAlgo::FOperatorID;
+
+		/* Convenience wrapper for execute function of an IOperator. */
+		struct FExecuteEntry final
+		{
+			FExecuteEntry(FOperatorID InOperatorID, IOperator& InOperator, IOperator::FExecuteFunction InFunc);
+
+			void Execute()
+			{
+				check(Operator);
+				check(Function);
+				Function(Operator);
+			}
+
+			FOperatorID OperatorID;
+			IOperator* Operator;
+			IOperator::FExecuteFunction Function;	
+		};
+
+		/* Convenience wrapper for post execute function of an IOperator. */
+		struct FPostExecuteEntry final
+		{
+			FPostExecuteEntry(FOperatorID InOperatorID, IOperator& InOperator, IOperator::FPostExecuteFunction InFunc);
+
+			void PostExecute()
+			{
+				check(Operator);
+				check(Function);
+				Function(Operator);
+			}
+
+			FOperatorID OperatorID;
+			IOperator* Operator;
+			IOperator::FPostExecuteFunction Function;	
+		};
+
+		/* Convenience wrapper for reset function of an IOperator. */
+		struct FResetEntry final
+		{
+			FResetEntry(FOperatorID InOperatorID, IOperator& InOperator, IOperator::FResetFunction InFunc);
+
+			void Reset(const IOperator::FResetParams& InParams)
+			{
+				check(Operator);
+				check(Function);
+				Function(Operator, InParams);
+			}
+
+			FOperatorID OperatorID;
+			IOperator* Operator;
+			IOperator::FResetFunction Function;	
+		};
+
+		/** Collection of data needed to support a dynamic operator*/
+		struct FDynamicGraphOperatorData : DirectedGraphAlgo::FGraphOperatorData
+		{
+			FDynamicGraphOperatorData(DirectedGraphAlgo::FGraphOperatorData&& InGraphOperatorData);
+
+			TArray<FExecuteEntry> ExecuteTable;
+			TArray<FPostExecuteEntry> PostExecuteTable;
+			TArray<FResetEntry> ResetTable;
+		};
+
+		/** Propagate FVertexInterfaceData updates through the operators in the Dynamic Graph Operator Data. 
+		 *
+		 * A change to an operator's input may result in a change to the operator's output. The updates to the
+		 * operator's output and any subsequent knock-on operator output updates need to be propagated through
+		 * all the relevant operators in the graph.
+		 *
+		 * @param InInitialOperatorID - ID of operator which will have it's input updated.
+		 * @param InVertexName - Vertex name on the initial operator which will have it's input updated.
+		 * @param InNewReference - New data reference to apply to the operators input vertex.
+		 * @param InOutGraphOperatorData - Graph data which will be updated with new references. 
+		 */
+		void PropagateBindUpdate(FOperatorID InInitialOperatorID, const FVertexName& InVertexName, const FAnyDataReference& InNewReference, FDynamicGraphOperatorData& InOutGraphOperatorData);
+
+		/** Iterate through the output operators and make force that their output data references
+		 * are reflected in the graph's FOutputVertexInterfaceData 
+		 */
+		void UpdateOutputVertexData(FDynamicGraphOperatorData& InOutGraphOperatorData);
+
+		/** Rebind the graph inputs, updating internal operator bindings as needed. */
+		void RebindGraphInputs(FInputVertexInterfaceData& InOutVertexData, FDynamicGraphOperatorData& InOutGraphOperatorData);
+
+		/** Rebind the graph inputs, updating internal operator bindings as needed. */
+		void RebindGraphOutputs(FOutputVertexInterfaceData& InOutVertexData, FDynamicGraphOperatorData& InOutGraphOperatorData);
+	}
+}
+
diff --git a/Engine/Plugins/Runtime/Metasound/Source/MetasoundGraphCore/Private/MetasoundDynamicOperator.cpp b/Engine/Plugins/Runtime/Metasound/Source/MetasoundGraphCore/Private/MetasoundDynamicOperator.cpp
new file mode 100644
--- /dev/null
+++ b/Engine/Plugins/Runtime/Metasound/Source/MetasoundGraphCore/Private/MetasoundDynamicOperator.cpp
@@ -0,0 +1,585 @@
+// Copyright Epic Games, Inc. All Rights Reserved.
+
+#include "MetasoundDynamicOperator.h"
+
+#include "Containers/Array.h"
+#include "MetasoundNodeInterface.h"
+#include "MetasoundOperatorInterface.h"
+#include "MetasoundVertexData.h"
+#include "MetasoundTrace.h"
+#include "Templates/SharedPointer.h"
+#include "Templates/UniquePtr.h"
+
+namespace Metasound
+{
+	namespace DynamicGraph
+	{
+		const TCHAR* LexToString(EDynamicOperatorTransformQueueAction InTransformResult)
+		{
+			switch (InTransformResult)
+			{
+				case EDynamicOperatorTransformQueueAction::Continue:
+					return TEXT("Continue");
+
+				case EDynamicOperatorTransformQueueAction::Fence:
+					return TEXT("Fence");
+
+				default:
+					{
+						checkNoEntry();
+						return TEXT("Unhandled");
+					}
+			}
+		}
+
+		namespace DynamicOperatorPrivate
+		{
+			// Table sorter helper so we don't rewrite this algorithm for each differe
+			// stack type (Execute/PostExecute/Reset)
+			struct FTableSorter
+			{
+				template<typename TableEntryType>
+				static void SortTable(const TArray<FOperatorID>& OperatorOrder, TArray<TableEntryType>& InOutTable)
+				{
+					uint32 TargetIndex = 0;
+					for (const FOperatorID& OperatorID : OperatorOrder)
+					{
+						uint32 CurrentIndex = InOutTable.IndexOfByPredicate([&](const TableEntryType& Entry) 
+							{ 
+								return Entry.OperatorID == OperatorID; 
+							}
+						);
+
+						// The operator may not exist in the stack if it doesn't
+						// have an applicable function.
+						if (INDEX_NONE != CurrentIndex)
+						{
+							InOutTable.Swap(TargetIndex, CurrentIndex);
+							TargetIndex++;
+						}
+					}
+				}
+			};
+
+			template<typename VertexInterfaceDataType>
+			class TScopeUnfreeze final
+			{
+			public:
+				explicit TScopeUnfreeze(VertexInterfaceDataType& InVertexData)
+				: bIsOriginallyFrozen(InVertexData.IsVertexInterfaceFrozen())
+				, VertexData(InVertexData)
+				{
+					VertexData.SetIsVertexInterfaceFrozen(false);
+				}
+
+				~TScopeUnfreeze()
+				{
+					VertexData.SetIsVertexInterfaceFrozen(bIsOriginallyFrozen);
+				}
+
+			private:
+				bool bIsOriginallyFrozen;
+				VertexInterfaceDataType& VertexData;
+			};
+
+		} // namespace DynamicOperatorPrivate
+
+		FDynamicOperator::FDynamicOperator(const FOperatorSettings& InSettings)
+		: DynamicOperatorData(InSettings)
+		{
+		}
+
+		FDynamicOperator::FDynamicOperator(DirectedGraphAlgo::FGraphOperatorData&& InGraphOperatorData, TSharedRef<TSpscQueue<TUniquePtr<IDynamicOperatorTransform>>> InTransformQueue)
+		: DynamicOperatorData(MoveTemp(InGraphOperatorData))
+		, TransformQueue(InTransformQueue)
+		{
+		}
+
+		FDataReferenceCollection FDynamicOperator::GetInputs() const
+		{
+			checkNoEntry();
+			return FDataReferenceCollection{};
+		}
+
+		FDataReferenceCollection FDynamicOperator::GetOutputs() const
+		{
+			checkNoEntry();
+			return FDataReferenceCollection{};
+		}
+
+		void FDynamicOperator::BindInputs(FInputVertexInterfaceData& InOutVertexData)
+		{
+			RebindGraphInputs(InOutVertexData, DynamicOperatorData);
+		}
+
+		void FDynamicOperator::BindOutputs(FOutputVertexInterfaceData& InOutVertexData)
+		{
+			RebindGraphOutputs(InOutVertexData, DynamicOperatorData);
+		}
+
+		IOperator::FResetFunction FDynamicOperator::GetResetFunction()
+		{
+			return &StaticReset;
+		}
+
+		IOperator::FExecuteFunction FDynamicOperator::GetExecuteFunction()
+		{
+			return &StaticExecute;
+		}
+
+		IOperator::FPostExecuteFunction FDynamicOperator::GetPostExecuteFunction()
+		{
+			return &StaticPostExecute;
+		}
+
+		void FDynamicOperator::FlushEnqueuedTransforms()
+		{
+			METASOUND_TRACE_CPUPROFILER_EVENT_SCOPE(Metasound::FDynamicOperator::FlushEnqueuedTransforms)
+
+			while (TOptional<TUniquePtr<IDynamicOperatorTransform>> Transform = TransformQueue->Dequeue())
+			{
+				if (Transform.IsSet())
+				{
+					if (Transform->IsValid())
+					{
+						(*Transform)->Transform(DynamicOperatorData);
+					}
+				}
+			}
+		}
+
+		void FDynamicOperator::ApplyTransformsUntilFence()
+		{
+			METASOUND_TRACE_CPUPROFILER_EVENT_SCOPE(Metasound::FDynamicOperator::ApplyTransformsUntilFence)
+
+			if (!bExecuteFenceIsSet)
+			{
+				while (TOptional<TUniquePtr<IDynamicOperatorTransform>> Transform = TransformQueue->Dequeue())
+				{
+					if (Transform.IsSet())
+					{
+						if (Transform->IsValid())
+						{
+							EDynamicOperatorTransformQueueAction Result = (*Transform)->Transform(DynamicOperatorData);
+							if (EDynamicOperatorTransformQueueAction::Fence == Result)
+							{
+								bExecuteFenceIsSet = true;
+								break;
+							}
+						}
+					}
+				}
+			}
+		}
+
+		void FDynamicOperator::Execute()
+		{
+			ApplyTransformsUntilFence();
+
+			for (FExecuteEntry& Entry : DynamicOperatorData.ExecuteTable)
+			{
+				Entry.Execute();
+			}
+		}
+
+		void FDynamicOperator::PostExecute()
+		{
+			for (FPostExecuteEntry& Entry : DynamicOperatorData.PostExecuteTable)
+			{
+				Entry.PostExecute();
+			}
+
+			bExecuteFenceIsSet = false;
+		}
+
+		void FDynamicOperator::Reset(const IOperator::FResetParams& InParams)
+		{
+			FlushEnqueuedTransforms();
+			for (FResetEntry& Entry : DynamicOperatorData.ResetTable)
+			{
+				Entry.Reset(InParams);
+			}
+		}
+
+		void FDynamicOperator::StaticReset(IOperator* InOperator, const IOperator::FResetParams& InParams)
+		{
+			static_cast<FDynamicOperator*>(InOperator)->Reset(InParams);
+		}
+
+		void FDynamicOperator::StaticExecute(IOperator* InOperator)
+		{
+			static_cast<FDynamicOperator*>(InOperator)->Execute();
+		}
+
+		void FDynamicOperator::StaticPostExecute(IOperator* InOperator)
+		{
+			static_cast<FDynamicOperator*>(InOperator)->PostExecute();
+		}
+
+		EDynamicOperatorTransformQueueAction FNullTransform::Transform(FDynamicGraphOperatorData& InGraphOperatorData)
+		{
+			return EDynamicOperatorTransformQueueAction::Continue;
+		}
+
+		// Set the order of operators in the graph
+		FSetOperatorOrder::FSetOperatorOrder(TArray<FOperatorID> InOrder)
+		: Order(MoveTemp(InOrder))
+		{
+		}
+
+		EDynamicOperatorTransformQueueAction FSetOperatorOrder::Transform(FDynamicGraphOperatorData& InGraphOperatorData)
+		{
+			METASOUND_TRACE_CPUPROFILER_EVENT_SCOPE(Metasound::DynamicOperator::SetOperatorOrder)
+
+			using namespace DynamicOperatorPrivate;
+
+			FTableSorter::SortTable(Order, InGraphOperatorData.ExecuteTable);
+			FTableSorter::SortTable(Order, InGraphOperatorData.PostExecuteTable);
+			FTableSorter::SortTable(Order, InGraphOperatorData.ResetTable);
+
+			return EDynamicOperatorTransformQueueAction::Continue;
+		}
+
+		// Add an operator to the grpah
+		FAddOperator::FAddOperator(FOperatorID InOperatorID, FOperatorInfo&& InInfo)
+		: OperatorID(InOperatorID)
+		, OperatorInfo(MoveTemp(InInfo))
+		{
+		}
+
+		EDynamicOperatorTransformQueueAction FAddOperator::Transform(FDynamicGraphOperatorData& InGraphOperatorData)
+		{
+			METASOUND_TRACE_CPUPROFILER_EVENT_SCOPE(Metasound::DynamicOperator::AddOperator)
+
+			if (OperatorInfo.Operator.IsValid())
+			{
+				IOperator* Operator = OperatorInfo.Operator.Get();
+
+				InGraphOperatorData.OperatorOrder.Add(OperatorID);
+				InGraphOperatorData.OperatorMap.Add(OperatorID, MoveTemp(OperatorInfo));
+
+				// Update execution tables
+				if (IOperator::FExecuteFunction ExecuteFunc = Operator->GetExecuteFunction())
+				{
+					InGraphOperatorData.ExecuteTable.Add(FExecuteEntry(OperatorID, *Operator, ExecuteFunc));
+				}
+
+				if (IOperator::FPostExecuteFunction PostExecuteFunc = Operator->GetPostExecuteFunction())
+				{
+					InGraphOperatorData.PostExecuteTable.Add(FPostExecuteEntry(OperatorID, *Operator, PostExecuteFunc));
+				}
+
+				if (IOperator::FResetFunction ResetFunc = Operator->GetResetFunction())
+				{
+					InGraphOperatorData.ResetTable.Add(FResetEntry(OperatorID, *Operator, ResetFunc));
+				}
+			}
+
+			return EDynamicOperatorTransformQueueAction::Continue;
+		}
+
+		// Remove an operator from the graph
+		FRemoveOperator::FRemoveOperator(FOperatorID InOperatorID)
+		: OperatorID(InOperatorID)
+		{
+		}
+
+		EDynamicOperatorTransformQueueAction FRemoveOperator::Transform(FDynamicGraphOperatorData& InGraphOperatorData) 
+		{
+			METASOUND_TRACE_CPUPROFILER_EVENT_SCOPE(Metasound::DynamicOperator::RemoveOperator)
+
+			InGraphOperatorData.OperatorOrder.Remove(OperatorID);
+			InGraphOperatorData.OperatorMap.Remove(OperatorID);
+
+			// Update execution tables
+			InGraphOperatorData.ExecuteTable.RemoveAll([&](const FExecuteEntry& InEntry) { return InEntry.OperatorID == OperatorID; });
+			InGraphOperatorData.PostExecuteTable.RemoveAll([&](const FPostExecuteEntry& InEntry) { return InEntry.OperatorID == OperatorID; });
+			InGraphOperatorData.ResetTable.RemoveAll([&](const FResetEntry& InEntry) { return InEntry.OperatorID == OperatorID; });
+
+			return EDynamicOperatorTransformQueueAction::Continue;
+		}
+
+		// Add an input to the graph
+		FAddInput::FAddInput(FOperatorID InOperatorID, const FVertexName& InVertexName, FAnyDataReference InDataReference)
+		: OperatorID(InOperatorID)
+		, VertexName(InVertexName)
+		, DataReference(InDataReference)
+		{
+		}
+
+		EDynamicOperatorTransformQueueAction FAddInput::Transform(FDynamicGraphOperatorData& InGraphOperatorData)
+		{
+			METASOUND_TRACE_CPUPROFILER_EVENT_SCOPE(Metasound::DynamicOperator::AddInput)
+
+			FOperatorInfo* OpInfo = InGraphOperatorData.OperatorMap.Find(OperatorID);
+			if (nullptr == OpInfo)
+			{
+				UE_LOG(LogMetaSound, Error, TEXT("Could not find operator with ID %d when adding input %s."), OperatorID, *VertexName.ToString());
+				return EDynamicOperatorTransformQueueAction::Continue;
+			}
+
+			const FInputDataVertex& OperatorInputVertex = OpInfo->VertexData.GetInputs().GetVertex(VertexName);
+
+			FInputVertexInterfaceData& GraphInputData = InGraphOperatorData.VertexData.GetInputs();
+			{
+				// Unfreeze interface so a new vertex can be added. 
+				DynamicOperatorPrivate::TScopeUnfreeze Unfreeze(GraphInputData);
+
+				// Update unfrozen vertex interface data.
+				GraphInputData.AddVertex(OperatorInputVertex);
+			}
+			GraphInputData.SetVertex(VertexName, DataReference);
+
+			InGraphOperatorData.InputVertexMap.Add(VertexName, OperatorID);
+
+			// Propagate the data reference update through the graph
+			PropagateBindUpdate(OperatorID, VertexName, DataReference, InGraphOperatorData);
+
+			// Refresh output vertex interface data in case any output nodes were updated
+			// when bind updates were propagated through the graph.
+			UpdateOutputVertexData(InGraphOperatorData);
+
+			return EDynamicOperatorTransformQueueAction::Continue;
+		}
+
+		// Remove an input from the graph
+		FRemoveInput::FRemoveInput(const FVertexName& InVertexName)
+		: VertexName(InVertexName)
+		{
+		}
+
+		EDynamicOperatorTransformQueueAction FRemoveInput::Transform(FDynamicGraphOperatorData& InGraphOperatorData)
+		{
+			METASOUND_TRACE_CPUPROFILER_EVENT_SCOPE(Metasound::DynamicOperator::RemoveInput)
+
+			InGraphOperatorData.InputVertexMap.Remove(VertexName);
+			{
+				// Unfreeze vertex data so vertex can be removed
+				DynamicOperatorPrivate::TScopeUnfreeze Unfreeze(InGraphOperatorData.VertexData.GetInputs());
+				InGraphOperatorData.VertexData.GetInputs().RemoveVertex(VertexName);
+			}
+
+			return EDynamicOperatorTransformQueueAction::Continue;
+		}
+
+		// Add an output to the graph
+		FAddOutput::FAddOutput(FOperatorID InOperatorID, const FVertexName& InVertexName)
+		: OperatorID(InOperatorID)
+		, VertexName(InVertexName)
+		{
+		}
+
+		EDynamicOperatorTransformQueueAction FAddOutput::Transform(FDynamicGraphOperatorData& InGraphOperatorData)
+		{
+			METASOUND_TRACE_CPUPROFILER_EVENT_SCOPE(Metasound::DynamicOperator::AddOutput)
+
+			FOperatorInfo* OpInfo = InGraphOperatorData.OperatorMap.Find(OperatorID);
+			if (nullptr == OpInfo)
+			{
+				UE_LOG(LogMetaSound, Error, TEXT("Could not find operator with ID %d when adding output %s."), OperatorID, *VertexName.ToString());
+				return EDynamicOperatorTransformQueueAction::Continue;
+			}
+
+			const FOutputDataVertex& OperatorOutputVertex = OpInfo->VertexData.GetOutputs().GetVertex(VertexName);
+			const FAnyDataReference* Ref = OpInfo->VertexData.GetOutputs().FindDataReference(VertexName);
+
+			if (nullptr == Ref)
+			{
+				UE_LOG(LogMetaSound, Error, TEXT("Could not find data reference when creating output %s"), *VertexName.ToString());
+				return EDynamicOperatorTransformQueueAction::Continue;
+			}
+
+
+			FOutputVertexInterfaceData& GraphOutputData = InGraphOperatorData.VertexData.GetOutputs();
+			{
+				// Unfreeze interface so a new vertex can be added. 
+				DynamicOperatorPrivate::TScopeUnfreeze Unfreeze(GraphOutputData);
+
+				// Update unfrozen vertex interface data.
+				GraphOutputData.AddVertex(OperatorOutputVertex);
+			}
+			GraphOutputData.SetVertex(VertexName, *Ref);
+
+			InGraphOperatorData.OutputVertexMap.Add(VertexName, OperatorID);
+
+			return EDynamicOperatorTransformQueueAction::Continue;
+		}
+
+		// Remove an output from the graph
+		FRemoveOutput::FRemoveOutput(const FVertexName& InVertexName)
+		: VertexName(InVertexName)
+		{
+		}
+
+		EDynamicOperatorTransformQueueAction FRemoveOutput::Transform(FDynamicGraphOperatorData& InGraphOperatorData)
+		{
+			METASOUND_TRACE_CPUPROFILER_EVENT_SCOPE(Metasound::DynamicOperator::RemoveOutput)
+
+			InGraphOperatorData.OutputVertexMap.Remove(VertexName);
+			{
+				// Unfreeze vertex data so vertex can be removed
+				DynamicOperatorPrivate::TScopeUnfreeze Unfreeze(InGraphOperatorData.VertexData.GetOutputs());
+				InGraphOperatorData.VertexData.GetOutputs().RemoveVertex(VertexName);
+			}
+			return EDynamicOperatorTransformQueueAction::Continue;
+		}
+
+		// Trigger a execution fence.
+		EDynamicOperatorTransformQueueAction FExecuteFence::Transform(FDynamicGraphOperatorData& InGraphOperatorData)
+		{
+			return EDynamicOperatorTransformQueueAction::Fence;
+		}
+
+		// Connect two operators in the graph
+		FConnectOperators::FConnectOperators(FOperatorID InFromOpID, const FName& InFromVert, FOperatorID InToOpID, const FName& InToVert)
+		: FromOpID(InFromOpID)
+		, ToOpID(InToOpID)
+		, FromVert(InFromVert)
+		, ToVert(InToVert)
+		{
+		}
+
+		EDynamicOperatorTransformQueueAction FConnectOperators::Transform(FDynamicGraphOperatorData& InGraphOperatorData)
+		{
+			using namespace DirectedGraphAlgo;
+
+			METASOUND_TRACE_CPUPROFILER_EVENT_SCOPE(Metasound::DynamicOperator::ConnectOperators)
+
+			FOperatorInfo* FromOpInfo = InGraphOperatorData.OperatorMap.Find(FromOpID);
+			FOperatorInfo* ToOpInfo = InGraphOperatorData.OperatorMap.Find(ToOpID);
+
+			if (nullptr == FromOpInfo)
+			{
+				UE_LOG(LogMetaSound, Error, TEXT("Could not find operator with ID %d when connecting from %d:%s to %d:%s"), FromOpID, FromOpID, *FromVert.ToString(), ToOpID, *ToVert.ToString());
+				return EDynamicOperatorTransformQueueAction::Continue;
+			}
+
+			if (nullptr == ToOpInfo)
+			{
+				UE_LOG(LogMetaSound, Error, TEXT("Could not find operator with ID %d when connecting from %d:%s to %d:%s"), ToOpID, FromOpID, *FromVert.ToString(), ToOpID, *ToVert.ToString());
+				return EDynamicOperatorTransformQueueAction::Continue;
+			}
+
+			const FAnyDataReference* FromRef = FromOpInfo->VertexData.GetOutputs().FindDataReference(FromVert);
+			if (nullptr == FromRef)
+			{
+				UE_LOG(LogMetaSound, Error, TEXT("Could not find output data reference with vertex name %s when connecting from %d:%s to %d:%s"), *FromVert.ToString(), FromOpID, *FromVert.ToString(), ToOpID, *ToVert.ToString());
+				return EDynamicOperatorTransformQueueAction::Continue;
+			}
+
+			// Propagate the data reference update through the graph
+			PropagateBindUpdate(ToOpID, ToVert, *FromRef, InGraphOperatorData);
+
+			// Refresh output vertex interface data in case any output nodes were updated
+			// when bind updates were propagated through the graph.
+			UpdateOutputVertexData(InGraphOperatorData);
+
+			FromOpInfo->OutputConnections.FindOrAdd(FromVert).Add(FGraphOperatorData::FVertexDestination{ToOpID, ToVert});
+
+			return EDynamicOperatorTransformQueueAction::Continue;
+		}
+
+		// Disconnect a specific edge
+		FSwapOperatorConnection::FSwapOperatorConnection(FOperatorID InOriginalFromOpID, const FName& InOriginalFromVert, FOperatorID InNewFromOpID, const FName& InNewFromVert, FOperatorID InToOpID, const FName& InToVert)
+		: ConnectTransform(InNewFromOpID, InNewFromVert, InToOpID, InToVert)
+		, OriginalFromOpID(InOriginalFromOpID)
+		, ToOpID(InToOpID)
+		, OriginalFromVert(InOriginalFromVert)
+		, ToVert(InToVert)
+		{
+		}
+
+		EDynamicOperatorTransformQueueAction FSwapOperatorConnection::Transform(FDynamicGraphOperatorData& InGraphOperatorData)
+		{
+			using namespace DirectedGraphAlgo;
+
+			METASOUND_TRACE_CPUPROFILER_EVENT_SCOPE(Metasound::DynamicOperator::DisconnectOperators)
+
+			// Make new connection.
+			EDynamicOperatorTransformQueueAction NextAction = ConnectTransform.Transform(InGraphOperatorData);
+			check(NextAction == EDynamicOperatorTransformQueueAction::Continue); // this should always be continue since next step must be performed. 
+
+			// Clean up the old connection.
+			FOperatorInfo* OriginalFromOpInfo = InGraphOperatorData.OperatorMap.Find(OriginalFromOpID);
+
+			if (nullptr == OriginalFromOpInfo)
+			{
+				UE_LOG(LogMetaSound, Error, TEXT("Could not find operator with ID %d when disconnecting from %d:%s to %d:%s"), OriginalFromOpID, OriginalFromOpID, *OriginalFromVert.ToString(), ToOpID, *ToVert.ToString());
+				return EDynamicOperatorTransformQueueAction::Continue;
+			}
+
+			// Remove destinations from operator map.
+			auto IsDestinationToRemove = [&](const FGraphOperatorData::FVertexDestination& Dst) 
+			{ 
+				return (Dst.OperatorID == ToOpID) && (Dst.VertexName == ToVert);
+			};
+			OriginalFromOpInfo->OutputConnections.FindOrAdd(OriginalFromVert).RemoveAll(IsDestinationToRemove);
+
+			return EDynamicOperatorTransformQueueAction::Continue;
+		}
+
+		// Set the value on an unconnected operator input
+		FSetOperatorInput::FSetOperatorInput(FOperatorID InToOpID, const FName& InToVert, const FLiteral& InLiteral, FLiteralAssignmentFunction InAssignFunc)
+		: ToOpID(InToOpID)
+		, ToVert(InToVert)
+		, Literal(InLiteral)
+		, AssignFunc(InAssignFunc)
+		{
+		}
+
+		EDynamicOperatorTransformQueueAction FSetOperatorInput::Transform(FDynamicGraphOperatorData& InGraphOperatorData)
+		{
+			using namespace DirectedGraphAlgo;
+
+			METASOUND_TRACE_CPUPROFILER_EVENT_SCOPE(Metasound::DynamicOperator::SetOperatorInput)
+
+			FOperatorInfo* ToOpInfo = InGraphOperatorData.OperatorMap.Find(ToOpID);
+
+			if (nullptr == ToOpInfo)
+			{
+				UE_LOG(LogMetaSound, Error, TEXT("Could not find operator with ID %d when setting value for %d:%s"), ToOpID, ToOpID, *ToVert.ToString());
+				return EDynamicOperatorTransformQueueAction::Continue;
+			}
+
+			const FAnyDataReference* Ref = ToOpInfo->VertexData.GetInputs().FindDataReference(ToVert);
+			if (nullptr == Ref)
+			{
+				UE_LOG(LogMetaSound, Error, TEXT("Could not find input data reference with vertex name %s when setting %d:%s"), *ToVert.ToString(), ToOpID, *ToVert.ToString());
+				return EDynamicOperatorTransformQueueAction::Continue;
+			}
+
+			AssignFunc(InGraphOperatorData.OperatorSettings, Literal, *Ref);
+
+			return EDynamicOperatorTransformQueueAction::Continue;
+		}
+
+		// Perform several transformations at once without executing the graph.
+		FAtomicTransform::FAtomicTransform(TArray<TUniquePtr<IDynamicOperatorTransform>> InTransforms)
+		: Transforms(MoveTemp(InTransforms))
+		{
+		}
+
+		EDynamicOperatorTransformQueueAction FAtomicTransform::Transform(FDynamicGraphOperatorData& InGraphOperatorData)
+		{
+			METASOUND_TRACE_CPUPROFILER_EVENT_SCOPE(Metasound::DynamicOperator::AtomicTransform)
+
+			EDynamicOperatorTransformQueueAction Result = EDynamicOperatorTransformQueueAction::Continue;
+
+			for (TUniquePtr<IDynamicOperatorTransform>& TransformPtr : Transforms)
+			{
+				if (TransformPtr.IsValid())
+				{
+					if (EDynamicOperatorTransformQueueAction::Continue != Result)
+					{
+						UE_LOG(LogMetaSound, Error, TEXT("Encountered unsupported dynamic operator transform result (%s) during atomic operator trasnform."), LexToString(Result));
+					}
+
+					Result = TransformPtr->Transform(InGraphOperatorData);
+				}
+			}
+
+			return Result;
+		}
+	} // namespace DynamicGraph
+}
+
diff --git a/Engine/Plugins/Runtime/Metasound/Source/MetasoundGraphCore/Private/MetasoundDynamicOperator.h b/Engine/Plugins/Runtime/Metasound/Source/MetasoundGraphCore/Private/MetasoundDynamicOperator.h
new file mode 100644
--- /dev/null
+++ b/Engine/Plugins/Runtime/Metasound/Source/MetasoundGraphCore/Private/MetasoundDynamicOperator.h
@@ -0,0 +1,219 @@
+// Copyright Epic Games, Inc. All Rights Reserved.
+
+#pragma once
+
+#include "Containers/Array.h"
+#include "Containers/SpscQueue.h"
+#include "MetasoundDynamicGraphAlgo.h"
+#include "MetasoundGraphAlgoPrivate.h"
+#include "MetasoundNodeInterface.h"
+#include "MetasoundOperatorInterface.h"
+#include "MetasoundVertexData.h"
+#include "Templates/UniquePtr.h"
+
+namespace Metasound
+{
+	namespace DynamicGraph
+	{
+		struct FDynamicGraphOperatorData;
+		using FLiteralAssignmentFunction = void(*)(const FOperatorSettings& InOperatorSettings, const FLiteral& InLiteral, const FAnyDataReference& OutDataRef);
+
+		/** Action to perform after a single transformation. */
+		enum class EDynamicOperatorTransformQueueAction : uint8
+		{
+			Continue, //< Perform next operator if it exists. 
+			Fence     //< Wait to perform next operator until the operator has been executed. 
+		};
+
+		/** Interface for transformation of dynamic graph operator data. */
+		class IDynamicOperatorTransform
+		{
+		public:
+			virtual ~IDynamicOperatorTransform() = default;
+			virtual EDynamicOperatorTransformQueueAction Transform(FDynamicGraphOperatorData& InGraphOperatorData) = 0;
+		};
+
+		/** An FDynamicOperator is a MetaSound operator which can dynamically
+		 * change it's topology. Changes are communicated to the dynamic operator
+		 * through a TransformationQueue. 
+		 */
+		class FDynamicOperator : public IOperator
+		{
+		public:
+
+			FDynamicOperator(const FOperatorSettings& InSettings);
+			FDynamicOperator(DirectedGraphAlgo::FGraphOperatorData&& InGraphOperatorData, TSharedRef<TSpscQueue<TUniquePtr<IDynamicOperatorTransform>>> TransformQueue);
+
+			virtual FDataReferenceCollection GetInputs() const override;
+			virtual FDataReferenceCollection GetOutputs() const override;
+
+			virtual void BindInputs(FInputVertexInterfaceData& InOutVertexData) override;
+			virtual void BindOutputs(FOutputVertexInterfaceData& InOutVertexData) override;
+
+			virtual IOperator::FResetFunction GetResetFunction() override;
+
+			virtual IOperator::FExecuteFunction GetExecuteFunction() override;
+
+			virtual IOperator::FPostExecuteFunction GetPostExecuteFunction() override;
+
+			// Force all transformations in the transaction queue to be applied. 
+			void FlushEnqueuedTransforms();
+
+		private:
+			void ApplyTransformsUntilFence();
+			void Execute();
+			void PostExecute();
+			void Reset(const IOperator::FResetParams& InParams);
+
+			static void StaticReset(IOperator* InOperator, const IOperator::FResetParams& InParams);
+			static void StaticExecute(IOperator* InOperator);
+			static void StaticPostExecute(IOperator* InOperator);
+
+			FDynamicGraphOperatorData DynamicOperatorData;
+
+			TSharedRef<TSpscQueue<TUniquePtr<IDynamicOperatorTransform>>> TransformQueue;
+			bool bExecuteFenceIsSet = false;
+		};
+
+		/** A transform which does nothing. */
+		class FNullTransform : public IDynamicOperatorTransform
+		{
+		public:
+			virtual EDynamicOperatorTransformQueueAction Transform(FDynamicGraphOperatorData& InGraphOperatorData) override;
+		};
+
+		/** A transform which determines the order of execution for operators. */
+		class FSetOperatorOrder : public IDynamicOperatorTransform
+		{
+		public:
+			FSetOperatorOrder(TArray<FOperatorID> InOrder);
+			virtual EDynamicOperatorTransformQueueAction Transform(FDynamicGraphOperatorData& InGraphOperatorData) override;
+		private:
+			TArray<FOperatorID> Order;
+		};
+
+		/** A transform which adds an operator. */
+		class FAddOperator : public IDynamicOperatorTransform
+		{
+		public:
+			FAddOperator(FOperatorID InOperatorID, FOperatorInfo&& InInfo);
+			virtual EDynamicOperatorTransformQueueAction Transform(FDynamicGraphOperatorData& InGraphOperatorData) override;
+		private:
+			FOperatorID OperatorID;
+			FOperatorInfo OperatorInfo;
+		};
+
+		/** A transform which removes an operator. */
+		class FRemoveOperator : public IDynamicOperatorTransform
+		{
+		public:
+			FRemoveOperator(FOperatorID InOperatorID);
+			virtual EDynamicOperatorTransformQueueAction Transform(FDynamicGraphOperatorData& InGraphOperatorData) override;
+		private:
+			FOperatorID OperatorID;
+		};
+
+		/** A transform which exposes an input to the graph. */
+		class FAddInput : public IDynamicOperatorTransform
+		{
+		public:
+			FAddInput(FOperatorID InOperatorID, const FVertexName& InVertexName, FAnyDataReference InDataReference);
+			virtual EDynamicOperatorTransformQueueAction Transform(FDynamicGraphOperatorData& InGraphOperatorData) override;
+		private:
+			FOperatorID OperatorID;
+			FVertexName VertexName;
+			FAnyDataReference DataReference;
+		};
+
+		/** A transform which removes an input to the graph. */
+		class FRemoveInput : public IDynamicOperatorTransform
+		{
+		public:
+			FRemoveInput(const FVertexName& InVertexName);
+			virtual EDynamicOperatorTransformQueueAction Transform(FDynamicGraphOperatorData& InGraphOperatorData) override;
+		private:
+			FVertexName VertexName;
+		};
+
+		/** A transform which exposes an output from the graph. */
+		class FAddOutput: public IDynamicOperatorTransform
+		{
+		public:
+			FAddOutput(FOperatorID InOperatorID, const FVertexName& InVertexName);
+			virtual EDynamicOperatorTransformQueueAction Transform(FDynamicGraphOperatorData& InGraphOperatorData) override;
+		private:
+			FOperatorID OperatorID;
+			FVertexName VertexName;
+		};
+
+		/** A transform which removes an output from the graph. */
+		class FRemoveOutput: public IDynamicOperatorTransform
+		{
+		public:
+			FRemoveOutput(const FVertexName& InVertexName);
+			virtual EDynamicOperatorTransformQueueAction Transform(FDynamicGraphOperatorData& InGraphOperatorData) override;
+		private:
+			FVertexName VertexName;
+		};
+
+		/** A transform which pauses the transformation queue until the dynamic operator has executed. */
+		class FExecuteFence : public IDynamicOperatorTransform
+		{
+		public:
+			virtual EDynamicOperatorTransformQueueAction Transform(FDynamicGraphOperatorData& InGraphOperatorData) override;
+		};
+
+		/** A transform which connects two vertices in the graph. */
+		class FConnectOperators : public IDynamicOperatorTransform
+		{
+		public:
+			FConnectOperators(FOperatorID InFromOpID, const FName& InFromVert, FOperatorID InToOpID, const FName& InToVert);
+			virtual EDynamicOperatorTransformQueueAction Transform(FDynamicGraphOperatorData& InGraphOperatorData) override;
+		private:
+			FOperatorID FromOpID;
+			FOperatorID ToOpID;
+			FName FromVert;
+			FName ToVert;
+		};
+
+		/** A transform which disconnects two vertices in the graph. */
+		class FSwapOperatorConnection : public IDynamicOperatorTransform
+		{
+		public:
+			FSwapOperatorConnection(FOperatorID InOriginalFromOpID, const FName& InOrignialFromVert, FOperatorID InNewFromOpID, const FName& InNewFromVert, FOperatorID InToOpID, const FName& InToVert);
+			virtual EDynamicOperatorTransformQueueAction Transform(FDynamicGraphOperatorData& InGraphOperatorData) override;
+		private:
+			FConnectOperators ConnectTransform;
+
+			FOperatorID OriginalFromOpID;
+			FOperatorID ToOpID;
+			FName OriginalFromVert;
+			FName ToVert;
+		};
+
+		/** A transform which sets the input to an operator to a specified literal. */
+		class FSetOperatorInput : public IDynamicOperatorTransform
+		{
+		public:
+			FSetOperatorInput(FOperatorID InToOpID, const FName& InToVert, const FLiteral& InLiteral, FLiteralAssignmentFunction InAssignFunc);
+			virtual EDynamicOperatorTransformQueueAction Transform(FDynamicGraphOperatorData& InGraphOperatorData) override;
+
+		private:
+			FOperatorID ToOpID;
+			FName ToVert;
+			FLiteral Literal;
+			FLiteralAssignmentFunction AssignFunc;
+		};
+
+		/** A transform which groups multiple transformations together and forces them to all complete
+		 * before the dynamic operator executes. */
+		class FAtomicTransform : public IDynamicOperatorTransform
+		{
+		public:
+			FAtomicTransform(TArray<TUniquePtr<IDynamicOperatorTransform>> InTransforms);
+			virtual EDynamicOperatorTransformQueueAction Transform(FDynamicGraphOperatorData& InGraphOperatorData) override;
+		private:
+			TArray<TUniquePtr<IDynamicOperatorTransform>> Transforms;
+		};
+	}
+}
diff --git a/Engine/Plugins/Runtime/Metasound/Source/MetasoundGraphCore/Private/MetasoundDynamicOperatorTransactor.cpp b/Engine/Plugins/Runtime/Metasound/Source/MetasoundGraphCore/Private/MetasoundDynamicOperatorTransactor.cpp
new file mode 100644
--- /dev/null
+++ b/Engine/Plugins/Runtime/Metasound/Source/MetasoundGraphCore/Private/MetasoundDynamicOperatorTransactor.cpp
@@ -0,0 +1,493 @@
+// Copyright Epic Games, Inc. All Rights Reserved.
+
+#include "MetasoundDynamicOperatorTransactor.h"
+
+#include "Containers/SpscQueue.h"
+#include "Containers/UnrealString.h"
+#include "MetasoundBuilderInterface.h"
+#include "MetasoundDynamicOperator.h"
+#include "MetasoundGraph.h"
+#include "MetasoundGraphAlgo.h"
+#include "MetasoundGraphAlgoPrivate.h"
+#include "MetasoundNodeInterface.h"
+#include "MetasoundOperatorInterface.h"
+#include "MetasoundTrace.h"
+#include "Templates/SharedPointer.h"
+#include "Templates/UniquePtr.h"
+
+namespace Metasound
+{
+	namespace DynamicGraph
+	{
+		namespace DynamicOperatorTransactorPrivate
+		{
+			// Literal nodes always have output vertex with this name. 
+			static FLazyName LiteralNodeOutputVertexName("Value");
+
+			TArray<FOperatorID> DetermineOperatorOrder(const IGraph& InGraph)
+			{
+				/* determine new operator order. */
+				TArray<const INode*> NodeOrder;
+
+				bool bSuccess = DirectedGraphAlgo::DepthFirstTopologicalSort(InGraph, NodeOrder);
+				if (!bSuccess)
+				{
+					UE_LOG(LogMetaSound, Error, TEXT("Cycles found in dynamic graph"));
+				}
+
+				TArray<FOperatorID> OperatorOrder;
+				Algo::Transform(NodeOrder, OperatorOrder, static_cast<FOperatorID(*)(const INode*)>(DirectedGraphAlgo::GetOperatorID)); //< Static cast to help deduce which overloaded version of GetOperatorID to call in Algo::Transform
+
+				return OperatorOrder;
+			}
+		}
+
+		bool operator<(const FDynamicOperatorTransactor::FLiteralNodeID& InLHS, const FDynamicOperatorTransactor::FLiteralNodeID& InRHS)
+		{
+			if (InLHS.ToNode < InRHS.ToNode)
+			{
+				return true;
+			}
+			else if (InRHS.ToNode < InLHS.ToNode)
+			{
+				return false;
+			}
+			else
+			{
+				return InLHS.ToVertex.FastLess(InRHS.ToVertex);
+			}
+		}
+
+		FDynamicOperatorTransactor::FDynamicOperatorTransactor(const IGraph& InGraph)
+		: Graph(InGraph)
+		{
+		}
+
+		FDynamicOperatorTransactor::FDynamicOperatorTransactor()
+		: Graph(TEXT(""), FGuid())
+		{
+		}
+
+		TSharedRef<TSpscQueue<TUniquePtr<IDynamicOperatorTransform>>> FDynamicOperatorTransactor::CreateTransformQueue(const FOperatorSettings& InOperatorSettings, const FMetasoundEnvironment& InEnvironment)
+		{
+			TSharedRef<TSpscQueue<TUniquePtr<IDynamicOperatorTransform>>> Queue = MakeShared<TSpscQueue<TUniquePtr<IDynamicOperatorTransform>>>();
+			OperatorInfos.Add(FDynamicOperatorInfo{InOperatorSettings, InEnvironment, Queue});
+
+			return Queue;
+		}
+
+		void FDynamicOperatorTransactor::AddNode(const FGuid& InNodeID, TUniquePtr<INode> InNode)
+		{
+			using namespace DirectedGraphAlgo;
+
+			METASOUND_TRACE_CPUPROFILER_EVENT_SCOPE(Metasound::FDynamicOperatorTransactor::AddNode);
+
+			if (!InNode)
+			{
+				return;
+			}
+
+			// Cache underlying pointer to node for later use. TUniquePtr<INode> will be moved to node
+			// map thus invalidating `TUniquePtr<INode> InNode`. 
+			INode* NodePtr = InNode.Get();
+
+			NodeMap.Add(InNodeID, MoveTemp(InNode));
+
+			auto CreateAddNodeTransform = [&](const FOperatorSettings& InOperatorSettings, const FMetasoundEnvironment& InEnvironment) -> TUniquePtr<IDynamicOperatorTransform>
+			{
+				return CreateAddOperatorTransform(*NodePtr, InOperatorSettings, InEnvironment);
+			};
+
+			EnqueueTransformOnOperatorQueues(CreateAddNodeTransform);
+		}
+
+		void FDynamicOperatorTransactor::RemoveNode(const FGuid& InNodeID)
+		{
+			using namespace DirectedGraphAlgo;
+
+			METASOUND_TRACE_CPUPROFILER_EVENT_SCOPE(Metasound::FDynamicOperatorTransactor::RemoveNode);
+
+			if (const INode* Node = FindNode(InNodeID))
+			{
+				FOperatorID OperatorID = GetOperatorID(Node);
+				auto CreateRemoveNodeTransform = [&](const FOperatorSettings& InOperatorSettings, const FMetasoundEnvironment& InEnvironment) -> TUniquePtr<IDynamicOperatorTransform>
+				{
+					return MakeUnique<FRemoveOperator>(OperatorID);
+				};
+				EnqueueTransformOnOperatorQueues(CreateRemoveNodeTransform);
+
+				Graph.RemoveDataEdgesWithNode(*Node);
+				NodeMap.Remove(InNodeID);
+			}
+			else
+			{
+				UE_LOG(LogMetaSound, Error, TEXT("No node found in dynamic transactor with ID %s"), *InNodeID.ToString());
+			}
+		}
+
+		/** Add an edge to the graph. */
+		void FDynamicOperatorTransactor::AddDataEdge(const FGuid& InFromNodeID, const FVertexName& InFromVertex, const FGuid& InToNodeID, const FVertexName& InToVertex)
+		{
+			using namespace DirectedGraphAlgo; 
+
+			METASOUND_TRACE_CPUPROFILER_EVENT_SCOPE(Metasound::FDynamicOperatorTransactor::AddDataEdge);
+
+			const INode* FromNode = FindNode(InFromNodeID);
+			const INode* ToNode = FindNode(InToNodeID);
+
+			if ((nullptr == ToNode) || (nullptr == FromNode))
+			{
+				UE_LOG(LogMetaSound, Error, TEXT("Cannot add edge from %s:%s to %s:%s because of missing node"), *InFromNodeID.ToString(), *InFromVertex.ToString(), *InToNodeID.ToString(), *InToVertex.ToString());
+				return;
+			}
+
+			AddDataEdgeInternal(*FromNode, InFromVertex, InToNodeID, *ToNode, InToVertex);
+		}
+
+		void FDynamicOperatorTransactor::RemoveDataEdge(const FGuid& InFromNodeID, const FVertexName& InFromVertex, const FGuid& InToNodeID, const FVertexName& InToVertex, TUniquePtr<INode> InReplacementLiteralNode)
+		{
+			using namespace DirectedGraphAlgo; 
+
+			METASOUND_TRACE_CPUPROFILER_EVENT_SCOPE(Metasound::FDynamicOperatorTransactor::RemoveDataEdge);
+
+			const INode* FromNode = FindNode(InFromNodeID);
+			const INode* ToNode = FindNode(InToNodeID);
+			const INode* ReplacementLiteralNode = InReplacementLiteralNode.Get(); // Cache pointer because TUniquePtr<INode> will get moved
+
+			if ((nullptr == ToNode) || (nullptr == FromNode))
+			{
+				UE_LOG(LogMetaSound, Error, TEXT("Cannot remove edge from %s:%s to %s:%s because of missing node"), *InFromNodeID.ToString(), *InFromVertex.ToString(), *InToNodeID.ToString(), *InToVertex.ToString());
+				return;
+			}
+
+			if (!ToNode->GetVertexInterface().ContainsInputVertex(InToVertex))
+			{
+				UE_LOG(LogMetaSound, Error, TEXT("Cannot remove edge from %s:%s to %s:%s because of destination node does not contain vertex %s."), *InFromNodeID.ToString(), *InFromVertex.ToString(), *InToNodeID.ToString(), *InToVertex.ToString(), *InToVertex.ToString());
+				return;
+			}
+
+			if (!InReplacementLiteralNode.IsValid())
+			{
+				UE_LOG(LogMetaSound, Error, TEXT("Cannot remove edge from %s:%s to %s:%s because of invalid pointer to replacement literal node."), *InFromNodeID.ToString(), *InFromVertex.ToString(), *InToNodeID.ToString(), *InToVertex.ToString());
+				return;
+			}
+
+			// Store literal node associated with the target of the literal value.
+			LiteralNodeMap.Add(FLiteralNodeID{InToNodeID, InToVertex}, MoveTemp(InReplacementLiteralNode));
+			/* remove edge from internal graph. */
+			bool bSuccess = Graph.RemoveDataEdge(*FromNode, InFromVertex, *ToNode, InToVertex);
+			if (!bSuccess)
+			{
+				UE_LOG(LogMetaSound, Warning, TEXT("Failed to remove edge from %s:%s to %s:%s on internal graph."), *InFromNodeID.ToString(), *InFromVertex.ToString(), *InToNodeID.ToString(), *InToVertex.ToString());
+				// No need to early exit here because edge likely did not exist.
+			}
+
+			bSuccess = Graph.AddDataEdge(*ReplacementLiteralNode, DynamicOperatorTransactorPrivate::LiteralNodeOutputVertexName, *ToNode, InToVertex);
+			if (!bSuccess)
+			{
+				UE_LOG(LogMetaSound, Error, TEXT("Failed to add literal for %s:%s on internal graph."), *InToNodeID.ToString(), *InToVertex.ToString());
+				return;
+			}
+
+			/* enqueue an update. */
+			const FOperatorID FromOperatorID = GetOperatorID(FromNode);
+			const FOperatorID LiteralOperatorID = GetOperatorID(ReplacementLiteralNode);
+			const FOperatorID ToOperatorID = GetOperatorID(ToNode);
+			TArray<FOperatorID> OperatorOrder = DynamicOperatorTransactorPrivate::DetermineOperatorOrder(Graph);
+
+			auto CreateRemoveEdgeTransform = [&](const FOperatorSettings& InOperatorSettings, const FMetasoundEnvironment& InEnvironment) -> TUniquePtr<IDynamicOperatorTransform>
+			{
+
+				// Add the literal node.
+				TUniquePtr<IDynamicOperatorTransform> AddNodeTransform = CreateAddOperatorTransform(*ReplacementLiteralNode, InOperatorSettings, InEnvironment);
+
+				// Connect literal node to target node.
+				TUniquePtr<IDynamicOperatorTransform> ConnectOperatorsTransform = MakeUnique<FConnectOperators>(LiteralOperatorID, DynamicOperatorTransactorPrivate::LiteralNodeOutputVertexName, ToOperatorID, InToVertex);
+
+				// Reorder the node graph.
+				TUniquePtr<IDynamicOperatorTransform> SetOrderTransform = MakeUnique<FSetOperatorOrder>(OperatorOrder);
+
+				if (AddNodeTransform.IsValid() && ConnectOperatorsTransform.IsValid() && SetOrderTransform.IsValid())
+				{
+					// Create an atomic transform so all sub-transforms happen before next execution.
+					TArray<TUniquePtr<IDynamicOperatorTransform>> AtomicTransforms;
+
+					AtomicTransforms.Emplace(MoveTemp(AddNodeTransform));
+					AtomicTransforms.Emplace(MoveTemp(ConnectOperatorsTransform));
+					AtomicTransforms.Emplace(MoveTemp(SetOrderTransform));
+
+					return MakeUnique<FAtomicTransform>(MoveTemp(AtomicTransforms));
+				}
+				else
+				{
+					UE_LOG(LogMetaSound, Error, TEXT("Cannot remove edge from %s:%s to %s:%s because of failure to create all transforms needed to perform operatorn."), *InFromNodeID.ToString(), *InFromVertex.ToString(), *InToNodeID.ToString(), *InToVertex.ToString());
+					return TUniquePtr<IDynamicOperatorTransform>(nullptr);
+				}
+			};
+
+			EnqueueTransformOnOperatorQueues(CreateRemoveEdgeTransform);
+		}
+
+		void FDynamicOperatorTransactor::SetValue(const FGuid& InNodeID, const FVertexName& InVertex, TUniquePtr<INode> InLiteralNode)
+		{
+			METASOUND_TRACE_CPUPROFILER_EVENT_SCOPE(Metasound::FDynamicOperatorTransactor::SetValue);
+
+			const INode* Node = FindNode(InNodeID);
+			const INode* LiteralNode = InLiteralNode.Get();
+
+			if (nullptr == Node)
+			{
+				UE_LOG(LogMetaSound, Error, TEXT("Cannot set node value of %s:%s because of missing node"), *InNodeID.ToString(), *InVertex.ToString());
+				return;
+			}
+
+			if (!InLiteralNode.IsValid())
+			{
+				UE_LOG(LogMetaSound, Error, TEXT("Cannot set value on %s:%s because of invalid pointer to literal node."), *InNodeID.ToString(), *InVertex.ToString());
+				return;
+			}
+
+			auto CreateAddNodeTransform = [&](const FOperatorSettings& InOperatorSettings, const FMetasoundEnvironment& InEnvironment) -> TUniquePtr<IDynamicOperatorTransform>
+			{
+				return CreateAddOperatorTransform(*LiteralNode, InOperatorSettings, InEnvironment);
+			};
+
+			EnqueueTransformOnOperatorQueues(CreateAddNodeTransform);
+
+			AddDataEdgeInternal(*LiteralNode, DynamicOperatorTransactorPrivate::LiteralNodeOutputVertexName, InNodeID, *Node, InVertex);
+
+			// Add literal node before calling "AddDataEdgeInternal" so that AddDataEdgeInternal can check if there is an existing literal node.
+			LiteralNodeMap.Add(FLiteralNodeID{InNodeID, InVertex}, MoveTemp(InLiteralNode));
+		}
+
+		/** Add an input data destination to describe how data provided 
+		 * outside this graph should be routed internally.
+		 *
+		 * @param InNode - Node which receives the data.
+		 * @param InVertexName - Key for input vertex on InNode.
+		 *
+		 */
+		void FDynamicOperatorTransactor::AddInputDataDestination(const FGuid& InNodeID, const FVertexName& InVertexName, const FLiteral& InDefaultLiteral, FReferenceCreationFunction InFunc)
+		{
+			METASOUND_TRACE_CPUPROFILER_EVENT_SCOPE(Metasound::FDynamicOperatorTransactor::AddInputDataDestination);
+
+			const INode* Node = FindNode(InNodeID);
+			if (nullptr == Node)
+			{
+				UE_LOG(LogMetaSound, Error, TEXT("Cannot add Input Data Destination %s:%s because of missing node"), *InNodeID.ToString(), *InVertexName.ToString());
+				return;
+			}
+
+			if (!Node->GetVertexInterface().ContainsInputVertex(InVertexName))
+			{
+				UE_LOG(LogMetaSound, Error, TEXT("Cannot add Input Data Destination %s:%s because of node does not contain input vertex with name %s."), *InNodeID.ToString(), *InVertexName.ToString(), *InVertexName.ToString());
+				return;
+			}
+
+			const FInputDataVertex& InputVertex = Node->GetVertexInterface().GetInputVertex(InVertexName);
+			EDataReferenceAccessType ReferenceAccessType = InputVertex.AccessType == EVertexAccessType::Value ? EDataReferenceAccessType::Value : EDataReferenceAccessType::Write;
+			FOperatorID OperatorID = DirectedGraphAlgo::GetOperatorID(Node);
+
+			Graph.AddInputDataDestination(*Node, InVertexName);
+
+			auto CreateAddInputTransform = [&](const FOperatorSettings& InOperatorSettings, const FMetasoundEnvironment& InEnvironment) -> TUniquePtr<IDynamicOperatorTransform>
+			{
+				TOptional<FAnyDataReference> NewDataReference = InFunc(InOperatorSettings, InDefaultLiteral, ReferenceAccessType);
+				if (NewDataReference.IsSet())
+				{
+					return MakeUnique<FAddInput>(OperatorID, InVertexName, *NewDataReference);
+				}
+				else
+				{
+					UE_LOG(LogMetaSound, Error, TEXT("Cannot add Input Data Destination %s:%s because of failure to create data reference."), *InNodeID.ToString(), *InVertexName.ToString());
+					return TUniquePtr<IDynamicOperatorTransform>(nullptr);
+				}
+			};
+
+			EnqueueTransformOnOperatorQueues(CreateAddInputTransform);
+		}
+
+		void FDynamicOperatorTransactor::RemoveInputDataDestination(const FVertexName& InVertexName)
+		{
+			METASOUND_TRACE_CPUPROFILER_EVENT_SCOPE(Metasound::FDynamicOperatorTransactor::RemoveInputDataDestination);
+
+			Graph.RemoveInputDataDestination(InVertexName);
+
+			auto CreateRemoveInputTransform = [&](const FOperatorSettings& InOperatorSettings, const FMetasoundEnvironment& InEnvironment) -> TUniquePtr<IDynamicOperatorTransform>
+			{
+				return MakeUnique<FRemoveInput>(InVertexName);
+			};
+
+			EnqueueTransformOnOperatorQueues(CreateRemoveInputTransform);
+		}
+
+		/** Add an output data source which describes routing of data which is 
+		 * owned this graph and exposed externally.
+		 *
+		 * @param InNode - Node which produces the data.
+		 * @param InVertexName - Key for output vertex on InNode.
+		 */
+		void FDynamicOperatorTransactor::AddOutputDataSource(const FGuid& InNodeID, const FVertexName& InVertexName)
+		{
+			METASOUND_TRACE_CPUPROFILER_EVENT_SCOPE(Metasound::FDynamicOperatorTransactor::AddOutputDataSource);
+
+			const INode* Node = FindNode(InNodeID);
+			if (nullptr == Node)
+			{
+				UE_LOG(LogMetaSound, Error, TEXT("Cannot add Output Data Source %s:%s because of missing node"), *InNodeID.ToString(), *InVertexName.ToString());
+				return;
+			}
+
+			Graph.AddOutputDataSource(*Node, InVertexName);
+			const FOperatorID OperatorID = DirectedGraphAlgo::GetOperatorID(Node);
+
+			auto CreateAddOutputTransform = [&](const FOperatorSettings& InOperatorSettings, const FMetasoundEnvironment& InEnvironment) -> TUniquePtr<IDynamicOperatorTransform>
+			{
+				return MakeUnique<FAddOutput>(OperatorID, InVertexName);
+			};
+
+			EnqueueTransformOnOperatorQueues(CreateAddOutputTransform);
+		}
+
+		void FDynamicOperatorTransactor::RemoveOutputDataSource(const FVertexName& InVertexName)
+		{
+			METASOUND_TRACE_CPUPROFILER_EVENT_SCOPE(Metasound::FDynamicOperatorTransactor::RemoveOutputDataSource);
+
+			Graph.RemoveOutputDataSource(InVertexName);
+
+			auto CreateRemoveOutputTransform = [&](const FOperatorSettings& InOperatorSettings, const FMetasoundEnvironment& InEnvironment) -> TUniquePtr<IDynamicOperatorTransform>
+			{
+				return MakeUnique<FRemoveOutput>(InVertexName);
+			};
+
+			EnqueueTransformOnOperatorQueues(CreateRemoveOutputTransform);
+		}
+
+		const IGraph& FDynamicOperatorTransactor::GetGraph() const
+		{
+			return Graph;
+		}
+
+		void FDynamicOperatorTransactor::AddDataEdgeInternal(const INode& InFromNode, const FVertexName& InFromVertex, const FGuid& InToNodeID, const INode& InToNode, const FVertexName& InToVertex)
+		{
+			using namespace DirectedGraphAlgo;
+
+			/* Determine if there is an existing literal node connected to the node. */
+			const FLiteralNodeID LiteralNodeKey{InToNodeID, InToVertex};
+			bool bIsAlreadyConnectedToALiteralNode = false;
+			FOperatorID PriorLiteralOperatorID;
+			if (const TUniquePtr<INode>* PriorLiteralNode = LiteralNodeMap.Find(LiteralNodeKey))
+			{
+				bIsAlreadyConnectedToALiteralNode = true;
+				PriorLiteralOperatorID = GetOperatorID(PriorLiteralNode->Get());
+				
+				Graph.RemoveDataEdge(**PriorLiteralNode, DynamicOperatorTransactorPrivate::LiteralNodeOutputVertexName, InToNode, InToVertex);
+				LiteralNodeMap.Remove(LiteralNodeKey);
+			}
+
+			/* add edge to internal graph. */
+			Graph.AddDataEdge(InFromNode, InFromVertex, InToNode, InToVertex);
+
+			// Find order of operators after adding edge. 
+			TArray<FOperatorID> OperatorOrder = DynamicOperatorTransactorPrivate::DetermineOperatorOrder(Graph);
+
+			/* enqueue an update. */
+			FOperatorID FromOperatorID = GetOperatorID(InFromNode);
+			FOperatorID ToOperatorID = GetOperatorID(InToNode);
+
+			auto CreateAddEdgeTransform = [&](const FOperatorSettings& InOperatorSettings, const FMetasoundEnvironment& InEnvironment) -> TUniquePtr<IDynamicOperatorTransform>
+			{
+				TArray<TUniquePtr<IDynamicOperatorTransform>> AtomicTransforms;
+
+				if (bIsAlreadyConnectedToALiteralNode)
+				{
+					AtomicTransforms.Add(MakeUnique<FRemoveOperator>(PriorLiteralOperatorID));
+				}
+				AtomicTransforms.Add(MakeUnique<FSetOperatorOrder>(OperatorOrder));
+				AtomicTransforms.Add(MakeUnique<FConnectOperators>(FromOperatorID, InFromVertex, ToOperatorID, InToVertex));
+
+				return MakeUnique<FAtomicTransform>(MoveTemp(AtomicTransforms));
+			};
+
+			EnqueueTransformOnOperatorQueues(CreateAddEdgeTransform);
+		}
+
+		const INode* FDynamicOperatorTransactor::FindNode(const FGuid& InNodeID) const
+		{
+			if (const TUniquePtr<INode>* NodePtr = NodeMap.Find(InNodeID))
+			{
+				return NodePtr->Get();
+			}
+			else
+			{
+				UE_LOG(LogMetaSound, Error, TEXT("Graph does not contain node with ID %s"), *InNodeID.ToString());
+				return nullptr;
+			}
+		}
+
+		TUniquePtr<IDynamicOperatorTransform> FDynamicOperatorTransactor::CreateAddOperatorTransform(const INode& InNode, const FOperatorSettings& InOperatorSettings, const FMetasoundEnvironment& InEnvironment) const
+		{
+			const FOperatorID OperatorID = DirectedGraphAlgo::GetOperatorID(InNode);
+			FVertexInterfaceData InterfaceData(InNode.GetVertexInterface());
+			FBuildOperatorParams OperatorParams
+			{
+				InNode,
+				InOperatorSettings,
+				InterfaceData.GetInputs(),
+				InEnvironment
+			};
+
+			FBuildResults Results;
+			TUniquePtr<IOperator> Operator = InNode.GetDefaultOperatorFactory()->CreateOperator(OperatorParams, Results);
+
+			for (const TUniquePtr<IOperatorBuildError>& Error : Results.Errors)
+			{
+				if (Error.IsValid())
+				{
+					UE_LOG(LogMetaSound, Warning, TEXT("Encountered error while building operator for node %s. %s:%s"), *InNode.GetMetadata().ClassName.GetFullName().ToString(), *Error->GetErrorType().ToString(), *Error->GetErrorDescription().ToString());
+				}
+			}
+
+			if (Operator.IsValid())
+			{
+				Operator->BindInputs(InterfaceData.GetInputs());
+				Operator->BindOutputs(InterfaceData.GetOutputs());
+
+				FOperatorInfo OpInfo
+				{
+					MoveTemp(Operator),
+					MoveTemp(InterfaceData)
+				};
+
+				return MakeUnique<FAddOperator>(OperatorID, MoveTemp(OpInfo));
+			}
+			else
+			{
+				return TUniquePtr<IDynamicOperatorTransform>(nullptr);
+			}
+		}
+
+
+		void FDynamicOperatorTransactor::EnqueueTransformOnOperatorQueues(FCreateTransformFunctionRef InFunc)
+		{
+			TArray<FDynamicOperatorInfo>::TIterator OperatorInfoIterator = OperatorInfos.CreateIterator();
+			while (OperatorInfoIterator)
+			{
+				TSharedPtr<TSpscQueue<TUniquePtr<IDynamicOperatorTransform>>> OperatorQueue = OperatorInfoIterator->Queue.Pin();
+				if (OperatorQueue.IsValid())
+				{
+					TUniquePtr<IDynamicOperatorTransform> Transform = InFunc(OperatorInfoIterator->OperatorSettings, OperatorInfoIterator->Environment);
+					if (Transform.IsValid())
+					{
+						OperatorQueue->Enqueue(MoveTemp(Transform));
+					}
+				}
+				else
+				{
+					OperatorInfoIterator.RemoveCurrent();
+				}
+				OperatorInfoIterator++;
+			}
+		}
+	}
+}
diff --git a/Engine/Plugins/Runtime/Metasound/Source/MetasoundGraphCore/Private/MetasoundGraph.cpp b/Engine/Plugins/Runtime/Metasound/Source/MetasoundGraphCore/Private/MetasoundGraph.cpp
--- a/Engine/Plugins/Runtime/Metasound/Source/MetasoundGraphCore/Private/MetasoundGraph.cpp
+++ b/Engine/Plugins/Runtime/Metasound/Source/MetasoundGraphCore/Private/MetasoundGraph.cpp
@@ -5,6 +5,16 @@
 
 namespace Metasound
 {
+	FGraph::FGraph(const IGraph& InOther)
+	: InstanceName(InOther.GetInstanceName())
+	, InstanceID(InOther.GetInstanceID())
+	, Metadata(InOther.GetMetadata())
+	, Edges(InOther.GetDataEdges())
+	, InputDestinations(InOther.GetInputDataDestinations())
+	, OutputSources(InOther.GetOutputDataSources())
+	{
+	}
+
 	FGraph::FGraph(const FString& InInstanceName, const FGuid& InInstanceID)
 	: InstanceName(InInstanceName)
 	, InstanceID(InInstanceID)
@@ -47,6 +57,24 @@ namespace Metasound
 		InputDestinations.Add(MakeDestinationDataVertexKey(InDestination), InDestination);
 	}
 
+	bool FGraph::RemoveInputDataDestination(const FVertexName& InVertexName)
+	{
+		bool bSuccessRemoveFromInterface = Metadata.DefaultInterface.GetInputInterface().Remove(InVertexName);
+
+		FInputDataDestinationCollection::TIterator Iter = InputDestinations.CreateIterator();
+		bool bSuccessRemoveFromDestinations = false;
+		do
+		{
+			if (Iter.Value().Vertex.VertexName == InVertexName)
+			{
+				Iter.RemoveCurrent();
+				bSuccessRemoveFromDestinations = true;
+			}
+		} while(++Iter);
+
+		return bSuccessRemoveFromInterface && bSuccessRemoveFromDestinations;
+	}
+
 	const FInputDataDestinationCollection& FGraph::GetInputDataDestinations() const
 	{
 		return InputDestinations;
@@ -72,6 +100,24 @@ namespace Metasound
 		OutputSources.Add(MakeSourceDataVertexKey(InSource), InSource);
 	}
 
+	bool FGraph::RemoveOutputDataSource(const FVertexName& InVertexName)
+	{
+		bool bSuccessRemoveFromInterface = Metadata.DefaultInterface.GetOutputInterface().Remove(InVertexName);
+
+		FOutputDataSourceCollection::TIterator Iter = OutputSources.CreateIterator();
+		bool bSuccessRemoveFromSources = false;
+		do
+		{
+			if (Iter.Value().Vertex.VertexName == InVertexName)
+			{
+				Iter.RemoveCurrent();
+				bSuccessRemoveFromSources = true;
+			}
+		} while(++Iter);
+
+		return bSuccessRemoveFromInterface && bSuccessRemoveFromSources;
+	}
+
 	const FOutputDataSourceCollection& FGraph::GetOutputDataSources() const
 	{
 		return OutputSources;
@@ -113,6 +159,26 @@ namespace Metasound
 		return true;
 	}
 
+	bool FGraph::RemoveDataEdge(const INode& FromNode, const FVertexName& FromVertexKey, const INode& ToNode, const FVertexName& ToVertexKey)
+	{
+		auto IsTargetDataEdge = [&](const FDataEdge& InEdge)
+		{
+			return (InEdge.To.Node == &ToNode) && (InEdge.To.Vertex.VertexName == ToVertexKey) && (InEdge.From.Node == &FromNode) && (InEdge.From.Vertex.VertexName == FromVertexKey);
+		};
+		int32 NumRemoved = Edges.RemoveAllSwap(IsTargetDataEdge);
+
+		return (NumRemoved > 0);
+	}
+
+	void FGraph::RemoveDataEdgesWithNode(const INode& InNode)
+	{
+		auto ContainsNode = [&](const FDataEdge& InEdge)
+		{
+			return (InEdge.To.Node == &InNode) || (InEdge.From.Node == &InNode);
+		};
+		Edges.RemoveAllSwap(ContainsNode);
+	}
+
 	const TArray<FDataEdge>& FGraph::GetDataEdges() const
 	{
 		return Edges;
diff --git a/Engine/Plugins/Runtime/Metasound/Source/MetasoundGraphCore/Private/MetasoundGraphAlgo.cpp b/Engine/Plugins/Runtime/Metasound/Source/MetasoundGraphCore/Private/MetasoundGraphAlgo.cpp
--- a/Engine/Plugins/Runtime/Metasound/Source/MetasoundGraphCore/Private/MetasoundGraphAlgo.cpp
+++ b/Engine/Plugins/Runtime/Metasound/Source/MetasoundGraphCore/Private/MetasoundGraphAlgo.cpp
@@ -304,7 +304,6 @@ namespace Metasound
 
 	namespace DirectedGraphAlgo
 	{
-		
 		TPimplPtr<FDirectedGraphAlgoAdapter> CreateDirectedGraphAlgoAdapter(const IGraph& InGraph)
 		{
 			return MakePimpl<FDirectedGraphAlgoAdapter>(InGraph);
@@ -475,68 +474,6 @@ namespace Metasound
 			OutNodes = NodesFromInput.Union(NodesFromOutput);
 		}
 
-		namespace GraphAlgoPrivate
-		{
-			struct FInputToUpdate
-			{
-				FOperatorID OperatorID;
-				FVertexName VertexName;
-				FAnyDataReference DataReference;
-			};
-
-			// Apply updates to data references through all the operators by following connections described in the FOperatorInfo map.
-			void PropagateBindUpdate(FOperatorID InInitialOperatorID, const FVertexName& InVertexName, const FAnyDataReference& InNewReference, TSortedMap<FOperatorID, FGraphOperatorData::FOperatorInfo>& InOperatorInfoMap, FOnPostOperatorRebind OnPostOperatorRebind)
-			{
-				TArray<FInputToUpdate> UpdateStack;
-				UpdateStack.Emplace(FInputToUpdate{InInitialOperatorID, InVertexName, InNewReference});
-
-				TArray<FVertexDataState> InitialOutputState;
-				TSortedVertexNameMap<FAnyDataReference> OutputUpdates;
-				while (UpdateStack.Num())
-				{
-					FInputToUpdate Current = UpdateStack.Pop();
-					if (FGraphOperatorData::FOperatorInfo* OpInfo = InOperatorInfoMap.Find(Current.OperatorID))
-					{
-						// Get current outputs
-						InitialOutputState.Reset();
-						GetVertexInterfaceDataState(OpInfo->VertexData.GetOutputs(), InitialOutputState);
-
-						// Set new input.
-						OpInfo->VertexData.GetInputs().SetVertex(Current.VertexName, MoveTemp(Current.DataReference));
-
-						// Bind inputs and outputs
-						OpInfo->Operator->BindInputs(OpInfo->VertexData.GetInputs());
-						OpInfo->Operator->BindOutputs(OpInfo->VertexData.GetOutputs());
-						OnPostOperatorRebind(Current.OperatorID, *(OpInfo->Operator));
-
-						// See if binding altered the outputs. 
-						OutputUpdates.Reset();
-						CompareVertexInterfaceDataToPriorState(OpInfo->VertexData.GetOutputs(), InitialOutputState, OutputUpdates);
-
-						// Any updates to the outputs need to be propagated through the graph.
-						for (const TPair<FVertexName, FAnyDataReference>& OutputUpdate : OutputUpdates)
-						{
-							const FVertexName& OutputVertexName = OutputUpdate.Get<0>();
-							const FAnyDataReference& OutputDataReference = OutputUpdate.Get<1>();
-
-							if (const TArray<FGraphOperatorData::FVertexDestination>* Destinations = OpInfo->OutputConnections.Find(OutputVertexName))
-							{
-								for (const FGraphOperatorData::FVertexDestination& Destination : *Destinations)
-								{
-									UpdateStack.Push({Destination.OperatorID, Destination.VertexName, OutputDataReference});
-								}
-							}
-						}
-					}
-					else
-					{
-						UE_LOG(LogMetaSound, Error, TEXT("Failed to rebind graph operator state. Could not find operator info with ID %s"), *LexToString(Current.OperatorID));
-					}
-				}
-			}
-
-		}
-
 		FOperatorID GetOperatorID(const INode& InNode)
 		{
 			return GetOperatorID(&InNode);
@@ -551,63 +488,5 @@ namespace Metasound
 		: OperatorSettings(InOperatorSettings)
 		{
 		}
-		
-		void RebindGraphInputs(FInputVertexInterfaceData& InOutVertexData, FRebindGraphDataParams& InParams, FOnPostOperatorRebind OnPostOperatorRebind)
-		{
-			// Bind and diff the graph's interface to determine if there is an update to any vertices
-			FInputVertexInterfaceData& InputVertexData = InParams.VertexData.GetInputs();
-			TArray<FVertexDataState> InitialVertexDataState;
-			GetVertexInterfaceDataState(InputVertexData, InitialVertexDataState);
-
-			InOutVertexData.Bind(InputVertexData);
-
-			TSortedVertexNameMap<FAnyDataReference> GraphInputsToUpdate;
-			CompareVertexInterfaceDataToPriorState(InputVertexData, InitialVertexDataState, GraphInputsToUpdate);
-
-			if (GraphInputsToUpdate.Num() > 0)
-			{
-				// Update the graph data by propagating the updates through the inputs nodes. 
-				for (const TPair<FVertexName, FAnyDataReference>& InputToUpdate : GraphInputsToUpdate)
-				{
-					const FVertexName& VertexName = InputToUpdate.Get<0>();
-					if (const FOperatorID* OperatorID = InParams.InputVertexMap.Find(VertexName))
-					{
-						GraphAlgoPrivate::PropagateBindUpdate(*OperatorID, VertexName, InputToUpdate.Get<1>(), InParams.OperatorMap, OnPostOperatorRebind);
-					}
-					else
-					{
-						UE_LOG(LogMetaSound, Error, TEXT("No input operator exists for input vertex %s"), *VertexName.ToString());
-					}
-				}
-
-				// Refresh output vertex interface data.
-				for (const TPair<FVertexName, FOperatorID>& OutputVertexInfo : InParams.OutputVertexMap)
-				{
-					const FVertexName& VertexName = OutputVertexInfo.Get<0>();
-					const FOperatorID& OperatorID = OutputVertexInfo.Get<1>();
-
-					if (const FGraphOperatorData::FOperatorInfo* OperatorInfo = InParams.OperatorMap.Find(OperatorID))
-					{
-						if (const FAnyDataReference* Ref = OperatorInfo->VertexData.GetOutputs().FindDataReference(OutputVertexInfo.Get<0>()))
-						{
-							InParams.VertexData.GetOutputs().SetVertex(VertexName, *Ref);
-						}
-						else if (InParams.VertexData.GetOutputs().IsVertexBound(VertexName))
-						{
-							UE_LOG(LogMetaSound, Error, TEXT("Output vertex (%s) lost data reference after rebinding graph"), *VertexName.ToString());
-						}
-					}
-					else
-					{
-						UE_LOG(LogMetaSound, Error, TEXT("Failed to rebind graph operator state outputs. Could not find output operator info with ID %s"), *LexToString(OperatorID));
-					}
-				}
-			}
-		}
-
-		void RebindGraphOutputs(FOutputVertexInterfaceData& InOutVertexData, FRebindGraphDataParams& InParams, FOnPostOperatorRebind OnPostOperatorRebind)
-		{
-			InOutVertexData.Bind(InParams.VertexData.GetOutputs());
-		}
 	}
 }
diff --git a/Engine/Plugins/Runtime/Metasound/Source/MetasoundGraphCore/Private/MetasoundGraphAlgoPrivate.h b/Engine/Plugins/Runtime/Metasound/Source/MetasoundGraphCore/Private/MetasoundGraphAlgoPrivate.h
--- a/Engine/Plugins/Runtime/Metasound/Source/MetasoundGraphCore/Private/MetasoundGraphAlgoPrivate.h
+++ b/Engine/Plugins/Runtime/Metasound/Source/MetasoundGraphCore/Private/MetasoundGraphAlgoPrivate.h
@@ -62,26 +62,5 @@ namespace Metasound
 			// Map with output vertex name as key, and OperatorID of output node as value.
 			TSortedVertexNameMap<FOperatorID> OutputVertexMap;
 		};
-
-		// Arguments for calling Rebind
-		struct FRebindGraphDataParams
-		{
-			FVertexInterfaceData& VertexData; 											//< Current vertex data of graph
-			TSortedMap<FOperatorID, FGraphOperatorData::FOperatorInfo>& OperatorMap; 	//< Map of operator data in graph
-			TSortedVertexNameMap<FOperatorID>& InputVertexMap; 							//< Map of graph input nodes
-			TSortedVertexNameMap<FOperatorID>& OutputVertexMap; 						//< Map of graph output nodes
-		};
-
-		using FOnPostOperatorRebind = TFunctionRef<void(FOperatorID, IOperator&)>;
-
-		// Rebind the graph inputs, updating internal operator bindings as needed
-		//
-		// When an internal node is rebound, the OnPostOperatorRebind will be called. 
-		void RebindGraphInputs(FInputVertexInterfaceData& InOutVertexData, FRebindGraphDataParams& InOutParams, FOnPostOperatorRebind OnPostOperatorRebind);
-
-		// Rebind the graph inputs, updating internal operator bindings as needed
-		//
-		// When an internal node is rebound, the OnPostOperatorRebind will be called. 
-		void RebindGraphOutputs(FOutputVertexInterfaceData& InOutVertexData, FRebindGraphDataParams& InOutParams, FOnPostOperatorRebind OnPostOperatorRebind);
 	}
 }
diff --git a/Engine/Plugins/Runtime/Metasound/Source/MetasoundGraphCore/Private/MetasoundOperatorBuilder.cpp b/Engine/Plugins/Runtime/Metasound/Source/MetasoundGraphCore/Private/MetasoundOperatorBuilder.cpp
--- a/Engine/Plugins/Runtime/Metasound/Source/MetasoundGraphCore/Private/MetasoundOperatorBuilder.cpp
+++ b/Engine/Plugins/Runtime/Metasound/Source/MetasoundGraphCore/Private/MetasoundOperatorBuilder.cpp
@@ -7,6 +7,8 @@
 #include "MetasoundBuildError.h"
 #include "MetasoundBuilderInterface.h"
 #include "MetasoundBuildError.h"
+#include "MetasoundDynamicOperator.h"
+#include "MetasoundDynamicOperatorTransactor.h"
 #include "MetasoundGraphAlgo.h"
 #include "MetasoundGraphAlgoPrivate.h"
 #include "MetasoundGraphLinter.h"
@@ -29,7 +31,6 @@ namespace Metasound
 			const FDirectedGraphAlgoAdapter& AlgoAdapter;
 			const FOperatorSettings& Settings;
 			const FMetasoundEnvironment& Environment;
-			const FOperatorBuilderSettings& BuilderSettings;
 			
 			FBuildResults& Results;
 			TUniquePtr<DirectedGraphAlgo::FGraphOperatorData> GraphOperatorData;
@@ -39,13 +40,11 @@ namespace Metasound
 				const FDirectedGraphAlgoAdapter& InAlgoAdapter,
 				const FOperatorSettings& InSettings,
 				const FMetasoundEnvironment& InEnvironment,
-				const FOperatorBuilderSettings& bInBuilderSettings,
 				FBuildResults& OutResults)
 			: Graph(InGraph)
 			, AlgoAdapter(InAlgoAdapter)
 			, Settings(InSettings)
 			, Environment(InEnvironment)
-			, BuilderSettings(bInBuilderSettings)
 			, Results(OutResults)
 			, GraphOperatorData(MakeUnique<DirectedGraphAlgo::FGraphOperatorData>(InSettings))
 			{
@@ -78,10 +77,40 @@ namespace Metasound
 	}
 
 	TUniquePtr<IOperator> FOperatorBuilder::BuildGraphOperator(const FBuildGraphOperatorParams& InParams, FBuildResults& OutResults) const
+	{
+		TUniquePtr<DirectedGraphAlgo::FGraphOperatorData> GraphData = BuildGraphOperatorData(InParams, OutResults);
+
+		if (GraphData.IsValid())
+		{
+			// Create graph operator from collection of node operators.
+			return CreateGraphOperator(MoveTemp(GraphData));
+		}
+
+		return TUniquePtr<IOperator>(nullptr);
+	}
+
+	TUniquePtr<IOperator> FOperatorBuilder::BuildDynamicGraphOperator(const FBuildGraphOperatorParams& InParams, DynamicGraph::FDynamicOperatorTransactor& InTransactor, FBuildResults& OutResults)
+	{
+		using namespace DynamicGraph;
+
+		TUniquePtr<DirectedGraphAlgo::FGraphOperatorData> GraphData = BuildGraphOperatorData(InParams, OutResults);
+
+		if (GraphData.IsValid())
+		{
+			TSharedRef<TSpscQueue<TUniquePtr<IDynamicOperatorTransform>>> TransformQueue = InTransactor.CreateTransformQueue(InParams.OperatorSettings, InParams.Environment);
+
+			return MakeUnique<FDynamicOperator>(MoveTemp(*GraphData), TransformQueue);
+		}
+
+		return TUniquePtr<IOperator>(nullptr);
+	}
+
+	TUniquePtr<DirectedGraphAlgo::FGraphOperatorData> FOperatorBuilder::BuildGraphOperatorData(const FBuildGraphOperatorParams& InParams, FBuildResults& OutResults) const
 	{
 		METASOUND_TRACE_CPUPROFILER_EVENT_SCOPE(Metasound::FOperatorBuilder::BuildGraphOperator);
 
 		using namespace OperatorBuilderPrivate;
+		using namespace DirectedGraphAlgo;
 
 		FBuildStatus BuildStatus = FBuildStatus::NoError;
 
@@ -116,7 +145,7 @@ namespace Metasound
 		// Possible early exit if edge validation fails.
 		if (BuildStatus > GetMaxErrorLevel())
 		{
-			return TUniquePtr<IOperator>(nullptr);
+			return TUniquePtr<FGraphOperatorData>(nullptr);
 		}
 
 		// Create algo adapter view of graph to cache graph operations.
@@ -125,10 +154,10 @@ namespace Metasound
 		if (!AlgoAdapter.IsValid())
 		{
 			AddBuildError<FInternalError>(OutResults.Errors, __FILE__, __LINE__);
-			return TUniquePtr<IOperator>(nullptr);
+			return TUniquePtr<FGraphOperatorData>(nullptr);
 		}
 
-		OperatorBuilder::FBuildContext BuildContext(InParams.Graph, *AlgoAdapter, InParams.OperatorSettings, InParams.Environment, BuilderSettings, OutResults);
+		OperatorBuilder::FBuildContext BuildContext(InParams.Graph, *AlgoAdapter, InParams.OperatorSettings, InParams.Environment, OutResults);
 
 		TArray<const INode*> SortedNodes;
 
@@ -143,7 +172,7 @@ namespace Metasound
 		// Check build status in case build routine should be exited early.
 		if (BuildStatus > GetMaxErrorLevel())
 		{
-			return TUniquePtr<IOperator>(nullptr);
+			return TUniquePtr<FGraphOperatorData>(nullptr);
 		}
 
 		InitializeOperatorInfo(InParams.Graph, SortedNodes, *BuildContext.GraphOperatorData);
@@ -154,7 +183,7 @@ namespace Metasound
 		// Check build status in case build routine should be exited early.
 		if (BuildStatus > GetMaxErrorLevel())
 		{
-			return TUniquePtr<IOperator>(nullptr);
+			return TUniquePtr<FGraphOperatorData>(nullptr);
 		}
 
 		// Create node operators from factories.
@@ -162,7 +191,7 @@ namespace Metasound
 
 		if (BuildStatus > GetMaxErrorLevel())
 		{
-			return TUniquePtr<IOperator>(nullptr);
+			return TUniquePtr<FGraphOperatorData>(nullptr);
 		}
 
 		if (BuilderSettings.bPopulateInternalDataReferences)
@@ -170,8 +199,15 @@ namespace Metasound
 			GatherInternalGraphDataReferences(BuildContext, SortedNodes, BuildContext.Results.InternalDataReferences);
 		}
 
-		// Create graph operator from collection of node operators.
-		return CreateGraphOperator(BuildContext);
+		// Gather the inputs for the graph data. 
+		BuildStatus |= GatherGraphDataReferences(BuildContext, BuildContext.GraphOperatorData->VertexData);
+
+		if (BuildStatus > GetMaxErrorLevel())
+		{
+			return TUniquePtr<FGraphOperatorData>(nullptr);
+		}
+
+		return MoveTemp(BuildContext.GraphOperatorData);
 	}
 
 	FOperatorBuilder::FBuildStatus FOperatorBuilder::DepthFirstTopologicalSort(OperatorBuilder::FBuildContext& InOutContext, TArray<const INode*>& OutNodes) const
@@ -540,25 +576,19 @@ namespace Metasound
 		return BuildStatus;
 	}
 
-	TUniquePtr<IOperator> FOperatorBuilder::CreateGraphOperator(OperatorBuilder::FBuildContext& InOutContext) const
+	TUniquePtr<IOperator> FOperatorBuilder::CreateGraphOperator(TUniquePtr<DirectedGraphAlgo::FGraphOperatorData>&& InGraphOperatorData) const
 	{
 		METASOUND_TRACE_CPUPROFILER_EVENT_SCOPE(Metasound::FOperatorBuilder::CreateGraphOperator);
-		FBuildStatus BuildStatus = GatherGraphDataReferences(InOutContext, InOutContext.GraphOperatorData->VertexData);
 
 		TUniquePtr<IOperator> GraphOperator;
 
-		if (BuildStatus > GetMaxErrorLevel())
-		{
-			return GraphOperator;
-		}
-
-		if (InOutContext.BuilderSettings.bEnableOperatorRebind)
+		if (BuilderSettings.bEnableOperatorRebind)
 		{
-			 GraphOperator = MakeUnique<FRebindableGraphOperator>(MoveTemp(InOutContext.GraphOperatorData));
+			 GraphOperator = MakeUnique<FRebindableGraphOperator>(MoveTemp(*InGraphOperatorData));
 		}
 		else
 		{
-			 GraphOperator = MakeUnique<FGraphOperator>(MoveTemp(InOutContext.GraphOperatorData));
+			 GraphOperator = MakeUnique<FGraphOperator>(MoveTemp(InGraphOperatorData));
 		}
 
 		return GraphOperator;
diff --git a/Engine/Plugins/Runtime/Metasound/Source/MetasoundGraphCore/Private/MetasoundRebindableGraphOperator.cpp b/Engine/Plugins/Runtime/Metasound/Source/MetasoundGraphCore/Private/MetasoundRebindableGraphOperator.cpp
--- a/Engine/Plugins/Runtime/Metasound/Source/MetasoundGraphCore/Private/MetasoundRebindableGraphOperator.cpp
+++ b/Engine/Plugins/Runtime/Metasound/Source/MetasoundGraphCore/Private/MetasoundRebindableGraphOperator.cpp
@@ -4,8 +4,8 @@
 
 #include "Containers/Array.h"
 #include "Containers/SortedMap.h"
+#include "MetasoundDynamicGraphAlgo.h"
 #include "MetasoundExecutableOperator.h"
-#include "MetasoundGraphAlgoPrivate.h"
 #include "MetasoundOperatorInterface.h"
 #include "MetasoundVertex.h"
 #include "MetasoundVertexData.h"
@@ -14,37 +14,9 @@
 
 namespace Metasound
 {
-	FRebindableGraphOperator::FRebindableGraphOperator(TUniquePtr<DirectedGraphAlgo::FGraphOperatorData> InOperatorData)
+	FRebindableGraphOperator::FRebindableGraphOperator(DirectedGraphAlgo::FGraphOperatorData&& InOperatorData)
 	: GraphOperatorData(MoveTemp(InOperatorData))
 	{
-		using namespace DirectedGraphAlgo;
-
-		if (GraphOperatorData.IsValid())
-		{
-			for (DirectedGraphAlgo::FOperatorID OperatorID : GraphOperatorData->OperatorOrder)
-			{
-				if (FGraphOperatorData::FOperatorInfo* OperatorInfo = GraphOperatorData->OperatorMap.Find(OperatorID))
-				{
-					check(OperatorInfo->Operator.IsValid());
-					IOperator& Operator = *(OperatorInfo->Operator);
-
-					if (IOperator::FExecuteFunction ExecuteFunc = Operator.GetExecuteFunction())
-					{
-						ExecuteStack.Emplace(FExecuteEntry{OperatorID, Operator, ExecuteFunc});
-					}
-
-					if (IOperator::FPostExecuteFunction PostExecuteFunc = Operator.GetPostExecuteFunction())
-					{
-						PostExecuteStack.Emplace(FPostExecuteEntry{OperatorID, Operator, PostExecuteFunc});
-					}
-
-					if (IOperator::FResetFunction ResetFunc = Operator.GetResetFunction())
-					{
-						ResetStack.Emplace(FResetEntry{OperatorID, Operator, ResetFunc});
-					}
-				}
-			}
-		}
 	}
 
 	FDataReferenceCollection FRebindableGraphOperator::GetInputs() const
@@ -62,53 +34,23 @@ namespace Metasound
 	// Bind the graph's interface data references to FVertexInterfaceData.
 	void FRebindableGraphOperator::BindInputs(FInputVertexInterfaceData& InOutVertexData)
 	{
-		using namespace DirectedGraphAlgo;
+		using namespace DynamicGraph;
 
-		if (GraphOperatorData.IsValid())
-		{
-			auto UpdateStackOnRebind = [&](FOperatorID InOperatorID, IOperator& InOperator) 
-			{ 
-				UpdateStacks(InOperatorID, InOperator);
-			};
-
-			FRebindGraphDataParams Params
-			{
-				GraphOperatorData->VertexData,
-				GraphOperatorData->OperatorMap,
-				GraphOperatorData->InputVertexMap,
-				GraphOperatorData->OutputVertexMap
-			};
-
-			RebindGraphInputs(InOutVertexData, Params, UpdateStackOnRebind);
-		}
+		RebindGraphInputs(InOutVertexData, GraphOperatorData);
 	}
 
 	void FRebindableGraphOperator::BindOutputs(FOutputVertexInterfaceData& InOutVertexData)
 	{
-		using namespace DirectedGraphAlgo;
-		if (GraphOperatorData.IsValid())
-		{
-			auto UpdateStackOnRebind = [&](FOperatorID InOperatorID, IOperator& InOperator) 
-			{ 
-				UpdateStacks(InOperatorID, InOperator);
-			};
-
-			FRebindGraphDataParams Params
-			{
-				GraphOperatorData->VertexData,
-				GraphOperatorData->OperatorMap,
-				GraphOperatorData->InputVertexMap,
-				GraphOperatorData->OutputVertexMap
-			};
-
-			RebindGraphOutputs(InOutVertexData, Params, UpdateStackOnRebind);
-		}
-	}
+		using namespace DynamicGraph;
 
+		RebindGraphOutputs(InOutVertexData, GraphOperatorData);
+	}
 
 	void FRebindableGraphOperator::Execute()
 	{
-		for (FExecuteEntry& Entry : ExecuteStack)
+		using namespace DynamicGraph;
+
+		for (FExecuteEntry& Entry : GraphOperatorData.ExecuteTable)
 		{
 			Entry.Execute();
 		}
@@ -116,7 +58,9 @@ namespace Metasound
 
 	void FRebindableGraphOperator::PostExecute()
 	{
-		for (FPostExecuteEntry& Entry : PostExecuteStack)
+		using namespace DynamicGraph;
+
+		for (FPostExecuteEntry& Entry : GraphOperatorData.PostExecuteTable)
 		{
 			Entry.PostExecute();
 		}
@@ -124,7 +68,9 @@ namespace Metasound
 
 	void FRebindableGraphOperator::Reset(const IOperator::FResetParams& InParams)
 	{
-		for (FResetEntry& Entry : ResetStack)
+		using namespace DynamicGraph;
+
+		for (FResetEntry& Entry : GraphOperatorData.ResetTable)
 		{
 			Entry.Reset(InParams);
 		}
@@ -140,84 +86,6 @@ namespace Metasound
 		check(InOperator);
 		static_cast<FRebindableGraphOperator*>(InOperator)->PostExecute();
 	}
-
-	template<typename EntryType, typename FunctionType>
-	void UpdateOperatorAndStack(const TArray<DirectedGraphAlgo::FOperatorID>& InOperatorOrder, TArray<EntryType>& InStack, DirectedGraphAlgo::FOperatorID InOperatorID, IOperator& InOperator, FunctionType InFunction)
-	{
-		using namespace DirectedGraphAlgo;
-
-		int32 EntryIndex = InStack.IndexOfByPredicate([&](const EntryType& InEntry) { return InEntry.OperatorID == InOperatorID; });
-
-		bool bEntryShouldExist = (nullptr != InFunction);
-		bool bEntryCurrentlyExists = (EntryIndex != INDEX_NONE);
-
-		if (bEntryCurrentlyExists && bEntryShouldExist)
-		{
-			// Update the existing entry
-			InStack[EntryIndex].Function = InFunction;
-		}
-		else if (bEntryCurrentlyExists)
-		{
-			// Remove the existing entry
-			InStack.RemoveAt(EntryIndex);
-		}
-		else if (bEntryShouldExist)
-		{
-			// Add new entry
-			int32 OperatorIndex = InOperatorOrder.Find(InOperatorID);
-			if (INDEX_NONE == OperatorIndex)
-			{
-				UE_LOG(LogMetaSound, Error, TEXT("Failed to update operator stack. OperatorID not found in operator order array."));
-				return;
-			}
-
-			int32 PreceedingOperatorIndex = OperatorIndex - 1;
-			int32 InsertLocation = 0; // Default to inserting at the beginning
-			while (PreceedingOperatorIndex >= 0)
-			{
-				FOperatorID PreceedingOperatorID = InOperatorOrder[PreceedingOperatorIndex];
-				PreceedingOperatorIndex--;
-
-				int32 PreceedingEntryIndex = InStack.IndexOfByPredicate([&](const EntryType& InEntry) { return InEntry.OperatorID == PreceedingOperatorID; });
-				if (PreceedingEntryIndex != INDEX_NONE)
-				{
-					InsertLocation = PreceedingEntryIndex + 1;
-					break;
-				}
-			}
-
-			InStack.Insert(EntryType{InOperatorID, InOperator, InFunction}, InsertLocation);
-		}
-	}
-
-	void FRebindableGraphOperator::UpdateStacks(DirectedGraphAlgo::FOperatorID InOperatorID, IOperator& InOperator)
-	{
-		check(GraphOperatorData);
-		UpdateOperatorAndStack(GraphOperatorData->OperatorOrder, ExecuteStack, InOperatorID, InOperator, InOperator.GetExecuteFunction());
-		UpdateOperatorAndStack(GraphOperatorData->OperatorOrder, PostExecuteStack, InOperatorID, InOperator, InOperator.GetPostExecuteFunction());
-		UpdateOperatorAndStack(GraphOperatorData->OperatorOrder, ResetStack, InOperatorID, InOperator, InOperator.GetResetFunction());
-	}
-
-	FRebindableGraphOperator::FExecuteEntry::FExecuteEntry(DirectedGraphAlgo::FOperatorID InOperatorID, IOperator& InOperator, IOperator::FExecuteFunction InFunc)
-	: OperatorID(InOperatorID)
-	, Operator(&InOperator)
-	, Function(InFunc)
-	{
-	}
-
-	FRebindableGraphOperator::FPostExecuteEntry::FPostExecuteEntry(DirectedGraphAlgo::FOperatorID InOperatorID, IOperator& InOperator, IOperator::FPostExecuteFunction InFunc)
-	: OperatorID(InOperatorID)
-	, Operator(&InOperator)
-	, Function(InFunc)
-	{
-	}
-
-	FRebindableGraphOperator::FResetEntry::FResetEntry(DirectedGraphAlgo::FOperatorID InOperatorID, IOperator& InOperator, IOperator::FResetFunction InFunc)
-	: OperatorID(InOperatorID)
-	, Operator(&InOperator)
-	, Function(InFunc)
-	{
-	}
 }
 
 
diff --git a/Engine/Plugins/Runtime/Metasound/Source/MetasoundGraphCore/Private/MetasoundRebindableGraphOperator.h b/Engine/Plugins/Runtime/Metasound/Source/MetasoundGraphCore/Private/MetasoundRebindableGraphOperator.h
--- a/Engine/Plugins/Runtime/Metasound/Source/MetasoundGraphCore/Private/MetasoundRebindableGraphOperator.h
+++ b/Engine/Plugins/Runtime/Metasound/Source/MetasoundGraphCore/Private/MetasoundRebindableGraphOperator.h
@@ -2,23 +2,22 @@
 
 #pragma once
 
-#include "Containers/Array.h"
-#include "Containers/SortedMap.h"
 #include "MetasoundExecutableOperator.h"
-#include "MetasoundGraphAlgoPrivate.h"
+#include "MetasoundDynamicGraphAlgo.h"
 #include "MetasoundOperatorInterface.h"
-#include "MetasoundVertex.h"
-#include "MetasoundVertexData.h"
-#include "Misc/Guid.h"
-#include "Templates/UniquePtr.h"
 
 namespace Metasound
 {
+	class FDataReferenceCollection;
+	class FInputVertexInterfaceData;
+	class FOutputVertexInterfaceData;
+
+
 	class FRebindableGraphOperator : public TExecutableOperator<FRebindableGraphOperator>
 	{
 	public:
 
-		FRebindableGraphOperator(TUniquePtr<DirectedGraphAlgo::FGraphOperatorData> InOperatorState);
+		FRebindableGraphOperator(DirectedGraphAlgo::FGraphOperatorData&& InOperatorState);
 
 		virtual FDataReferenceCollection GetInputs() const override;
 
@@ -37,61 +36,7 @@ namespace Metasound
 	private:
 		static void StaticPostExecute(IOperator* InOperator);
 
-		void UpdateStacks(DirectedGraphAlgo::FOperatorID InOperatorID, IOperator& InOperator);
-
-		struct FExecuteEntry
-		{
-			FExecuteEntry(DirectedGraphAlgo::FOperatorID InOperatorID, IOperator& InOperator, IOperator::FExecuteFunction InFunc);
-
-			void Execute()
-			{
-				check(Operator);
-				check(Function);
-				Function(Operator);
-			}
-
-			DirectedGraphAlgo::FOperatorID OperatorID;
-			IOperator* Operator;
-			IOperator::FExecuteFunction Function;	
-		};
-
-		struct FPostExecuteEntry
-		{
-			FPostExecuteEntry(DirectedGraphAlgo::FOperatorID InOperatorID, IOperator& InOperator, IOperator::FPostExecuteFunction InFunc);
-
-			void PostExecute()
-			{
-				check(Operator);
-				check(Function);
-				Function(Operator);
-			}
-
-			DirectedGraphAlgo::FOperatorID OperatorID;
-			IOperator* Operator;
-			IOperator::FPostExecuteFunction Function;	
-		};
-
-		struct FResetEntry
-		{
-			FResetEntry(DirectedGraphAlgo::FOperatorID InOperatorID, IOperator& InOperator, IOperator::FResetFunction InFunc);
-
-			void Reset(const IOperator::FResetParams& InParams)
-			{
-				check(Operator);
-				check(Function);
-				Function(Operator, InParams);
-			}
-
-			DirectedGraphAlgo::FOperatorID OperatorID;
-			IOperator* Operator;
-			IOperator::FResetFunction Function;	
-		};
-
-		TArray<FExecuteEntry> ExecuteStack;
-		TArray<FPostExecuteEntry> PostExecuteStack;
-		TArray<FResetEntry> ResetStack;
-
-		TUniquePtr<DirectedGraphAlgo::FGraphOperatorData> GraphOperatorData;
+		DynamicGraph::FDynamicGraphOperatorData GraphOperatorData;
 	};
 }
 
diff --git a/Engine/Plugins/Runtime/Metasound/Source/MetasoundGraphCore/Private/MetasoundVertexData.cpp b/Engine/Plugins/Runtime/Metasound/Source/MetasoundGraphCore/Private/MetasoundVertexData.cpp
--- a/Engine/Plugins/Runtime/Metasound/Source/MetasoundGraphCore/Private/MetasoundVertexData.cpp
+++ b/Engine/Plugins/Runtime/Metasound/Source/MetasoundGraphCore/Private/MetasoundVertexData.cpp
@@ -72,8 +72,6 @@ namespace Metasound
 			Set(MoveTemp(InReference));
 		}
 
-
-
 		void FInputBinding::Set(FAnyDataReference&& InAnyDataReference)
 		{
 			check(Vertex.DataTypeName == InAnyDataReference.GetDataTypeName());
@@ -502,6 +500,30 @@ namespace Metasound
 		return FindChecked(InVertexName)->GetVertex();
 	}
 
+	void FInputVertexInterfaceData::AddVertex(const FInputDataVertex& InVertex)
+	{
+		if (IsVertexInterfaceFrozen())
+		{
+			UE_LOG(LogMetaSound, Warning, TEXT("Cannot add vertex with name '%s'. Vertex interface is frozen"), *InVertex.VertexName.ToString());
+		}
+		else
+		{
+			Bindings.Emplace(InVertex);
+		}
+	}
+
+	void FInputVertexInterfaceData::RemoveVertex(const FVertexName& InVertexName)
+	{
+		if (IsVertexInterfaceFrozen())
+		{
+			UE_LOG(LogMetaSound, Warning, TEXT("Cannot remove vertex with name '%s'. Vertex interface is frozen"), *InVertexName.ToString());
+		}
+		else
+		{
+			Bindings.RemoveAll([&](const FInputBinding& Binding) { return InVertexName == Binding.GetVertex().VertexName; });
+		}
+	}
+
 	EDataReferenceAccessType FInputVertexInterfaceData::GetVertexDataAccessType(const FVertexName& InVertexName) const
 	{
 		return MetasoundVertexDataPrivate::GetVertexDataAccessType(Bindings, InVertexName);
@@ -640,6 +662,30 @@ namespace Metasound
 		return FindChecked(InVertexName)->GetVertex();
 	}
 
+	void FOutputVertexInterfaceData::AddVertex(const FOutputDataVertex& InVertex)
+	{
+		if (IsVertexInterfaceFrozen())
+		{
+			UE_LOG(LogMetaSound, Warning, TEXT("Cannot add vertex with name '%s'. Vertex interface is frozen"), *InVertex.VertexName.ToString());
+		}
+		else
+		{
+			Bindings.Emplace(InVertex);
+		}
+	}
+
+	void FOutputVertexInterfaceData::RemoveVertex(const FVertexName& InVertexName)
+	{
+		if (IsVertexInterfaceFrozen())
+		{
+			UE_LOG(LogMetaSound, Warning, TEXT("Cannot remove vertex with name '%s'. Vertex interface is frozen"), *InVertexName.ToString());
+		}
+		else
+		{
+			Bindings.RemoveAll([&](const FOutputBinding& Binding) { return InVertexName == Binding.GetVertex().VertexName; });
+		}
+	}
+
 	EDataReferenceAccessType FOutputVertexInterfaceData::GetVertexDataAccessType(const FVertexName& InVertexName) const
 	{
 		return MetasoundVertexDataPrivate::GetVertexDataAccessType(Bindings, InVertexName);
diff --git a/Engine/Plugins/Runtime/Metasound/Source/MetasoundGraphCore/Public/MetasoundDynamicOperatorTransactor.h b/Engine/Plugins/Runtime/Metasound/Source/MetasoundGraphCore/Public/MetasoundDynamicOperatorTransactor.h
new file mode 100644
--- /dev/null
+++ b/Engine/Plugins/Runtime/Metasound/Source/MetasoundGraphCore/Public/MetasoundDynamicOperatorTransactor.h
@@ -0,0 +1,122 @@
+// Copyright Epic Games, Inc. All Rights Reserved.
+
+#pragma once
+
+#include "Containers/SpscQueue.h"
+#include "Containers/UnrealString.h"
+#include "MetasoundBuilderInterface.h"
+#include "MetasoundGraph.h"
+#include "MetasoundOperatorInterface.h"
+#include "MetasoundNodeInterface.h"
+#include "Templates/SharedPointer.h"
+#include "Templates/UniquePtr.h"
+
+namespace Metasound
+{
+	class FOperatorSettings;
+	struct FLiteral;
+	class FAnyDataReference;
+
+	namespace DynamicGraph
+	{
+		class IDynamicOperatorTransform;
+		class FDynamicOperator;
+
+		using FLiteralAssignmentFunction = void(*)(const FOperatorSettings& InOperatorSettings, const FLiteral& InLiteral, const FAnyDataReference& OutDataRef);
+		using FReferenceCreationFunction = TOptional<FAnyDataReference>(*)(const FOperatorSettings& InSettings, const FLiteral& InLiteral, EDataReferenceAccessType InAccessType);
+
+		class METASOUNDGRAPHCORE_API FDynamicOperatorTransactor
+		{
+		public:
+			FDynamicOperatorTransactor& operator=(const FDynamicOperatorTransactor&) = delete;
+			FDynamicOperatorTransactor(const FDynamicOperatorTransactor&) = delete;
+
+			FDynamicOperatorTransactor();
+			FDynamicOperatorTransactor(const IGraph& InGraph);
+
+			/** Create a queue for communication with a dynamic operator. */
+			TSharedRef<TSpscQueue<TUniquePtr<IDynamicOperatorTransform>>> CreateTransformQueue(const FOperatorSettings& InOperatorSettings, const FMetasoundEnvironment& InEnvironment);
+
+			/** Add a node to the graph. */
+			void AddNode(const FGuid& InNodeID, TUniquePtr<INode> InNode);
+
+			/** Remove a node from the graph. */
+			void RemoveNode(const FGuid& InNodeID);
+
+			/** Add an edge to the graph, connecting two vertices from two 
+			 * nodes. 
+			 *
+			 * @param FromNode - Node which contains the output vertex.
+			 * @param FromVertex - Key of the vertex in the FromNode.
+			 * @param ToNode - Node which contains the input vertex.
+			 * @param ToVertex - Key of the vertex in the ToNode.
+			 */
+			void AddDataEdge(const FGuid& InFromNodeID, const FVertexName& InFromVertex, const FGuid& InToNodeID, const FVertexName& InToVertex);
+
+			/** Remove the given data edge. */
+			void RemoveDataEdge(const FGuid& InFromNode, const FVertexName& InFromVertex, const FGuid& InToNode, const FVertexName& InToVertex, TUniquePtr<INode> InReplacementLiteralNode);
+
+			/** Set the value on a unconnected node input vertex. */
+			void SetValue(const FGuid& InNodeID, const FVertexName& InVertex, TUniquePtr<INode> InLiteralNode);
+
+			/** Add an input data destination to describe how data provided 
+			 * outside this graph should be routed internally.
+			 *
+			 * @param InNode - Node which receives the data.
+			 * @param InVertexName - Key for input vertex on InNode.
+			 */
+			void AddInputDataDestination(const FGuid& InNode, const FVertexName& InVertexName, const FLiteral& InDefaultLiteral, FReferenceCreationFunction InFunc);
+
+			/** Remove an exposed input from the graph. */
+			void RemoveInputDataDestination(const FVertexName& InVertexName);
+
+			/** Add an output data source which describes routing of data which is 
+			 * owned this graph and exposed externally.
+			 *
+			 * @param InNode - Node which produces the data.
+			 * @param InVertexName - Key for output vertex on InNode.
+			 */
+			void AddOutputDataSource(const FGuid& InNode, const FVertexName& InVertexName);
+
+			/** Remove an exposed output from the graph. */
+			void RemoveOutputDataSource(const FVertexName& InVertexName);
+
+			/** Return internal version of graph. */
+			const IGraph& GetGraph() const;
+
+		private:
+
+			void AddDataEdgeInternal(const INode& InFromNode, const FVertexName& InFromVertex, const FGuid& InToNodeID, const INode& InToNode, const FVertexName& InToVertex);
+
+			const INode* FindNode(const FGuid& InNodeID) const;
+
+			using FCreateTransformFunctionRef = TFunctionRef<TUniquePtr<IDynamicOperatorTransform>(const FOperatorSettings& InOperatorSettings, const FMetasoundEnvironment& InEnvironment)>;
+
+			TUniquePtr<IDynamicOperatorTransform> CreateAddOperatorTransform(const INode& InNode, const FOperatorSettings& InOperatorSettings, const FMetasoundEnvironment& InEnvironment) const;
+
+			void EnqueueTransformOnOperatorQueues(FCreateTransformFunctionRef InFunc);
+
+			FGraph Graph;
+
+			struct FDynamicOperatorInfo
+			{
+				FOperatorSettings OperatorSettings;
+				FMetasoundEnvironment Environment;
+				TWeakPtr< TSpscQueue< TUniquePtr< IDynamicOperatorTransform >>> Queue;
+			};
+
+			TArray<FDynamicOperatorInfo> OperatorInfos;
+
+			TSortedMap<FGuid, TUniquePtr<INode>> NodeMap;
+			struct FLiteralNodeID
+			{
+				FGuid ToNode;
+				FVertexName ToVertex;
+			};
+			friend bool operator<(const FLiteralNodeID& InLHS, const FLiteralNodeID& InRHS);
+
+			TSortedMap<FLiteralNodeID, TUniquePtr<INode>> LiteralNodeMap;
+		};
+	}
+}
+
diff --git a/Engine/Plugins/Runtime/Metasound/Source/MetasoundGraphCore/Public/MetasoundGraph.h b/Engine/Plugins/Runtime/Metasound/Source/MetasoundGraphCore/Public/MetasoundGraph.h
--- a/Engine/Plugins/Runtime/Metasound/Source/MetasoundGraphCore/Public/MetasoundGraph.h
+++ b/Engine/Plugins/Runtime/Metasound/Source/MetasoundGraphCore/Public/MetasoundGraph.h
@@ -17,23 +17,26 @@ namespace Metasound
 	class METASOUNDGRAPHCORE_API FGraph : public IGraph
 	{
 		public:
+
+			FGraph(const IGraph& InOther);
 			FGraph(const FString& InInstanceName, const FGuid& InInstanceID);
+
 			virtual ~FGraph() = default;
 
 			/** Return the name of this specific instance of the node class. */
-			const FVertexName& GetInstanceName() const override;
+			virtual const FVertexName& GetInstanceName() const override;
 
 			/** Return the ID of this specific instance of the node class. */
-			const FGuid& GetInstanceID() const override;
+			virtual const FGuid& GetInstanceID() const override;
 
 			/** Return metadata about this graph. */
-			const FNodeClassMetadata& GetMetadata() const override;
+			virtual const FNodeClassMetadata& GetMetadata() const override;
 
 			/** Retrieve all the edges associated with a graph. */
-			const TArray<FDataEdge>& GetDataEdges() const override;
+			virtual const TArray<FDataEdge>& GetDataEdges() const override;
 
 			/** Return the current vertex interface. */
-			const FVertexInterface& GetVertexInterface() const override;
+			virtual const FVertexInterface& GetVertexInterface() const override;
 
 			/** Set the vertex interface. If the vertex was successfully changed, returns true. 
 			 *
@@ -41,7 +44,7 @@ namespace Metasound
 			 *
 			 * @return True on success, false otherwise.
 			 */
-			bool SetVertexInterface(const FVertexInterface& InInterface) override;
+			virtual bool SetVertexInterface(const FVertexInterface& InInterface) override;
 
 			/** Expresses whether a specific vertex interface is supported.
 			 *
@@ -49,13 +52,13 @@ namespace Metasound
 			 *
 			 * @return True if the interface is supported, false otherwise. 
 			 */
-			bool IsVertexInterfaceSupported(const FVertexInterface& InInterface) const override;
+			virtual bool IsVertexInterfaceSupported(const FVertexInterface& InInterface) const override;
 
 			/** Get vertices which contain input parameters. */
-			const FInputDataDestinationCollection& GetInputDataDestinations() const override;
+			virtual const FInputDataDestinationCollection& GetInputDataDestinations() const override;
 
 			/** Get vertices which contain output parameters. */
-			const FOutputDataSourceCollection& GetOutputDataSources() const override;
+			virtual const FOutputDataSourceCollection& GetOutputDataSources() const override;
 
 			/** Add an edge to the graph. */
 			void AddDataEdge(const FDataEdge& InEdge);
@@ -84,11 +87,15 @@ namespace Metasound
 			 *                    the edge should be removed.
 			 */
 			template<typename PredicateType>
+			UE_DEPRECATED(5.3, "Removing data edges by predicate is no longer supported")
 			void RemoveDataEdgeByPredicate(const PredicateType& Predicate)
 			{
 				Edges.RemoveAllSwap(Predicate);
 			}
 
+			/** Removes all edges with connected to the node. */
+			void RemoveDataEdgesWithNode(const INode& InNode);
+
 			/** Add an input data destination to describe how data provided 
 			 * outside this graph should be routed internally.
 			 *
@@ -106,6 +113,14 @@ namespace Metasound
 			 */
 			void AddInputDataDestination(const FInputDataDestination& InDestination);
 
+			/** Remove an input data destination by vertex name
+			 * 
+			 * @param InVertexName - Name of vertex to remove.
+			 *
+			 * @return True if destination existed and was successfully removed. False otherwise.
+			 */
+			bool RemoveInputDataDestination(const FVertexName& InVertexName);
+
 			/** Add an output data source which describes routing of data which is 
 			 * owned this graph and exposed externally.
 			 *
@@ -122,12 +137,19 @@ namespace Metasound
 			 */
 			void AddOutputDataSource(const FOutputDataSource& InSource);
 
-			// TODO: Add ability to remove things.
+			/** Remove an output data source by vertex name
+			 * 
+			 * @param InVertexName - Name of vertex to remove.
+			 *
+			 * @return True if source existed and was successfully removed. False otherwise.
+			 */
+			bool RemoveOutputDataSource(const FVertexName& InVertexName);
 
 			/** Return a reference to the default operator factory. */
 			virtual FOperatorFactorySharedRef GetDefaultOperatorFactory() const override;
 
 		private:
+
 			class FFactory : public IOperatorFactory
 			{
 			public:
diff --git a/Engine/Plugins/Runtime/Metasound/Source/MetasoundGraphCore/Public/MetasoundOperatorBuilder.h b/Engine/Plugins/Runtime/Metasound/Source/MetasoundGraphCore/Public/MetasoundOperatorBuilder.h
--- a/Engine/Plugins/Runtime/Metasound/Source/MetasoundGraphCore/Public/MetasoundOperatorBuilder.h
+++ b/Engine/Plugins/Runtime/Metasound/Source/MetasoundGraphCore/Public/MetasoundOperatorBuilder.h
@@ -26,6 +26,11 @@ namespace Metasound
 		struct FBuildContext;
 	}
 
+	namespace DynamicGraph
+	{
+		class FDynamicOperatorTransactor;
+	}
+
 	/** FOperatorBuilder builds an IOperator from an IGraph. */
 	class METASOUNDGRAPHCORE_API FOperatorBuilder : public IOperatorBuilder
 	{
@@ -49,8 +54,21 @@ namespace Metasound
 			 */
 			virtual TUniquePtr<IOperator> BuildGraphOperator(const FBuildGraphOperatorParams& InParams, FBuildResults& OutResults) const override;
 
+			/** Create an dynamic IOperator from an IGraph.
+			 *
+			 * @param InParams     - Params of the current build
+			 * @param InTransactor - A dynamic operator transactor for communicating to the operator.
+			 * @param OutResults   - Results data pertaining to the given build operator result.
+			 *
+			 * @return A TUniquePtr to an IOperator. If the processes was unsuccessful, 
+			 *         the returned pointer will contain a nullptr and be invalid.
+			 */
+			TUniquePtr<IOperator> BuildDynamicGraphOperator(const FBuildGraphOperatorParams& InParams, DynamicGraph::FDynamicOperatorTransactor& InTransactor, FBuildResults& OutResults);
+
 		private:
 
+			TUniquePtr<DirectedGraphAlgo::FGraphOperatorData> BuildGraphOperatorData(const FBuildGraphOperatorParams& InParams, FBuildResults& OutResults) const;
+
 			// Handles build status of current build operation.
 			struct FBuildStatus
 			{
@@ -126,7 +144,7 @@ namespace Metasound
 			FBuildStatus CreateOperators(OperatorBuilder::FBuildContext& InOutContext, const TArray<const INode*>& InSortedNodes, const FInputVertexInterfaceData& InExternalInputData) const;
 
 			// Create the final graph operator from the provided build context.
-			TUniquePtr<IOperator> CreateGraphOperator(OperatorBuilder::FBuildContext& InOutContext) const;
+			TUniquePtr<IOperator> CreateGraphOperator(TUniquePtr<DirectedGraphAlgo::FGraphOperatorData>&& InGraphOperatorData) const;
 
 			FBuildStatus::EStatus GetMaxErrorLevel() const;
 
diff --git a/Engine/Plugins/Runtime/Metasound/Source/MetasoundGraphCore/Public/MetasoundVertex.h b/Engine/Plugins/Runtime/Metasound/Source/MetasoundGraphCore/Public/MetasoundVertex.h
--- a/Engine/Plugins/Runtime/Metasound/Source/MetasoundGraphCore/Public/MetasoundVertex.h
+++ b/Engine/Plugins/Runtime/Metasound/Source/MetasoundGraphCore/Public/MetasoundVertex.h
@@ -323,7 +323,7 @@ namespace Metasound
 			{
 			}
 
-			bool operator()(const VertexType& InOther)
+			bool operator()(const VertexType& InOther) const
 			{
 				return InOther.VertexName == NameRef;
 			}
diff --git a/Engine/Plugins/Runtime/Metasound/Source/MetasoundGraphCore/Public/MetasoundVertexData.h b/Engine/Plugins/Runtime/Metasound/Source/MetasoundGraphCore/Public/MetasoundVertexData.h
--- a/Engine/Plugins/Runtime/Metasound/Source/MetasoundGraphCore/Public/MetasoundVertexData.h
+++ b/Engine/Plugins/Runtime/Metasound/Source/MetasoundGraphCore/Public/MetasoundVertexData.h
@@ -597,6 +597,12 @@ namespace Metasound
 		/** Return the vertex associated with the vertex name. */
 		const FInputDataVertex& GetVertex(const FVertexName& InVertexName) const;
 
+		/** Add a vertex. VertexInterfaceData must be unfrozen. */
+		void AddVertex(const FInputDataVertex& InVertex);
+
+		/** Remove a vertex. VertexInterfaceData must be unfrozen. */
+		void RemoveVertex(const FVertexName& InVertexName);
+
 		/** Returns true if a vertex with the given vertex name exists and is bound
 		 * to a data reference. */
 		bool IsVertexBound(const FVertexName& InVertexName) const;
@@ -908,6 +914,12 @@ namespace Metasound
 		/** Return the vertex associated with the vertex name. */
 		const FOutputDataVertex& GetVertex(const FVertexName& InVertexName) const;
 
+		/** Add a vertex. VertexInterfaceData must be unfrozen. */
+		void AddVertex(const FOutputDataVertex& InVertex);
+
+		/** Remove a vertex. VertexInterfaceData must be unfrozen. */
+		void RemoveVertex(const FVertexName& InVertexName);
+
 		/** Returns the access type of a bound vertex. If the vertex does not exist
 		 * or if it is unbound, this will return EDataReferenceAccessType::None
 		 */
