Copy soundbase/soundwave settings from referenced MetaSound source to new preset
- Add option to BuildToAsset to copy settings from a given template SoundWave
- Update CreatePreset right click asset action to use builder instead of factory

diff --git a/Engine/Plugins/Runtime/Metasound/Source/MetasoundEditor/Private/MetasoundAssetTypeActions.cpp b/Engine/Plugins/Runtime/Metasound/Source/MetasoundEditor/Private/MetasoundAssetTypeActions.cpp
--- a/Engine/Plugins/Runtime/Metasound/Source/MetasoundEditor/Private/MetasoundAssetTypeActions.cpp
+++ b/Engine/Plugins/Runtime/Metasound/Source/MetasoundEditor/Private/MetasoundAssetTypeActions.cpp
@@ -9,11 +9,13 @@
 #include "IContentBrowserSingleton.h"
 #include "Metasound.h"
 #include "MetasoundAssetBase.h"
+#include "MetasoundBuilderSubsystem.h"
 #include "MetasoundSource.h"
 #include "MetasoundEditor.h"
 #include "MetasoundEditorGraphBuilder.h"
 #include "MetasoundEditorModule.h"
 #include "MetasoundEditorSettings.h"
+#include "MetasoundEditorSubsystem.h"
 #include "MetasoundFactory.h"
 #include "MetasoundUObjectRegistry.h"
 #include "ObjectEditorUtils.h"
@@ -51,19 +53,28 @@ namespace Metasound
 			{
 				if (const UContentBrowserAssetContextMenuContext* Context = UContentBrowserAssetContextMenuContext::FindContextWithAssets(MenuContext))
 				{
-					for (TClass* MetaSound : Context->LoadSelectedObjects<TClass>())
+					for (TClass* ReferencedMetaSound : Context->LoadSelectedObjects<TClass>())
 					{
 						FString PackagePath;
 						FString AssetName;
 
-						IAssetTools::Get().CreateUniqueAssetName(MetaSound->GetOutermost()->GetName(), TEXT("_Preset"), PackagePath, AssetName);
+						IAssetTools::Get().CreateUniqueAssetName(ReferencedMetaSound->GetOutermost()->GetName(), TEXT("_Preset"), PackagePath, AssetName);
 
-						TFactory* Factory = NewObject<TFactory>();
-						check(Factory);
+						EMetaSoundBuilderResult BuilderResult;
+						UMetaSoundBuilderBase& Builder = UMetaSoundBuilderSubsystem::GetChecked().CreatePresetBuilder(FName(AssetName), ReferencedMetaSound, BuilderResult);
 
-						Factory->ReferencedMetaSoundObject = MetaSound;
+						if (BuilderResult != EMetaSoundBuilderResult::Succeeded)
+						{
+							UE_LOG(LogMetaSound, Error, TEXT("Error creating a preset builder for MetaSound '%s'"), *AssetName);
+							return;
+						}
 
-						IContentBrowserSingleton::Get().CreateNewAsset(AssetName, FPackageName::GetLongPackagePath(PackagePath), TClass::StaticClass(), Factory);
+						UMetaSoundEditorSubsystem& MetaSoundEditorSubsystem = UMetaSoundEditorSubsystem::GetChecked();
+						MetaSoundEditorSubsystem.BuildToAsset(&Builder, MetaSoundEditorSubsystem.GetDefaultAuthor(), AssetName, FPackageName::GetLongPackagePath(PackagePath), BuilderResult);
+						if (BuilderResult != EMetaSoundBuilderResult::Succeeded)
+						{
+							UE_LOG(LogMetaSound, Error, TEXT("Error building to asset when creating preset '%s'"), *AssetName);
+						}
 					}
 				}
 			}
diff --git a/Engine/Plugins/Runtime/Metasound/Source/MetasoundEditor/Private/MetasoundEditorGraphBuilder.cpp b/Engine/Plugins/Runtime/Metasound/Source/MetasoundEditor/Private/MetasoundEditorGraphBuilder.cpp
--- a/Engine/Plugins/Runtime/Metasound/Source/MetasoundEditor/Private/MetasoundEditorGraphBuilder.cpp
+++ b/Engine/Plugins/Runtime/Metasound/Source/MetasoundEditor/Private/MetasoundEditorGraphBuilder.cpp
@@ -24,6 +24,7 @@
 #include "MetasoundEditorGraphValidation.h"
 #include "MetasoundEditorModule.h"
 #include "MetasoundEditorSettings.h"
+#include "MetasoundEditorSubsystem.h"
 #include "MetasoundFactory.h"
 #include "MetasoundFrontendDataTypeRegistry.h"
 #include "MetasoundFrontendDocumentAccessPtr.h"
@@ -123,7 +124,7 @@ namespace Metasound
 				// If no graph is set, MetaSound has been created outside of asset factory, so initialize it here.
 				if (!MetaSoundAsset->GetGraph())
 				{
-					UMetaSoundBaseFactory::InitAsset(InMetaSound);
+					UMetaSoundEditorSubsystem::GetChecked().InitAsset(InMetaSound);
 				}
 
 				bEditorGraphModified |= FGraphBuilder::SynchronizeGraphMembers(InMetaSound);
diff --git a/Engine/Plugins/Runtime/Metasound/Source/MetasoundEditor/Private/MetasoundEditorSubsystem.cpp b/Engine/Plugins/Runtime/Metasound/Source/MetasoundEditor/Private/MetasoundEditorSubsystem.cpp
--- a/Engine/Plugins/Runtime/Metasound/Source/MetasoundEditor/Private/MetasoundEditorSubsystem.cpp
+++ b/Engine/Plugins/Runtime/Metasound/Source/MetasoundEditor/Private/MetasoundEditorSubsystem.cpp
@@ -16,7 +16,7 @@
 #define LOCTEXT_NAMESPACE "MetaSoundEditor"
 
 
-TScriptInterface<IMetaSoundDocumentInterface> UMetaSoundEditorSubsystem::BuildToAsset(UMetaSoundBuilderBase* InBuilder, const FString& Author, const FString& AssetName, const FString& PackagePath, EMetaSoundBuilderResult& OutResult)
+TScriptInterface<IMetaSoundDocumentInterface> UMetaSoundEditorSubsystem::BuildToAsset(UMetaSoundBuilderBase* InBuilder, const FString& Author, const FString& AssetName, const FString& PackagePath, EMetaSoundBuilderResult& OutResult, const USoundWave* TemplateSoundWave)
 {
 	using namespace Metasound;
 	using namespace Metasound::Frontend;
@@ -37,11 +37,27 @@ TScriptInterface<IMetaSoundDocumentInterface> UMetaSoundEditorSubsystem::BuildTo
 			BuilderOptions.ExistingMetaSound = NewMetaSound;
 			TScriptInterface<IMetaSoundDocumentInterface> DocInterface = InBuilder->Build(Parent, BuilderOptions);
 
+			const bool bIsSource = &BuilderUClass == UMetaSoundSource::StaticClass();
 			if (InBuilder->IsPreset())
 			{
 				FMetasoundAssetBase* PresetAsset = IMetasoundUObjectRegistry::Get().GetObjectAsAssetBase(NewMetaSound);
 				check(PresetAsset);
 				PresetAsset->ConformObjectDataToInterfaces();
+
+				// Only use referenced object's soundwave settings for sources if not overridden 
+				if (TemplateSoundWave == nullptr && bIsSource)
+				{
+					if (const UObject* ReferencedObject = InBuilder->GetReferencedPresetAsset())
+					{
+						TemplateSoundWave = CastChecked<USoundWave>(ReferencedObject);
+					}
+				}
+			}
+
+			// Template sound wave settings only apply to sources 
+			if (TemplateSoundWave != nullptr && bIsSource)
+			{
+				SetSoundWaveSettingsFromTemplate(NewMetaSound, TemplateSoundWave);
 			}
 
 			InitEdGraph(*NewMetaSound);
@@ -54,6 +70,72 @@ TScriptInterface<IMetaSoundDocumentInterface> UMetaSoundEditorSubsystem::BuildTo
 	return nullptr;
 }
 
+void UMetaSoundEditorSubsystem::SetSoundWaveSettingsFromTemplate(UObject* NewMetasound, const USoundWave* TemplateSoundWave)
+{
+	USoundWave* NewMetaSoundWave = CastChecked<USoundWave>(NewMetasound);
+
+	NewMetaSoundWave->SoundClassObject = TemplateSoundWave->SoundClassObject;
+	NewMetaSoundWave->AttenuationSettings = TemplateSoundWave->AttenuationSettings;
+	NewMetaSoundWave->ModulationSettings = TemplateSoundWave->ModulationSettings;
+	
+	// Concurrency
+	NewMetaSoundWave->bOverrideConcurrency = TemplateSoundWave->bOverrideConcurrency;
+	NewMetaSoundWave->ConcurrencySet = TemplateSoundWave->ConcurrencySet;
+	NewMetaSoundWave->ConcurrencyOverrides = TemplateSoundWave->ConcurrencyOverrides;
+
+	// Submix and bus sends 
+	NewMetaSoundWave->bEnableSubmixSends = TemplateSoundWave->bEnableSubmixSends;
+	NewMetaSoundWave->SoundSubmixSends = TemplateSoundWave->SoundSubmixSends;
+	NewMetaSoundWave->bEnableBusSends = TemplateSoundWave->bEnableBusSends;
+	NewMetaSoundWave->BusSends = TemplateSoundWave->BusSends;
+	NewMetaSoundWave->PreEffectBusSends = TemplateSoundWave->PreEffectBusSends;
+}
+
+void UMetaSoundEditorSubsystem::InitAsset(UObject& InNewMetaSound, UObject* InReferencedMetaSound)
+{
+	using namespace Metasound;
+	using namespace Metasound::Editor;
+	using namespace Metasound::Frontend;
+
+	TScriptInterface<IMetaSoundDocumentInterface> DocInterface = &InNewMetaSound;
+	FMetaSoundFrontendDocumentBuilder Builder(DocInterface);
+	Builder.InitDocument();
+#if WITH_EDITORONLY_DATA
+	Builder.InitNodeLocations();
+#endif // WITH_EDITORONLY_DATA
+
+	const FString& Author = GetDefaultAuthor();
+	Builder.SetAuthor(Author);
+
+	// Initialize asset as a preset
+	if (InReferencedMetaSound)
+	{
+		// Ensure the referenced MetaSound is registered already
+		RegisterGraphWithFrontend(*InReferencedMetaSound);
+
+		// Initialize preset with referenced Metasound 
+		TScriptInterface<IMetaSoundDocumentInterface> ReferencedDocInterface = InReferencedMetaSound;
+		const IMetaSoundDocumentInterface* ReferencedInterface = ReferencedDocInterface.GetInterface();
+		check(ReferencedInterface);
+		Builder.ConvertToPreset(ReferencedInterface->GetDocument());
+
+		// Update asset object data from interfaces 
+		FMetasoundAssetBase* PresetAsset = IMetasoundUObjectRegistry::Get().GetObjectAsAssetBase(&InNewMetaSound);
+		check(PresetAsset);
+		PresetAsset->ConformObjectDataToInterfaces();
+
+		// Copy sound wave settings to preset for sources
+		if (&ReferencedInterface->GetBaseMetaSoundUClass() == UMetaSoundSource::StaticClass())
+		{
+			SetSoundWaveSettingsFromTemplate(&InNewMetaSound, CastChecked<USoundWave>(InReferencedMetaSound));
+		}
+	}
+
+	// Initial graph generation is not something to be managed by the transaction
+	// stack, so don't track dirty state until after initial setup if necessary.
+	InitEdGraph(InNewMetaSound);
+}
+
 void UMetaSoundEditorSubsystem::InitEdGraph(UObject& InMetaSound)
 {
 	using namespace Metasound;
diff --git a/Engine/Plugins/Runtime/Metasound/Source/MetasoundEditor/Private/MetasoundFactory.cpp b/Engine/Plugins/Runtime/Metasound/Source/MetasoundEditor/Private/MetasoundFactory.cpp
--- a/Engine/Plugins/Runtime/Metasound/Source/MetasoundEditor/Private/MetasoundFactory.cpp
+++ b/Engine/Plugins/Runtime/Metasound/Source/MetasoundEditor/Private/MetasoundFactory.cpp
@@ -23,7 +23,7 @@ namespace Metasound
 			T* MetaSoundObject = NewObject<T>(InParent, InName, InFlags);
 			check(MetaSoundObject);
 
-			UMetaSoundBaseFactory::InitAsset(*MetaSoundObject, InReferencedMetaSoundObject);
+			UMetaSoundEditorSubsystem::GetChecked().InitAsset(*MetaSoundObject, InReferencedMetaSoundObject);
 			return MetaSoundObject;
 		}
 	}
diff --git a/Engine/Plugins/Runtime/Metasound/Source/MetasoundEditor/Public/MetasoundEditorSubsystem.h b/Engine/Plugins/Runtime/Metasound/Source/MetasoundEditor/Public/MetasoundEditorSubsystem.h
--- a/Engine/Plugins/Runtime/Metasound/Source/MetasoundEditor/Public/MetasoundEditorSubsystem.h
+++ b/Engine/Plugins/Runtime/Metasound/Source/MetasoundEditor/Public/MetasoundEditorSubsystem.h
@@ -19,8 +19,14 @@ class METASOUNDEDITOR_API UMetaSoundEditorSubsystem : public UEditorSubsystem
 	GENERATED_BODY()
 
 public:
+	// Build the given builder to a MetaSound asset
+	// @param TemplateSoundWave - SoundWave settings such as attenuation, modulation, and sound class will be copied from the optional TemplateSoundWave. 
+	// For preset builders, TemplateSoundWave will override the template values from the referenced asset
 	UFUNCTION(BlueprintCallable, Category = "Audio|MetaSound|Builder|Editor", meta = (WorldContext = "Parent", ExpandEnumAsExecs = "OutResult"))
-	UPARAM(DisplayName = "MetaSound Asset") TScriptInterface<IMetaSoundDocumentInterface> BuildToAsset(UMetaSoundBuilderBase* InBuilder, const FString& Author, const FString& AssetName, const FString& PackagePath, EMetaSoundBuilderResult& OutResult);
+	UPARAM(DisplayName = "MetaSound Asset") TScriptInterface<IMetaSoundDocumentInterface> BuildToAsset(UMetaSoundBuilderBase* InBuilder, const FString& Author, const FString& AssetName, const FString& PackagePath, EMetaSoundBuilderResult& OutResult, const USoundWave* TemplateSoundWave = nullptr);
+	
+	// Initialize the UObject asset, with an optional MetaSound to be referenced if the asset is a preset
+	void InitAsset(UObject& InNewMetaSound, UObject* InReferencedMetaSound = nullptr);
 
 	// Initialize UMetasoundEditorGraph for a given MetaSound object
 	void InitEdGraph(UObject& InMetaSound);
@@ -34,6 +40,9 @@ public:
 	// Get the default author for a MetaSound asset
 	const FString GetDefaultAuthor();
 
+	// Copy over sound wave settings such as attenuation, modulation, and sound class from the template sound wave to the MetaSound 
+	void SetSoundWaveSettingsFromTemplate(UObject* NewMetasound, const USoundWave* TemplateSoundWave);
+
 	static UMetaSoundEditorSubsystem& GetChecked();
 	static const UMetaSoundEditorSubsystem& GetConstChecked(); 
 };
diff --git a/Engine/Plugins/Runtime/Metasound/Source/MetasoundEditor/Public/MetasoundFactory.h b/Engine/Plugins/Runtime/Metasound/Source/MetasoundEditor/Public/MetasoundFactory.h
--- a/Engine/Plugins/Runtime/Metasound/Source/MetasoundEditor/Public/MetasoundFactory.h
+++ b/Engine/Plugins/Runtime/Metasound/Source/MetasoundEditor/Public/MetasoundFactory.h
@@ -20,7 +20,8 @@ class UMetaSoundBaseFactory : public UFactory
 	// using the provided MetaSound as a Reference
 	UPROPERTY(EditAnywhere, Transient, Category = Factory)
 	TObjectPtr<UObject> ReferencedMetaSoundObject;
-
+	
+	UE_DEPRECATED(5.3, "Use UMetaSoundEditorSubsystem::InitAsset instead.")
 	static void InitAsset(UObject& InNewMetaSound, UObject* InReferencedMetaSound = nullptr);
 };
 
diff --git a/Engine/Plugins/Runtime/Metasound/Source/MetasoundEngine/Private/MetasoundBuilderSubsystem.cpp b/Engine/Plugins/Runtime/Metasound/Source/MetasoundEngine/Private/MetasoundBuilderSubsystem.cpp
--- a/Engine/Plugins/Runtime/Metasound/Source/MetasoundEngine/Private/MetasoundBuilderSubsystem.cpp
+++ b/Engine/Plugins/Runtime/Metasound/Source/MetasoundEngine/Private/MetasoundBuilderSubsystem.cpp
@@ -9,6 +9,7 @@
 #include "Interfaces/MetasoundOutputFormatInterfaces.h"
 #include "Interfaces/MetasoundFrontendSourceInterface.h"
 #include "Metasound.h"
+#include "MetasoundAssetSubsystem.h"
 #include "MetasoundFrontendDocument.h"
 #include "MetasoundFrontendRegistries.h"
 #include "MetasoundFrontendSearchEngine.h"
@@ -490,6 +491,35 @@ void UMetaSoundBuilderBase::InitNodeLocations()
 	Builder.InitNodeLocations();
 }
 
+UObject* UMetaSoundBuilderBase::GetReferencedPresetAsset() const
+{
+	using namespace Metasound::Frontend;
+	if (!IsPreset())
+	{
+		return nullptr;
+	}
+
+	// Find the single external node which is the referenced preset asset, 
+	// and find the asset with its registry key 
+	auto FindExternalNode = [this](const FMetasoundFrontendNode& Node) 
+	{
+		const FMetasoundFrontendClass* Class = Builder.FindDependency(Node.ClassID);
+		return Class->Metadata.GetType() == EMetasoundFrontendClassType::External;
+	};
+	const FMetasoundFrontendNode* Node = Builder.GetDocument().RootGraph.Graph.Nodes.FindByPredicate(FindExternalNode);
+	if (Node != nullptr)
+	{
+		const FMetasoundFrontendClass* NodeClass = Builder.FindDependency(Node->ClassID);
+		const FNodeRegistryKey NodeClassRegistryKey = NodeRegistryKey::CreateKey(NodeClass->Metadata);
+		if (FMetasoundAssetBase* Asset = IMetaSoundAssetManager::GetChecked().TryLoadAssetFromKey(NodeClassRegistryKey))
+		{
+			return Asset->GetOwningAsset();
+		}
+	}
+
+	return nullptr;
+}
+
 bool UMetaSoundBuilderBase::InterfaceIsDeclared(FName InterfaceName) const
 {
 	return Builder.IsInterfaceDeclared(InterfaceName);
@@ -927,6 +957,24 @@ UMetaSoundPatchBuilder* UMetaSoundBuilderSubsystem::CreatePatchPresetBuilder(FNa
 	return &Builder;
 }
 
+UMetaSoundBuilderBase& UMetaSoundBuilderSubsystem::CreatePresetBuilder(FName BuilderName, const TScriptInterface<IMetaSoundDocumentInterface>& ReferencedPatchClass, EMetaSoundBuilderResult& OutResult)
+{
+	const UClass& Class = ReferencedPatchClass->GetBaseMetaSoundUClass();
+	if (&Class == UMetaSoundSource::StaticClass())
+	{
+		return *CreateSourcePresetBuilder(BuilderName, ReferencedPatchClass, OutResult);
+	}
+	else if (&Class == UMetaSoundPatch::StaticClass())
+	{
+		return *CreatePatchPresetBuilder(BuilderName, ReferencedPatchClass, OutResult);
+	}
+	else
+	{
+		checkf(false, TEXT("UClass '%s' cannot be built to a MetaSound preset"), *Class.GetFullName());
+		return Metasound::Engine::BuilderSubsystemPrivate::CreateTransientBuilder<UMetaSoundPatchBuilder>();
+	}
+}
+
 UMetaSoundSourceBuilder* UMetaSoundBuilderSubsystem::CreateSourcePresetBuilder(FName BuilderName, const TScriptInterface<IMetaSoundDocumentInterface>& ReferencedNodeClass, EMetaSoundBuilderResult& OutResult)
 {
 	UMetaSoundSourceBuilder& Builder = Metasound::Engine::BuilderSubsystemPrivate::CreateTransientBuilder<UMetaSoundSourceBuilder>();
diff --git a/Engine/Plugins/Runtime/Metasound/Source/MetasoundEngine/Public/MetasoundBuilderSubsystem.h b/Engine/Plugins/Runtime/Metasound/Source/MetasoundEngine/Public/MetasoundBuilderSubsystem.h
--- a/Engine/Plugins/Runtime/Metasound/Source/MetasoundEngine/Public/MetasoundBuilderSubsystem.h
+++ b/Engine/Plugins/Runtime/Metasound/Source/MetasoundEngine/Public/MetasoundBuilderSubsystem.h
@@ -251,6 +251,10 @@ public:
 	// Returns node output's data if valid (including things like name and datatype).
 	UFUNCTION(BlueprintCallable, Category = "Audio|MetaSound|Builder", meta = (ExpandEnumAsExecs = "OutResult"))
 	void GetNodeOutputData(const FMetaSoundBuilderNodeOutputHandle& OutputHandle, FName& Name, FName& DataType, EMetaSoundBuilderResult& OutResult);
+	
+	// Return the asset referenced by this preset builder. Returns nullptr if the builder is not a preset.
+	UFUNCTION(BlueprintCallable, Category = "Audio|MetaSound|Builder")
+	UObject* GetReferencedPresetAsset() const;
 
 	// Returns if a given interface is declared.
 	UFUNCTION(BlueprintCallable, Category = "Audio|MetaSound|Builder")
@@ -471,6 +475,8 @@ public:
 	UFUNCTION(BlueprintCallable, Category = "Audio|MetaSound|Builder",  meta = (ExpandEnumAsExecs = "OutResult"))
 	UPARAM(DisplayName = "Patch Preset Builder") UMetaSoundPatchBuilder* CreatePatchPresetBuilder(FName BuilderName, const TScriptInterface<IMetaSoundDocumentInterface>& ReferencedPatchClass, EMetaSoundBuilderResult& OutResult);
 
+	UMetaSoundBuilderBase& CreatePresetBuilder(FName BuilderName, const TScriptInterface<IMetaSoundDocumentInterface>& ReferencedPatchClass, EMetaSoundBuilderResult& OutResult);
+	
 	UFUNCTION(BlueprintCallable, Category = "Audio|MetaSound|Builder", meta = (ExpandEnumAsExecs = "OutResult"))
 	UPARAM(DisplayName = "Source Preset Builder") UMetaSoundSourceBuilder* CreateSourcePresetBuilder(FName BuilderName, const TScriptInterface<IMetaSoundDocumentInterface>& ReferencedSourceClass, EMetaSoundBuilderResult& OutResult);
 
