Nonunity fix for copying over Sound Wave properties for MetaSound presets

#rnx

diff --git a/Engine/Plugins/Runtime/Metasound/Source/MetasoundEditor/Private/MetasoundEditorSubsystem.cpp b/Engine/Plugins/Runtime/Metasound/Source/MetasoundEditor/Private/MetasoundEditorSubsystem.cpp
--- a/Engine/Plugins/Runtime/Metasound/Source/MetasoundEditor/Private/MetasoundEditorSubsystem.cpp
+++ b/Engine/Plugins/Runtime/Metasound/Source/MetasoundEditor/Private/MetasoundEditorSubsystem.cpp
@@ -10,6 +10,8 @@
 #include "MetasoundEditorSettings.h"
 #include "MetasoundFactory.h"
 #include "MetasoundUObjectRegistry.h"
+#include "Sound/SoundSourceBusSend.h"
+#include "Sound/SoundSubmixSend.h"
 
 #include UE_INLINE_GENERATED_CPP_BY_NAME(MetasoundEditorSubsystem)
 
