Fix for editor crash when using BuildToAsset over a pre-existing asset due to GC nuking the referenced builder that is in use
#rnx

diff --git a/Engine/Plugins/Runtime/Metasound/Source/MetasoundEditor/Private/MetasoundEditorSubsystem.cpp b/Engine/Plugins/Runtime/Metasound/Source/MetasoundEditor/Private/MetasoundEditorSubsystem.cpp
--- a/Engine/Plugins/Runtime/Metasound/Source/MetasoundEditor/Private/MetasoundEditorSubsystem.cpp
+++ b/Engine/Plugins/Runtime/Metasound/Source/MetasoundEditor/Private/MetasoundEditorSubsystem.cpp
@@ -25,11 +25,18 @@ TScriptInterface<IMetaSoundDocumentInterface> UMetaSoundEditorSubsystem::BuildTo
 
 	if (InBuilder)
 	{
-		constexpr UObject* Parent = nullptr;
-		constexpr UFactory* Factory = nullptr;
-
 		// Not about to follow this lack of const correctness down a multidecade in the works rabbit hole.
 		UClass& BuilderUClass = const_cast<UClass&>(InBuilder->GetBuilderUClass());
+
+		// AddToRoot to avoid builder getting gc'ed during CreateAsset call below, as the builder 
+		// may be unreferenced by other UObjects and it must be persistent to finish initializing.
+		const bool bWasRooted = InBuilder->IsRooted();
+		if (!bWasRooted)
+		{
+			InBuilder->AddToRoot();
+		}
+
+		constexpr UFactory* Factory = nullptr;
 		if (UObject* NewMetaSound = IAssetTools::Get().CreateAsset(AssetName, PackagePath, &BuilderUClass, Factory))
 		{
 			InBuilder->InitNodeLocations();
@@ -37,6 +44,7 @@ TScriptInterface<IMetaSoundDocumentInterface> UMetaSoundEditorSubsystem::BuildTo
 
 			FMetaSoundBuilderOptions BuilderOptions { FName(*AssetName) };
 			BuilderOptions.ExistingMetaSound = NewMetaSound;
+			constexpr UObject* Parent = nullptr;
 			TScriptInterface<IMetaSoundDocumentInterface> DocInterface = InBuilder->Build(Parent, BuilderOptions);
 
 			const bool bIsSource = &BuilderUClass == UMetaSoundSource::StaticClass();
@@ -66,6 +74,11 @@ TScriptInterface<IMetaSoundDocumentInterface> UMetaSoundEditorSubsystem::BuildTo
 			OutResult = EMetaSoundBuilderResult::Succeeded;
 			return NewMetaSound;
 		}
+
+		if (!bWasRooted)
+		{
+			InBuilder->RemoveFromRoot();
+		}
 	}
 
 	OutResult = EMetaSoundBuilderResult::Failed;
