Fix build for editoronly data in program using metasound fronend.
#rnx

diff --git a/Engine/Plugins/Runtime/Metasound/Source/MetasoundFrontend/Private/MetasoundFrontendGraphController.cpp b/Engine/Plugins/Runtime/Metasound/Source/MetasoundFrontend/Private/MetasoundFrontendGraphController.cpp
--- a/Engine/Plugins/Runtime/Metasound/Source/MetasoundFrontend/Private/MetasoundFrontendGraphController.cpp
+++ b/Engine/Plugins/Runtime/Metasound/Source/MetasoundFrontend/Private/MetasoundFrontendGraphController.cpp
@@ -1283,7 +1283,9 @@ namespace Metasound
 
 			return FText::GetEmpty();
 		}
+#endif // WITH_EDITOR
 
+#if WITH_EDITORONLY_DATA
 		int32 FGraphController::GetSortOrderIndexForInput(const FVertexName& InName) const
 		{
 			if (const FMetasoundFrontendClassInput* Desc = FindInputDescriptionWithName(InName))
@@ -1303,6 +1305,9 @@ namespace Metasound
 
 			return 0;
 		}
+#endif // WITH_EDITORONLY_DATA
+
+#if WITH_EDITOR
 
 		void FGraphController::SetInputDisplayName(const FVertexName& InName, const FText& InDisplayName)
 		{
diff --git a/Engine/Plugins/Runtime/Metasound/Source/MetasoundFrontend/Private/MetasoundFrontendGraphController.h b/Engine/Plugins/Runtime/Metasound/Source/MetasoundFrontend/Private/MetasoundFrontendGraphController.h
--- a/Engine/Plugins/Runtime/Metasound/Source/MetasoundFrontend/Private/MetasoundFrontendGraphController.h
+++ b/Engine/Plugins/Runtime/Metasound/Source/MetasoundFrontend/Private/MetasoundFrontendGraphController.h
@@ -158,9 +158,14 @@ namespace Metasound
 
 			// Set the display name for the output with the given name
 			void SetOutputDisplayName(const FVertexName& InName, const FText& InDisplayName) override;
+#endif // WITH_EDITOR
 
+#if WITH_EDITORONLY_DATA
 			int32 GetSortOrderIndexForInput(const FVertexName& InName) const override;
 			int32 GetSortOrderIndexForOutput(const FVertexName& InName) const override;
+#endif // WITH_EDITORONLY_DATA
+
+#if WITH_EDITOR
 			void SetSortOrderIndexForInput(const FVertexName& InName, int32 InSortOrderIndex) override;
 			void SetSortOrderIndexForOutput(const FVertexName& InName, int32 InSortOrderIndex) override;
 #endif // WITH_EDITOR
diff --git a/Engine/Plugins/Runtime/Metasound/Source/MetasoundFrontend/Private/MetasoundFrontendInvalidController.h b/Engine/Plugins/Runtime/Metasound/Source/MetasoundFrontend/Private/MetasoundFrontendInvalidController.h
--- a/Engine/Plugins/Runtime/Metasound/Source/MetasoundFrontend/Private/MetasoundFrontendInvalidController.h
+++ b/Engine/Plugins/Runtime/Metasound/Source/MetasoundFrontend/Private/MetasoundFrontendInvalidController.h
@@ -215,7 +215,7 @@ namespace Metasound
 			virtual FConstInputHandle GetInputWithID(FGuid InVertexID) const override { return IInputController::GetInvalidHandle(); }
 			virtual FConstOutputHandle GetOutputWithID(FGuid InVertexID) const override { return IOutputController::GetInvalidHandle(); }
 
-#if WITH_EDITOR
+#if WITH_EDITORONLY_DATA
 			virtual const FMetasoundFrontendNodeStyle& GetNodeStyle() const override { static const FMetasoundFrontendNodeStyle Invalid; return Invalid; }
 			virtual void SetNodeStyle(const FMetasoundFrontendNodeStyle& InNodeStyle) override { }
 #endif // WITH_EDITOR
@@ -395,9 +395,12 @@ namespace Metasound
 			virtual void SetOutputDescription(const FVertexName& InName, const FText& InDescription) override { }
 			virtual void SetInputDisplayName(const FVertexName& InName, const FText& InDisplayName) override { }
 			virtual void SetOutputDisplayName(const FVertexName& InName, const FText& InDisplayName) override { }
-
+#endif // WITH_EDITOR
+#if WITH_EDITORONLY_DATA
 			virtual int32 GetSortOrderIndexForInput(const FVertexName& InName) const { return 0; }
 			virtual int32 GetSortOrderIndexForOutput(const FVertexName& InName) const { return 0; }
+#endif // WITH_EDITORONLY_DATA
+#if WITH_EDITOR
 			virtual void SetSortOrderIndexForInput(const FVertexName& InName, int32 InSortOrderIndex) { }
 			virtual void SetSortOrderIndexForOutput(const FVertexName& InName, int32 InSortOrderIndex) { }
 #endif // WITH_EDITOR
diff --git a/Engine/Plugins/Runtime/Metasound/Source/MetasoundFrontend/Private/MetasoundFrontendNodeController.cpp b/Engine/Plugins/Runtime/Metasound/Source/MetasoundFrontend/Private/MetasoundFrontendNodeController.cpp
--- a/Engine/Plugins/Runtime/Metasound/Source/MetasoundFrontend/Private/MetasoundFrontendNodeController.cpp
+++ b/Engine/Plugins/Runtime/Metasound/Source/MetasoundFrontend/Private/MetasoundFrontendNodeController.cpp
@@ -206,6 +206,9 @@ namespace Metasound
 			static const FMetasoundFrontendClassStyle Invalid;
 			return Invalid;
 		}
+#endif // WITH_EDITOR
+
+#if WITH_EDITORONLY_DATA
 
 		const FMetasoundFrontendNodeStyle& FBaseNodeController::GetNodeStyle() const
 		{
@@ -225,6 +228,9 @@ namespace Metasound
 				Node->Style = InStyle;
 			}
 		}
+#endif // WITH_EDITORONLY_DATA
+
+#if WITH_EDITOR
 
 		const FText& FBaseNodeController::GetDescription() const
 		{
diff --git a/Engine/Plugins/Runtime/Metasound/Source/MetasoundFrontend/Private/MetasoundFrontendNodeController.h b/Engine/Plugins/Runtime/Metasound/Source/MetasoundFrontend/Private/MetasoundFrontendNodeController.h
--- a/Engine/Plugins/Runtime/Metasound/Source/MetasoundFrontend/Private/MetasoundFrontendNodeController.h
+++ b/Engine/Plugins/Runtime/Metasound/Source/MetasoundFrontend/Private/MetasoundFrontendNodeController.h
@@ -53,10 +53,12 @@ namespace Metasound
 			virtual const FMetasoundFrontendInterfaceStyle& GetInputStyle() const override;
 			virtual const FMetasoundFrontendInterfaceStyle& GetOutputStyle() const override;
 			virtual const FMetasoundFrontendClassStyle& GetClassStyle() const override;
+#endif // WITH_EDITOR
 
+#if WITH_EDITORONLY_DATA
 			virtual const FMetasoundFrontendNodeStyle& GetNodeStyle() const override;
 			virtual void SetNodeStyle(const FMetasoundFrontendNodeStyle& InStyle) override;
-#endif // WITH_EDITOR
+#endif // WITH_EDITORONLY_DATA
 
 			virtual bool DiffAgainstRegistryInterface(FClassInterfaceUpdates& OutInterfaceUpdates, bool bInUseHighestMinorVersion) const override;
 
diff --git a/Engine/Plugins/Runtime/Metasound/Source/MetasoundFrontend/Public/MetasoundFrontendController.h b/Engine/Plugins/Runtime/Metasound/Source/MetasoundFrontend/Public/MetasoundFrontendController.h
--- a/Engine/Plugins/Runtime/Metasound/Source/MetasoundFrontend/Public/MetasoundFrontendController.h
+++ b/Engine/Plugins/Runtime/Metasound/Source/MetasoundFrontend/Public/MetasoundFrontendController.h
@@ -580,10 +580,10 @@ namespace Metasound
 			/** Returns the name of this node. */
 			virtual const FVertexName& GetNodeName() const = 0;
 
-#if WITH_EDITOR
+#if WITH_EDITORONLY_DATA
 			virtual const FMetasoundFrontendNodeStyle& GetNodeStyle() const = 0;
 			virtual void SetNodeStyle(const FMetasoundFrontendNodeStyle& InStyle) = 0;
-#endif // WITH_EDITOR
+#endif // WITH_EDITORONLY_DATA
 
 			/** Returns the ID associated with this node. */
 			virtual FGuid GetID() const = 0;
@@ -897,13 +897,17 @@ namespace Metasound
 
 			/** Set the description for the output with the given name. */
 			virtual void SetOutputDescription(const FVertexName& InName, const FText& InDescription) = 0;
+#endif // WITH_EDITOR
 
+#if WITH_EDITORONLY_DATA
 			/** Returns the sort order index for the input with the given name. Returns 0 if not found or unset. */
 			virtual int32 GetSortOrderIndexForInput(const FVertexName& InName) const = 0;
 
 			/** Returns the sort order index for the input with the given name. Returns 0 if not found or unset. */
 			virtual int32 GetSortOrderIndexForOutput(const FVertexName& InName) const = 0;
+#endif // WITH_EDITORONLY_DATA
 
+#if WITH_EDITOR
 			/** Sets the sort order index for the input with the given name. No-ops if input not found. */
 			virtual void SetSortOrderIndexForInput(const FVertexName& InName, int32 InSortOrderIndex) = 0;
 
