Deprecate IOperator::GetInputs IOPerator::GetOutputs and IOperator::Bind.
Add debug info for developers to determine which node is needs updating.

diff --git a/Engine/Plugins/Runtime/Metasound/Source/MetasoundGraphCore/Private/MetasoundOperatorBuilder.cpp b/Engine/Plugins/Runtime/Metasound/Source/MetasoundGraphCore/Private/MetasoundOperatorBuilder.cpp
--- a/Engine/Plugins/Runtime/Metasound/Source/MetasoundGraphCore/Private/MetasoundOperatorBuilder.cpp
+++ b/Engine/Plugins/Runtime/Metasound/Source/MetasoundGraphCore/Private/MetasoundOperatorBuilder.cpp
@@ -436,6 +436,8 @@ namespace Metasound
 	FOperatorBuilder::FBuildStatus FOperatorBuilder::CreateOperators(OperatorBuilder::FBuildContext& InOutContext, const TArray<const INode*>& InSortedNodes, const FInputVertexInterfaceData& InExternalInputData) const
 	{
 		METASOUND_TRACE_CPUPROFILER_EVENT_SCOPE(Metasound::FOperatorBuilder::CreateOperators);
+		METASOUND_DEBUG_DECLARE_SCOPE;
+
 		using namespace DirectedGraphAlgo;
 
 		FBuildStatus BuildStatus;
@@ -446,7 +448,7 @@ namespace Metasound
 		// Call operator factory for each node.
 		for (const INode* Node : InSortedNodes)
 		{
-			//ThreadLocalDebug::SetActiveNodeClass(Node->GetMetadata());
+			METASOUND_DEBUG_SET_ACTIVE_NODE_SCOPE(Node);
 
 			FOperatorID OperatorID = GetOperatorID(Node);
 			FGraphOperatorData::FOperatorInfo& OperatorInfo = OperatorMap.FindChecked(OperatorID);
@@ -500,8 +502,6 @@ namespace Metasound
 					}
 				}
 			}
-
-			//ThreadLocalDebug::ResetActiveNodeClass();
 		}
 
 		return BuildStatus;
diff --git a/Engine/Plugins/Runtime/Metasound/Source/MetasoundGraphCore/Private/MetasoundOperatorInterface.cpp b/Engine/Plugins/Runtime/Metasound/Source/MetasoundGraphCore/Private/MetasoundOperatorInterface.cpp
new file mode 100644
--- /dev/null
+++ b/Engine/Plugins/Runtime/Metasound/Source/MetasoundGraphCore/Private/MetasoundOperatorInterface.cpp
@@ -0,0 +1,90 @@
+// Copyright Epic Games, Inc. All Rights Reserved.
+
+#include "MetasoundOperatorInterface.h"
+
+#include "HAL/IConsoleManager.h"
+#include "MetasoundDataReferenceCollection.h"
+#include "MetasoundLog.h"
+#include "MetasoundThreadLocalDebug.h"
+
+#include <atomic>
+
+namespace Metasound
+{
+	namespace MetasoundOperatorInterfacePrivate
+	{
+		static bool bEnableMetaSoundOperatorMissingOverrideLog = false;
+	}
+}
+
+FAutoConsoleVariableRef CVarMetaSoundEnableOperatorMissingOverrideLog(
+	TEXT("au.MetaSound.Debug.EnableOperatorMissingOverrideLog"),
+	Metasound::MetasoundOperatorInterfacePrivate::bEnableMetaSoundOperatorMissingOverrideLog,
+	TEXT("Enables additional logging on operators with missing overrides\n")
+	TEXT("Default: false"),
+	ECVF_Default);
+
+namespace Metasound
+{
+	namespace MetasoundOperatorInterfacePrivate
+	{
+		static std::atomic<bool> bHasWarnedMissingBindImplementation = false;
+
+		void WarnMissingBindImplementation()
+		{
+			if (!bHasWarnedMissingBindImplementation)
+			{
+				bHasWarnedMissingBindImplementation = true;
+				UE_LOG(LogMetaSound, Warning, TEXT("One or more MetaSound IOperator derived classes have not overridden the BindInputs(...) and/or BindOutputs(...) methods. Please override these methods and ensure that the IOperator derived class conforms to the most recent IOperator API. Enable console variable au.MetaSound.Debug.EnableOperatorMissingOverrideLog for additional information in logs."));
+			}
+
+			if (bEnableMetaSoundOperatorMissingOverrideLog)
+			{
+				UE_LOG(LogMetaSound, Warning, TEXT("MetaSound Node %s produced IOperator which does not override BindInputs(...) or BindOutputs(...). Please override these methods and ensure that the IOperator derived class conforms to the most recent IOperator API."), METASOUND_DEBUG_ACTIVE_NODE_NAME);
+			}
+		}
+	}
+
+	FDataReferenceCollection IOperator::GetInputs() const 
+	{
+		return FDataReferenceCollection();
+	}
+
+	FDataReferenceCollection IOperator::GetOutputs() const
+	{
+		return FDataReferenceCollection();
+	}
+
+	void IOperator::Bind(FVertexInterfaceData& InVertexData) const
+	{
+		PRAGMA_DISABLE_DEPRECATION_WARNINGS
+		InVertexData.GetInputs().Set(GetInputs());
+		InVertexData.GetOutputs().Set(GetOutputs());
+		PRAGMA_ENABLE_DEPRECATION_WARNINGS
+	}
+
+	
+	void IOperator::BindInputs(FInputVertexInterfaceData& InVertexData)
+	{
+		MetasoundOperatorInterfacePrivate::WarnMissingBindImplementation();
+
+		FVertexInterfaceData Data;
+		Data.GetInputs() = InVertexData;
+		PRAGMA_DISABLE_DEPRECATION_WARNINGS
+		Bind(Data);
+		PRAGMA_ENABLE_DEPRECATION_WARNINGS
+		InVertexData = Data.GetInputs();
+	}
+
+	void IOperator::BindOutputs(FOutputVertexInterfaceData& InVertexData) 
+	{
+		MetasoundOperatorInterfacePrivate::WarnMissingBindImplementation();
+
+		FVertexInterfaceData Data;
+		Data.GetOutputs() = InVertexData;
+		PRAGMA_DISABLE_DEPRECATION_WARNINGS
+		Bind(Data);
+		PRAGMA_ENABLE_DEPRECATION_WARNINGS
+		InVertexData = Data.GetOutputs();
+	}
+}
diff --git a/Engine/Plugins/Runtime/Metasound/Source/MetasoundGraphCore/Private/MetasoundThreadLocalDebug.cpp b/Engine/Plugins/Runtime/Metasound/Source/MetasoundGraphCore/Private/MetasoundThreadLocalDebug.cpp
--- a/Engine/Plugins/Runtime/Metasound/Source/MetasoundGraphCore/Private/MetasoundThreadLocalDebug.cpp
+++ b/Engine/Plugins/Runtime/Metasound/Source/MetasoundGraphCore/Private/MetasoundThreadLocalDebug.cpp
@@ -2,43 +2,95 @@
 
 #include "MetasoundThreadLocalDebug.h"
 
+#if UE_METASOUND_DEBUG_ENABLED
+
 #include "MetasoundNodeInterface.h"
 #include "Containers/UnrealString.h"
 
-#define ENABLE_METASOUND_THREAD_LOCAL_DEBUG (DO_CHECK || DO_ENSURE)
-
 namespace Metasound
 {
 	namespace ThreadLocalDebug
 	{
+		class FDebugInfo
+		{
+		public:
+			// Set the active node for the debug info.
+			void SetActiveNode(const INode* InNode);
+
+			// Return the active node.
+			const INode* GetActiveNode() const;
+
+			// Returns the class name and version string for the active node in the current thread.
+			const TCHAR* GetActiveNodeClassNameAndVersion() const;
+
+		private:
+
+			const INode* ActiveNode = nullptr;
+			FString NodeClassNameAndVersion;
+		};
+
 		namespace ThreadLocalDebugPrivate
 		{
-#if ENABLE_METASOUND_THREAD_LOCAL_DEBUG
-			static thread_local FString ActiveNodeClassNameAndVersion;
-#endif
+			static thread_local FDebugInfo DebugInfoOnThisThread;
 		}
 
-		void SetActiveNodeClass(const FNodeClassMetadata& InMetadata)
+		void FDebugInfo::SetActiveNode(const INode* InNode)
 		{
-#if ENABLE_METASOUND_THREAD_LOCAL_DEBUG
-			ThreadLocalDebugPrivate::ActiveNodeClassNameAndVersion = FString::Format(TEXT("{0} v{1}.{2}"), {InMetadata.ClassName.GetFullName().ToString(), InMetadata.MajorVersion, InMetadata.MinorVersion});
-#endif
+			ActiveNode = InNode;
+			if (ActiveNode)
+			{
+				const FNodeClassMetadata& Metadata = ActiveNode->GetMetadata();
+
+				NodeClassNameAndVersion = FString::Format(TEXT("{0} v{1}.{2}"), {Metadata.ClassName.GetFullName().ToString(), Metadata.MajorVersion, Metadata.MinorVersion});
+			}
+			else
+			{
+				NodeClassNameAndVersion = TEXT("[No Active Debug Node Set]");
+			}
 		}
 
-		void ResetActiveNodeClass()
+		const INode* FDebugInfo::GetActiveNode() const
 		{
-#if ENABLE_METASOUND_THREAD_LOCAL_DEBUG
-			ThreadLocalDebugPrivate::ActiveNodeClassNameAndVersion = TEXT("");
-#endif
+			return ActiveNode;
 		}
 
-		const TCHAR* GetActiveNodeClassNameAndVersion()
+		const TCHAR* FDebugInfo::GetActiveNodeClassNameAndVersion() const
 		{
-#if ENABLE_METASOUND_THREAD_LOCAL_DEBUG
-			return *ThreadLocalDebugPrivate::ActiveNodeClassNameAndVersion;	
-#else
+			if (ActiveNode)
+			{
+				return *NodeClassNameAndVersion;
+			}
 			return TEXT("");
-#endif
+		}
+
+		const TCHAR* GetActiveNodeClassNameAndVersionOnThisThread()
+		{
+			return ThreadLocalDebugPrivate::DebugInfoOnThisThread.GetActiveNodeClassNameAndVersion();
+		}
+
+		FDebugInfo* GetDebugInfoOnThisThread()
+		{
+			return &ThreadLocalDebugPrivate::DebugInfoOnThisThread;
+		}
+
+		FScopeDebugActiveNode::FScopeDebugActiveNode(FDebugInfo* InDebugInfo, const INode* InNode)
+		: DebugInfo(InDebugInfo)
+		{
+			if (DebugInfo)
+			{
+				PriorNode = DebugInfo->GetActiveNode();
+				DebugInfo->SetActiveNode(InNode);
+			}
+		}
+
+		FScopeDebugActiveNode::~FScopeDebugActiveNode()
+		{
+			if (DebugInfo)
+			{
+				DebugInfo->SetActiveNode(PriorNode);
+			}
 		}
 	}
 }
+
+#endif
diff --git a/Engine/Plugins/Runtime/Metasound/Source/MetasoundGraphCore/Private/MetasoundThreadLocalDebug.h b/Engine/Plugins/Runtime/Metasound/Source/MetasoundGraphCore/Private/MetasoundThreadLocalDebug.h
--- a/Engine/Plugins/Runtime/Metasound/Source/MetasoundGraphCore/Private/MetasoundThreadLocalDebug.h
+++ b/Engine/Plugins/Runtime/Metasound/Source/MetasoundGraphCore/Private/MetasoundThreadLocalDebug.h
@@ -3,21 +3,91 @@
 #pragma once
 
 #include "HAL/Platform.h"
+#include "Containers/UnrealString.h"
+
+/** MetaSound Thread Local Debugging allows node info to be propagated down to 
+ * calling functions for debug purposes. It passes the information by storing thread_local
+ * debug data which can be accessed on other functions. thread_local access can
+ * have negative performance implications, and so this debug info is disabled by
+ * default when DO_CHECK and DO_ENSURE are false. 
+ *
+ * To set which node is active in the current thread, use the following macros.
+ *
+ * 	{
+ * 		// Declare a debug scope to minimize access to thread locals
+ * 		METASOUND_DEBUG_DECLARE_SCOPE
+ *
+ * 		for (const INode* Node : Nodes)
+ * 		{
+ * 			// Set the active node in this scope to be Node
+ * 			METASOUND_DEBUG_SET_ACTIVE_NODE_SCOPE(Node)
+ *
+ *
+ * 			// Any calls in this scope can get debug info about the "active node"
+ * 			// by using the METASOUND_DEBUG macros
+ * 			DoOtherThings().
+ * 		}
+ * 	}
+ *
+ *
+ * 	In another function in the same callstack, debug info can be accessed.
+ *
+ * 	{ 
+ * 		UE_LOG(LogMetaSound, Error, TEXT("Error with node %s"), METASOUND_DEBUG_ACTIVE_NODE_NAME);
+ * 	}
+ */
+
+#ifndef UE_METASOUND_DEBUG_ENABLED
+#define UE_METASOUND_DEBUG_ENABLED (DO_CHECK || DO_ENSURE)
+#endif
+
+#if UE_METASOUND_DEBUG_ENABLED
 
 namespace Metasound
 {
-	struct FNodeClassMetadata;
+	class INode;
 
 	// Debug util for setting stack scope debug variables
 	namespace ThreadLocalDebug
 	{
-		// Sets the current active node class for the current thread.
-		void SetActiveNodeClass(const FNodeClassMetadata& InMetadata);
+		class FDebugInfo;
+
+		FDebugInfo* GetDebugInfoOnThisThread(); 
+		const TCHAR* GetActiveNodeClassNameAndVersionOnThisThread(); 
 
-		// Resets the current active node class for the current thread.
-		void ResetActiveNodeClass();
+		struct FScopeDebugActiveNode final
+		{
+			explicit FScopeDebugActiveNode(FDebugInfo* InDebugInfo, const INode* InNode);
+			~FScopeDebugActiveNode();
 
-		// Returns the class name and version string for the active node in the current thread.
-		const TCHAR* GetActiveNodeClassNameAndVersion();
+			FScopeDebugActiveNode(const FScopeDebugActiveNode&) = delete;
+			FScopeDebugActiveNode(FScopeDebugActiveNode&&) = delete;
+			FScopeDebugActiveNode& operator=(const FScopeDebugActiveNode&) = delete;
+			FScopeDebugActiveNode& operator=(FScopeDebugActiveNode&&) = delete;
+
+		private:
+
+			const INode* PriorNode = nullptr;
+			FDebugInfo* DebugInfo = nullptr;
+		};
 	}
 }
+
+#endif // if UE_METASOUND_DEBUG_ENABLED
+
+#if UE_METASOUND_DEBUG_ENABLED
+
+#define METASOUND_DEBUG_SCOPE_NAME __MetaSoundDebugScope
+#define METASOUND_DEBUG_DECLARE_SCOPE ::Metasound::ThreadLocalDebug::FDebugInfo* METASOUND_DEBUG_SCOPE_NAME = ::Metasound::ThreadLocalDebug::GetDebugInfoOnThisThread();
+#define METASOUND_DEBUG_SET_ACTIVE_NODE_SCOPE(InNodePtr) ::Metasound::ThreadLocalDebug::FScopeDebugActiveNode __MetaSoundDebugScope_ActiveNode(METASOUND_DEBUG_SCOPE_NAME, InNodePtr);
+#define METASOUND_DEBUG_ACTIVE_NODE_NAME (::Metasound::ThreadLocalDebug::GetActiveNodeClassNameAndVersionOnThisThread())
+
+#else // if UE_METASOUND_DEBUG_ENABLED
+
+#define METASOUND_DEBUG_SCOPE_NAME 
+#define METASOUND_DEBUG_DECLARE_SCOPE 
+#define METASOUND_DEBUG_SET_ACTIVE_NODE_SCOPE(InNodePtr)
+#define METASOUND_DEBUG_ACTIVE_NODE_NAME TEXT("[Metasound debug disable. Use UE_METASOUND_DEBUG_ENABLED to enable]")
+
+#endif // if UE_METASOUND_DEBUG_ENABLED
+
diff --git a/Engine/Plugins/Runtime/Metasound/Source/MetasoundGraphCore/Private/MetasoundVertexData.cpp b/Engine/Plugins/Runtime/Metasound/Source/MetasoundGraphCore/Private/MetasoundVertexData.cpp
--- a/Engine/Plugins/Runtime/Metasound/Source/MetasoundGraphCore/Private/MetasoundVertexData.cpp
+++ b/Engine/Plugins/Runtime/Metasound/Source/MetasoundGraphCore/Private/MetasoundVertexData.cpp
@@ -37,7 +37,7 @@ namespace Metasound
 					}
 			}
 
-			checkf(bIsCompatible, TEXT("Vertex access type \"%s\" is incompatible with data access type \"%s\" on vertex \"%s\" on node \"%s\""), *LexToString(InDataVertex.AccessType), *LexToString(ReferenceAccessType), *InDataVertex.VertexName.ToString(), ThreadLocalDebug::GetActiveNodeClassNameAndVersion());
+			checkf(bIsCompatible, TEXT("Vertex access type \"%s\" is incompatible with data access type \"%s\" on vertex \"%s\" on node \"%s\""), *LexToString(InDataVertex.AccessType), *LexToString(ReferenceAccessType), *InDataVertex.VertexName.ToString(), METASOUND_DEBUG_ACTIVE_NODE_NAME);
 		}
 #endif // #if ENABLE_METASOUND_ACCESS_TYPE_COMPATIBILITY_TEST
 		EVertexAccessType DataReferenceAccessTypeToVertexAccessType(EDataReferenceAccessType InReferenceAccessType)
diff --git a/Engine/Plugins/Runtime/Metasound/Source/MetasoundGraphCore/Public/MetasoundExecutableOperator.h b/Engine/Plugins/Runtime/Metasound/Source/MetasoundGraphCore/Public/MetasoundExecutableOperator.h
--- a/Engine/Plugins/Runtime/Metasound/Source/MetasoundGraphCore/Public/MetasoundExecutableOperator.h
+++ b/Engine/Plugins/Runtime/Metasound/Source/MetasoundGraphCore/Public/MetasoundExecutableOperator.h
@@ -180,16 +180,6 @@ namespace Metasound
 			{
 				return nullptr;
 			}
-
-			virtual FDataReferenceCollection GetInputs() const override
-			{
-				return FDataReferenceCollection{};
-			}
-
-			virtual FDataReferenceCollection GetOutputs() const override
-			{
-				return FDataReferenceCollection{};
-			}
 	};
 
 	/** FExecuter
@@ -268,30 +258,6 @@ namespace Metasound
 				return Operator.IsValid();
 			}
 
-			virtual FDataReferenceCollection GetInputs() const override
-			{
-				static FDataReferenceCollection EmptyCollection;
-
-				if (Operator.IsValid())
-				{
-					return Operator->GetInputs();
-				}
-
-				return EmptyCollection;
-			}
-
-			virtual FDataReferenceCollection GetOutputs() const override
-			{
-				static FDataReferenceCollection EmptyCollection;
-
-				if (Operator.IsValid())
-				{
-					return Operator->GetOutputs();
-				}
-
-				return EmptyCollection;
-			}
-
 			virtual void BindInputs(FInputVertexInterfaceData& InVertexData) override
 			{
 				if (Operator.IsValid())
diff --git a/Engine/Plugins/Runtime/Metasound/Source/MetasoundGraphCore/Public/MetasoundOperatorInterface.h b/Engine/Plugins/Runtime/Metasound/Source/MetasoundGraphCore/Public/MetasoundOperatorInterface.h
--- a/Engine/Plugins/Runtime/Metasound/Source/MetasoundGraphCore/Public/MetasoundOperatorInterface.h
+++ b/Engine/Plugins/Runtime/Metasound/Source/MetasoundGraphCore/Public/MetasoundOperatorInterface.h
@@ -29,51 +29,109 @@ namespace Metasound
 
 		virtual ~IOperator() {}
 
-		/** Return the input parameters associated with this IOperator.
-		 *
-		 * Implementations of IOperator should populate and return their input parameters
-		 * if they want to enable callers to query and write to their parameters. Most IOperator
-		 * implementations will return an empty FDataReferenceCollection because thier inputs are
-		 * set during the IOperator's construction and do not need to be updated afterwards. The
-		 * exceptions are input operators and graph operators which need to interface with external
-		 * systems even after the operator is created.
+		/** GetInputs() has been deprecated in favor of BindInputs(...). Please update
+		 * your code by removing the implementation to GetInputs() and implementing
+		 * BindInputs(...).  In future releases, the GetInputs() virtual method
+		 * will be removed.
+		 *
+		 * Note: IOperators which do not correctly implement BindInputs(...) and
+		 * BindOutputs(...) will not function correctly with MetaSound operator
+		 * caching and live auditioning in the MetaSound Builder BP API. 
 		 */
-		virtual FDataReferenceCollection GetInputs() const = 0;
+		UE_DEPRECATED(5.3, "GetInputs() has been replaced by BindInputs(FInputVertexInterfaceData&).")
+		virtual FDataReferenceCollection METASOUNDGRAPHCORE_API GetInputs() const;
 
-		/** Return the output parameters associated with this IOperator.
+		/** GetOutputs() has been deprecated in favor of BindOutputs(...). Please update
+		 * your code by removing the implementation to GetOutputs() and implementing
+		 * BindOutputs(...).  In future releases, the GetOutputs() virtual method
+		 * will be removed.
 		 *
-		 * Implementations of IOperator should return a collection of their output parameters
-		 * which other nodes can read.
+		 * Note: IOperators which do not correctly implement BindInputs(...) and
+		 * BindOutputs(...) will not function correctly with MetaSound operator
+		 * caching and live auditioning in the MetaSound Builder BP API. 
 		 */
-		virtual FDataReferenceCollection GetOutputs() const = 0;
+		UE_DEPRECATED(5.3, "GetOutputs() has been replaced by BindOutputs(FOutputVertexInterfaceData&)")
+		virtual FDataReferenceCollection METASOUNDGRAPHCORE_API GetOutputs() const;
 
-		/** Bind data references to the interface of this operator.
+		/** Bind(...) has been deprecated in favor of BindInputs(...) and BindOutputs(...). 
+		 * Please update your code by removing the implementation to Bind(...) 
+		 * and implementing BindOutputs(...).  In future releases, the Bind(...) 
+		 * virtual method will be removed.
 		 *
-		 * Operators should bind all data references to the vertices of the provided
-		 * InVertexData. This method is meant to replace GetInputs() and GetOutputs()
-		 * because it provides better error checking against mismatched vertex names.
+		 * Note: IOperators which do not correctly implement BindInputs(...) and
+		 * BindOutputs(...) will not function correctly with MetaSound operator
+		 * caching and live auditioning in the MetaSound Builder BP API. 
 		 */
-		virtual void Bind(FVertexInterfaceData& InVertexData) const
-		{
-			InVertexData.GetInputs().Set(GetInputs());
-			InVertexData.GetOutputs().Set(GetOutputs());
-		}
+		UE_DEPRECATED(5.3, "Bind(FVertexInterfaceData&) has been replaced by BindInputs(FInputVertexInterfaceDAta&) and BindOutputs(FOutputVertexInterfaceData&)")
+		virtual void METASOUNDGRAPHCORE_API Bind(FVertexInterfaceData& InVertexData) const;
 
-		virtual void BindInputs(FInputVertexInterfaceData& InVertexData)
-		{
-			FVertexInterfaceData Data;
-			Data.GetInputs() = InVertexData;
-			Bind(Data);
-			InVertexData = Data.GetInputs();
-		}
+		/** BindInputs binds data references in the IOperator with the FInputVertexInterfaceData.
+		 *
+		 * All input data references should be bound to the InVertexData to support
+		 * other MetaSound systems such as MetaSound BP API, Operator Caching, 
+		 * and live auditioning. 
+		 *
+		 * Note: The virtual function IOPerator::BindInputs(...) will be made a
+		 * pure virtual when IOperator::GetInputs() is removed at or after release 5.5
+		 *
+		 * Note: Binding an data reference may update the which underlying object 
+		 * the reference points to. Any operator which caches pointers or values 
+		 * from data references must update their cached pointers in the call to 
+		 * BindInputs(...).  Operators which do not cache the underlying pointer 
+		 * of a data reference do not need to update anything after BindInputs(...)
+		 *
+		 * Example:
+		 * 	FMyOperator::FMyOperator(TDataReadReference<float> InGain, TDataReadReference<FAudioBuffer> InAudioBuffer)
+		 * 	: Gain(InGain)
+		 * 	, AudioBuffer(InAudioBuffer)
+		 * 	{
+		 * 		MyBufferPtr = AudioBuffer->GetData();		
+		 * 	}
+		 *
+		 * 	void FMyOperator::BindInputs(FInputVertexInterfaceData& InVertexData)
+		 * 	{
+		 * 		InVertexData.BindReadVertex("Gain", Gain);
+		 * 		InVertexData.BindReadVertex("Audio", AudioBuffer);
+		 *
+		 * 		// Update MyBufferPtr in case AudioBuffer references a new FAudioBuffer.
+		 * 		MyBufferPtr = AudioBuffer->GetData();
+		 * 	}
+		 */
+		virtual void METASOUNDGRAPHCORE_API BindInputs(FInputVertexInterfaceData& InVertexData);
 
-		virtual void BindOutputs(FOutputVertexInterfaceData& InVertexData) 
-		{
-			FVertexInterfaceData Data;
-			Data.GetOutputs() = InVertexData;
-			Bind(Data);
-			InVertexData = Data.GetOutputs();
-		}
+		/** BindOutputs binds data references in the IOperator with the FOutputVertexInterfaceData.
+		 *
+		 * All output data references should be bound to the InVertexData to support
+		 * other MetaSound systems such as MetaSound BP API, Operator Caching, 
+		 * and live auditioning. 
+		 *
+		 * Note: The virtual function IOPerator::BindOutputs(...) will be made a
+		 * pure virtual when IOperator::GetOutputs() is removed at or after release 5.5
+		 *
+		 * Note: Binding an data reference may update the which underlying object 
+		 * the reference points to. Any operator which caches pointers or values 
+		 * from data references must update their cached pointers in the call to 
+		 * BindOutputs(...).  Operators which do not cache the underlying pointer 
+		 * of a data reference do not need to update anything after BindOutputs(...)
+		 *
+		 * Example:
+		 * 	FMyOperator::FMyOperator(TDataWriteReference<float> InGain, TDataWriteReference<FAudioBuffer> InAudioBuffer)
+		 * 	: Gain(InGain)
+		 * 	, AudioBuffer(InAudioBuffer)
+		 * 	{
+		 * 		MyBufferPtr = AudioBuffer->GetData();		
+		 * 	}
+		 *
+		 * 	void FMyOperator::BindOutputs(FOutputVertexInterfaceData& InVertexData)
+		 * 	{
+		 * 		InVertexData.BindReadVertex("Gain", Gain);
+		 * 		InVertexData.BindReadVertex("Audio", AudioBuffer);
+		 *
+		 * 		// Update MyBufferPtr in case AudioBuffer references a new FAudioBuffer.
+		 * 		MyBufferPtr = AudioBuffer->GetData();
+		 * 	}
+		 */
+		virtual void METASOUNDGRAPHCORE_API BindOutputs(FOutputVertexInterfaceData& InVertexData);
 
 		/** Pointer to initialize function for an operator.
 		 *
