Fix setting array parameters on metasound source

diff --git a/Engine/Plugins/Runtime/Metasound/Source/MetasoundFrontend/Public/MetasoundDataTypeRegistrationMacro.h b/Engine/Plugins/Runtime/Metasound/Source/MetasoundFrontend/Public/MetasoundDataTypeRegistrationMacro.h
--- a/Engine/Plugins/Runtime/Metasound/Source/MetasoundFrontend/Public/MetasoundDataTypeRegistrationMacro.h
+++ b/Engine/Plugins/Runtime/Metasound/Source/MetasoundFrontend/Public/MetasoundDataTypeRegistrationMacro.h
@@ -474,13 +474,10 @@ namespace Metasound
 				{
 					if constexpr (std::is_copy_assignable_v<TDataType> && bIsParsable)
 					{
-						if constexpr (!TIsArrayType<TDataType>::Value)
+						this->LiteralAssignmentFunction = [](const FOperatorSettings& InOperatorSettings, const FLiteral& InLiteral, const FAnyDataReference& OutDataRef)
 						{
-							this->LiteralAssignmentFunction = [](const FOperatorSettings& InOperatorSettings, const FLiteral& InLiteral, const FAnyDataReference& OutDataRef)
-							{
-								*OutDataRef.GetWritableValue<TDataType>() = TDataTypeLiteralFactory<TDataType>::CreateExplicitArgs(InOperatorSettings, InLiteral);
-							};
-						}
+							*OutDataRef.GetWritableValue<TDataType>() = TDataTypeLiteralFactory<TDataType>::CreateExplicitArgs(InOperatorSettings, InLiteral);
+						};
 					}
 				}
 
