Initial work to pipe builder updates to dynamic graph
- Add delegates to respond to builder changes to document
- Add delegate handlers to source builder to forward requests to audition sound if live update enabled
- Misc piping for live updating
#ok forgithub public

diff --git a/Engine/Plugins/Runtime/Metasound/Source/MetasoundEngine/Private/MetasoundAssetSubsystem.cpp b/Engine/Plugins/Runtime/Metasound/Source/MetasoundEngine/Private/MetasoundAssetSubsystem.cpp
--- a/Engine/Plugins/Runtime/Metasound/Source/MetasoundEngine/Private/MetasoundAssetSubsystem.cpp
+++ b/Engine/Plugins/Runtime/Metasound/Source/MetasoundEngine/Private/MetasoundAssetSubsystem.cpp
@@ -8,6 +8,7 @@
 #include "Engine/AssetManager.h"
 #include "Metasound.h"
 #include "MetasoundAssetBase.h"
+#include "MetasoundBuilderSubsystem.h"
 #include "MetasoundFrontendDocument.h"
 #include "MetasoundFrontendRegistries.h"
 #include "MetasoundFrontendTransform.h"
@@ -18,6 +19,7 @@
 #include "Misc/CoreDelegates.h"
 #include "UObject/NoExportTypes.h"
 
+
 #include UE_INLINE_GENERATED_CPP_BY_NAME(MetasoundAssetSubsystem)
 
 
@@ -572,7 +574,11 @@ void UMetaSoundAssetSubsystem::RemoveAsset(const UObject& InObject)
 	{
 		const FNodeClassInfo ClassInfo = MetaSoundAsset->GetAssetClassInfo();
 		FNodeRegistryKey RegistryKey = FMetasoundFrontendRegistryContainer::Get()->GetRegistryKey(ClassInfo);
-		AssetSubsystemPrivate::RemoveIfExactMatch(PathMap, RegistryKey, FSoftObjectPath(&InObject));
+		const FSoftObjectPath ObjectPath(&InObject);
+		if (AssetSubsystemPrivate::RemoveIfExactMatch(PathMap, RegistryKey, ObjectPath))
+		{
+			UMetaSoundBuilderSubsystem::GetChecked().DetachBuilderFromAsset(ClassInfo.ClassName);
+		}
 	}
 }
 
@@ -585,7 +591,11 @@ void UMetaSoundAssetSubsystem::RemoveAsset(const FAssetData& InAssetData)
 	if (ensureAlways(AssetSubsystemPrivate::GetAssetClassInfo(InAssetData, ClassInfo)))
 	{
 		FNodeRegistryKey RegistryKey = FMetasoundFrontendRegistryContainer::Get()->GetRegistryKey(ClassInfo);
-		AssetSubsystemPrivate::RemoveIfExactMatch(PathMap, RegistryKey, InAssetData.GetSoftObjectPath());
+		const FSoftObjectPath ObjectPath = InAssetData.GetSoftObjectPath();
+		if (AssetSubsystemPrivate::RemoveIfExactMatch(PathMap, RegistryKey, ObjectPath))
+		{
+			UMetaSoundBuilderSubsystem::GetChecked().DetachBuilderFromAsset(ClassInfo.ClassName);
+		}
 	}
 }
 
diff --git a/Engine/Plugins/Runtime/Metasound/Source/MetasoundEngine/Private/MetasoundBuilderSubsystem.cpp b/Engine/Plugins/Runtime/Metasound/Source/MetasoundEngine/Private/MetasoundBuilderSubsystem.cpp
--- a/Engine/Plugins/Runtime/Metasound/Source/MetasoundEngine/Private/MetasoundBuilderSubsystem.cpp
+++ b/Engine/Plugins/Runtime/Metasound/Source/MetasoundEngine/Private/MetasoundBuilderSubsystem.cpp
@@ -10,6 +10,9 @@
 #include "Interfaces/MetasoundFrontendSourceInterface.h"
 #include "Metasound.h"
 #include "MetasoundAssetSubsystem.h"
+#include "MetasoundDataReference.h"
+#include "MetasoundDynamicOperatorTransactor.h"
+#include "MetasoundFrontendDataTypeRegistry.h"
 #include "MetasoundFrontendDocument.h"
 #include "MetasoundFrontendRegistries.h"
 #include "MetasoundFrontendSearchEngine.h"
@@ -20,6 +23,7 @@
 #include "MetasoundSource.h"
 #include "MetasoundTrace.h"
 #include "MetasoundUObjectRegistry.h"
+#include "MetasoundVertex.h"
 
 #include UE_INLINE_GENERATED_CPP_BY_NAME(MetasoundBuilderSubsystem)
 
@@ -28,6 +32,8 @@ namespace Metasound::Engine
 {
 	namespace BuilderSubsystemPrivate
 	{
+		static int32 LiveUpdateEnabledCVar = 0;
+
 		template <typename TLiteralType>
 		FMetasoundFrontendLiteral CreatePODMetaSoundLiteral(const TLiteralType& Value, FName& OutDataType)
 		{
@@ -38,6 +44,16 @@ namespace Metasound::Engine
 			return Literal;
 		}
 
+		TUniquePtr<INode> CreateDynamicNodeFromFrontendLiteral(const FName DataType, const FMetasoundFrontendLiteral& InLiteral)
+		{
+			using namespace Frontend;
+			FLiteral ValueLiteral = InLiteral.ToLiteral(DataType);
+
+			// TODO: Node name "Literal" is always the same.  Consolidate and deprecate providing unique node name to avoid unnecessary FName table bloat.
+			FLiteralNodeConstructorParams Params { "Literal", FGuid::NewGuid(), MoveTemp(ValueLiteral) };
+			return IDataTypeRegistry::Get().CreateLiteralNode(DataType, MoveTemp(Params));
+		}
+
 		template <typename BuilderClass>
 		BuilderClass& CreateTransientBuilder()
 		{
@@ -46,22 +62,22 @@ namespace Metasound::Engine
 			const FName ObjectName = MakeUniqueObjectName(TransientPackage, BuilderClass::StaticClass());
 			TObjectPtr<BuilderClass> NewBuilder = NewObject<BuilderClass>(TransientPackage, ObjectName, NewObjectFlags);
 			check(NewBuilder);
-			NewBuilder->CreateTransientDocument();
+			NewBuilder->InitFrontendBuilder();
 			return *NewBuilder.Get();
 		}
 	} // namespace BuilderSubsystemPrivate
 } // namespace Metasound::Engine
 
+FAutoConsoleVariableRef CVarMetaSoundEditorAudioConnectSpacing(
+	TEXT("au.MetaSound.Builder.Audition.LiveUpdatesEnabled"),
+	Metasound::Engine::BuilderSubsystemPrivate::LiveUpdateEnabledCVar,
+	TEXT("Whether or not the MetaSound builder's audition feature supports live updates (mutating a MetaSound graph while it is being executed at runtime).\n")
+	TEXT("Values: 0 (Disabled), !=0 (Enabled)"),
+	ECVF_Default);
+
+
 FMetaSoundBuilderNodeOutputHandle UMetaSoundBuilderBase::AddGraphInputNode(FName Name, FName DataType, FMetasoundFrontendLiteral DefaultValue, EMetaSoundBuilderResult& OutResult, bool bIsConstructorInput)
 {
-	FMetasoundFrontendClassInput Description;
-	Description.Name = Name;
-	Description.TypeName = DataType;
-	Description.NodeID = FGuid::NewGuid();
-	Description.VertexID = FGuid::NewGuid();
-	Description.DefaultLiteral = static_cast<FMetasoundFrontendLiteral>(DefaultValue);
-	Description.AccessType = bIsConstructorInput ? EMetasoundFrontendVertexAccessType::Value : EMetasoundFrontendVertexAccessType::Reference;
-
 	FMetaSoundBuilderNodeOutputHandle NewHandle;
 
 	const FMetasoundFrontendNode* Node = Builder.FindGraphInputNode(Name);
@@ -71,6 +87,13 @@ FMetaSoundBuilderNodeOutputHandle UMetaSoundBuilderBase::AddGraphInputNode(FName
 	}
 	else
 	{
+		FMetasoundFrontendClassInput Description;
+		Description.Name = Name;
+		Description.TypeName = DataType;
+		Description.NodeID = FGuid::NewGuid();
+		Description.VertexID = FGuid::NewGuid();
+		Description.DefaultLiteral = static_cast<FMetasoundFrontendLiteral>(DefaultValue);
+		Description.AccessType = bIsConstructorInput ? EMetasoundFrontendVertexAccessType::Value : EMetasoundFrontendVertexAccessType::Reference;
 		Node = Builder.AddGraphInput(Description);
 	}
 
@@ -89,13 +112,6 @@ FMetaSoundBuilderNodeOutputHandle UMetaSoundBuilderBase::AddGraphInputNode(FName
 
 FMetaSoundBuilderNodeInputHandle UMetaSoundBuilderBase::AddGraphOutputNode(FName Name, FName DataType, FMetasoundFrontendLiteral DefaultValue, EMetaSoundBuilderResult& OutResult, bool bIsConstructorOutput)
 {
-	FMetasoundFrontendClassOutput Description;
-	Description.Name = Name;
-	Description.TypeName = DataType;
-	Description.NodeID = FGuid::NewGuid();
-	Description.VertexID = FGuid::NewGuid();
-	Description.AccessType = bIsConstructorOutput ? EMetasoundFrontendVertexAccessType::Value : EMetasoundFrontendVertexAccessType::Reference;
-
 	FMetaSoundBuilderNodeInputHandle NewHandle;
 	const FMetasoundFrontendNode* Node = Builder.FindGraphOutputNode(Name);
 	if (Node)
@@ -104,6 +120,12 @@ FMetaSoundBuilderNodeInputHandle UMetaSoundBuilderBase::AddGraphOutputNode(FName
 	}
 	else
 	{
+		FMetasoundFrontendClassOutput Description;
+		Description.Name = Name;
+		Description.TypeName = DataType;
+		Description.NodeID = FGuid::NewGuid();
+		Description.VertexID = FGuid::NewGuid();
+		Description.AccessType = bIsConstructorOutput ? EMetasoundFrontendVertexAccessType::Value : EMetasoundFrontendVertexAccessType::Reference;
 		Node = Builder.AddGraphOutput(Description);
 	}
 
@@ -149,7 +171,6 @@ FMetaSoundNodeHandle UMetaSoundBuilderBase::AddNode(const TScriptInterface<IMeta
 
 			if (const FMetasoundFrontendNode* NewNode = Builder.AddGraphNode(NodeClassGraph))
 			{
-				FMetasoundFrontendClassMetadata NodeClassMetadata = NodeClassGraph.Metadata;
 				NewHandle.NodeID = NewNode->GetID();
 			}
 		}
@@ -159,6 +180,21 @@ FMetaSoundNodeHandle UMetaSoundBuilderBase::AddNode(const TScriptInterface<IMeta
 	return NewHandle;
 }
 
+FMetaSoundNodeHandle UMetaSoundBuilderBase::AddNodeByClassName(const FMetasoundFrontendClassName& ClassName, int32 MajorVersion, EMetaSoundBuilderResult& OutResult)
+{
+	using namespace Metasound;
+	using namespace Metasound::Frontend;
+
+	FMetaSoundNodeHandle NewHandle;
+	if (const FMetasoundFrontendNode* NewNode = Builder.AddNodeByClassName(ClassName, MajorVersion))
+	{
+		NewHandle.NodeID = NewNode->GetID();
+	}
+
+	OutResult = NewHandle.IsSet() ? EMetaSoundBuilderResult::Succeeded : EMetaSoundBuilderResult::Failed;
+	return NewHandle;
+}
+
 bool UMetaSoundBuilderBase::ContainsNode(const FMetaSoundNodeHandle& NodeHandle) const
 {
 	return Builder.ContainsNode(NodeHandle.NodeID);
@@ -456,6 +492,7 @@ FMetaSoundNodeHandle UMetaSoundBuilderBase::FindGraphInputNode(FName InputName,
 {
 	if (const FMetasoundFrontendNode* GraphInputNode = Builder.FindGraphInputNode(InputName))
 	{
+		OutResult = EMetaSoundBuilderResult::Succeeded;
 		return FMetaSoundNodeHandle { GraphInputNode->GetID() };
 	}
 
@@ -477,12 +514,17 @@ FMetaSoundNodeHandle UMetaSoundBuilderBase::FindGraphOutputNode(FName OutputName
 	return { };
 }
 
-void UMetaSoundBuilderBase::CreateTransientDocument()
+UMetaSoundBuilderDocument* UMetaSoundBuilderBase::CreateTransientDocumentObject() const
 {
-	UMetaSoundBuilderDocument* NewDoc = NewObject<UMetaSoundBuilderDocument>();
-	NewDoc->SetBaseMetaSoundUClass(GetBuilderUClass());
-	TScriptInterface<IMetaSoundDocumentInterface> MetaSoundSourceDocInterface = NewDoc;
-	Builder = FMetaSoundFrontendDocumentBuilder(MetaSoundSourceDocInterface);
+	UMetaSoundBuilderDocument* DocObject = NewObject<UMetaSoundBuilderDocument>();
+	DocObject->SetBaseMetaSoundUClass(GetBuilderUClass());
+	return DocObject;
+}
+
+void UMetaSoundBuilderBase::InitFrontendBuilder()
+{
+	UMetaSoundBuilderDocument* DocObject = CreateTransientDocumentObject();
+	Builder = FMetaSoundFrontendDocumentBuilder(DocObject);
 	Builder.InitDocument();
 }
 
@@ -576,6 +618,12 @@ void UMetaSoundBuilderBase::RemoveNode(const FMetaSoundNodeHandle& NodeHandle, E
 	OutResult = bNodeRemoved ? EMetaSoundBuilderResult::Succeeded : EMetaSoundBuilderResult::Failed;
 }
 
+void UMetaSoundBuilderBase::RemoveNodeInputDefault(const FMetaSoundBuilderNodeInputHandle& InputHandle, EMetaSoundBuilderResult& OutResult)
+{
+	const bool bInputDefaultRemoved = Builder.RemoveNodeInputDefault(InputHandle.NodeID, InputHandle.VertexID);
+	OutResult = bInputDefaultRemoved ? EMetaSoundBuilderResult::Succeeded : EMetaSoundBuilderResult::Failed;
+}
+
 #if WITH_EDITOR
 void UMetaSoundBuilderBase::SetAuthor(const FString& InAuthor)
 {
@@ -585,9 +633,6 @@ void UMetaSoundBuilderBase::SetAuthor(const FString& InAuthor)
 
 void UMetaSoundBuilderBase::SetNodeInputDefault(const FMetaSoundBuilderNodeInputHandle& InputHandle, const FMetasoundFrontendLiteral& Literal, EMetaSoundBuilderResult& OutResult)
 {
-	using namespace Metasound;
-	using namespace Metasound::Frontend;
-
 	const bool bInputDefaultSet = Builder.SetNodeInputDefault(InputHandle.NodeID, InputHandle.VertexID, Literal);
 	OutResult = bInputDefaultSet ? EMetaSoundBuilderResult::Succeeded : EMetaSoundBuilderResult::Failed;
 }
@@ -653,6 +698,30 @@ void UMetaSoundBuilderBase::GetNodeInputData(const FMetaSoundBuilderNodeInputHan
 	}
 }
 
+FMetasoundFrontendLiteral UMetaSoundBuilderBase::GetNodeInputDefault(const FMetaSoundBuilderNodeInputHandle& InputHandle, EMetaSoundBuilderResult& OutResult)
+{
+	if (const FMetasoundFrontendLiteral* Default = Builder.GetNodeInputDefault(InputHandle.NodeID, InputHandle.VertexID))
+	{
+		OutResult = EMetaSoundBuilderResult::Succeeded;
+		return *Default;
+	}
+
+	OutResult = EMetaSoundBuilderResult::Failed;
+	return { };
+}
+
+FMetasoundFrontendLiteral UMetaSoundBuilderBase::GetNodeInputClassDefault(const FMetaSoundBuilderNodeInputHandle& InputHandle, EMetaSoundBuilderResult& OutResult)
+{
+	if (const FMetasoundFrontendLiteral* Default = Builder.GetNodeInputClassDefault(InputHandle.NodeID, InputHandle.VertexID))
+	{
+		OutResult = EMetaSoundBuilderResult::Succeeded;
+		return *Default;
+	}
+
+	OutResult = EMetaSoundBuilderResult::Failed;
+	return { };
+}
+
 void UMetaSoundBuilderBase::GetNodeOutputData(const FMetaSoundBuilderNodeOutputHandle& OutputHandle, FName& Name, FName& DataType, EMetaSoundBuilderResult& OutResult)
 {
 	if (const FMetasoundFrontendVertex* Vertex = Builder.FindNodeOutput(OutputHandle.NodeID, OutputHandle.VertexID))
@@ -711,7 +780,7 @@ void UMetaSoundBuilderBase::ConvertToPreset(const TScriptInterface<IMetaSoundDoc
 
 TScriptInterface<IMetaSoundDocumentInterface> UMetaSoundPatchBuilder::Build(UObject* Parent, const FMetaSoundBuilderOptions& InBuilderOptions) const
 {
-	return BuildInternal<UMetaSoundPatch>(Parent, InBuilderOptions);
+	return &BuildInternal<UMetaSoundPatch>(Parent, InBuilderOptions);
 }
 
 const UClass& UMetaSoundPatchBuilder::GetBuilderUClass() const
@@ -719,9 +788,52 @@ const UClass& UMetaSoundPatchBuilder::GetBuilderUClass() const
 	return *UMetaSoundPatch::StaticClass();
 }
 
-void UMetaSoundSourceBuilder::Audition(UObject* Parent, UAudioComponent* AudioComponent, FOnCreateAuditionGeneratorHandleDelegate CreateGenerator)
+UMetaSoundBuilderBase& UMetaSoundBuilderSubsystem::AttachBuilderToAssetChecked(UObject& InObject) const
 {
-	using namespace Metasound;
+	const UClass* BaseClass = InObject.GetClass();
+	if (BaseClass == UMetaSoundSource::StaticClass())
+	{
+		UMetaSoundSourceBuilder* NewBuilder = AttachSourceBuilderToAsset(CastChecked<UMetaSoundSource>(&InObject));
+		return *NewBuilder;
+	}
+	else if (BaseClass == UMetaSoundPatch::StaticClass())
+	{
+		UMetaSoundPatchBuilder* NewBuilder = AttachPatchBuilderToAsset(CastChecked<UMetaSoundPatch>(&InObject));
+		return *NewBuilder;
+	}
+	else
+	{
+		checkf(false, TEXT("UClass '%s' is not a base MetaSound that supports attachment via the MetaSoundBuilderSubsystem"), *BaseClass->GetFullName());
+		return Metasound::Engine::BuilderSubsystemPrivate::CreateTransientBuilder<UMetaSoundPatchBuilder>();
+	}
+}
+
+UMetaSoundPatchBuilder* UMetaSoundBuilderSubsystem::AttachPatchBuilderToAsset(UMetaSoundPatch* InPatch) const
+{
+	if (InPatch)
+	{
+		return &AttachBuilderToAssetCheckedPrivate<UMetaSoundPatchBuilder>(InPatch);
+	}
+
+	return nullptr;
+}
+
+UMetaSoundSourceBuilder* UMetaSoundBuilderSubsystem::AttachSourceBuilderToAsset(UMetaSoundSource* InSource) const
+{
+	if (InSource)
+	{
+		UMetaSoundSourceBuilder& SourceBuilder = AttachBuilderToAssetCheckedPrivate<UMetaSoundSourceBuilder>(InSource);
+		SourceBuilder.AuditionSound = InSource;
+		return &SourceBuilder;
+	}
+
+	return nullptr;
+}
+
+void UMetaSoundSourceBuilder::Audition(UObject* Parent, UAudioComponent* AudioComponent, FOnCreateAuditionGeneratorHandleDelegate CreateGenerator, bool bLiveUpdatesEnabled)
+{
+	using namespace Metasound::Engine;
+	using namespace Metasound::Engine::BuilderSubsystemPrivate;
 	using namespace Metasound::Frontend;
 
 	METASOUND_TRACE_CPUPROFILER_EVENT_SCOPE(UMetaSoundSourceBuilder::Audition);
@@ -732,33 +844,44 @@ void UMetaSoundSourceBuilder::Audition(UObject* Parent, UAudioComponent* AudioCo
 		return;
 	}
 
-	FMetaSoundBuilderOptions BuilderOptions;
-	BuilderOptions.bAddToRegistry = true;
-
-	if (AuditionSound.IsValid())
+	if (!bIsAttached)
 	{
-		BuilderOptions.ExistingMetaSound = AuditionSound.Get();
-		if (USoundBase* Sound = AudioComponent->GetSound())
+		FMetaSoundBuilderOptions BuilderOptions;
+		BuilderOptions.bAddToRegistry = true;
+
+		if (AuditionSound.IsValid())
 		{
-			if (Sound != AuditionSound)
+			BuilderOptions.ExistingMetaSound = AuditionSound.Get();
+			if (USoundBase* Sound = AudioComponent->GetSound())
 			{
-				UE_LOG(LogMetaSound, Warning, TEXT("MetaSoundBuilder '%s' supplied AudioComponent with unlinked sound '%s'. Stopping sound and replacing with builder's audition sound."),
-					*GetFullName(),
-					*Sound->GetFullName());
+				if (Sound != AuditionSound)
+				{
+					UE_LOG(LogMetaSound, Warning, TEXT("MetaSoundBuilder '%s' supplied AudioComponent with unlinked sound '%s'. Stopping sound and replacing with builder's audition sound."),
+						*GetFullName(),
+						*Sound->GetFullName());
+				}
 			}
 		}
-	}
-	else
-	{
-		BuilderOptions.Name = MakeUniqueObjectName(nullptr, UMetaSoundSource::StaticClass(), FName(GetName() + TEXT("_Audition")));
+		else
+		{
+			BuilderOptions.Name = MakeUniqueObjectName(nullptr, UMetaSoundSource::StaticClass(), FName(GetName() + TEXT("_Audition")));
+		}
+
+		AuditionSound = &BuildInternal<UMetaSoundSource>(Parent, BuilderOptions);
 	}
 
-	AuditionSound = BuildInternal<UMetaSoundSource>(Parent, BuilderOptions);
-	if (AuditionSound.IsValid())
+	AudioComponent->SetSound(AuditionSound.Get());
+
+	if (bLiveUpdatesEnabled && !BuilderSubsystemPrivate::LiveUpdateEnabledCVar)
 	{
-		AudioComponent->SetSound(AuditionSound.Get());
+		UE_LOG(LogMetaSound, Warning, TEXT("Attempting to audition MetaSoundBuilder '%s' with live update enabled, but feature's console variable is disabled. Request is being ignored and updates will not be audible."),
+			*GetFullName(),
+			* AuditionSound->GetFullName());
 	}
 
+	const bool bEnableDynamicGenerators = bLiveUpdatesEnabled && BuilderSubsystemPrivate::LiveUpdateEnabledCVar;
+	AuditionSound->EnableDynamicGenerators(bEnableDynamicGenerators);
+
 	if (CreateGenerator.IsBound())
 	{
 		UMetasoundGeneratorHandle* NewHandle = UMetasoundGeneratorHandle::CreateMetaSoundGeneratorHandle(AudioComponent);
@@ -769,9 +892,40 @@ void UMetaSoundSourceBuilder::Audition(UObject* Parent, UAudioComponent* AudioCo
 	AudioComponent->Play();
 }
 
+bool UMetaSoundSourceBuilder::ExecuteAuditionableTransaction(FAuditionableTransaction Transaction) const
+{
+	using namespace Metasound::Engine::BuilderSubsystemPrivate;
+	using namespace Metasound::DynamicGraph;
+
+	METASOUND_TRACE_CPUPROFILER_EVENT_SCOPE(UMetaSoundSourceBuilder::ExecuteAuditionableTransaction);
+
+	if (LiveUpdateEnabledCVar && AuditionSound.IsValid())
+	{
+		TSharedPtr<FDynamicOperatorTransactor> Transactor = AuditionSound->GetDynamicGeneratorTransactor();
+		if (Transactor.IsValid())
+		{
+			return Transaction(*Transactor);
+		}
+	}
+
+	return false;
+}
+
 TScriptInterface<IMetaSoundDocumentInterface> UMetaSoundSourceBuilder::Build(UObject* Parent, const FMetaSoundBuilderOptions& InBuilderOptions) const
 {
-	return BuildInternal<UMetaSoundSource>(Parent, InBuilderOptions);
+	return &BuildInternal<UMetaSoundSource>(Parent, InBuilderOptions);
+}
+
+void UMetaSoundSourceBuilder::InitFrontendBuilder()
+{
+	using namespace Metasound::Frontend;
+
+	TSharedRef<FDocumentModifyDelegates> DocumentDelegates = MakeShared<FDocumentModifyDelegates>();
+	InitDelegates(*DocumentDelegates);
+
+	UMetaSoundBuilderDocument* DocObject = CreateTransientDocumentObject();
+	Builder = FMetaSoundFrontendDocumentBuilder(DocObject, DocumentDelegates);
+	Builder.InitDocument();
 }
 
 const Metasound::Engine::FOutputAudioFormatInfoPair* UMetaSoundSourceBuilder::FindOutputAudioFormatInfo() const
@@ -794,6 +948,286 @@ const UClass& UMetaSoundSourceBuilder::GetBuilderUClass() const
 	return *UMetaSoundSource::StaticClass();
 }
 
+bool UMetaSoundSourceBuilder::GetLiveUpdatesEnabled() const
+{
+	using namespace Metasound::Engine::BuilderSubsystemPrivate;
+
+	if (!LiveUpdateEnabledCVar)
+	{
+		return false;
+	}
+
+	if (!AuditionSound.IsValid())
+	{
+		return false;
+	}
+	
+	return AuditionSound->GetDynamicGeneratorTransactor().IsValid();
+}
+
+void UMetaSoundSourceBuilder::InitDelegates(Metasound::Frontend::FDocumentModifyDelegates& OutDocumentDelegates) const
+{
+	OutDocumentDelegates.EdgeDelegates.OnEdgeAdded.AddUObject(this, &UMetaSoundSourceBuilder::OnEdgeAdded);
+	OutDocumentDelegates.EdgeDelegates.OnRemovingEdge.AddUObject(this, &UMetaSoundSourceBuilder::OnRemovingEdge);
+
+	OutDocumentDelegates.InterfaceDelegates.OnInputAdded.AddUObject(this, &UMetaSoundSourceBuilder::OnInputAdded);
+	OutDocumentDelegates.InterfaceDelegates.OnOutputAdded.AddUObject(this, &UMetaSoundSourceBuilder::OnOutputAdded);
+	OutDocumentDelegates.InterfaceDelegates.OnRemovingInput.AddUObject(this, &UMetaSoundSourceBuilder::OnRemovingInput);
+	OutDocumentDelegates.InterfaceDelegates.OnRemovingOutput.AddUObject(this, &UMetaSoundSourceBuilder::OnRemovingOutput);
+
+	OutDocumentDelegates.NodeDelegates.OnNodeAdded.AddUObject(this, &UMetaSoundSourceBuilder::OnNodeAdded);
+	OutDocumentDelegates.NodeDelegates.OnNodeInputLiteralSet.AddUObject(this, &UMetaSoundSourceBuilder::OnNodeInputLiteralSet);
+	OutDocumentDelegates.NodeDelegates.OnRemovingNode.AddUObject(this, &UMetaSoundSourceBuilder::OnRemovingNode);
+	OutDocumentDelegates.NodeDelegates.OnRemovingNodeInputLiteral.AddUObject(this, &UMetaSoundSourceBuilder::OnRemovingNodeInputLiteral);
+}
+
+void UMetaSoundSourceBuilder::OnEdgeAdded(int32 EdgeIndex) const
+{
+	using namespace Metasound::DynamicGraph;
+
+	const FMetasoundFrontendDocument& Doc = Builder.GetDocument();
+	const FMetasoundFrontendEdge& NewEdge = Doc.RootGraph.Graph.Edges[EdgeIndex];
+	ExecuteAuditionableTransaction([this, &NewEdge](Metasound::DynamicGraph::FDynamicOperatorTransactor& Transactor)
+	{
+		const FMetaSoundFrontendDocumentBuilder& Builder = GetConstBuilder();
+		const FMetasoundFrontendVertex* FromNodeOutput = Builder.FindNodeOutput(NewEdge.FromNodeID, NewEdge.FromVertexID);
+		const FMetasoundFrontendVertex* ToNodeInput = Builder.FindNodeInput(NewEdge.ToNodeID, NewEdge.ToVertexID);
+		if (FromNodeOutput && ToNodeInput)
+		{
+			Transactor.AddDataEdge(NewEdge.FromNodeID, FromNodeOutput->Name, NewEdge.ToNodeID, ToNodeInput->Name);
+			return true;
+		}
+
+		return false;
+	});
+}
+
+TOptional<Metasound::FAnyDataReference> UMetaSoundSourceBuilder::CreateDataReference(const Metasound::FOperatorSettings& InOperatorSettings, const Metasound::FLiteral& InLiteral, Metasound::EDataReferenceAccessType AccessType)
+{
+	using namespace Metasound;
+	using namespace Metasound::Frontend;
+
+	// TODO: move to TFunctionRef
+	const FName DataType = GetMetasoundDataTypeName<float>();
+	return IDataTypeRegistry::Get().CreateDataReference(DataType, AccessType, InLiteral, InOperatorSettings);
+};
+
+void UMetaSoundSourceBuilder::OnInputAdded(int32 InputIndex) const
+{
+	using namespace Metasound::DynamicGraph;
+
+	ExecuteAuditionableTransaction([this, InputIndex](FDynamicOperatorTransactor& Transactor)
+	{
+		using namespace Metasound;
+		using namespace Metasound::Frontend;
+
+		const FMetasoundFrontendDocument& Doc = Builder.GetDocument();
+		const FMetasoundFrontendGraphClass& GraphClass = Doc.RootGraph;
+		const FMetasoundFrontendClassInput& NewInput = GraphClass.Interface.Inputs[InputIndex];
+		const FLiteral NewInputLiteral = NewInput.DefaultLiteral.ToLiteral(NewInput.TypeName);
+		Transactor.AddInputDataDestination(NewInput.NodeID, NewInput.Name, NewInputLiteral, &UMetaSoundSourceBuilder::CreateDataReference);
+		return true;
+	});
+}
+
+void UMetaSoundSourceBuilder::OnNodeAdded(int32 NodeIndex) const
+{
+	using namespace Metasound::DynamicGraph;
+
+	ExecuteAuditionableTransaction([this, NodeIndex](FDynamicOperatorTransactor& Transactor)
+	{
+		using namespace Metasound;
+		using namespace Metasound::Frontend;
+
+		const FMetasoundFrontendDocument& Doc = Builder.GetDocument();
+		const FMetasoundFrontendGraphClass& GraphClass = Doc.RootGraph;
+		const FMetasoundFrontendNode& AddedNode = GraphClass.Graph.Nodes[NodeIndex];
+		const Metasound::FNodeInitData InitData { AddedNode.Name, AddedNode.GetID() };
+
+		const FMetasoundFrontendClass* NodeClass = Builder.FindDependency(AddedNode.ClassID);
+		checkf(NodeClass, TEXT("Node successfully added to graph but document is missing associated dependency"));
+
+		auto MatchesVertexNodeID = [&InitData](const FMetasoundFrontendClassVertex& Vertex) { return Vertex.NodeID == InitData.InstanceID; };
+		switch (NodeClass->Metadata.GetType())
+		{
+			case EMetasoundFrontendClassType::Input:
+			case EMetasoundFrontendClassType::Output:
+			case EMetasoundFrontendClassType::External:
+			{
+			const FNodeRegistryKey& ClassKey = NodeRegistryKey::CreateKey(NodeClass->Metadata);
+				FMetasoundFrontendRegistryContainer& Registry = *FMetasoundFrontendRegistryContainer::Get();
+				TUniquePtr<INode> NewNode = Registry.CreateNode(ClassKey, InitData);
+				Transactor.AddNode(InitData.InstanceID, MoveTemp(NewNode));
+				return true;
+			}
+			break;
+
+			case EMetasoundFrontendClassType::Graph:
+			case EMetasoundFrontendClassType::Invalid:
+			case EMetasoundFrontendClassType::Literal:
+			case EMetasoundFrontendClassType::Template:
+			case EMetasoundFrontendClassType::Variable:
+			case EMetasoundFrontendClassType::VariableAccessor:
+			case EMetasoundFrontendClassType::VariableDeferredAccessor:
+			case EMetasoundFrontendClassType::VariableMutator:
+			default:
+			{
+				static_assert(static_cast<int32>(EMetasoundFrontendClassType::Invalid) == 10, "Possible missing EMetasoundFrontendClassType case coverage");
+				UE_LOG(LogMetaSound, Error, TEXT("Builder '%s' Audition Error: Addition of node was not propagated to active runtime graph: Runtime manipualtion of node ClassType unsupported."), *GetName());
+				return false;
+			}
+		}
+	});
+}
+
+void UMetaSoundSourceBuilder::OnNodeInputLiteralSet(int32 NodeIndex, int32 VertexIndex, int32 LiteralIndex) const
+{
+	using namespace Metasound::DynamicGraph;
+
+	ExecuteAuditionableTransaction([this, &NodeIndex, &VertexIndex, &LiteralIndex](FDynamicOperatorTransactor& Transactor)
+	{
+		using namespace Metasound;
+		using namespace Metasound::Engine;
+		using namespace Metasound::Frontend;
+
+		const FMetasoundFrontendGraphClass& GraphClass = Builder.GetDocument().RootGraph;
+		const FMetasoundFrontendNode& Node = GraphClass.Graph.Nodes[NodeIndex];
+		const FMetasoundFrontendVertex& Input = Node.Interface.Inputs[VertexIndex];
+		const FMetasoundFrontendLiteral& InputDefault = Node.InputLiterals[LiteralIndex].Value;
+		TUniquePtr<INode> LiteralNode = BuilderSubsystemPrivate::CreateDynamicNodeFromFrontendLiteral(Input.TypeName, InputDefault);
+		Transactor.SetValue(Node.GetID(), Input.Name, MoveTemp(LiteralNode));
+
+		return true;
+	});
+}
+
+void UMetaSoundSourceBuilder::OnOutputAdded(int32 OutputIndex) const
+{
+	using namespace Metasound::DynamicGraph;
+
+	ExecuteAuditionableTransaction([this, OutputIndex](FDynamicOperatorTransactor& Transactor)
+	{
+		using namespace Metasound::Frontend;
+
+		const FMetasoundFrontendDocument& Doc = Builder.GetDocument();
+		const FMetasoundFrontendGraphClass& GraphClass = Doc.RootGraph;
+		const FMetasoundFrontendClassOutput& NewOutput = GraphClass.Interface.Outputs[OutputIndex];
+
+		Transactor.AddOutputDataSource(NewOutput.NodeID, NewOutput.Name);
+		return true;
+	});
+}
+
+void UMetaSoundSourceBuilder::OnRemovingEdge(int32 EdgeIndex) const
+{
+	using namespace Metasound::DynamicGraph;
+
+	const FMetasoundFrontendDocument& Doc = Builder.GetDocument();
+	const FMetasoundFrontendEdge& EdgeBeingRemoved = Doc.RootGraph.Graph.Edges[EdgeIndex];
+	ExecuteAuditionableTransaction([this, EdgeBeingRemoved](FDynamicOperatorTransactor& Transactor)
+	{
+		using namespace Metasound;
+		using namespace Metasound::Engine;
+
+		const FMetaSoundFrontendDocumentBuilder& Builder = GetConstBuilder();
+		const FMetasoundFrontendVertex* FromNodeOutput = Builder.FindNodeOutput(EdgeBeingRemoved.FromNodeID, EdgeBeingRemoved.FromVertexID);
+		const FMetasoundFrontendVertex* ToNodeInput = Builder.FindNodeInput(EdgeBeingRemoved.ToNodeID, EdgeBeingRemoved.ToVertexID);
+		if (FromNodeOutput && ToNodeInput)
+		{
+			const FMetasoundFrontendLiteral* InputDefault = Builder.GetNodeInputDefault(EdgeBeingRemoved.ToNodeID, EdgeBeingRemoved.ToVertexID);
+			if (!InputDefault)
+			{
+				InputDefault = Builder.GetNodeInputClassDefault(EdgeBeingRemoved.ToNodeID, EdgeBeingRemoved.ToVertexID);
+			}
+
+			if (ensureAlwaysMsgf(InputDefault, TEXT("Could not dynamically assign default literal upon removing edge: literal should be assigned by either the frontend document's input or the class definition")))
+			{
+				TUniquePtr<INode> LiteralNode = BuilderSubsystemPrivate::CreateDynamicNodeFromFrontendLiteral(ToNodeInput->TypeName, *InputDefault);
+				Transactor.RemoveDataEdge(EdgeBeingRemoved.FromNodeID, FromNodeOutput->Name, EdgeBeingRemoved.ToNodeID, ToNodeInput->Name, MoveTemp(LiteralNode));
+				return true;
+			}
+		}
+
+		return false;
+	});
+}
+
+void UMetaSoundSourceBuilder::OnRemovingInput(int32 InputIndex) const
+{
+	using namespace Metasound::DynamicGraph;
+
+	ExecuteAuditionableTransaction([this, InputIndex](FDynamicOperatorTransactor& Transactor)
+	{
+		using namespace Metasound::Frontend;
+
+		const FMetasoundFrontendDocument& Doc = Builder.GetDocument();
+		const FMetasoundFrontendGraphClass& GraphClass = Doc.RootGraph;
+		const FMetasoundFrontendClassInput& InputBeingRemoved = GraphClass.Interface.Inputs[InputIndex];
+
+		Transactor.RemoveInputDataDestination(InputBeingRemoved.Name);
+		return true;
+	});
+}
+
+void UMetaSoundSourceBuilder::OnRemovingNode(int32 NodeIndex) const
+{
+	using namespace Metasound::DynamicGraph;
+	ExecuteAuditionableTransaction([this, NodeIndex](FDynamicOperatorTransactor& Transactor)
+	{
+		using namespace Metasound::Frontend;
+
+		const FMetasoundFrontendDocument& Doc = Builder.GetDocument();
+		const FMetasoundFrontendGraphClass& GraphClass = Doc.RootGraph;
+		const FMetasoundFrontendNode& NodeBeingRemoved = GraphClass.Graph.Nodes[NodeIndex];
+		const FGuid& NodeID = NodeBeingRemoved.GetID();
+		Transactor.RemoveNode(NodeID);
+		return true;
+	});
+}
+
+void UMetaSoundSourceBuilder::OnRemovingNodeInputLiteral(int32 NodeIndex, int32 VertexIndex, int32 LiteralIndex) const
+{
+	using namespace Metasound::DynamicGraph;
+
+	ExecuteAuditionableTransaction([this, &NodeIndex, &VertexIndex, &LiteralIndex](FDynamicOperatorTransactor& Transactor)
+	{
+		using namespace Metasound;
+		using namespace Metasound::Engine;
+		using namespace Metasound::Frontend;
+
+		const TArray<FMetasoundFrontendNode>& Nodes = Builder.GetDocument().RootGraph.Graph.Nodes;
+		const FMetasoundFrontendNode& Node = Nodes[NodeIndex];
+		const FMetasoundFrontendVertex& Input = Node.Interface.Inputs[VertexIndex];
+
+		const FMetasoundFrontendLiteral* InputDefault = Builder.GetNodeInputClassDefault(Node.GetID(), Input.VertexID);
+		if (ensureAlwaysMsgf(InputDefault, TEXT("Could not dynamically assign default literal from class definition upon removing input literal: document's dependency entry invalid and has no default assigned")))
+		{
+			TUniquePtr<INode> LiteralNode = BuilderSubsystemPrivate::CreateDynamicNodeFromFrontendLiteral(Input.TypeName, *InputDefault);
+			Transactor.SetValue(Node.GetID(), Input.Name, MoveTemp(LiteralNode));
+			return true;
+		}
+
+		return false;
+	});
+}
+
+void UMetaSoundSourceBuilder::OnRemovingOutput(int32 OutputIndex) const
+{
+	using namespace Metasound::DynamicGraph;
+
+	ExecuteAuditionableTransaction([this, OutputIndex](FDynamicOperatorTransactor& Transactor)
+	{
+		using namespace Metasound::Frontend;
+
+		const FMetasoundFrontendDocument& Doc = Builder.GetDocument();
+		const FMetasoundFrontendGraphClass& GraphClass = Doc.RootGraph;
+		const FMetasoundFrontendClassOutput& OutputBeingRemoved = GraphClass.Interface.Outputs[OutputIndex];
+
+		Transactor.RemoveOutputDataSource(OutputBeingRemoved.Name);
+		return true;
+	});
+}
+
 void UMetaSoundSourceBuilder::SetFormat(EMetaSoundOutputAudioFormat OutputFormat, EMetaSoundBuilderResult& OutResult)
 {
 	using namespace Metasound::Engine;
@@ -998,53 +1432,6 @@ const UMetaSoundBuilderSubsystem& UMetaSoundBuilderSubsystem::GetConstChecked()
 	return *BuilderSubsystem;
 }
 
-UMetaSoundBuilderBase& UMetaSoundBuilderSubsystem::AttachBuilderToObjectChecked(TScriptInterface<IMetaSoundDocumentInterface> InInterface) const
-{
-	checkf(InInterface.GetInterface(), TEXT("Failed to attach MetaSound builder to asset: Interface not supplied"))
-	const UClass* BaseClass = &InInterface->GetBaseMetaSoundUClass();
-	if (BaseClass == UMetaSoundSource::StaticClass())
-	{
-		UMetaSoundSourceBuilder* NewBuilder = AttachSourceBuilderToAsset(CastChecked<UMetaSoundSource>(InInterface.GetObject()));
-		return *NewBuilder;
-	}
-	else if (BaseClass == UMetaSoundPatch::StaticClass())
-	{
-		UMetaSoundPatchBuilder* NewBuilder = AttachPatchBuilderToAsset(CastChecked<UMetaSoundPatch>(InInterface.GetObject()));
-		return *NewBuilder;
-	}
-	else
-	{
-		checkf(false, TEXT("UClass '%s' is not a base MetaSound that supports attachment via the MetaSoundBuilderSubsystem"), *BaseClass->GetFullName());
-		return Metasound::Engine::BuilderSubsystemPrivate::CreateTransientBuilder<UMetaSoundPatchBuilder>();
-	}
-}
-
-UMetaSoundPatchBuilder* UMetaSoundBuilderSubsystem::AttachPatchBuilderToAsset(UMetaSoundPatch* InPatch) const
-{
-	if (InPatch)
-	{
-		TObjectPtr<UMetaSoundPatchBuilder> NewBuilder = NewObject<UMetaSoundPatchBuilder>(InPatch);
-		TScriptInterface<IMetaSoundDocumentInterface> MetaSoundPatchDocInterface = InPatch;
-		NewBuilder->Builder = FMetaSoundFrontendDocumentBuilder(MetaSoundPatchDocInterface);
-		return NewBuilder;
-	}
-
-	return nullptr;
-}
-
-UMetaSoundSourceBuilder* UMetaSoundBuilderSubsystem::AttachSourceBuilderToAsset(UMetaSoundSource* InSource) const
-{
-	if (InSource)
-	{
-		TObjectPtr<UMetaSoundSourceBuilder> NewBuilder = NewObject<UMetaSoundSourceBuilder>(InSource);
-		TScriptInterface<IMetaSoundDocumentInterface> MetaSoundSourceDocInterface = InSource;
-		NewBuilder->Builder = FMetaSoundFrontendDocumentBuilder(MetaSoundSourceDocInterface);
-		return NewBuilder;
-	}
-
-	return nullptr;
-}
-
 FMetasoundFrontendLiteral UMetaSoundBuilderSubsystem::CreateBoolMetaSoundLiteral(bool Value, FName& OutDataType)
 {
 	return Metasound::Engine::BuilderSubsystemPrivate::CreatePODMetaSoundLiteral(Value, OutDataType);
@@ -1104,14 +1491,45 @@ FMetasoundFrontendLiteral UMetaSoundBuilderSubsystem::CreateMetaSoundLiteralFrom
 	return FMetasoundFrontendLiteral { Param };
 }
 
+bool UMetaSoundBuilderSubsystem::DetachBuilderFromAsset(const FMetasoundFrontendClassName& InClassName)
+{
+	return AssetBuilders.Remove(InClassName.GetFullName()) > 0;
+}
+
+const Metasound::Frontend::FDocumentModifyDelegates* UMetaSoundBuilderSubsystem::FindModifyDelegates(const FMetasoundFrontendClassName& InClassName) const
+{
+	using namespace Metasound::Frontend;
+
+	TWeakObjectPtr<UMetaSoundBuilderBase> BuilderPtr = AssetBuilders.FindRef(InClassName.GetFullName());
+	if (BuilderPtr.IsValid())
+	{
+		return &BuilderPtr->GetConstBuilder().GetDocumentDelegates();
+	}
+
+	return nullptr;
+}
+
 UMetaSoundPatchBuilder* UMetaSoundBuilderSubsystem::FindPatchBuilder(FName BuilderName)
 {
-	return PatchBuilders.FindRef(BuilderName);
+	return NamedPatchBuilders.FindRef(BuilderName);
 }
 
 UMetaSoundSourceBuilder* UMetaSoundBuilderSubsystem::FindSourceBuilder(FName BuilderName)
 {
-	return SourceBuilders.FindRef(BuilderName);
+	return NamedSourceBuilders.FindRef(BuilderName);
+}
+
+void UMetaSoundBuilderSubsystem::Initialize(FSubsystemCollectionBase& Collection)
+{
+	using namespace Metasound::Frontend;
+
+	IMetaSoundDocumentBuilderRegistry::Set([]() -> IMetaSoundDocumentBuilderRegistry&
+	{
+		check(GEngine);
+		UMetaSoundBuilderSubsystem* BuilderSubsystem = GEngine->GetEngineSubsystem<UMetaSoundBuilderSubsystem>();
+		check(BuilderSubsystem);
+		return static_cast<IMetaSoundDocumentBuilderRegistry&>(*BuilderSubsystem);
+	});
 }
 
 bool UMetaSoundBuilderSubsystem::IsInterfaceRegistered(FName InInterfaceName) const
@@ -1124,20 +1542,20 @@ bool UMetaSoundBuilderSubsystem::IsInterfaceRegistered(FName InInterfaceName) co
 
 void UMetaSoundBuilderSubsystem::RegisterPatchBuilder(FName BuilderName, UMetaSoundPatchBuilder* Builder)
 {
-	PatchBuilders.FindOrAdd(BuilderName) = Builder;
+	NamedPatchBuilders.FindOrAdd(BuilderName) = Builder;
 }
 
 void UMetaSoundBuilderSubsystem::RegisterSourceBuilder(FName BuilderName, UMetaSoundSourceBuilder* Builder)
 {
-	SourceBuilders.FindOrAdd(BuilderName) = Builder;
+	NamedSourceBuilders.FindOrAdd(BuilderName) = Builder;
 }
 
 bool UMetaSoundBuilderSubsystem::UnregisterPatchBuilder(FName BuilderName)
 {
-	return PatchBuilders.Remove(BuilderName) > 0;
+	return NamedPatchBuilders.Remove(BuilderName) > 0;
 }
 
 bool UMetaSoundBuilderSubsystem::UnregisterSourceBuilder(FName BuilderName)
 {
-	return SourceBuilders.Remove(BuilderName) > 0;
+	return NamedSourceBuilders.Remove(BuilderName) > 0;
 }
diff --git a/Engine/Plugins/Runtime/Metasound/Source/MetasoundEngine/Public/MetasoundBuilderSubsystem.h b/Engine/Plugins/Runtime/Metasound/Source/MetasoundEngine/Public/MetasoundBuilderSubsystem.h
--- a/Engine/Plugins/Runtime/Metasound/Source/MetasoundEngine/Public/MetasoundBuilderSubsystem.h
+++ b/Engine/Plugins/Runtime/Metasound/Source/MetasoundEngine/Public/MetasoundBuilderSubsystem.h
@@ -2,12 +2,18 @@
 #pragma once
 
 #include "Containers/Map.h"
+#include "Delegates/IDelegateInstance.h"
 #include "Engine/Engine.h"
 #include "Interfaces/MetasoundOutputFormatInterfaces.h"
+#include "MetasoundDocumentInterface.h"
+#include "MetasoundDynamicOperatorTransactor.h"
 #include "MetasoundFrontendController.h"
 #include "MetasoundFrontendDocument.h"
 #include "MetasoundFrontendDocumentBuilder.h"
+#include "MetasoundFrontendDocumentCacheInterface.h"
+#include "MetasoundFrontendDocumentModifyDelegates.h"
 #include "Subsystems/EngineSubsystem.h"
+#include "Templates/Function.h"
 #include "UObject/NoExportTypes.h"
 #include "UObject/Object.h"
 #include "UObject/ObjectMacros.h"
@@ -137,6 +143,11 @@ public:
 	UFUNCTION(BlueprintCallable, Category = "Audio|MetaSound|Builder", DisplayName = "Add MetaSound Node From Asset Class", meta = (ExpandEnumAsExecs = "OutResult"))
 	UPARAM(DisplayName = "Node Handle") FMetaSoundNodeHandle AddNode(const TScriptInterface<IMetaSoundDocumentInterface>& NodeClass, EMetaSoundBuilderResult& OutResult);
 
+	// Adds node referencing the highest native class version of the given class name to the document.
+	// Returns a node handle to the created node if successful, or an invalid handle if it failed.
+	UFUNCTION(BlueprintCallable, Category = "Audio|MetaSound|Builder", DisplayName = "Add MetaSound Node By ClassName", meta = (ExpandEnumAsExecs = "OutResult"))
+	UPARAM(DisplayName = "Node Handle") FMetaSoundNodeHandle AddNodeByClassName(const FMetasoundFrontendClassName& ClassName, int32 MajorVersion, EMetaSoundBuilderResult& OutResult);
+
 	// Connects node output to a node input. Does *NOT* provide loop detection for performance reasons.  Loop detection is checked on class registration when built or played.
 	// Returns succeeded if connection made, failed if connection already exists with input, the data types do not match, or the connection is not supported due to access type
 	// incompatibility (ex. constructor input to non-constructor input).
@@ -248,6 +259,14 @@ public:
 	UFUNCTION(BlueprintCallable, Category = "Audio|MetaSound|Builder", meta = (ExpandEnumAsExecs = "OutResult"))
 	void GetNodeInputData(const FMetaSoundBuilderNodeInputHandle& InputHandle, FName& Name, FName& DataType, EMetaSoundBuilderResult& OutResult);
 
+	// Returns node input's literal value if set on graph, otherwise fails and returns default literal.
+	UFUNCTION(BlueprintCallable, Category = "Audio|MetaSound|Builder", meta = (ExpandEnumAsExecs = "OutResult"))
+	FMetasoundFrontendLiteral GetNodeInputDefault(const FMetaSoundBuilderNodeInputHandle& InputHandle, EMetaSoundBuilderResult& OutResult);
+
+	// Returns node input's class literal value if set, otherwise fails and returns default literal.
+	UFUNCTION(BlueprintCallable, Category = "Audio|MetaSound|Builder", meta = (ExpandEnumAsExecs = "OutResult"))
+	FMetasoundFrontendLiteral GetNodeInputClassDefault(const FMetaSoundBuilderNodeInputHandle& InputHandle, EMetaSoundBuilderResult& OutResult);
+
 	// Returns node output's data if valid (including things like name and datatype).
 	UFUNCTION(BlueprintCallable, Category = "Audio|MetaSound|Builder", meta = (ExpandEnumAsExecs = "OutResult"))
 	void GetNodeOutputData(const FMetaSoundBuilderNodeOutputHandle& OutputHandle, FName& Name, FName& DataType, EMetaSoundBuilderResult& OutResult);
@@ -301,6 +320,11 @@ public:
 	UFUNCTION(BlueprintCallable, Category = "Audio|MetaSound|Builder", meta = (ExpandEnumAsExecs = "OutResult"))
 	void RemoveNode(const FMetaSoundNodeHandle& NodeHandle, EMetaSoundBuilderResult& OutResult);
 
+	// Removes node input literal default if set, reverting the value to be whatever the node class defaults the value to.
+	// Returns success if value was removed, false if not removed (i.e. wasn't set to begin with).
+	UFUNCTION(BlueprintCallable, Category = "Audio|MetaSound|Builder", meta = (ExpandEnumAsExecs = "OutResult"))
+	void RemoveNodeInputDefault(const FMetaSoundBuilderNodeInputHandle& InputHandle, EMetaSoundBuilderResult& OutResult);
+
 #if WITH_EDITOR
 	// Sets the author of the MetaSound.
 	void SetAuthor(const FString& InAuthor);
@@ -310,21 +334,20 @@ public:
 	UFUNCTION(BlueprintCallable, Category = "Audio|MetaSound|Builder", meta = (ExpandEnumAsExecs = "OutResult"))
 	void SetNodeInputDefault(const FMetaSoundBuilderNodeInputHandle& NodeInputHandle, const FMetasoundFrontendLiteral& Literal, EMetaSoundBuilderResult& OutResult);
 
-	// Sets the input node's default value, overriding the default if the graph is a preset. 
+	// Sets the input node's default value, overriding the default provided by the referenced graph if the graph is a preset.
 	UFUNCTION(BlueprintCallable, Category = "Audio|MetaSound|Builder", meta = (ExpandEnumAsExecs = "OutResult"))
 	void SetGraphInputDefault(FName InputName, const FMetasoundFrontendLiteral& Literal, EMetaSoundBuilderResult& OutResult);
 
 	virtual TScriptInterface<IMetaSoundDocumentInterface> Build(UObject* Parent, const FMetaSoundBuilderOptions& Options) const PURE_VIRTUAL(UMetaSoundBuilderBase::Build, return { }; );
 
-	// Constructs a transient UMetaSoundBuilderDocument & attaches builder to it
-	void CreateTransientDocument();
-
 	// Returns the base MetaSound UClass the builder is operating on (ex. MetaSoundSource, MetaSoundPatch, etc. not to be confused with the underlying document builder's MetaSound class type)
 	virtual const UClass& GetBuilderUClass() const PURE_VIRTUAL(UMetaSoundBuilderBase::Build, return *UClass::StaticClass(); );
 
 	// Initializes and ensures all nodes have a position (required prior to exporting to an asset if expected to be viewed in the editor).
 	void InitNodeLocations();
 
+	virtual void InitFrontendBuilder();
+
 protected:
 	const FMetaSoundFrontendDocumentBuilder& GetConstBuilder() const
 	{
@@ -333,7 +356,7 @@ protected:
 
 	// Runs build, conforming the document and corresponding object data on a MetaSound UObject to that managed by this builder.
 	template <typename UClassType>
-	UClassType* BuildInternal(UObject* Parent, const FMetaSoundBuilderOptions& BuilderOptions) const
+	UClassType& BuildInternal(UObject* Parent, const FMetaSoundBuilderOptions& BuilderOptions) const
 	{
 		using namespace Metasound::Frontend;
 
@@ -366,37 +389,49 @@ protected:
 				ObjectName = MakeUniqueObjectName(Parent, UClassType::StaticClass(), BuilderOptions.Name);
 			}
 
+			if (!Parent)
+			{
+				Parent = GetTransientPackage();
+			}
+
 			MetaSound = NewObject<UClassType>(Parent, ObjectName, RF_Public | RF_Transient);
 		}
 
-		if (MetaSound)
+		checkf(MetaSound, TEXT("Failed to build MetaSound from builder '%s'"), *GetPathName());
+		FMetasoundFrontendDocument NewDocument = GetConstBuilder().GetDocument();
 		{
-			FMetasoundFrontendDocument NewDocument = GetConstBuilder().GetDocument();
-			{
-				// This is required to ensure the newly build document has a unique class
-				// identifier to avoid collisions if added to the Frontend class registry
-				// (either below or at a later point in time).
-				constexpr bool bResetVersion = false;
-				FMetaSoundFrontendDocumentBuilder::InitGraphClassMetadata(NewDocument.RootGraph.Metadata, bResetVersion);
-			}
+			// This is required to ensure the newly build document has a unique class
+			// identifier to avoid collisions if added to the Frontend class registry
+			// (either below or at a later point in time).
+			constexpr bool bResetVersion = false;
+			FMetaSoundFrontendDocumentBuilder::InitGraphClassMetadata(NewDocument.RootGraph.Metadata, bResetVersion);
+		}
 
-			MetaSound->SetDocument(MoveTemp(NewDocument));
-			MetaSound->ConformObjectDataToInterfaces();
+		MetaSound->SetDocument(MoveTemp(NewDocument));
+		MetaSound->ConformObjectDataToInterfaces();
 
-			if (BuilderOptions.bAddToRegistry)
-			{
-				MetaSound->RegisterGraphWithFrontend();
-			}
-
-			UE_LOG(LogMetaSound, VeryVerbose, TEXT("MetaSound '%s' built from '%s'"), *BuilderOptions.Name.ToString(), *GetFullName());
+		if (BuilderOptions.bAddToRegistry)
+		{
+			MetaSound->RegisterGraphWithFrontend();
 		}
 
-		return MetaSound;
+		UE_LOG(LogMetaSound, VeryVerbose, TEXT("MetaSound '%s' built from '%s'"), *BuilderOptions.Name.ToString(), *GetFullName());
+		return *MetaSound;
 	}
 
+	// Constructs a transient UMetaSoundBuilderDocument to be acted upon by this builder instance.
+	UMetaSoundBuilderDocument* CreateTransientDocumentObject() const;
+
+
 	UPROPERTY()
 	FMetaSoundFrontendDocumentBuilder Builder;
 
+	// If true, builder is attached to an existing asset (directly making changes to its document). If false, builder is
+	// operating on a transient document which must be built to an asset prior to use by the MetaSound Frontend where it
+	// is then available for referencing, playback, auditioning, etc. by core execution.
+	UPROPERTY()
+	bool bIsAttached = false;
+
 	// Friending allows for swapping the builder in certain circumstances where desired (eg. attaching a builder to an existing asset)
 	friend class UMetaSoundBuilderSubsystem;
 };
@@ -422,11 +457,15 @@ class METASOUNDENGINE_API UMetaSoundSourceBuilder : public UMetaSoundBuilderBase
 
 public:
 	UFUNCTION(BlueprintCallable, Category = "Audio|MetaSound|Builder", meta = (WorldContext = "Parent", AdvancedDisplay = "2"))
-	void Audition(UObject* Parent, UAudioComponent* AudioComponent, FOnCreateAuditionGeneratorHandleDelegate OnCreateGenerator);
+	void Audition(UObject* Parent, UAudioComponent* AudioComponent, FOnCreateAuditionGeneratorHandleDelegate OnCreateGenerator, bool bLiveUpdatesEnabled = false);
 
 	UFUNCTION(BlueprintCallable, Category = "Audio|MetaSound|Builder", meta = (WorldContext = "Parent"))
 	virtual UPARAM(DisplayName = "MetaSound") TScriptInterface<IMetaSoundDocumentInterface> Build(UObject* Parent, const FMetaSoundBuilderOptions& Options) const override;
 
+	// Returns whether or not live updates are both globally enabled (via cvar) and are enabled on this builder's last built sound, which may or may not still be playing.
+	UFUNCTION(BlueprintCallable, Category = "Audio|MetaSound|Builder")
+	bool GetLiveUpdatesEnabled() const;
+
 	// Sets the output audio format of the source
 	UFUNCTION(BlueprintCallable, Category = "Audio|MetaSound|Builder", meta = (ExpandEnumAsExecs = "OutResult"))
 	void SetFormat(EMetaSoundOutputAudioFormat OutputFormat, EMetaSoundBuilderResult& OutResult);
@@ -434,31 +473,64 @@ public:
 	const Metasound::Engine::FOutputAudioFormatInfoPair* FindOutputAudioFormatInfo() const;
 
 	virtual const UClass& GetBuilderUClass() const override;
+	virtual void InitFrontendBuilder() override;
 
 private:
+	static TOptional<Metasound::FAnyDataReference> CreateDataReference(const Metasound::FOperatorSettings& InOperatorSettings, const Metasound::FLiteral& InLiteral, Metasound::EDataReferenceAccessType AccessType);
+
+	void InitDelegates(Metasound::Frontend::FDocumentModifyDelegates& OutDocumentDelegates) const;
+	void OnEdgeAdded(int32 EdgeIndex) const;
+	void OnInputAdded(int32 InputIndex) const;
+	void OnNodeAdded(int32 NodeIndex) const;
+	void OnNodeInputLiteralSet(int32 NodeIndex, int32 VertexIndex, int32 LiteralIndex) const;
+	void OnOutputAdded(int32 OutputIndex) const;
+	void OnRemovingEdge(int32 EdgeIndex) const;
+	void OnRemovingInput(int32 InputIndex) const;
+	void OnRemovingNode(int32 NodeIndex) const;
+	void OnRemovingNodeInputLiteral(int32 NodeIndex, int32 VertexIndex, int32 LiteralIndex) const;
+	void OnRemovingOutput(int32 OutputIndex) const;
+
 	TWeakObjectPtr<UMetaSoundSource> AuditionSound;
+
+	using FAuditionableTransaction = TFunctionRef<bool(Metasound::DynamicGraph::FDynamicOperatorTransactor&)>;
+	bool ExecuteAuditionableTransaction(FAuditionableTransaction Transaction) const;
+
+	friend class UMetaSoundBuilderSubsystem;
 };
 
 /** The subsystem in charge of tracking MetaSound builders */
 UCLASS()
-class METASOUNDENGINE_API UMetaSoundBuilderSubsystem : public UEngineSubsystem
+class METASOUNDENGINE_API UMetaSoundBuilderSubsystem : public UEngineSubsystem, public Metasound::Frontend::IMetaSoundDocumentBuilderRegistry
 {
 	GENERATED_BODY()
 
 	UPROPERTY()
-	TMap<FName, TObjectPtr<UMetaSoundPatchBuilder>> PatchBuilders;
+	TMap<FName, TObjectPtr<UMetaSoundPatchBuilder>> NamedPatchBuilders;
 
 	UPROPERTY()
-	TMap<FName, TObjectPtr<UMetaSoundSourceBuilder>> SourceBuilders;
+	TMap<FName, TObjectPtr<UMetaSoundSourceBuilder>> NamedSourceBuilders;
+
+	UPROPERTY()
+	mutable TMap<FName, TWeakObjectPtr<UMetaSoundBuilderBase>> AssetBuilders;
 
 public:
+	virtual void Initialize(FSubsystemCollectionBase& Collection) override;
+
+	virtual const Metasound::Frontend::FDocumentModifyDelegates* FindModifyDelegates(const FMetasoundFrontendClassName& InClassName) const override;
+
 	static UMetaSoundBuilderSubsystem& GetChecked();
 	static const UMetaSoundBuilderSubsystem& GetConstChecked();
 
-	UMetaSoundBuilderBase& AttachBuilderToObjectChecked(TScriptInterface<IMetaSoundDocumentInterface> InInterface) const;
+	UMetaSoundBuilderBase& AttachBuilderToAssetChecked(UObject& InMetaSound) const;
+
+	// May move these to UFUNCTIONS in the future, but right now too many unguarded cases
+	// where a user could mutate these and potentially improperly reference transient assets.
 	UMetaSoundPatchBuilder* AttachPatchBuilderToAsset(UMetaSoundPatch* InPatch) const;
 	UMetaSoundSourceBuilder* AttachSourceBuilderToAsset(UMetaSoundSource* InSource) const;
 
+	// Unregisters builder for MetaSound class asset.
+	bool DetachBuilderFromAsset(const FMetasoundFrontendClassName& InClassName);
+
 	UFUNCTION(BlueprintCallable, Category = "Audio|MetaSound|Builder",  meta = (ExpandEnumAsExecs = "OutResult"))
 	UPARAM(DisplayName = "Patch Builder") UMetaSoundPatchBuilder* CreatePatchBuilder(FName BuilderName, EMetaSoundBuilderResult& OutResult);
 
@@ -535,4 +607,27 @@ public:
 
 	UFUNCTION(BlueprintCallable, Category = "Audio|MetaSound|Builder")
 	UPARAM(DisplayName = "Unregistered") bool UnregisterSourceBuilder(FName BuilderName);
+
+private:
+	template <typename BuilderClass>
+	BuilderClass& AttachBuilderToAssetCheckedPrivate(UObject* InMetaSoundObject) const
+	{
+		check(InMetaSoundObject);
+		TScriptInterface<IMetaSoundDocumentInterface> DocInterface = InMetaSoundObject;
+		const FMetasoundFrontendDocument& Document = static_cast<const IMetaSoundDocumentInterface*>(DocInterface.GetInterface())->GetDocument();
+		const FName FullClassName = Document.RootGraph.Metadata.GetClassName().GetFullName();
+		TWeakObjectPtr<UMetaSoundBuilderBase> Builder = AssetBuilders.FindRef(FullClassName);
+		if (Builder.IsValid())
+		{
+			return *CastChecked<BuilderClass>(Builder.Get());
+		}
+
+		TObjectPtr<BuilderClass> NewBuilder = NewObject<BuilderClass>(InMetaSoundObject);
+		check(NewBuilder);
+		NewBuilder->Builder = FMetaSoundFrontendDocumentBuilder(DocInterface);
+		NewBuilder->bIsAttached = true;
+		TObjectPtr<UMetaSoundBuilderBase> NewBuilderBase = CastChecked<UMetaSoundBuilderBase>(NewBuilder);
+		AssetBuilders.Add(FullClassName, NewBuilderBase);
+		return *NewBuilder;
+	}
 };
diff --git a/Engine/Plugins/Runtime/Metasound/Source/MetasoundEngine/Public/MetasoundSource.h b/Engine/Plugins/Runtime/Metasound/Source/MetasoundEngine/Public/MetasoundSource.h
--- a/Engine/Plugins/Runtime/Metasound/Source/MetasoundEngine/Public/MetasoundSource.h
+++ b/Engine/Plugins/Runtime/Metasound/Source/MetasoundEngine/Public/MetasoundSource.h
@@ -52,6 +52,8 @@ class METASOUNDENGINE_API UMetaSoundSource : public USoundWaveProcedural, public
 	GENERATED_BODY()
 
 	friend struct Metasound::FMetaSoundEngineAssetHelper;
+	friend class UMetaSoundSourceBuilder;
+
 protected:
 	UPROPERTY(EditAnywhere, Category = CustomView)
 	FMetasoundFrontendDocument RootMetasoundDocument;
diff --git a/Engine/Plugins/Runtime/Metasound/Source/MetasoundEngineTest/Private/EngineTestMetasoundSource.cpp b/Engine/Plugins/Runtime/Metasound/Source/MetasoundEngineTest/Private/EngineTestMetasoundSource.cpp
--- a/Engine/Plugins/Runtime/Metasound/Source/MetasoundEngineTest/Private/EngineTestMetasoundSource.cpp
+++ b/Engine/Plugins/Runtime/Metasound/Source/MetasoundEngineTest/Private/EngineTestMetasoundSource.cpp
@@ -6,17 +6,26 @@
 #include "Interfaces/IPluginManager.h"
 #include "Interfaces/MetasoundOutputFormatInterfaces.h"
 #include "Interfaces/MetasoundFrontendSourceInterface.h"
+#include "MetasoundBuilderSubsystem.h"
 #include "MetasoundFrontendController.h"
 #include "MetasoundFrontendSearchEngine.h"
 #include "MetasoundSource.h"
 #include "Misc/AutomationTest.h"
 #include "Misc/Paths.h"
 #include "Tests/AutomationCommon.h"
+#include "../../MetasoundGraphCore/Public/MetasoundDataReference.h"
 
 #if WITH_DEV_AUTOMATION_TESTS
 
 namespace EngineTestMetasoundSourcePrivate
 {
+	struct FInitTestBuilderOutput
+	{
+		FMetaSoundBuilderNodeOutputHandle OnPlayOutput;
+		FMetaSoundBuilderNodeInputHandle OnFinishedInput;
+		TArray<FMetaSoundBuilderNodeInputHandle> AudioOutNodeInputs;
+	};
+
 	static FString GetPluginContentDirectory()
 	{
 		TSharedPtr<IPlugin> Plugin = IPluginManager::Get().FindPlugin(TEXT("Metasound"));
@@ -65,6 +74,91 @@ namespace EngineTestMetasoundSourcePrivate
 		return Node;
 	}
 
+	UMetaSoundSourceBuilder& CreateMetaSoundSourceBuilder(EMetaSoundOutputAudioFormat OutputFormat, bool bIsOneShot, FInitTestBuilderOutput& Output)
+	{
+		using namespace Audio;
+		using namespace Metasound;
+		using namespace Metasound::Engine;
+		using namespace Metasound::Frontend;
+
+		EMetaSoundBuilderResult Result;
+		UMetaSoundSourceBuilder* Builder = UMetaSoundBuilderSubsystem::GetChecked().CreateSourceBuilder(
+			"Unit Test Graph Builder",
+			Output.OnPlayOutput,
+			Output.OnFinishedInput,
+			Output.AudioOutNodeInputs,
+			Result,
+			OutputFormat,
+			bIsOneShot);
+		checkf(Result == EMetaSoundBuilderResult::Succeeded, TEXT("Failed to create MetaSoundSourceBuilder"));
+
+		return *Builder;
+	}
+
+	UMetaSoundSourceBuilder& CreateMetaSoundMonoSourceSinGenBuilder(
+		FAutomationTestBase& Test,
+		FMetaSoundBuilderNodeInputHandle* GenInputNodeFreq = nullptr,
+		FMetaSoundBuilderNodeInputHandle* MonoOutNodeInput = nullptr,
+		float InDefaultFreq = 100.f)
+	{
+		using namespace EngineTestMetasoundSourcePrivate;
+		using namespace Metasound;
+		using namespace Metasound::Engine;
+		using namespace Metasound::Frontend;
+
+		constexpr EMetaSoundOutputAudioFormat OutputFormat = EMetaSoundOutputAudioFormat::Mono;
+		constexpr bool bIsOneShot = false;
+		FInitTestBuilderOutput Output;
+		UMetaSoundSourceBuilder& Builder = CreateMetaSoundSourceBuilder(EMetaSoundOutputAudioFormat::Mono, bIsOneShot, Output);
+
+		EMetaSoundBuilderResult Result = EMetaSoundBuilderResult::Failed;
+		if (MonoOutNodeInput)
+		{
+			*MonoOutNodeInput = { };
+		}
+
+		// Input on Play
+		const FMetaSoundNodeHandle OnPlayOutputNode = Builder.FindGraphInputNode(SourceInterface::Inputs::OnPlay, Result);
+		Test.AddErrorIfFalse(Result == EMetaSoundBuilderResult::Succeeded && OnPlayOutputNode.IsSet(), TEXT("Failed to create MetaSound OnPlay input"));
+
+		// Input Frequency
+		FMetasoundFrontendLiteral DefaultFreq;
+		DefaultFreq.Set(InDefaultFreq);
+		const FMetaSoundBuilderNodeOutputHandle FrequencyNodeOutput = Builder.AddGraphInputNode("Frequency", GetMetasoundDataTypeName<float>(), DefaultFreq, Result);
+		Test.AddErrorIfFalse(Result == EMetaSoundBuilderResult::Succeeded && FrequencyNodeOutput.IsSet(), TEXT("Failed to create new MetaSound graph input"));
+
+		// Sine Oscillator Node
+		const FMetaSoundNodeHandle OscNode = Builder.AddNodeByClassName({ "UE", "Sine", "Audio" }, 1, Result);
+		Test.AddErrorIfFalse(Result == EMetaSoundBuilderResult::Succeeded && OscNode.IsSet(), TEXT("Failed to create new MetaSound node by class name"));
+
+		// Make connections:
+		const FMetaSoundBuilderNodeInputHandle OscNodeFrequencyInput = Builder.FindNodeInputByName(OscNode, "Frequency", Result);
+		if (GenInputNodeFreq)
+		{
+			*GenInputNodeFreq = OscNodeFrequencyInput;
+		}
+		Test.AddErrorIfFalse(Result == EMetaSoundBuilderResult::Succeeded && OscNodeFrequencyInput.IsSet(), TEXT("Failed to find Sine Oscillator node input 'Frequency'"));
+
+		const FMetaSoundBuilderNodeOutputHandle OscNodeAudioOutput = Builder.FindNodeOutputByName(OscNode, "Audio", Result);
+		Test.AddErrorIfFalse(Result == EMetaSoundBuilderResult::Succeeded && OscNodeAudioOutput.IsSet(), TEXT("Failed to find Sine Oscillator node output 'Audio'"));
+
+		// Frequency input "Frequency" -> oscillator "Frequency"
+		Builder.ConnectNodes(FrequencyNodeOutput, OscNodeFrequencyInput, Result);
+		Test.AddErrorIfFalse(Result == EMetaSoundBuilderResult::Succeeded, TEXT("Failed to connect 'Frequency' input to node input 'Frequency'"));
+
+		// Oscillator to Output Node
+		Test.AddErrorIfFalse(Output.AudioOutNodeInputs.Num() == 1, TEXT("Should only ever have one output node for mono"));
+		if (MonoOutNodeInput)
+		{
+			*MonoOutNodeInput = Output.AudioOutNodeInputs.Last();
+		}
+
+		Builder.ConnectNodes(OscNodeAudioOutput, Output.AudioOutNodeInputs.Last(), Result);
+		Test.AddErrorIfFalse(Result == EMetaSoundBuilderResult::Succeeded, TEXT("Failed to connect 'Audio' Sine Oscillator output to MetaSound graph's 'Mono Output'"));
+
+		return Builder;
+	}
+
 	FMetasoundFrontendDocument CreateMetaSoundMonoSourceDocument()
 	{
 		using namespace Audio;
@@ -129,7 +223,7 @@ namespace EngineTestMetasoundSourcePrivate
 
 		return Document;
 	}
-}
+} // EngineTestMetasoundSourcePrivate
 
 DEFINE_LATENT_AUTOMATION_COMMAND_ONE_PARAMETER(FAudioComponentPlayLatentCommand, UAudioComponent*, AudioComponent);
 
@@ -155,6 +249,81 @@ bool FAudioComponentStopLatentCommand::Update()
 	return false;
 }
 
+DEFINE_LATENT_AUTOMATION_COMMAND_THREE_PARAMETER(FMetaSoundSourceBuilderAuditionLatentCommand, UMetaSoundSourceBuilder*, Builder, UAudioComponent*, AudioComponent, bool, bEnableLiveUpdates);
+
+bool FMetaSoundSourceBuilderAuditionLatentCommand::Update()
+{
+	if (Builder && AudioComponent)
+	{
+		Builder->Audition(nullptr, AudioComponent, { }, true);
+		return true;
+	}
+	return false;
+}
+
+DEFINE_LATENT_AUTOMATION_COMMAND_THREE_PARAMETER(FMetaSoundSourceBuilderDisconnectInputLatentCommand, FAutomationTestBase&, Test, UMetaSoundBuilderBase*, Builder, FMetaSoundBuilderNodeInputHandle, InputToDisconnect);
+
+bool FMetaSoundSourceBuilderDisconnectInputLatentCommand::Update()
+{
+	EMetaSoundBuilderResult Result = EMetaSoundBuilderResult::Failed;
+	if (Builder)
+	{
+		Builder->DisconnectNodeInput(InputToDisconnect, Result);
+		Test.AddErrorIfFalse(Result == EMetaSoundBuilderResult::Succeeded, TEXT("Failed to disconnect MetaSound node's input"));
+	}
+
+	return Result == EMetaSoundBuilderResult::Succeeded;
+}
+
+DEFINE_LATENT_AUTOMATION_COMMAND_FOUR_PARAMETER(FMetaSoundSourceBuilderSetLiteralLatentCommand, FAutomationTestBase&, Test, UMetaSoundBuilderBase*, Builder, FMetaSoundBuilderNodeInputHandle, NodeInput, FMetasoundFrontendLiteral, NewValue);
+
+bool FMetaSoundSourceBuilderSetLiteralLatentCommand::Update()
+{
+	EMetaSoundBuilderResult Result = EMetaSoundBuilderResult::Failed;
+	if (Builder)
+	{
+		Builder->SetNodeInputDefault(NodeInput, NewValue, Result);
+		Test.AddErrorIfFalse(Result == EMetaSoundBuilderResult::Succeeded, TEXT("Failed to disconnect MetaSound node's input"));
+	}
+
+	return Result == EMetaSoundBuilderResult::Succeeded;
+}
+
+DEFINE_LATENT_AUTOMATION_COMMAND_THREE_PARAMETER(FMetaSoundSourceBuilderRemoveNodeDefaultLiteralLatentCommand, FAutomationTestBase&, Test, UMetaSoundBuilderBase*, Builder, FMetaSoundBuilderNodeInputHandle, NodeInput);
+
+bool FMetaSoundSourceBuilderRemoveNodeDefaultLiteralLatentCommand::Update()
+{
+	EMetaSoundBuilderResult Result = EMetaSoundBuilderResult::Failed;
+	if (Builder)
+	{
+		Builder->RemoveNodeInputDefault(NodeInput, Result);
+		Test.AddErrorIfFalse(Result == EMetaSoundBuilderResult::Succeeded, TEXT("Failed to disconnect MetaSound node's input"));
+	}
+
+	return Result == EMetaSoundBuilderResult::Succeeded;
+}
+
+DEFINE_LATENT_AUTOMATION_COMMAND_THREE_PARAMETER(FMetaSoundSourceBuilderCreateAndConnectTriGeneratorNodeLatentCommand, FAutomationTestBase&, Test, UMetaSoundSourceBuilder*, Builder, FMetaSoundBuilderNodeInputHandle, AudioOutNodeInput);
+
+bool FMetaSoundSourceBuilderCreateAndConnectTriGeneratorNodeLatentCommand::Update()
+{
+	EMetaSoundBuilderResult Result = EMetaSoundBuilderResult::Failed;
+	if (Builder)
+	{
+		// Tri Oscillator Node
+		FMetaSoundNodeHandle TriNode = Builder->AddNodeByClassName({ "UE", "Triangle", "Audio" }, 1, Result);
+		Test.AddErrorIfFalse(Result == EMetaSoundBuilderResult::Succeeded && TriNode.IsSet(), TEXT("Failed to create node by class name 'UE:Triangle:Audio"));
+
+		FMetaSoundBuilderNodeOutputHandle TriNodeAudioOutput = Builder->FindNodeOutputByName(TriNode, "Audio", Result);
+		Test.AddErrorIfFalse(Result == EMetaSoundBuilderResult::Succeeded, TEXT("Failed to find Triangle Oscillator node output 'Audio'"));
+
+		Builder->ConnectNodes(TriNodeAudioOutput, AudioOutNodeInput, Result);
+		Test.AddErrorIfFalse(Result == EMetaSoundBuilderResult::Succeeded, TEXT("Failed to connect 'Audio' Triangle Oscillator output to MetaSound graph's 'Mono Output'"));
+	}
+
+	return Result == EMetaSoundBuilderResult::Succeeded;
+}
+
 DEFINE_LATENT_AUTOMATION_COMMAND_ONE_PARAMETER(FAudioComponentRemoveFromRootLatentCommand, UAudioComponent*, AudioComponent);
 
 bool FAudioComponentRemoveFromRootLatentCommand::Update()
@@ -167,8 +336,21 @@ bool FAudioComponentRemoveFromRootLatentCommand::Update()
 	return false;
 }
 
-// This test imports a metasound source from a JSON file, attempts to play it, and exports it as a json.
-IMPLEMENT_SIMPLE_AUTOMATION_TEST(FAudioMetasoundSourceTest, "Audio.Metasound.BuildAndPlayMetasoundSource", EAutomationTestFlags::EditorContext | EAutomationTestFlags::EngineFilter)
+DEFINE_LATENT_AUTOMATION_COMMAND_ONE_PARAMETER(FBuilderRemoveFromRootLatentCommand, UMetaSoundBuilderBase*, Builder);
+
+bool FBuilderRemoveFromRootLatentCommand::Update()
+{
+	if (Builder)
+	{
+		Builder->RemoveFromRoot();
+		return true;
+	}
+	return false;
+}
+
+
+// This test creates a MetaSound from the legacy controller document editing system and attempts to play it directly.
+IMPLEMENT_SIMPLE_AUTOMATION_TEST(FAudioMetasoundSourceTest, "Audio.Metasound.PlayMetasoundSource", EAutomationTestFlags::EditorContext | EAutomationTestFlags::EngineFilter)
 bool FAudioMetasoundSourceTest::RunTest(const FString& Parameters)
 {
 	UMetaSoundSource* MetaSoundSource = NewObject<UMetaSoundSource>(GetTransientPackage(), FName(*LexToString(FGuid::NewGuid())));;
@@ -180,7 +362,7 @@ bool FAudioMetasoundSourceTest::RunTest(const FString& Parameters)
 		{
 
 			UAudioComponent* AudioComponent = FAudioDevice::CreateComponent(MetaSoundSource);
-			AddErrorIfFalse(AudioComponent != nullptr, "Couldn't create audio component!");
+			AddErrorIfFalse(AudioComponent != nullptr, "Failed to create audio component");
 
 			if (AudioComponent)
 			{
@@ -201,7 +383,161 @@ bool FAudioMetasoundSourceTest::RunTest(const FString& Parameters)
 	return true;
  }
 
-// TODO: add separate JSON test.
+// This test creates a MetaSound source from a SourceBuilder, adds a simple sine tone generator with a connected graph input frequency, and attempts to audition it.
+IMPLEMENT_SIMPLE_AUTOMATION_TEST(FAudioMetasoundSourceBuilderTest, "Audio.Metasound.AuditionMetasoundSource", EAutomationTestFlags::EditorContext | EAutomationTestFlags::EngineFilter)
+bool FAudioMetasoundSourceBuilderTest::RunTest(const FString& Parameters)
+{
+	using namespace EngineTestMetasoundSourcePrivate;
+// 	using namespace Metasound;
+// 	using namespace Metasound::Engine;
+// 	using namespace Metasound::Frontend;
 
+	FMetaSoundBuilderNodeInputHandle MonoOutNodeInput;
+	UMetaSoundSourceBuilder& Builder = CreateMetaSoundMonoSourceSinGenBuilder(*this, nullptr, &MonoOutNodeInput);
+	Builder.AddToRoot();
 
-#endif //WITH_DEV_AUTOMATION_TESTS
+	if (FAudioDevice* AudioDevice = GEngine->GetMainAudioDeviceRaw())
+	{
+		constexpr bool bAddToRegistry = false;
+		TScriptInterface<IMetaSoundDocumentInterface> BuiltDocumentInterface = Builder.Build(nullptr, FMetaSoundBuilderOptions { "BuildAndPlayMetasoundSource", bAddToRegistry });
+		UMetaSoundSource* MetaSoundSource = CastChecked<UMetaSoundSource>(BuiltDocumentInterface.GetObject());
+		UAudioComponent* AudioComponent = FAudioDevice::CreateComponent(MetaSoundSource);
+		AddErrorIfFalse(AudioComponent != nullptr, "Failed to create audio component");
+
+		if (AudioComponent)
+		{
+			AudioComponent->bAutoActivate = false;
+			AudioComponent->bIsUISound = true; // play while "paused"
+			AudioComponent->AudioDeviceID = AudioDevice->DeviceID;
+
+			AudioComponent->bAllowSpatialization = false;
+			AudioComponent->SetVolumeMultiplier(1.0f);
+			AudioComponent->AddToRoot();
+
+			constexpr bool bEnableLiveUpdate = false;
+			ADD_LATENT_AUTOMATION_COMMAND(FMetaSoundSourceBuilderAuditionLatentCommand(&Builder, AudioComponent, bEnableLiveUpdate));
+			ADD_LATENT_AUTOMATION_COMMAND(FEngineWaitLatentCommand(2.f));
+			ADD_LATENT_AUTOMATION_COMMAND(FAudioComponentStopLatentCommand(AudioComponent));
+			ADD_LATENT_AUTOMATION_COMMAND(FEngineWaitLatentCommand(0.5f));
+			ADD_LATENT_AUTOMATION_COMMAND(FBuilderRemoveFromRootLatentCommand(&Builder));
+			ADD_LATENT_AUTOMATION_COMMAND(FAudioComponentRemoveFromRootLatentCommand(AudioComponent));
+		}
+	}
+
+	return true;
+}
+
+// This test creates a MetaSound source from a SourceBuilder, adds a simple sine tone generator with a connected graph input frequency, attempts to audition it, and then switches to a new tri tone generator during playback.
+IMPLEMENT_SIMPLE_AUTOMATION_TEST(FAudioMetasoundSourceBuilderLiveUpdateNode, "Audio.Metasound.LiveUpdateNodeMetaSoundSource", EAutomationTestFlags::EditorContext | EAutomationTestFlags::EngineFilter)
+bool FAudioMetasoundSourceBuilderLiveUpdateNode::RunTest(const FString& Parameters)
+{
+	using namespace EngineTestMetasoundSourcePrivate;
+	using namespace Metasound;
+	using namespace Metasound::Engine;
+	using namespace Metasound::Frontend;
+
+	FMetaSoundBuilderNodeInputHandle MonoOutNodeInput;
+	UMetaSoundSourceBuilder& Builder = CreateMetaSoundMonoSourceSinGenBuilder(*this, nullptr, &MonoOutNodeInput);
+	Builder.AddToRoot();
+
+	if (FAudioDevice* AudioDevice = GEngine->GetMainAudioDeviceRaw())
+	{
+		UAudioComponent* AudioComponent = NewObject<UAudioComponent>();
+		AddErrorIfFalse(AudioComponent != nullptr, "Failed to create audio component");
+
+		if (AudioComponent)
+		{
+			AudioComponent->bAutoActivate = false;
+			AudioComponent->bIsUISound = true; // play while "paused"
+			AudioComponent->AudioDeviceID = AudioDevice->DeviceID;
+
+			AudioComponent->bAllowSpatialization = false;
+			AudioComponent->SetVolumeMultiplier(1.0f);
+			AudioComponent->AddToRoot();
+
+			constexpr bool bEnableLiveUpdate = true;
+			ADD_LATENT_AUTOMATION_COMMAND(FMetaSoundSourceBuilderAuditionLatentCommand(&Builder, AudioComponent, bEnableLiveUpdate));
+			ADD_LATENT_AUTOMATION_COMMAND(FEngineWaitLatentCommand(2.f));
+
+			// Disconnects graph audio output from existing sinosc output
+			ADD_LATENT_AUTOMATION_COMMAND(FMetaSoundSourceBuilderDisconnectInputLatentCommand(*this, &Builder, MonoOutNodeInput));
+
+			ADD_LATENT_AUTOMATION_COMMAND(FMetaSoundSourceBuilderCreateAndConnectTriGeneratorNodeLatentCommand(*this, &Builder, MonoOutNodeInput));
+
+			// TODO: Once we have the builder hooked up to the dynamic graph backend, remove this and voila!
+			ADD_LATENT_AUTOMATION_COMMAND(FMetaSoundSourceBuilderAuditionLatentCommand(&Builder, AudioComponent, bEnableLiveUpdate));
+
+			ADD_LATENT_AUTOMATION_COMMAND(FEngineWaitLatentCommand(2.f));
+			ADD_LATENT_AUTOMATION_COMMAND(FAudioComponentStopLatentCommand(AudioComponent));
+			ADD_LATENT_AUTOMATION_COMMAND(FEngineWaitLatentCommand(0.5f));
+			ADD_LATENT_AUTOMATION_COMMAND(FAudioComponentRemoveFromRootLatentCommand(AudioComponent));
+			ADD_LATENT_AUTOMATION_COMMAND(FBuilderRemoveFromRootLatentCommand(&Builder));
+		}
+	}
+
+	return true;
+}
+
+// This test creates a MetaSound source from a SourceBuilder, adds a simple sine tone generator with a connected graph input frequency, attempts to audition it,
+// disconnects frequency input, sets the sinosc frequency literal value to a new value, and finally removes the literal value default to have it return to the
+// class default.
+IMPLEMENT_SIMPLE_AUTOMATION_TEST(FAudioMetasoundSourceBuilderLiveUpdateLiteral, "Audio.Metasound.LiveUpdateLiteralMetaSoundSource", EAutomationTestFlags::EditorContext | EAutomationTestFlags::EngineFilter)
+bool FAudioMetasoundSourceBuilderLiveUpdateLiteral::RunTest(const FString& Parameters)
+{
+	using namespace EngineTestMetasoundSourcePrivate;
+
+	FMetaSoundBuilderNodeInputHandle MonoOutNodeInput;
+	FMetaSoundBuilderNodeInputHandle GenNodeFreqInput;
+	UMetaSoundSourceBuilder& Builder = CreateMetaSoundMonoSourceSinGenBuilder(*this, &GenNodeFreqInput , &MonoOutNodeInput, 220.f);
+	Builder.AddToRoot();
+
+	if (FAudioDevice* AudioDevice = GEngine->GetMainAudioDeviceRaw())
+	{
+		UAudioComponent* AudioComponent = NewObject<UAudioComponent>();
+		AddErrorIfFalse(AudioComponent != nullptr, "Failed to create audio component");
+
+		if (AudioComponent)
+		{
+			AudioComponent->bAutoActivate = false;
+			AudioComponent->bIsUISound = true; // play while "paused"
+			AudioComponent->AudioDeviceID = AudioDevice->DeviceID;
+
+			AudioComponent->bAllowSpatialization = false;
+			AudioComponent->SetVolumeMultiplier(1.0f);
+			AudioComponent->AddToRoot();
+
+			constexpr bool bEnableLiveUpdate = true;
+			ADD_LATENT_AUTOMATION_COMMAND(FMetaSoundSourceBuilderAuditionLatentCommand(&Builder, AudioComponent, bEnableLiveUpdate));
+			ADD_LATENT_AUTOMATION_COMMAND(FEngineWaitLatentCommand(0.25f));
+
+			// TODO: Once we have the builder hooked up to the dynamic graph backend, remove additional calls to
+			// FMetaSoundSourceBuilderAuditionLatentCommand. The dynamic graph should automatically make changes audible.
+
+			// Disconnects graph audio output from existing sinosc output
+			ADD_LATENT_AUTOMATION_COMMAND(FMetaSoundSourceBuilderDisconnectInputLatentCommand(*this, &Builder, GenNodeFreqInput));
+			ADD_LATENT_AUTOMATION_COMMAND(FMetaSoundSourceBuilderAuditionLatentCommand(&Builder, AudioComponent, bEnableLiveUpdate));
+			ADD_LATENT_AUTOMATION_COMMAND(FEngineWaitLatentCommand(0.25f));
+
+			FName DataTypeName;
+			FMetasoundFrontendLiteral NewValue = UMetaSoundBuilderSubsystem::GetChecked().CreateFloatMetaSoundLiteral(880.f, DataTypeName);
+			AddErrorIfFalse(DataTypeName == Metasound::GetMetasoundDataTypeName<float>(),
+				"Setting MetaSound Float literal returns non-float DataTypeName.");
+			ADD_LATENT_AUTOMATION_COMMAND(FMetaSoundSourceBuilderSetLiteralLatentCommand(*this, &Builder, GenNodeFreqInput, NewValue));
+			ADD_LATENT_AUTOMATION_COMMAND(FMetaSoundSourceBuilderAuditionLatentCommand(&Builder, AudioComponent, bEnableLiveUpdate));
+			ADD_LATENT_AUTOMATION_COMMAND(FEngineWaitLatentCommand(0.25f));
+
+			ADD_LATENT_AUTOMATION_COMMAND(FMetaSoundSourceBuilderRemoveNodeDefaultLiteralLatentCommand(*this, &Builder, GenNodeFreqInput));
+			ADD_LATENT_AUTOMATION_COMMAND(FMetaSoundSourceBuilderAuditionLatentCommand(&Builder, AudioComponent, bEnableLiveUpdate));
+			ADD_LATENT_AUTOMATION_COMMAND(FEngineWaitLatentCommand(0.25f));
+
+			ADD_LATENT_AUTOMATION_COMMAND(FAudioComponentStopLatentCommand(AudioComponent));
+			ADD_LATENT_AUTOMATION_COMMAND(FEngineWaitLatentCommand(0.25f));
+
+			ADD_LATENT_AUTOMATION_COMMAND(FAudioComponentRemoveFromRootLatentCommand(AudioComponent));
+			ADD_LATENT_AUTOMATION_COMMAND(FBuilderRemoveFromRootLatentCommand(&Builder));
+		}
+	}
+
+	return true;
+}
+#endif // WITH_DEV_AUTOMATION_TESTS
diff --git a/Engine/Plugins/Runtime/Metasound/Source/MetasoundFrontend/Private/MetasoundDocumentInterface.cpp b/Engine/Plugins/Runtime/Metasound/Source/MetasoundFrontend/Private/MetasoundDocumentInterface.cpp
new file mode 100644
--- /dev/null
+++ b/Engine/Plugins/Runtime/Metasound/Source/MetasoundFrontend/Private/MetasoundDocumentInterface.cpp
@@ -0,0 +1,28 @@
+// Copyright Epic Games, Inc. All Rights Reserved.
+
+#include "MetasoundDocumentInterface.h"
+
+
+namespace Metasound::Frontend
+{
+	namespace DocumentBuilderRegistryPrivate
+	{
+		TUniqueFunction<IMetaSoundDocumentBuilderRegistry&()> GetInstance;
+	} // namespace DocumentBuilderRegistryPrivate
+
+	IMetaSoundDocumentBuilderRegistry& IMetaSoundDocumentBuilderRegistry::GetChecked()
+	{
+		using namespace DocumentBuilderRegistryPrivate;
+
+		checkf(GetInstance, TEXT("Failed to return MetaSoundDocumentBuilderRegistry instance: Registry has not been initialized"));
+		return GetInstance();
+	}
+
+	void IMetaSoundDocumentBuilderRegistry::Set(TUniqueFunction<IMetaSoundDocumentBuilderRegistry&()>&& InGetInstance)
+	{
+		using namespace DocumentBuilderRegistryPrivate;
+
+		checkf(!GetInstance, TEXT("Failed to initialize MetaSoundDocumentBuilderRegistry getter: Cannot reinitialize once initialized."))
+		GetInstance = MoveTemp(InGetInstance);
+	}
+} // namespace Metasound::Frontend
\ No newline at end of file
diff --git a/Engine/Plugins/Runtime/Metasound/Source/MetasoundFrontend/Private/MetasoundFrontendDocumentBuilder.cpp b/Engine/Plugins/Runtime/Metasound/Source/MetasoundFrontend/Private/MetasoundFrontendDocumentBuilder.cpp
--- a/Engine/Plugins/Runtime/Metasound/Source/MetasoundFrontend/Private/MetasoundFrontendDocumentBuilder.cpp
+++ b/Engine/Plugins/Runtime/Metasound/Source/MetasoundFrontend/Private/MetasoundFrontendDocumentBuilder.cpp
@@ -15,6 +15,7 @@
 #include "MetasoundFrontendController.h"
 #include "MetasoundFrontendDocumentCache.h"
 #include "MetasoundFrontendDocument.h"
+#include "MetasoundFrontendDocumentModifyDelegates.h"
 #include "MetasoundFrontendRegistries.h"
 #include "MetasoundFrontendSearchEngine.h"
 #include "MetasoundFrontendTransform.h"
@@ -350,18 +351,20 @@ namespace Metasound::Frontend
 #endif // WITH_EDITORONLY_DATA
 
 		public:
-			bool Execute(FMetaSoundFrontendDocumentBuilder& OutBuilder, FMetasoundFrontendDocument& OutDoc)
+			bool Execute(FMetaSoundFrontendDocumentBuilder& OutBuilder, FMetasoundFrontendDocument& OutDoc, FDocumentModifyDelegates& OutDelegates)
 			{
 				bool bDidEdit = false;
 
 				for (const FMetasoundFrontendInterface& Interface : Options.InterfacesToRemove)
 				{
-					if (OutDoc.Interfaces.Remove(Interface.Version) > 0)
+					if (OutDoc.Interfaces.Contains(Interface.Version))
 					{
+						OutDelegates.InterfaceDelegates.OnRemovingInterface.Broadcast(Interface);
 						bDidEdit = true;
 #if WITH_EDITORONLY_DATA
 						OutDoc.Metadata.ModifyContext.AddInterfaceModified(Interface.Version.Name);
 #endif // WITH_EDITORONLY_DATA
+						OutDoc.Interfaces.Remove(Interface.Version);
 					}
 				}
 
@@ -371,6 +374,7 @@ namespace Metasound::Frontend
 					OutDoc.Interfaces.Add(Interface.Version, &bAlreadyInSet);
 					if (!bAlreadyInSet)
 					{
+						OutDelegates.InterfaceDelegates.OnInterfaceAdded.Broadcast(Interface);
 						bDidEdit = true;
 #if WITH_EDITORONLY_DATA
 						OutDoc.Metadata.ModifyContext.AddInterfaceModified(Interface.Version.Name);
@@ -445,6 +449,12 @@ namespace Metasound::Frontend
 	{
 	}
 
+	FModifyInterfaceOptions::FModifyInterfaceOptions(TArray<FMetasoundFrontendInterface>&& InInterfacesToRemove, TArray<FMetasoundFrontendInterface>&& InInterfacesToAdd)
+		: InterfacesToRemove(MoveTemp(InInterfacesToRemove))
+		, InterfacesToAdd(MoveTemp(InInterfacesToAdd))
+	{
+	}
+
 	FModifyInterfaceOptions::FModifyInterfaceOptions(const TArray<FMetasoundFrontendVersion>& InInterfaceVersionsToRemove, const TArray<FMetasoundFrontendVersion>& InInterfaceVersionsToAdd)
 	{
 		Algo::Transform(InInterfaceVersionsToRemove, InterfacesToRemove, [](const FMetasoundFrontendVersion& Version)
@@ -472,89 +482,52 @@ namespace Metasound::Frontend
 } // namespace Metasound::Frontend
 
 
-FMetaSoundFrontendDocumentBuilder::FMetaSoundFrontendDocumentBuilder(TScriptInterface<IMetaSoundDocumentInterface> InDocumentInterface)
-	: DocumentInterface(InDocumentInterface)
+FMetaSoundFrontendDocumentBuilder::FMetaSoundFrontendDocumentBuilder()
+	: DocumentDelegates(MakeShared<Metasound::Frontend::FDocumentModifyDelegates>())
 {
-	if (DocumentInterface)
-	{
-		ReloadCache();
-	}
 }
 
-FMetaSoundFrontendDocumentBuilder::FMetaSoundFrontendDocumentBuilder(const FMetaSoundFrontendDocumentBuilder& InBuilder)
+FMetaSoundFrontendDocumentBuilder::FMetaSoundFrontendDocumentBuilder(TScriptInterface<IMetaSoundDocumentInterface> InDocumentInterface)
+	: DocumentInterface(InDocumentInterface)
+	, DocumentDelegates(MakeShared<Metasound::Frontend::FDocumentModifyDelegates>())
 {
-	DocumentInterface = InBuilder.DocumentInterface;
 	if (DocumentInterface)
 	{
 		ReloadCache();
 	}
 }
 
-FMetaSoundFrontendDocumentBuilder::FMetaSoundFrontendDocumentBuilder(FMetaSoundFrontendDocumentBuilder&& InBuilder)
-	: DocumentInterface(MoveTemp(InBuilder.DocumentInterface))
-	, DocumentCache(MoveTemp(InBuilder.DocumentCache))
-{
-}
-
-FMetaSoundFrontendDocumentBuilder& FMetaSoundFrontendDocumentBuilder::operator=(FMetaSoundFrontendDocumentBuilder&& InRHS)
-{
-	using namespace Metasound::Frontend;
-	DocumentInterface = MoveTemp(InRHS.DocumentInterface);
-	DocumentCache = MoveTemp(InRHS.DocumentCache);
-	return *this;
-}
-
-FMetaSoundFrontendDocumentBuilder& FMetaSoundFrontendDocumentBuilder::operator=(const FMetaSoundFrontendDocumentBuilder& InRHS)
+FMetaSoundFrontendDocumentBuilder::FMetaSoundFrontendDocumentBuilder(TScriptInterface<IMetaSoundDocumentInterface> InDocumentInterface, TSharedRef<Metasound::Frontend::FDocumentModifyDelegates> InDocumentDelegates)
+	: DocumentInterface(InDocumentInterface)
+	, DocumentDelegates(InDocumentDelegates)
 {
-	using namespace Metasound::Frontend;
-
-	DocumentInterface = InRHS.DocumentInterface;
 	if (DocumentInterface)
 	{
 		ReloadCache();
 	}
-	return *this;
 }
 
-const FMetasoundFrontendClass* FMetaSoundFrontendDocumentBuilder::AddGraphDependency(const FMetasoundFrontendGraphClass& InClass)
+const FMetasoundFrontendClass* FMetaSoundFrontendDocumentBuilder::AddDependency(const FMetasoundFrontendClass& InClass)
 {
 	using namespace Metasound::Frontend;
 
-	checkf(InClass.Metadata.GetType() == EMetasoundFrontendClassType::Graph, TEXT("Call should only be used for 'Graph' type dependencies. Use 'AddNativeDependency' for other class types."));
-
 	FMetasoundFrontendDocument& Document = GetDocument();
 	const FMetasoundFrontendClass* Dependency = nullptr;
 
 	FMetasoundFrontendClass NewDependency = InClass;
 
-	// All dependencies are listed as 'External' from the perspective of the owning document.
+	// All 'Graph' dependencies are listed as 'External' from the perspective of the owning document.
 	// This makes them implementation agnostic to accommodate nativization of assets.
-	NewDependency.Metadata.SetType(EMetasoundFrontendClassType::External);
-	NewDependency.ID = FGuid::NewGuid();
-	Dependency = &Document.Dependencies.Emplace_GetRef(MoveTemp(NewDependency));
-
-	const int32 NewIndex = Document.Dependencies.Num() - 1;
-	DocumentCache->OnDependencyAdded(NewIndex);
-
-	return Dependency;
-}
-
-const FMetasoundFrontendClass* FMetaSoundFrontendDocumentBuilder::AddNativeDependency(const FMetasoundFrontendClassMetadata& InClassMetadata)
-{
-	using namespace Metasound::Frontend;
-
-	checkf(InClassMetadata.GetType() != EMetasoundFrontendClassType::Graph, TEXT("Graph dependencies must be added via 'AddGraphDependency' to avoid overhead of registration where not necessary"));
-
-	FMetasoundFrontendDocument& Document = GetDocument();
-	const FMetasoundFrontendClass* Dependency = nullptr;
+	if (NewDependency.Metadata.GetType() == EMetasoundFrontendClassType::Graph)
+	{
+		NewDependency.Metadata.SetType(EMetasoundFrontendClassType::External);
+	}
 
-	const FNodeRegistryKey RegistryKey = NodeRegistryKey::CreateKey(InClassMetadata);
-	FMetasoundFrontendClass NewDependency = GenerateClass(RegistryKey);
 	NewDependency.ID = FGuid::NewGuid();
 	Dependency = &Document.Dependencies.Emplace_GetRef(MoveTemp(NewDependency));
 
 	const int32 NewIndex = Document.Dependencies.Num() - 1;
-	DocumentCache->OnDependencyAdded(NewIndex);
+	DocumentDelegates->OnDependencyAdded.Broadcast(NewIndex);
 
 	return Dependency;
 }
@@ -563,15 +536,14 @@ const FMetasoundFrontendEdge* FMetaSoundFrontendDocumentBuilder::AddEdge(FMetaso
 {
 	using namespace Metasound::Frontend;
 
-	FMetasoundFrontendDocument& Document = GetDocument();
-	FMetasoundFrontendGraph& Graph = Document.RootGraph.Graph;
 	if (CanAddEdge(InNewEdge))
 	{
-		IDocumentGraphEdgeCache& EdgeCache = DocumentCache->GetEdgeCache();
-
+		FMetasoundFrontendDocument& Document = GetDocument();
+		FMetasoundFrontendGraph& Graph = Document.RootGraph.Graph;
+		const IDocumentGraphEdgeCache& EdgeCache = DocumentCache->GetEdgeCache();
 		const FMetasoundFrontendEdge& NewEdge = Graph.Edges.Add_GetRef(MoveTemp(InNewEdge));
 		const int32 NewIndex = Graph.Edges.Num() - 1;
-		EdgeCache.OnEdgeAdded(NewIndex);
+		DocumentDelegates->EdgeDelegates.OnEdgeAdded.Broadcast(NewIndex);
 		return &NewEdge;
 	}
 
@@ -668,6 +640,7 @@ bool FMetaSoundFrontendDocumentBuilder::AddEdgesFromMatchingInterfaceNodeOutputs
 	}
 
 	const IDocumentGraphNodeCache& NodeCache = DocumentCache->GetNodeCache();
+	const IDocumentGraphInterfaceCache& InterfaceCache = DocumentCache->GetInterfaceCache();
 
 	const TSet<FMetasoundFrontendVersion>& NodeInterfaces = NodeDocument->Interfaces;
 	const TSet<FMetasoundFrontendVersion> CommonInterfaces = NodeInterfaces.Intersect(GetDocument().Interfaces);
@@ -678,12 +651,14 @@ bool FMetaSoundFrontendDocumentBuilder::AddEdgesFromMatchingInterfaceNodeOutputs
 		const FInterfaceRegistryKey InterfaceKey = GetInterfaceRegistryKey(Version);
 		if (const IInterfaceRegistryEntry* RegistryEntry = IInterfaceRegistry::Get().FindInterfaceRegistryEntry(InterfaceKey))
 		{
-			Algo::Transform(RegistryEntry->GetInterface().Outputs, EdgesToMake, [this, &NodeCache, InNodeID](const FMetasoundFrontendClassOutput& Output)
+			Algo::Transform(RegistryEntry->GetInterface().Outputs, EdgesToMake, [this, &NodeCache, &InterfaceCache, InNodeID](const FMetasoundFrontendClassOutput& Output)
 			{
 				const FMetasoundFrontendGraph& Graph = GetDocument().RootGraph.Graph;
 				const FMetasoundFrontendVertex* NodeVertex = NodeCache.FindOutputVertex(InNodeID, Output.Name);
 				check(NodeVertex);
-				const FMetasoundFrontendNode* OutputNode = NodeCache.FindOutputNode(Output.Name);
+				const FMetasoundFrontendClassOutput* OutputClass = InterfaceCache.FindOutput(Output.Name);
+				check(OutputClass);
+				const FMetasoundFrontendNode* OutputNode = NodeCache.FindNode(OutputClass->NodeID);
 				check(OutputNode);
 				const TArray<FMetasoundFrontendVertex>& Inputs = OutputNode->Interface.Inputs;
 				check(!Inputs.IsEmpty());
@@ -710,6 +685,7 @@ bool FMetaSoundFrontendDocumentBuilder::AddEdgesFromMatchingInterfaceNodeInputsT
 	}
 
 	const IDocumentGraphNodeCache& NodeCache = DocumentCache->GetNodeCache();
+	const IDocumentGraphInterfaceCache& InterfaceCache = DocumentCache->GetInterfaceCache();
 
 	const TSet<FMetasoundFrontendVersion>& NodeInterfaces = NodeDocument->Interfaces;
 	const TSet<FMetasoundFrontendVersion> CommonInterfaces = NodeInterfaces.Intersect(GetDocument().Interfaces);
@@ -720,12 +696,14 @@ bool FMetaSoundFrontendDocumentBuilder::AddEdgesFromMatchingInterfaceNodeInputsT
 		const FInterfaceRegistryKey InterfaceKey = GetInterfaceRegistryKey(Version);
 		if (const IInterfaceRegistryEntry* RegistryEntry = IInterfaceRegistry::Get().FindInterfaceRegistryEntry(InterfaceKey))
 		{
-			Algo::Transform(RegistryEntry->GetInterface().Inputs, EdgesToMake, [this, &NodeCache, InNodeID](const FMetasoundFrontendClassInput& Input)
+			Algo::Transform(RegistryEntry->GetInterface().Inputs, EdgesToMake, [this, &NodeCache, &InterfaceCache, InNodeID](const FMetasoundFrontendClassInput& Input)
 			{
 				const FMetasoundFrontendGraph& Graph = GetDocument().RootGraph.Graph;
 				const FMetasoundFrontendVertex* NodeVertex = NodeCache.FindOutputVertex(InNodeID, Input.Name);
 				check(NodeVertex);
-				const FMetasoundFrontendNode* InputNode = NodeCache.FindInputNode(Input.Name);
+				const FMetasoundFrontendClassInput* InputClass = InterfaceCache.FindInput(Input.Name);
+				check(InputClass);
+				const FMetasoundFrontendNode* InputNode = NodeCache.FindNode(InputClass->NodeID);
 				check(InputNode);
 				const TArray<FMetasoundFrontendVertex>& Outputs = InputNode->Interface.Outputs;
 				check(!Outputs.IsEmpty());
@@ -752,10 +730,12 @@ const FMetasoundFrontendNode* FMetaSoundFrontendDocumentBuilder::AddGraphInput(c
 		UE_LOG(LogMetaSound, Error, TEXT("TypeName unset when attempting to add class input '%s'"), *InClassInput.Name.ToString());
 		return nullptr;
 	}
-	else if (const FMetasoundFrontendNode* Node = DocumentCache->GetNodeCache().FindInputNode(InClassInput.Name))
+	else if (const FMetasoundFrontendClassInput* Input = DocumentCache->GetInterfaceCache().FindInput(InClassInput.Name))
 	{
 		UE_LOG(LogMetaSound, Error, TEXT("Attempting to add MetaSound graph input '%s' when input with name already exists"), *InClassInput.Name.ToString());
-		return Node;
+		const FMetasoundFrontendNode* OutputNode = DocumentCache->GetNodeCache().FindNode(Input->NodeID);
+		check(OutputNode);
+		return OutputNode;
 	}
 
 	auto FindRegistryClass = [&InClassInput](FMetasoundFrontendClass& OutClass) -> bool
@@ -790,7 +770,7 @@ const FMetasoundFrontendNode* FMetaSoundFrontendDocumentBuilder::AddGraphInput(c
 	{
 		if(!FindDependency(Class.Metadata))
 		{
-			AddNativeDependency(Class.Metadata);
+			AddDependency(Class);
 		}
 
 		auto FinalizeNode = [&InClassInput](FMetasoundFrontendNode& InOutNode, const Metasound::Frontend::FNodeRegistryKey&)
@@ -799,12 +779,14 @@ const FMetasoundFrontendNode* FMetaSoundFrontendDocumentBuilder::AddGraphInput(c
 		};
 		if (FMetasoundFrontendNode* NewNode = AddNodeInternal(Class.Metadata, FinalizeNode, InClassInput.NodeID))
 		{
+			const int32 NewIndex = RootGraph.Interface.Inputs.Num();
 			FMetasoundFrontendClassInput& NewInput = RootGraph.Interface.Inputs.Add_GetRef(InClassInput);
 			if (!NewInput.VertexID.IsValid())
 			{
 				NewInput.VertexID = FGuid::NewGuid();
 			}
 
+			DocumentDelegates->InterfaceDelegates.OnInputAdded.Broadcast(NewIndex);
 #if WITH_EDITORONLY_DATA
 			Document.Metadata.ModifyContext.AddMemberIDModified(InClassInput.NodeID);
 #endif // WITH_EDITORONLY_DATA
@@ -831,10 +813,10 @@ const FMetasoundFrontendNode* FMetaSoundFrontendDocumentBuilder::AddGraphOutput(
 		UE_LOG(LogMetaSound, Error, TEXT("TypeName unset when attempting to add class output '%s'"), *InClassOutput.Name.ToString());
 		return nullptr;
 	}
-	else if (const FMetasoundFrontendNode* Node = DocumentCache->GetNodeCache().FindInputNode(InClassOutput.Name))
+	else if (const FMetasoundFrontendClassInput* Input = DocumentCache->GetInterfaceCache().FindInput(InClassOutput.Name))
 	{
 		UE_LOG(LogMetaSound, Error, TEXT("Attempting to add MetaSound graph output '%s' when output with name already exists"), *InClassOutput.Name.ToString());
-		return Node;
+		return DocumentCache->GetNodeCache().FindNode(Input->NodeID);
 	}
 
 	auto FindRegistryClass = [&InClassOutput](FMetasoundFrontendClass& OutClass) -> bool
@@ -869,7 +851,7 @@ const FMetasoundFrontendNode* FMetaSoundFrontendDocumentBuilder::AddGraphOutput(
 	{
 		if (!FindDependency(Class.Metadata))
 		{
-			AddNativeDependency(Class.Metadata);
+			AddDependency(Class);
 		}
 
 		auto FinalizeNode = [&InClassOutput](FMetasoundFrontendNode& InOutNode, const Metasound::Frontend::FNodeRegistryKey&)
@@ -878,12 +860,14 @@ const FMetasoundFrontendNode* FMetaSoundFrontendDocumentBuilder::AddGraphOutput(
 		};
 		if (FMetasoundFrontendNode* NewNode = AddNodeInternal(Class.Metadata, FinalizeNode, InClassOutput.NodeID))
 		{
+			const int32 NewIndex = RootGraph.Interface.Outputs.Num();
 			FMetasoundFrontendClassOutput& NewOutput = RootGraph.Interface.Outputs.Add_GetRef(InClassOutput);
 			if (!NewOutput.VertexID.IsValid())
 			{
 				NewOutput.VertexID = FGuid::NewGuid();
 			}
 
+			DocumentDelegates->InterfaceDelegates.OnOutputAdded.Broadcast(NewIndex);
 #if WITH_EDITORONLY_DATA
 			Document.Metadata.ModifyContext.AddMemberIDModified(InClassOutput.NodeID);
 #endif // WITH_EDITORONLY_DATA
@@ -918,17 +902,11 @@ bool FMetaSoundFrontendDocumentBuilder::AddInterface(FName InterfaceName)
 				UE_LOG(LogMetaSound, Error, TEXT("DocumentBuilder failed to add MetaSound Interface '%s' to document: is not set to be modifiable for given UClass '%s'"), *InterfaceName.ToString(), *BuilderClassPath.ToString());
 				return false;
 			}
-		}
 
-		const FModifyRootGraphInterfaces ModifyInterfaces({ }, { Interface.Version });
-		if (ModifyInterfaces.Transform(GetDocument()))
-		{
-			// TODO: Add ability to update local builder caches dynamically
-			// in ModifyRootGraphInterfaces transform once it moves away
-			// from using controllers internally. Then cache updates can
-			// be more isolated and not use this full rebuild call.
-			ReloadCache();
-			return true;
+			TArray<FMetasoundFrontendInterface> InterfacesToAdd;
+			InterfacesToAdd.Add(Entry->GetInterface());
+			FModifyInterfaceOptions Options({ }, MoveTemp(InterfacesToAdd));
+			ModifyInterfaces(MoveTemp(Options));
 		}
 	}
 
@@ -962,22 +940,49 @@ const FMetasoundFrontendNode* FMetaSoundFrontendDocumentBuilder::AddGraphNode(co
 
 	if (!FindDependency(NewClassMetadata))
 	{
-		AddGraphDependency(InGraphClass);
+		AddDependency(InGraphClass);
 	}
 
 	return AddNodeInternal(NewClassMetadata, FinalizeNode, InNodeID);
 }
 
+const FMetasoundFrontendNode* FMetaSoundFrontendDocumentBuilder::AddNodeByClassName(const FMetasoundFrontendClassName& InClassName, int32 InMajorVersion, FGuid InNodeID)
+{
+	using namespace Metasound;
+	using namespace Metasound::Frontend;
+
+	FMetasoundFrontendClass RegisteredClass;
+	if (ISearchEngine::Get().FindClassWithHighestMinorVersion(InClassName, InMajorVersion, RegisteredClass))
+	{
+		// Dependency is considered "External" when looked up or added as a dependency to a graph
+		RegisteredClass.Metadata.SetType(EMetasoundFrontendClassType::External);
+
+		const FMetasoundFrontendClass* Dependency = FindDependency(RegisteredClass.Metadata);
+		if (!Dependency)
+		{
+			Dependency = AddDependency(RegisteredClass);
+		}
+
+		if (Dependency)
+		{
+			return AddNodeInternal(Dependency->Metadata, [](const FMetasoundFrontendNode& Node, const Metasound::Frontend::FNodeRegistryKey& ClassKey) { return Node.Name; }, InNodeID);
+		}
+	}
+
+	UE_LOG(LogMetaSound, Warning, TEXT("Failed to add new dependency with name '%s': Node not added"), *InClassName.ToString());
+	return nullptr;
+}
+
 FMetasoundFrontendNode* FMetaSoundFrontendDocumentBuilder::AddNodeInternal(const FMetasoundFrontendClassMetadata& InClassMetadata, FFinalizeNodeFunctionRef FinalizeNode, FGuid InNodeID)
 {
 	METASOUND_TRACE_CPUPROFILER_EVENT_SCOPE(FMetaSoundFrontendDocumentBuilder::AddNodeInternal);
 
 	using namespace Metasound::Frontend;
 
-	FMetasoundFrontendDocument& Document = GetDocument();
 	const FNodeRegistryKey ClassKey = NodeRegistryKey::CreateKey(InClassMetadata);
 	if (const FMetasoundFrontendClass* Dependency = DocumentCache->FindDependency(ClassKey))
 	{
+		FMetasoundFrontendDocument& Document = GetDocument();
 		FMetasoundFrontendGraph& Graph = Document.RootGraph.Graph;
 		TArray<FMetasoundFrontendNode>& Nodes = Graph.Nodes;
 		FMetasoundFrontendNode& Node = Nodes.Emplace_GetRef(*Dependency);
@@ -985,8 +990,8 @@ FMetasoundFrontendNode* FMetaSoundFrontendDocumentBuilder::AddNodeInternal(const
 		FinalizeNode(Node, ClassKey);
 
 		const int32 NewIndex = Nodes.Num() - 1;
-		IDocumentGraphNodeCache& NodeCache = DocumentCache->GetNodeCache();
-		NodeCache.OnNodeAdded(NewIndex);
+		const IDocumentGraphNodeCache& NodeCache = DocumentCache->GetNodeCache();
+		DocumentDelegates->NodeDelegates.OnNodeAdded.Broadcast(NewIndex);
 
 #if WITH_EDITORONLY_DATA
 		Document.Metadata.ModifyContext.AddNodeIDModified(InNodeID);
@@ -1061,7 +1066,7 @@ void FMetaSoundFrontendDocumentBuilder::ClearGraph()
 bool FMetaSoundFrontendDocumentBuilder::ContainsEdge(const FMetasoundFrontendEdge& InEdge) const
 {
 	using namespace Metasound::Frontend;
-	IDocumentGraphEdgeCache& EdgeCache = DocumentCache->GetEdgeCache();
+	const IDocumentGraphEdgeCache& EdgeCache = DocumentCache->GetEdgeCache();
 	return EdgeCache.ContainsEdge(InEdge);
 }
 
@@ -1177,20 +1182,26 @@ bool FMetaSoundFrontendDocumentBuilder::FindInterfaceInputNodes(FName InterfaceN
 
 	OutInputs.Reset();
 
-	const IDocumentGraphNodeCache& NodeCache = DocumentCache->GetNodeCache();
-
 	FMetasoundFrontendInterface Interface;
 	if (ISearchEngine::Get().FindInterfaceWithHighestVersion(InterfaceName, Interface))
 	{
 		if (GetDocument().Interfaces.Contains(Interface.Version))
 		{
+			const IDocumentGraphNodeCache& NodeCache = DocumentCache->GetNodeCache();
+			const IDocumentGraphInterfaceCache& InterfaceCache = DocumentCache->GetInterfaceCache();
+
 			TArray<const FMetasoundFrontendNode*> InterfaceInputs;
 			for (const FMetasoundFrontendClassInput& Input : Interface.Inputs)
 			{
-				if (const int32* NodeIndex = NodeCache.FindInputNodeIndex(Input.Name))
+				const FMetasoundFrontendClassInput* ClassInput = InterfaceCache.FindInput(Input.Name);
+				if (!ClassInput)
 				{
-					const FMetasoundFrontendNode& Node = GetDocument().RootGraph.Graph.Nodes[*NodeIndex];
-					InterfaceInputs.Add(&Node);
+					return false;
+				}
+
+				if (const FMetasoundFrontendNode* Node = NodeCache.FindNode(ClassInput->NodeID))
+				{
+					InterfaceInputs.Add(Node);
 				}
 				else
 				{
@@ -1217,13 +1228,21 @@ bool FMetaSoundFrontendDocumentBuilder::FindInterfaceOutputNodes(FName Interface
 	{
 		if (GetDocument().Interfaces.Contains(Interface.Version))
 		{
+			const IDocumentGraphNodeCache& NodeCache = DocumentCache->GetNodeCache();
+			const IDocumentGraphInterfaceCache& InterfaceCache = DocumentCache->GetInterfaceCache();
+
 			TArray<const FMetasoundFrontendNode*> InterfaceOutputs;
 			for (const FMetasoundFrontendClassOutput& Output : Interface.Outputs)
 			{
-				if (const int32* NodeIndex = DocumentCache->GetNodeCache().FindOutputNodeIndex(Output.Name))
+				const FMetasoundFrontendClassOutput* ClassOutput = InterfaceCache.FindOutput(Output.Name);
+				if (!ClassOutput)
 				{
-					const FMetasoundFrontendNode& Node = GetDocument().RootGraph.Graph.Nodes[*NodeIndex];
-					InterfaceOutputs.Add(&Node);
+					return false;
+				}
+
+				if (const FMetasoundFrontendNode* Node = NodeCache.FindNode(ClassOutput->NodeID))
+				{
+					InterfaceOutputs.Add(Node);
 				}
 				else
 				{
@@ -1241,28 +1260,38 @@ bool FMetaSoundFrontendDocumentBuilder::FindInterfaceOutputNodes(FName Interface
 
 const FMetasoundFrontendClassInput* FMetaSoundFrontendDocumentBuilder::FindGraphInput(FName InputName) const
 {
-	const TArray<FMetasoundFrontendClassInput>& Inputs = GetDocument().RootGraph.Interface.Inputs;
-	return Inputs.FindByPredicate([&InputName](const FMetasoundFrontendClassInput& Input) { return Input.Name == InputName; });
+	return DocumentCache->GetInterfaceCache().FindInput(InputName);
 }
 
 const FMetasoundFrontendNode* FMetaSoundFrontendDocumentBuilder::FindGraphInputNode(FName InputName) const
 {
 	using namespace Metasound::Frontend;
-	const IDocumentGraphNodeCache& NodeCache = DocumentCache->GetNodeCache();
-	return NodeCache.FindInputNode(InputName);
+
+	if (const FMetasoundFrontendClassInput* InputClass = FindGraphInput(InputName))
+	{
+		const IDocumentGraphNodeCache& NodeCache = DocumentCache->GetNodeCache();
+		return NodeCache.FindNode(InputClass->NodeID);
+	}
+
+	return nullptr;
 }
 
 const FMetasoundFrontendClassOutput* FMetaSoundFrontendDocumentBuilder::FindGraphOutput(FName OutputName) const
 {
-	const TArray<FMetasoundFrontendClassOutput>& Outputs = GetDocument().RootGraph.Interface.Outputs;
-	return Outputs.FindByPredicate([&OutputName](const FMetasoundFrontendClassOutput& Output) { return Output.Name == OutputName; });
+	return DocumentCache->GetInterfaceCache().FindOutput(OutputName);
 }
 
 const FMetasoundFrontendNode* FMetaSoundFrontendDocumentBuilder::FindGraphOutputNode(FName OutputName) const
 {
 	using namespace Metasound::Frontend;
-	const IDocumentGraphNodeCache& NodeCache = DocumentCache->GetNodeCache();
-	return NodeCache.FindOutputNode(OutputName);
+
+	if (const FMetasoundFrontendClassOutput* OutputClass = FindGraphOutput(OutputName))
+	{
+		const IDocumentGraphNodeCache& NodeCache = DocumentCache->GetNodeCache();
+		return NodeCache.FindNode(OutputClass->NodeID);
+	}
+
+	return nullptr;
 }
 
 const FMetasoundFrontendVertex* FMetaSoundFrontendDocumentBuilder::FindNodeInput(const FGuid& InNodeID, const FGuid& InVertexID) const
@@ -1392,6 +1421,11 @@ const FMetasoundFrontendDocument& FMetaSoundFrontendDocumentBuilder::GetDocument
 	return DocumentInterface->GetDocument();
 }
 
+const Metasound::Frontend::FDocumentModifyDelegates& FMetaSoundFrontendDocumentBuilder::GetDocumentDelegates() const
+{
+	return *DocumentDelegates;
+}
+
 const IMetaSoundDocumentInterface& FMetaSoundFrontendDocumentBuilder::GetDocumentInterface() const
 {
 	IMetaSoundDocumentInterface* Interface = DocumentInterface.GetInterface();
@@ -1399,6 +1433,67 @@ const IMetaSoundDocumentInterface& FMetaSoundFrontendDocumentBuilder::GetDocumen
 	return *Interface;
 }
 
+const FMetasoundFrontendLiteral* FMetaSoundFrontendDocumentBuilder::GetNodeInputDefault(const FGuid& InNodeID, const FGuid& InVertexID) const
+{
+	using namespace Metasound::Frontend;
+
+	const IDocumentGraphNodeCache& NodeCache = DocumentCache->GetNodeCache();
+	if (const int32* NodeIndex = NodeCache.FindNodeIndex(InNodeID))
+	{
+		const FMetasoundFrontendGraph& Graph = GetDocument().RootGraph.Graph;
+		const FMetasoundFrontendNode& Node = Graph.Nodes[*NodeIndex];
+
+		auto IsVertex = [&InVertexID](const FMetasoundFrontendVertex& Vertex) { return Vertex.VertexID == InVertexID; };
+		const int32 VertexIndex = Node.Interface.Inputs.IndexOfByPredicate(IsVertex);
+		if (VertexIndex != INDEX_NONE)
+		{
+			const FMetasoundFrontendVertex& NodeInput = Node.Interface.Inputs[VertexIndex];
+
+			// If default not found on node, check class definition
+			auto IsLiteral = [&InVertexID](const FMetasoundFrontendVertexLiteral& Literal) { return Literal.VertexID == InVertexID; };
+			const int32 LiteralIndex = Node.InputLiterals.IndexOfByPredicate(IsLiteral);
+			if (LiteralIndex != INDEX_NONE)
+			{
+				return &Node.InputLiterals[LiteralIndex].Value;
+			}
+		}
+	}
+
+	return nullptr;
+}
+
+const FMetasoundFrontendLiteral* FMetaSoundFrontendDocumentBuilder::GetNodeInputClassDefault(const FGuid& InNodeID, const FGuid& InVertexID) const
+{
+	using namespace Metasound::Frontend;
+
+	const IDocumentGraphNodeCache& NodeCache = DocumentCache->GetNodeCache();
+	if (const int32* NodeIndex = NodeCache.FindNodeIndex(InNodeID))
+	{
+		const FMetasoundFrontendGraph& Graph = GetDocument().RootGraph.Graph;
+		const FMetasoundFrontendNode& Node = Graph.Nodes[*NodeIndex];
+
+		auto IsVertex = [&InVertexID](const FMetasoundFrontendVertex& Vertex) { return Vertex.VertexID == InVertexID; };
+		const int32 VertexIndex = Node.Interface.Inputs.IndexOfByPredicate(IsVertex);
+		if (VertexIndex != INDEX_NONE)
+		{
+			const FMetasoundFrontendVertex& NodeInput = Node.Interface.Inputs[VertexIndex];
+			if (const FMetasoundFrontendClass* Class = DocumentCache->FindDependency(Node.ClassID))
+			{
+// 				auto IsClassInput = [&NodeInput](const FMetasoundFrontendClassInput& Input)
+// 				{
+// 					return FMetasoundFrontendVertex::IsFunctionalEquivalent(Input, NodeInput);
+// 				};
+				if (const FMetasoundFrontendClassInput* ClassInput = Class->Interface.Inputs.FindByPredicate(IsVertex))
+				{
+					return &ClassInput->DefaultLiteral;
+				}
+			}
+		}
+	}
+
+	return nullptr;
+}
+
 void FMetaSoundFrontendDocumentBuilder::InitGraphClassMetadata(FMetasoundFrontendClassMetadata& InOutMetadata, bool bResetVersion)
 {
 	InOutMetadata.SetClassName(FMetasoundFrontendClassName(FName(), *FGuid::NewGuid().ToString(), FName()));
@@ -1531,13 +1626,22 @@ bool FMetaSoundFrontendDocumentBuilder::ModifyInterfaces(Metasound::Frontend::FM
 
 	DocumentBuilderPrivate::FModifyInterfacesImpl Context(MoveTemp(InOptions));
 	FMetasoundFrontendDocument& Doc = GetDocument();
-	return Context.Execute(*this, Doc);
+	return Context.Execute(*this, Doc, *DocumentDelegates);
 }
 
 void FMetaSoundFrontendDocumentBuilder::ReloadCache()
 {
 	using namespace Metasound::Frontend;
-	DocumentCache = MakeUnique<FDocumentCache>(GetDocument());
+	if (!DocumentCache.IsValid())
+	{
+		DocumentCache = MakeShared<FDocumentCache>(GetDocument());
+	}
+
+	// Must be called after the constructor to allow for passing shared pointer to
+	// sub-caches (i.e. can't be RAII because initialization requires creating
+	// additional shared pointers in the ctor which is explicitly forbidden)
+	TSharedPtr<FDocumentCache> CacheConcrete = StaticCastSharedPtr<FDocumentCache>(DocumentCache);
+	CacheConcrete->Init(*DocumentDelegates);
 }
 
 bool FMetaSoundFrontendDocumentBuilder::RemoveDependency(const FGuid& InClassID)
@@ -1562,15 +1666,15 @@ bool FMetaSoundFrontendDocumentBuilder::RemoveDependency(const FGuid& InClassID)
 		}
 
 		const int32 LastIndex = Dependencies.Num() - 1;
-		DocumentCache->OnRemovingDependency(Index);
+		DocumentDelegates->OnRemovingDependency.Broadcast(Index);
 		if (Index != LastIndex)
 		{
-			DocumentCache->OnRemovingDependency(LastIndex);
+			DocumentDelegates->OnRemovingDependency.Broadcast(LastIndex);
 		}
 		Dependencies.RemoveAtSwap(Index, 1, false);
 		if (Index != LastIndex)
 		{
-			DocumentCache->OnDependencyAdded(Index);
+			DocumentDelegates->OnDependencyAdded.Broadcast(Index);
 		}
 	}
 
@@ -1600,15 +1704,15 @@ bool FMetaSoundFrontendDocumentBuilder::RemoveDependency(EMetasoundFrontendClass
 		}
 
 		const int32 LastIndex = Dependencies.Num() - 1;
-		DocumentCache->OnRemovingDependency(Index);
+		DocumentDelegates->OnRemovingDependency.Broadcast(Index);
 		if (Index != LastIndex)
 		{
-			DocumentCache->OnRemovingDependency(LastIndex);
+			DocumentDelegates->OnRemovingDependency.Broadcast(LastIndex);
 		}
 		Dependencies.RemoveAtSwap(Index, 1, false);
 		if (Index != LastIndex)
 		{
-			DocumentCache->OnDependencyAdded(Index);
+			DocumentDelegates->OnDependencyAdded.Broadcast(Index);
 		}
 	}
 
@@ -1622,20 +1726,20 @@ bool FMetaSoundFrontendDocumentBuilder::RemoveEdge(const FMetasoundFrontendEdge&
 	FMetasoundFrontendDocument& Document = GetDocument();
 	FMetasoundFrontendGraph& Graph = Document.RootGraph.Graph;
 	TArray<FMetasoundFrontendEdge>& Edges = Graph.Edges;
-	IDocumentGraphEdgeCache& EdgeCache = DocumentCache->GetEdgeCache();
+	const IDocumentGraphEdgeCache& EdgeCache = DocumentCache->GetEdgeCache();
 	if (const int32* IndexPtr = EdgeCache.FindEdgeIndexToNodeInput(EdgeToRemove.ToNodeID, EdgeToRemove.ToVertexID))
 	{
 		const int32 Index = *IndexPtr;
 		const int32 LastIndex = Edges.Num() - 1;
-		EdgeCache.OnRemovingEdge(Index);
+		DocumentDelegates->EdgeDelegates.OnRemovingEdge.Broadcast(Index);
 		if (Index != LastIndex)
 		{
-			EdgeCache.OnRemovingEdge(LastIndex);
+			DocumentDelegates->EdgeDelegates.OnRemovingEdge.Broadcast(LastIndex);
 		}
 		Edges.RemoveAtSwap(Index, 1, false);
 		if (Index != LastIndex)
 		{
-			EdgeCache.OnEdgeAdded(Index);
+			DocumentDelegates->EdgeDelegates.OnEdgeAdded.Broadcast(Index);
 		}
 	}
 
@@ -1697,6 +1801,47 @@ bool FMetaSoundFrontendDocumentBuilder::RemoveNamedEdges(const TSet<Metasound::F
 	return bSuccess;
 }
 
+bool FMetaSoundFrontendDocumentBuilder::RemoveNodeInputDefault(const FGuid& InNodeID, const FGuid& InVertexID)
+{
+	using namespace Metasound::Frontend;
+
+	const IDocumentGraphNodeCache& NodeCache = DocumentCache->GetNodeCache();
+	if (const int32* NodeIndex = NodeCache.FindNodeIndex(InNodeID))
+	{
+		FMetasoundFrontendGraph& Graph = GetDocument().RootGraph.Graph;
+		FMetasoundFrontendNode& Node = Graph.Nodes[*NodeIndex];
+
+		auto IsVertex = [&InVertexID](const FMetasoundFrontendVertex& Vertex) { return Vertex.VertexID == InVertexID; };
+		const int32 VertexIndex = Node.Interface.Inputs.IndexOfByPredicate(IsVertex);
+		if (VertexIndex != INDEX_NONE)
+		{
+			auto IsLiteral = [&InVertexID](const FMetasoundFrontendVertexLiteral& Literal) { return Literal.VertexID == InVertexID; };
+			const int32 LiteralIndex = Node.InputLiterals.IndexOfByPredicate(IsLiteral);
+			if (LiteralIndex != INDEX_NONE)
+			{
+				const FOnMetaSoundFrontendDocumentMutateNodeInputLiteralArray& OnRemovingNodeInputLiteral = DocumentDelegates->NodeDelegates.OnRemovingNodeInputLiteral;
+				const int32 LastIndex = Node.InputLiterals.Num() - 1;
+				OnRemovingNodeInputLiteral.Broadcast(*NodeIndex, VertexIndex, LastIndex);
+				if (LiteralIndex != LastIndex)
+				{
+					OnRemovingNodeInputLiteral.Broadcast(*NodeIndex, VertexIndex, LiteralIndex);
+				}
+
+				constexpr bool bAllowShrinking = false;
+				Node.InputLiterals.RemoveAtSwap(LiteralIndex, 1, bAllowShrinking);
+				if (LiteralIndex != LastIndex)
+				{
+					const FOnMetaSoundFrontendDocumentMutateNodeInputLiteralArray& OnNodeInputLiteralSet = DocumentDelegates->NodeDelegates.OnNodeInputLiteralSet;
+					OnNodeInputLiteralSet.Broadcast(*NodeIndex, VertexIndex, LiteralIndex);
+				}
+				return true;
+			}
+		}
+	}
+
+	return false;
+}
+
 bool FMetaSoundFrontendDocumentBuilder::RemoveEdgesByNodeClassInterfaceBindings(const FGuid& InFromNodeID, const FGuid& InToNodeID)
 {
 	using namespace Metasound;
@@ -1720,25 +1865,25 @@ bool FMetaSoundFrontendDocumentBuilder::RemoveEdgesFromNodeOutput(const FGuid& I
 {
 	using namespace Metasound::Frontend;
 
-	FMetasoundFrontendDocument& Document = GetDocument();
-	FMetasoundFrontendGraph& Graph = Document.RootGraph.Graph;
-	IDocumentGraphEdgeCache& EdgeCache = DocumentCache->GetEdgeCache();
-
+	const IDocumentGraphEdgeCache& EdgeCache = DocumentCache->GetEdgeCache();
 	if (const TArray<int32>* Indices = EdgeCache.FindEdgeIndicesFromNodeOutput(InNodeID, InVertexID))
 	{
+		FMetasoundFrontendDocument& Document = GetDocument();
+		FMetasoundFrontendGraph& Graph = Document.RootGraph.Graph;
 		TArray<int32> IndicesCopy = *Indices; // Copy off indices as the array may be modified when notifying the cache in the loop below
 		for (int32 Index : IndicesCopy)
 		{
 			const int32 LastIndex = Graph.Edges.Num() - 1;
-			EdgeCache.OnRemovingEdge(Index);
+			DocumentDelegates->EdgeDelegates.OnRemovingEdge.Broadcast(Index);
 			if (Index != LastIndex)
 			{
-				EdgeCache.OnRemovingEdge(LastIndex);
+				DocumentDelegates->EdgeDelegates.OnRemovingEdge.Broadcast(LastIndex);
 			}
-			Graph.Edges.RemoveAtSwap(Index, 1, false);
+			constexpr bool bAllowShrinking = false;
+			Graph.Edges.RemoveAtSwap(Index, 1, bAllowShrinking);
 			if (Index != LastIndex)
 			{
-				EdgeCache.OnEdgeAdded(Index);
+				DocumentDelegates->EdgeDelegates.OnEdgeAdded.Broadcast(Index);
 			}
 		}
 
@@ -1752,22 +1897,22 @@ bool FMetaSoundFrontendDocumentBuilder::RemoveEdgeToNodeInput(const FGuid& InNod
 {
 	using namespace Metasound::Frontend;
 
-	FMetasoundFrontendDocument& Document = GetDocument();
-	IDocumentGraphEdgeCache& EdgeCache = DocumentCache->GetEdgeCache();
-	FMetasoundFrontendGraph& Graph = Document.RootGraph.Graph;
+	const IDocumentGraphEdgeCache& EdgeCache = DocumentCache->GetEdgeCache();
 	if (const int32* IndexPtr = EdgeCache.FindEdgeIndexToNodeInput(InNodeID, InVertexID))
 	{
+		FMetasoundFrontendDocument& Document = GetDocument();
+		FMetasoundFrontendGraph& Graph = Document.RootGraph.Graph;
 		const int32 Index = *IndexPtr; // Copy off indices as the pointer may be modified when notifying the cache below
 		const int32 LastIndex = Graph.Edges.Num() - 1;
-		EdgeCache.OnRemovingEdge(Index);
+		DocumentDelegates->EdgeDelegates.OnRemovingEdge.Broadcast(Index);
 		if (Index != LastIndex)
 		{
-			EdgeCache.OnRemovingEdge(LastIndex);
+			DocumentDelegates->EdgeDelegates.OnRemovingEdge.Broadcast(LastIndex);
 		}
 		Graph.Edges.RemoveAtSwap(Index, 1, false);
 		if (Index != LastIndex)
 		{
-			EdgeCache.OnEdgeAdded(Index);
+			DocumentDelegates->EdgeDelegates.OnEdgeAdded.Broadcast(Index);
 		}
 
 		return true;
@@ -1785,8 +1930,12 @@ bool FMetaSoundFrontendDocumentBuilder::RemoveGraphInput(FName InInputName)
 		{
 			TArray<FMetasoundFrontendClassInput>& Inputs = GetDocument().RootGraph.Interface.Inputs;
 			auto InputNameMatches = [InInputName](const FMetasoundFrontendClassInput& Input) { return Input.Name == InInputName; };
-			if (Inputs.RemoveAllSwap(InputNameMatches))
+			const int32 Index = Inputs.IndexOfByPredicate(InputNameMatches);
+			if (Index != INDEX_NONE)
 			{
+				DocumentDelegates->InterfaceDelegates.OnRemovingInput.Broadcast(Index);
+				constexpr bool bAllowShrinking = false;
+				Inputs.RemoveAtSwap(Index, 1, bAllowShrinking);
 				if (IsDependencyReferenced(ClassID))
 				{
 					return true;
@@ -1814,8 +1963,12 @@ bool FMetaSoundFrontendDocumentBuilder::RemoveGraphOutput(FName InOutputName)
 		{
 			TArray<FMetasoundFrontendClassOutput>& Outputs = GetDocument().RootGraph.Interface.Outputs;
 			auto OutputNameMatches = [InOutputName](const FMetasoundFrontendClassOutput& Output) { return Output.Name == InOutputName; };
-			if (Outputs.RemoveAllSwap(OutputNameMatches))
+			const int32 Index = Outputs.IndexOfByPredicate(OutputNameMatches);
+			if (Index != INDEX_NONE)
 			{
+				DocumentDelegates->InterfaceDelegates.OnRemovingOutput.Broadcast(Index);
+				constexpr bool bAllowShrinking = false;
+				Outputs.RemoveAtSwap(Index, 1, bAllowShrinking);
 				if (IsDependencyReferenced(ClassID))
 				{
 					return true;
@@ -1857,17 +2010,11 @@ bool FMetaSoundFrontendDocumentBuilder::RemoveInterface(FName InterfaceName)
 				UE_LOG(LogMetaSound, Error, TEXT("DocumentBuilder failed to remove MetaSound Interface '%s' to document: is not set to be modifiable for given UClass '%s'"), *InterfaceName.ToString(), *BuilderClassPath.ToString());
 				return false;
 			}
-		}
 
-		const FModifyRootGraphInterfaces ModifyInterfaces({ { Interface.Version }, { } });
-		if (ModifyInterfaces.Transform(GetDocument()))
-		{
-			// TODO: Add ability to update local builder caches dynamically
-			// in ModifyRootGraphInterfaces transform once it moves away
-			// from using controllers internally. Then cache updates can
-			// be more isolated and not use this full rebuild call.
-			ReloadCache();
-			return true;
+			TArray<FMetasoundFrontendInterface> InterfacesToRemove;
+			InterfacesToRemove.Add(Entry->GetInterface());
+			FModifyInterfaceOptions Options(MoveTemp(InterfacesToRemove), { });
+			return ModifyInterfaces(MoveTemp(Options));
 		}
 	}
 
@@ -1880,14 +2027,15 @@ bool FMetaSoundFrontendDocumentBuilder::RemoveNode(const FGuid& InNodeID)
 
 	using namespace Metasound::Frontend;
 
-	FMetasoundFrontendDocument& Document = GetDocument();
-	IDocumentGraphNodeCache& NodeCache = DocumentCache->GetNodeCache();
+	const IDocumentGraphNodeCache& NodeCache = DocumentCache->GetNodeCache();
 
-	FMetasoundFrontendGraph& Graph = Document.RootGraph.Graph;
-	TArray<FMetasoundFrontendNode>& Nodes = Graph.Nodes;
 	if (const int32* IndexPtr = NodeCache.FindNodeIndex(InNodeID))
 	{
 		const int32 Index = *IndexPtr; // Copy off indices as the pointer may be modified when notifying the cache below
+
+		FMetasoundFrontendDocument& Document = GetDocument();
+		FMetasoundFrontendGraph& Graph = Document.RootGraph.Graph;
+		TArray<FMetasoundFrontendNode>& Nodes = Graph.Nodes;
 		const FMetasoundFrontendNode& Node = Nodes[Index];
 
 		for (const FMetasoundFrontendVertex& Vertex : Node.Interface.Inputs)
@@ -1900,17 +2048,17 @@ bool FMetaSoundFrontendDocumentBuilder::RemoveNode(const FGuid& InNodeID)
 			RemoveEdgesFromNodeOutput(InNodeID, Vertex.VertexID);
 		}
 
-		NodeCache.OnRemovingNode(Index);
+		DocumentDelegates->NodeDelegates.OnRemovingNode.Broadcast(Index);
 
 		const int32 LastIndex = Nodes.Num() - 1;
 		if (Index != LastIndex)
 		{
-			NodeCache.OnRemovingNode(LastIndex);
+			DocumentDelegates->NodeDelegates.OnRemovingNode.Broadcast(LastIndex);
 		}
 		Nodes.RemoveAtSwap(Index, 1, false);
 		if (Index != LastIndex)
 		{
-			NodeCache.OnNodeAdded(Index);
+			DocumentDelegates->NodeDelegates.OnNodeAdded.Broadcast(Index);
 		}
 
 #if WITH_EDITORONLY_DATA
@@ -1946,7 +2094,7 @@ bool FMetaSoundFrontendDocumentBuilder::SetGraphInputDefault(FName InputName, co
 			// Set the input as no longer inheriting default for presets
 			if (IsPreset())
 			{
-				return SetInputInheritsDefault(InputName, /*bInputInheritsDefault=*/false);
+				return SetGraphInputInheritsDefault(InputName, /*bInputInheritsDefault=*/false);
 			}
 			return true;
 		}
@@ -1956,33 +2104,59 @@ bool FMetaSoundFrontendDocumentBuilder::SetGraphInputDefault(FName InputName, co
 	return false;
 }
 
+bool FMetaSoundFrontendDocumentBuilder::SetGraphInputInheritsDefault(FName InName, bool bInputInheritsDefault)
+{
+	FMetasoundFrontendGraphClassPresetOptions& PresetOptions = GetDocument().RootGraph.PresetOptions;
+	if (bInputInheritsDefault)
+	{
+		if (PresetOptions.bIsPreset)
+		{
+			return PresetOptions.InputsInheritingDefault.Add(InName).IsValidId();
+		}
+	}
+	else
+	{
+		if (PresetOptions.bIsPreset)
+		{
+			return PresetOptions.InputsInheritingDefault.Remove(InName) > 0;
+		}
+	}
+
+	return false;
+}
+
 bool FMetaSoundFrontendDocumentBuilder::SetNodeInputDefault(const FGuid& InNodeID, const FGuid& InVertexID, const FMetasoundFrontendLiteral& InLiteral)
 {
 	using namespace Metasound::Frontend;
 
-	if (const int32* Index = DocumentCache->GetNodeCache().FindNodeIndex(InNodeID))
+	const IDocumentGraphNodeCache& NodeCache = DocumentCache->GetNodeCache();
+	if (const int32* NodeIndex = NodeCache.FindNodeIndex(InNodeID))
 	{
 		FMetasoundFrontendGraph& Graph = GetDocument().RootGraph.Graph;
-		FMetasoundFrontendNode& Node = Graph.Nodes[*Index];
-		TArray<FMetasoundFrontendVertex>& Inputs = Node.Interface.Inputs;
+		FMetasoundFrontendNode& Node = Graph.Nodes[*NodeIndex];
 
 		auto IsVertex = [&InVertexID](const FMetasoundFrontendVertex& Vertex) { return Vertex.VertexID == InVertexID; };
-		if (Inputs.ContainsByPredicate(IsVertex))
+		int32 VertexIndex = Node.Interface.Inputs.IndexOfByPredicate(IsVertex);
+		if (VertexIndex != INDEX_NONE)
 		{
 			FMetasoundFrontendVertexLiteral NewVertexLiteral;
 			NewVertexLiteral.VertexID = InVertexID;
 			NewVertexLiteral.Value = InLiteral;
 
 			auto IsLiteral = [&InVertexID](const FMetasoundFrontendVertexLiteral& Literal) { return Literal.VertexID == InVertexID; };
-			if (FMetasoundFrontendVertexLiteral* InputLiteral = Node.InputLiterals.FindByPredicate(IsLiteral))
+			int32 LiteralIndex = Node.InputLiterals.IndexOfByPredicate(IsLiteral);
+			if (LiteralIndex == INDEX_NONE)
 			{
-				*InputLiteral = MoveTemp(NewVertexLiteral);
+				LiteralIndex = Node.InputLiterals.Num();
+				Node.InputLiterals.Add(MoveTemp(NewVertexLiteral));
 			}
 			else
 			{
-				Node.InputLiterals.Add(MoveTemp(NewVertexLiteral));
+				Node.InputLiterals[LiteralIndex] = MoveTemp(NewVertexLiteral);
 			}
 
+			const FOnMetaSoundFrontendDocumentMutateNodeInputLiteralArray& OnNodeInputLiteralSet = DocumentDelegates->NodeDelegates.OnNodeInputLiteralSet;
+			OnNodeInputLiteralSet.Broadcast(*NodeIndex, VertexIndex, LiteralIndex);
 			return true;
 		}
 	}
@@ -1990,28 +2164,6 @@ bool FMetaSoundFrontendDocumentBuilder::SetNodeInputDefault(const FGuid& InNodeI
 	return false;
 }
 
-
-bool FMetaSoundFrontendDocumentBuilder::SetInputInheritsDefault(FName InName, bool bInputInheritsDefault)
-{
-	FMetasoundFrontendGraphClassPresetOptions& PresetOptions = GetDocument().RootGraph.PresetOptions;
-	if (bInputInheritsDefault)
-	{
-		if (PresetOptions.bIsPreset)
-		{
-			return PresetOptions.InputsInheritingDefault.Add(InName).IsValidId();
-		}
-	}
-	else
-	{
-		if (PresetOptions.bIsPreset)
-		{
-			return PresetOptions.InputsInheritingDefault.Remove(InName) > 0;
-		}
-	}
-
-	return false;
-}
-
 bool FMetaSoundFrontendDocumentBuilder::SwapGraphInput(const FMetasoundFrontendClassVertex& InExistingInputVertex, const FMetasoundFrontendClassVertex& InNewInputVertex)
 {
 	using namespace Metasound::Frontend;
@@ -2026,24 +2178,28 @@ bool FMetaSoundFrontendDocumentBuilder::SwapGraphInput(const FMetasoundFrontendC
 	}
 
 	const IDocumentGraphNodeCache& NodeCache = DocumentCache->GetNodeCache();
+	const IDocumentGraphInterfaceCache& InterfaceCache = DocumentCache->GetInterfaceCache();
 
 	// 2. Gather data from existing member/node needed to swap
 	TArray<FMetasoundFrontendEdge> RemovedEdges;
 
-	const FMetasoundFrontendNode* ExistingInputNode = NodeCache.FindInputNode(InExistingInputVertex.Name);
-	checkf(ExistingInputNode, TEXT("'SwapGraphInput' failed to find original input node"));
-	const FGuid NodeID = ExistingInputNode->GetID();
+	const FMetasoundFrontendClassInput* ExistingInputClass = InterfaceCache.FindInput(InExistingInputVertex.Name);
+	checkf(ExistingInputClass, TEXT("'SwapGraphInput' failed to find original graph input"));
+	const FGuid NodeID = ExistingInputClass->NodeID;
 
 #if WITH_EDITOR
 	TMap<FGuid, FVector2D> Locations;
 #endif // WITH_EDITOR
 	{
+		const FMetasoundFrontendNode* ExistingInputNode = NodeCache.FindNode(NodeID);
+		check(ExistingInputNode);
+
 #if WITH_EDITOR
 		Locations = ExistingInputNode->Style.Display.Locations;
 #endif // WITH_EDITOR
 
 		const FGuid VertexID = ExistingInputNode->Interface.Outputs.Last().VertexID;
-		TArray<const FMetasoundFrontendEdge*> Edges = DocumentCache->GetEdgeCache().FindEdges(ExistingInputNode->GetID(), VertexID);
+		TArray<const FMetasoundFrontendEdge*> Edges = DocumentCache->GetEdgeCache().FindEdges(NodeID, VertexID);
 		Algo::Transform(Edges, RemovedEdges, [](const FMetasoundFrontendEdge* Edge) { return *Edge; });
 	}
 
@@ -2100,16 +2256,22 @@ bool FMetaSoundFrontendDocumentBuilder::SwapGraphOutput(const FMetasoundFrontend
 	}
 
 	const IDocumentGraphNodeCache& NodeCache = DocumentCache->GetNodeCache();
-	const FMetasoundFrontendNode* ExistingOutputNode = NodeCache.FindOutputNode(InExistingOutputVertex.Name);
-	checkf(ExistingOutputNode, TEXT("'SwapGraphOutput' failed to find original output node"));
-	const FGuid NodeID = ExistingOutputNode->GetID();
+	const IDocumentGraphInterfaceCache& InterfaceCache = DocumentCache->GetInterfaceCache();
 
 	// 2. Gather data from existing member/node needed to swap
 	TArray<FMetasoundFrontendEdge> RemovedEdges;
+
+	const FMetasoundFrontendClassOutput* ExistingOutputClass = InterfaceCache.FindOutput(InExistingOutputVertex.Name);
+	checkf(ExistingOutputClass, TEXT("'SwapGraphOutput' failed to find original graph output"));
+	const FGuid NodeID = ExistingOutputClass->NodeID;
+
 #if WITH_EDITOR
 	TMap<FGuid, FVector2D> Locations;
 #endif // WITH_EDITOR
 	{
+		const FMetasoundFrontendNode* ExistingOutputNode = NodeCache.FindNode(NodeID);
+		check(ExistingOutputNode);
+
 #if WITH_EDITOR
 		Locations = ExistingOutputNode->Style.Display.Locations;
 #endif // WITH_EDITOR
diff --git a/Engine/Plugins/Runtime/Metasound/Source/MetasoundFrontend/Private/MetasoundFrontendDocumentCache.cpp b/Engine/Plugins/Runtime/Metasound/Source/MetasoundFrontend/Private/MetasoundFrontendDocumentCache.cpp
--- a/Engine/Plugins/Runtime/Metasound/Source/MetasoundFrontend/Private/MetasoundFrontendDocumentCache.cpp
+++ b/Engine/Plugins/Runtime/Metasound/Source/MetasoundFrontend/Private/MetasoundFrontendDocumentCache.cpp
@@ -13,7 +13,16 @@ namespace Metasound::Frontend
 	FDocumentCache::FDocumentCache(const FMetasoundFrontendDocument& InDocument)
 		: Document(&InDocument)
 	{
-		METASOUND_TRACE_CPUPROFILER_EVENT_SCOPE(FDocumentCache::FDocumentCache);
+	}
+
+	void FDocumentCache::Init(FDocumentModifyDelegates& OutDelegates)
+	{
+		METASOUND_TRACE_CPUPROFILER_EVENT_SCOPE(FDocumentCache::Init);
+
+		check(Document);
+
+		IDToIndex.Reset();
+		KeyToIndex.Reset();
 
 		for (int32 Index = 0; Index < Document->Dependencies.Num(); ++Index)
 		{
@@ -24,9 +33,26 @@ namespace Metasound::Frontend
 			IDToIndex.Add(Class.ID, Index);
 		}
 
-		// Dependencies must be cached first as the NodeCache is dependent on them
-		NodeCache = MakeUnique<FDocumentGraphNodeCache>(InDocument, *this);
-		EdgeCache = MakeUnique<FDocumentGraphEdgeCache>(InDocument);
+		OutDelegates.OnDependencyAdded.AddSP(this, &FDocumentCache::OnDependencyAdded);
+		OutDelegates.OnRemovingDependency.AddSP(this, &FDocumentCache::OnRemovingDependency);
+
+		if (!EdgeCache.IsValid())
+		{
+			EdgeCache = MakeShared<FDocumentGraphEdgeCache>(AsShared());
+		}
+		EdgeCache->Init(OutDelegates.EdgeDelegates);
+
+		if (!NodeCache.IsValid())
+		{
+			NodeCache = MakeShared<FDocumentGraphNodeCache>(AsShared());
+		}
+		NodeCache->Init(OutDelegates.NodeDelegates);
+
+		if (!InterfaceCache.IsValid())
+		{
+			InterfaceCache = MakeShared<FDocumentGraphInterfaceCache>(AsShared());
+		}
+		InterfaceCache->Init(OutDelegates.InterfaceDelegates);
 	}
 
 	const FMetasoundFrontendDocument& FDocumentCache::GetDocument() const
@@ -70,6 +96,21 @@ namespace Metasound::Frontend
 		return IDToIndex.Find(InClassID);
 	}
 
+	const IDocumentGraphEdgeCache& FDocumentCache::GetEdgeCache() const
+	{
+		return *EdgeCache;
+	}
+
+	const IDocumentGraphNodeCache& FDocumentCache::GetNodeCache() const
+	{
+		return *NodeCache;
+	}
+
+	const IDocumentGraphInterfaceCache& FDocumentCache::GetInterfaceCache() const
+	{
+		return *InterfaceCache;
+	}
+
 	void FDocumentCache::OnDependencyAdded(int32 NewIndex)
 	{
 		using namespace Metasound::Frontend;
@@ -90,61 +131,119 @@ namespace Metasound::Frontend
 		IDToIndex.Remove(DependencyBeingRemoved.ID);
 	}
 
-	FDocumentGraphNodeCache::FDocumentGraphNodeCache(const FMetasoundFrontendDocument& InDocument, const FDocumentCache& InParent)
-		: Document(&InDocument)
-		, Parent(&InParent)
+	FDocumentGraphInterfaceCache::FDocumentGraphInterfaceCache(TSharedRef<IDocumentCache> ParentCache)
+		: Parent(ParentCache)
 	{
-		const FMetasoundFrontendGraph& Graph = GetDocument().RootGraph.Graph;
-		const TArray<FMetasoundFrontendNode>& Nodes = Graph.Nodes;
-		for (int32 Index = 0; Index < Nodes.Num(); ++Index)
+	}
+
+	const FMetasoundFrontendClassInput* FDocumentGraphInterfaceCache::FindInput(FName InputName) const
+	{
+		if (const int32* Index = InputNameToIndex.Find(InputName))
 		{
-			const FMetasoundFrontendNode& Node = Nodes[Index];
-			const FMetasoundFrontendClass* Class = InParent.FindDependency(Node.ClassID);
-			check(Class);
+			const FMetasoundFrontendDocument& Doc = Parent->GetDocument();
+			return &Doc.RootGraph.Interface.Inputs[*Index];
+		}
 
-			const FGuid& NodeID = Node.GetID();
-			IDToIndex.Add(NodeID, Index);
-			ClassIDToNodeIndices.FindOrAdd(Node.ClassID).Add(Index);
+		return nullptr;
+	}
 
-			switch (Class->Metadata.GetType())
-			{
-				case EMetasoundFrontendClassType::Output:
-				{
-					OutputNameToIndex.Add(Node.Name, Index);
-				}
-				break;
-
-				case EMetasoundFrontendClassType::Input:
-				{
-					InputNameToIndex.Add(Node.Name, Index);
-				}
-				break;
-
-				default:
-				break;
-			}
+	const FMetasoundFrontendClassOutput* FDocumentGraphInterfaceCache::FindOutput(FName OutputName) const
+	{
+		if (const int32* Index = OutputNameToIndex.Find(OutputName))
+		{
+			const FMetasoundFrontendDocument& Doc = Parent->GetDocument();
+			return &Doc.RootGraph.Interface.Outputs[*Index];
 		}
+
+		return nullptr;
 	}
 
-	const FMetasoundFrontendDocument& FDocumentGraphNodeCache::GetDocument() const
+	void FDocumentGraphInterfaceCache::Init(FInterfaceModifyDelegates& OutDelegates)
 	{
-		check(Document);
-		return *Document;
+		InputNameToIndex.Reset();
+		OutputNameToIndex.Reset();
+
+		const FMetasoundFrontendGraphClass& GraphClass = Parent->GetDocument().RootGraph;
+		const TArray<FMetasoundFrontendClassInput>& Inputs = GraphClass.Interface.Inputs;
+		for (int32 Index = 0; Index < Inputs.Num(); ++Index)
+		{
+			InputNameToIndex.Add(Inputs[Index].Name, Index);
+		};
+
+		const TArray<FMetasoundFrontendClassOutput>& Outputs = GraphClass.Interface.Outputs;
+		for (int32 Index = 0; Index < Outputs.Num(); ++Index)
+		{
+			OutputNameToIndex.Add(Outputs[Index].Name, Index);
+		};
+
+		OutDelegates.OnInputAdded.AddSP(this, &FDocumentGraphInterfaceCache::OnInputAdded);
+		OutDelegates.OnOutputAdded.AddSP(this, &FDocumentGraphInterfaceCache::OnOutputAdded);
+		OutDelegates.OnRemovingInput.AddSP(this, &FDocumentGraphInterfaceCache::OnRemovingInput);
+		OutDelegates.OnRemovingOutput.AddSP(this, &FDocumentGraphInterfaceCache::OnRemovingOutput);
 	}
 
-	bool FDocumentGraphNodeCache::ContainsNode(const FGuid& InNodeID) const
+	void FDocumentGraphInterfaceCache::OnInputAdded(int32 NewIndex)
 	{
-		return IDToIndex.Contains(InNodeID);
+		const FMetasoundFrontendDocument& Document = Parent->GetDocument();
+		const FMetasoundFrontendGraphClass& GraphClass = Document.RootGraph;
+		const FMetasoundFrontendClassInput& Input = GraphClass.Interface.Inputs[NewIndex];
+		InputNameToIndex.Add(Input.Name, NewIndex);
+	}
+
+	void FDocumentGraphInterfaceCache::OnOutputAdded(int32 NewIndex)
+	{
+		const FMetasoundFrontendDocument& Document = Parent->GetDocument();
+		const FMetasoundFrontendGraphClass& GraphClass = Document.RootGraph;
+		const FMetasoundFrontendClassOutput& Output = GraphClass.Interface.Outputs[NewIndex];
+		OutputNameToIndex.Add(Output.Name, NewIndex);
+	}
+
+	void FDocumentGraphInterfaceCache::OnRemovingInput(int32 IndexBeingRemoved)
+	{
+		const FMetasoundFrontendDocument& Document = Parent->GetDocument();
+		const FMetasoundFrontendGraphClass& GraphClass = Document.RootGraph;
+		const FMetasoundFrontendClassInput& Input = GraphClass.Interface.Inputs[IndexBeingRemoved];
+		InputNameToIndex.Remove(Input.Name);
+	}
+
+	void FDocumentGraphInterfaceCache::OnRemovingOutput(int32 IndexBeingRemoved)
+	{
+		const FMetasoundFrontendDocument& Document = Parent->GetDocument();
+		const FMetasoundFrontendGraphClass& GraphClass = Document.RootGraph;
+		const FMetasoundFrontendClassOutput& Output = GraphClass.Interface.Outputs[IndexBeingRemoved];
+		OutputNameToIndex.Remove(Output.Name);
 	}
 
-	const int32* FDocumentGraphNodeCache::FindInputNodeIndex(FName InputName) const
+	FDocumentGraphNodeCache::FDocumentGraphNodeCache(TSharedRef<IDocumentCache> ParentCache)
+		: Parent(ParentCache)
 	{
-		return InputNameToIndex.Find(InputName);
 	}
 
-	const int32* FDocumentGraphNodeCache::FindOutputNodeIndex(FName OutputName) const
+	void FDocumentGraphNodeCache::Init(FNodeModifyDelegates& OutDelegates)
 	{
-		return OutputNameToIndex.Find(OutputName);
+		IDToIndex.Reset();
+		ClassIDToNodeIndices.Reset();
+
+		const FMetasoundFrontendGraph& Graph = Parent->GetDocument().RootGraph.Graph;
+		const TArray<FMetasoundFrontendNode>& Nodes = Graph.Nodes;
+		for (int32 Index = 0; Index < Nodes.Num(); ++Index)
+		{
+			const FMetasoundFrontendNode& Node = Nodes[Index];
+			const FMetasoundFrontendClass* Class = Parent->FindDependency(Node.ClassID);
+			check(Class);
+
+			const FGuid& NodeID = Node.GetID();
+			IDToIndex.Add(NodeID, Index);
+			ClassIDToNodeIndices.FindOrAdd(Node.ClassID).Add(Index);
+		}
+
+		OutDelegates.OnNodeAdded.AddSP(this, &FDocumentGraphNodeCache::OnNodeAdded);
+		OutDelegates.OnRemovingNode.AddSP(this, &FDocumentGraphNodeCache::OnRemovingNode);
+	}
+
+	bool FDocumentGraphNodeCache::ContainsNode(const FGuid& InNodeID) const
+	{
+		return IDToIndex.Contains(InNodeID);
 	}
 
 	const int32* FDocumentGraphNodeCache::FindNodeIndex(const FGuid& InNodeID) const
@@ -157,7 +256,8 @@ namespace Metasound::Frontend
 		TArray<const FMetasoundFrontendNode*> Nodes;
 		if (const TArray<int32>* NodeIndices = ClassIDToNodeIndices.Find(InClassID))
 		{
-			const FMetasoundFrontendGraph& Graph = GetDocument().RootGraph.Graph;
+			const FMetasoundFrontendDocument& Document = Parent->GetDocument();
+			const FMetasoundFrontendGraph& Graph = Document.RootGraph.Graph;
 			Algo::Transform(*NodeIndices, Nodes, [&Graph](const int32& Index) { return &Graph.Nodes[Index]; });
 		}
 		return Nodes;
@@ -167,7 +267,9 @@ namespace Metasound::Frontend
 	{
 		if (const int32* NodeIndex = IDToIndex.Find(InNodeID))
 		{
-			return &GetDocument().RootGraph.Graph.Nodes[*NodeIndex];
+			const FMetasoundFrontendDocument& Document = Parent->GetDocument();
+			const FMetasoundFrontendGraph& Graph = Document.RootGraph.Graph;
+			return &Graph.Nodes[*NodeIndex];
 		}
 
 		return nullptr;
@@ -209,26 +311,6 @@ namespace Metasound::Frontend
 		return FoundVertices;
 	}
 
-	const FMetasoundFrontendNode* FDocumentGraphNodeCache::FindOutputNode(FName OutputName) const
-	{
-		if (const int32* Index = OutputNameToIndex.Find(OutputName))
-		{
-			return &GetDocument().RootGraph.Graph.Nodes[*Index];
-		}
-
-		return nullptr;
-	}
-
-	const FMetasoundFrontendNode* FDocumentGraphNodeCache::FindInputNode(FName InputName) const
-	{
-		if (const int32* Index = InputNameToIndex.Find(InputName))
-		{
-			return &GetDocument().RootGraph.Graph.Nodes[*Index];
-		}
-
-		return nullptr;
-	}
-
 	const FMetasoundFrontendVertex* FDocumentGraphNodeCache::FindInputVertex(const FGuid& InNodeID, const FGuid& InVertexID) const
 	{
 		if (const FMetasoundFrontendNode* Node = FindNode(InNodeID))
@@ -287,31 +369,11 @@ namespace Metasound::Frontend
 
 	void FDocumentGraphNodeCache::OnNodeAdded(int32 InNewIndex)
 	{
-		const FMetasoundFrontendNode& Node = GetDocument().RootGraph.Graph.Nodes[InNewIndex];
+		const FMetasoundFrontendDocument& Document = Parent->GetDocument();
+		const FMetasoundFrontendGraph& Graph = Document.RootGraph.Graph;
+		const FMetasoundFrontendNode& Node = Graph.Nodes[InNewIndex];
 		IDToIndex.Add(Node.GetID(), InNewIndex);
 		ClassIDToNodeIndices.FindOrAdd(Node.ClassID).Add(InNewIndex);
-
-		check(Parent);
-		const FMetasoundFrontendClass* Dependency = Parent->FindDependency(Node.ClassID);
-		check(Dependency);
-
-		switch (Dependency->Metadata.GetType())
-		{
-			case EMetasoundFrontendClassType::Input:
-			{
-				InputNameToIndex.Add(Node.Name, InNewIndex);
-			}
-			break;
-
-			case EMetasoundFrontendClassType::Output:
-			{
-				OutputNameToIndex.Add(Node.Name, InNewIndex);
-			}
-			break;
-
-			default:
-			break;
-		}
 	}
 
 	void FDocumentGraphNodeCache::OnRemovingNode(int32 IndexBeingRemoved)
@@ -319,7 +381,8 @@ namespace Metasound::Frontend
 		using namespace Metasound::Frontend;
 
 		// If not found, transaction was missed
-		const FMetasoundFrontendNode& NodeBeingRemoved = GetDocument().RootGraph.Graph.Nodes[IndexBeingRemoved];
+		const FMetasoundFrontendDocument& Document = Parent->GetDocument();
+		const FMetasoundFrontendNode& NodeBeingRemoved = Document.RootGraph.Graph.Nodes[IndexBeingRemoved];
 
 		IDToIndex.Remove(NodeBeingRemoved.GetID());
 		TArray<int32>& NodeIndices = ClassIDToNodeIndices.FindChecked(NodeBeingRemoved.ClassID);
@@ -328,33 +391,12 @@ namespace Metasound::Frontend
 		{
 			ClassIDToNodeIndices.Remove(NodeBeingRemoved.ClassID);
 		}
-
-		check(Parent);
-		const FMetasoundFrontendClass* Dependency = Parent->FindDependency(NodeBeingRemoved.ClassID);
-		check(Dependency);
-		switch (Dependency->Metadata.GetType())
-		{
-			case EMetasoundFrontendClassType::Input:
-			{
-				InputNameToIndex.Remove(NodeBeingRemoved.Name);
-			}
-			break;
-
-			case EMetasoundFrontendClassType::Output:
-			{
-				OutputNameToIndex.Remove(NodeBeingRemoved.Name);
-			}
-			break;
-
-			default:
-			break;
-		}
 	}
 
-	FDocumentGraphEdgeCache::FDocumentGraphEdgeCache(const FMetasoundFrontendDocument& InDocument)
-		: Document(&InDocument)
+	FDocumentGraphEdgeCache::FDocumentGraphEdgeCache(TSharedRef<IDocumentCache> ParentCache)
+		: Parent(ParentCache)
 	{
-		const TArray<FMetasoundFrontendEdge>& Edges = GetDocument().RootGraph.Graph.Edges;
+		const TArray<FMetasoundFrontendEdge>& Edges = Parent->GetDocument().RootGraph.Graph.Edges;
 		for (int32 Index = 0; Index < Edges.Num(); ++Index)
 		{
 			const FMetasoundFrontendEdge& Edge = Edges[Index];
@@ -363,12 +405,22 @@ namespace Metasound::Frontend
 		}
 	}
 
+	void FDocumentGraphEdgeCache::Init(FEdgeModifyDelegates& OutDelegates)
+	{
+		OutputToEdgeIndices.Reset();
+		InputToEdgeIndex.Reset();
+
+		OutDelegates.OnEdgeAdded.AddSP(this, &FDocumentGraphEdgeCache::OnEdgeAdded);
+		OutDelegates.OnRemovingEdge.AddSP(this, &FDocumentGraphEdgeCache::OnRemovingEdge);
+	}
+
 	bool FDocumentGraphEdgeCache::ContainsEdge(const FMetasoundFrontendEdge& InEdge) const
 	{
 		const FMetasoundFrontendVertexHandle InputPair = { InEdge.ToNodeID, InEdge.ToVertexID };
 		if (const int32* Index = InputToEdgeIndex.Find(InputPair))
 		{
-			const FMetasoundFrontendEdge& Edge = GetDocument().RootGraph.Graph.Edges[*Index];
+			const FMetasoundFrontendDocument& Document = Parent->GetDocument();
+			const FMetasoundFrontendEdge& Edge = Document.RootGraph.Graph.Edges[*Index];
 			return Edge.FromNodeID == InEdge.FromNodeID && Edge.FromVertexID == InEdge.FromVertexID;
 		}
 
@@ -379,7 +431,9 @@ namespace Metasound::Frontend
 	{
 		TArray<const FMetasoundFrontendEdge*> Edges;
 
-		const FMetasoundFrontendGraph& Graph = GetDocument().RootGraph.Graph;
+		const FMetasoundFrontendDocument& Document = Parent->GetDocument();
+		const FMetasoundFrontendGraph& Graph = Document.RootGraph.Graph;
+
 		const FMetasoundFrontendVertexHandle Handle { InNodeID, InVertexID };
 		if (const int32* Index = InputToEdgeIndex.Find(Handle))
 		{
@@ -404,12 +458,6 @@ namespace Metasound::Frontend
 		return OutputToEdgeIndices.Find({ InNodeID, InVertexID });
 	}
 
-	const FMetasoundFrontendDocument& FDocumentGraphEdgeCache::GetDocument() const
-	{
-		check(Document);
-		return *Document;
-	}
-
 	bool FDocumentGraphEdgeCache::IsNodeInputConnected(const FGuid& InNodeID, const FGuid& InVertexID) const
 	{
 		return InputToEdgeIndex.Contains(FMetasoundFrontendVertexHandle { InNodeID, InVertexID });
@@ -422,14 +470,18 @@ namespace Metasound::Frontend
 
 	void FDocumentGraphEdgeCache::OnEdgeAdded(int32 InNewIndex)
 	{
-		const FMetasoundFrontendEdge& NewEdge = GetDocument().RootGraph.Graph.Edges[InNewIndex];
+		const FMetasoundFrontendDocument& Document = Parent->GetDocument();
+		const FMetasoundFrontendGraph& Graph = Document.RootGraph.Graph;
+		const FMetasoundFrontendEdge& NewEdge = Graph.Edges[InNewIndex];
 		InputToEdgeIndex.Add(FMetasoundFrontendVertexHandle { NewEdge.ToNodeID, NewEdge.ToVertexID }, InNewIndex);
 		OutputToEdgeIndices.FindOrAdd({ NewEdge.FromNodeID, NewEdge.FromVertexID }).Add(InNewIndex);
 	}
 
 	void FDocumentGraphEdgeCache::OnRemovingEdge(int32 IndexToRemove)
 	{
-		const FMetasoundFrontendEdge& EdgeToRemove = GetDocument().RootGraph.Graph.Edges[IndexToRemove];
+		const FMetasoundFrontendDocument& Document = Parent->GetDocument();
+		const FMetasoundFrontendGraph& Graph = Document.RootGraph.Graph;
+		const FMetasoundFrontendEdge& EdgeToRemove = Graph.Edges[IndexToRemove];
 		InputToEdgeIndex.Remove(FMetasoundFrontendVertexHandle { EdgeToRemove.ToNodeID, EdgeToRemove.ToVertexID });
 		TArray<int32>& Indices = OutputToEdgeIndices.FindChecked(EdgeToRemove.GetFromVertexHandle());
 		Indices.RemoveAllSwap([&IndexToRemove](const int32& Index) { return Index == IndexToRemove; }, false /* bAllowShrinking */);
diff --git a/Engine/Plugins/Runtime/Metasound/Source/MetasoundFrontend/Private/MetasoundFrontendDocumentCache.h b/Engine/Plugins/Runtime/Metasound/Source/MetasoundFrontend/Private/MetasoundFrontendDocumentCache.h
--- a/Engine/Plugins/Runtime/Metasound/Source/MetasoundFrontend/Private/MetasoundFrontendDocumentCache.h
+++ b/Engine/Plugins/Runtime/Metasound/Source/MetasoundFrontend/Private/MetasoundFrontendDocumentCache.h
@@ -20,7 +20,7 @@ namespace Metasound::Frontend
 	{
 	public:
 		FDocumentGraphEdgeCache() = default;
-		FDocumentGraphEdgeCache(const FMetasoundFrontendDocument& InDocument);
+		FDocumentGraphEdgeCache(TSharedRef<IDocumentCache> ParentCache);
 		virtual ~FDocumentGraphEdgeCache() = default;
 
 		// IDocumentGraphEdgeCache implementation
@@ -33,15 +33,11 @@ namespace Metasound::Frontend
 		virtual const int32* FindEdgeIndexToNodeInput(const FGuid& InNodeID, const FGuid& InVertexID) const override;
 		virtual const TArray<int32>* FindEdgeIndicesFromNodeOutput(const FGuid& InNodeID, const FGuid& InVertexID) const override;
 
-	protected:
-		virtual void OnEdgeAdded(int32 InNewIndex) override;
-		virtual void OnRemovingEdge(int32 IndexBeingRemoved) override;
-		virtual const FMetasoundFrontendDocument& GetDocument() const override;
-
-		// Friended to ensure only builder is the one managing cache manipulation via add/remove calls above
-		friend struct ::FMetaSoundFrontendDocumentBuilder;
+		void Init(FEdgeModifyDelegates& OutDelegates);
 
 	private:
+		void OnEdgeAdded(int32 InNewIndex);
+		void OnRemovingEdge(int32 IndexBeingRemoved);
 
 		// Cache of outputs NodeId/VertexId pairs to associated edge indices
 		TMap<FMetasoundFrontendVertexHandle, TArray<int32>> OutputToEdgeIndices;
@@ -49,27 +45,24 @@ namespace Metasound::Frontend
 		// Cache of Input NodeId/VertexId pairs to associated edge indices
 		TMap<FMetasoundFrontendVertexHandle, int32> InputToEdgeIndex;
 
-		const FMetasoundFrontendDocument* Document = nullptr;
+		TSharedPtr<IDocumentCache> Parent;
 	};
 
+
 	class FDocumentGraphNodeCache : public IDocumentGraphNodeCache
 	{
 	public:
 		FDocumentGraphNodeCache() = default;
-		FDocumentGraphNodeCache(const FMetasoundFrontendDocument& InDocument, const FDocumentCache& InParent);
+		FDocumentGraphNodeCache(TSharedRef<IDocumentCache> ParentCache);
 		virtual ~FDocumentGraphNodeCache() = default;
 
 		// IDocumentGraphNodeCache implementation
 		virtual bool ContainsNode(const FGuid& InNodeID) const override;
 		virtual bool ContainsNodesOfClassID(const FGuid& InClassID) const override;
 
-		virtual const int32* FindInputNodeIndex(FName InputName) const override;
-		virtual const int32* FindOutputNodeIndex(FName OutputName) const override;
 		virtual const int32* FindNodeIndex(const FGuid& InNodeID) const override;
 		virtual TArray<const FMetasoundFrontendNode*> FindNodesOfClassID(const FGuid& InClassID) const override;
 		virtual const FMetasoundFrontendNode* FindNode(const FGuid& InNodeID) const override;
-		virtual const FMetasoundFrontendNode* FindInputNode(FName InputName) const override;
-		virtual const FMetasoundFrontendNode* FindOutputNode(FName OutputName) const override;
 
 		virtual TArray<const FMetasoundFrontendVertex*> FindNodeInputs(const FGuid& InNodeID, FName TypeName = FName()) const override;
 		virtual TArray<const FMetasoundFrontendVertex*> FindNodeOutputs(const FGuid& InNodeID, FName TypeName = FName()) const override;
@@ -79,21 +72,11 @@ namespace Metasound::Frontend
 		virtual const FMetasoundFrontendVertex* FindOutputVertex(const FGuid& InNodeID, const FGuid& InVertexID) const override;
 		virtual const FMetasoundFrontendVertex* FindOutputVertex(const FGuid& InNodeID, FName InVertexName) const override;
 
-	protected:
-		virtual void OnNodeAdded(int32 NewIndex) override;
-		virtual void OnRemovingNode(int32 IndexBeingRemoved) override;
-		virtual const FMetasoundFrontendDocument& GetDocument() const override;
-
-		// Friended to ensure only builder is the one managing cache manipulation via add/remove calls above
-		friend struct ::FMetaSoundFrontendDocumentBuilder;
+		void Init(FNodeModifyDelegates& OutDelegates);
 
 	private:
-
-		// Cache of Input name to array index of node
-		TMap<FName, int32> InputNameToIndex;
-
-		// Cache of Output name to array index of node
-		TMap<FName, int32> OutputNameToIndex;
+		void OnNodeAdded(int32 NewIndex);
+		void OnRemovingNode(int32 IndexBeingRemoved);
 
 		// Cache of NodeId to array index of node
 		TSortedMap<FGuid, int32> IDToIndex;
@@ -101,10 +84,39 @@ namespace Metasound::Frontend
 		// Cache of ClassID to referencing node indices
 		TSortedMap<FGuid, TArray<int32>> ClassIDToNodeIndices;
 
-		const FMetasoundFrontendDocument* Document = nullptr;
-		const FDocumentCache* Parent = nullptr;
+		TSharedPtr<IDocumentCache> Parent;
+	};
+
+
+	class FDocumentGraphInterfaceCache : public IDocumentGraphInterfaceCache
+	{
+	public:
+		FDocumentGraphInterfaceCache() = default;
+		FDocumentGraphInterfaceCache(TSharedRef<IDocumentCache> ParentCache);
+		virtual ~FDocumentGraphInterfaceCache() = default;
+
+		// IDocumentGraphInterfaceCache implementation
+		virtual const FMetasoundFrontendClassInput* FindInput(FName InputName) const override;
+		virtual const FMetasoundFrontendClassOutput* FindOutput(FName OutputName) const override;
+
+		void Init(FInterfaceModifyDelegates& OutDelegates);
+
+	private:
+		void OnInputAdded(int32 NewIndex);
+		void OnOutputAdded(int32 NewIndex);
+		void OnRemovingInput(int32 IndexBeingRemoved);
+		void OnRemovingOutput(int32 IndexBeingRemoved);
+
+		// Cache of Input name to array index of input
+		TMap<FName, int32> InputNameToIndex;
+
+		// Cache of Output name to array index of output
+		TMap<FName, int32> OutputNameToIndex;
+
+		TSharedPtr<IDocumentCache> Parent;
 	};
 
+
 	class FDocumentCache : public IDocumentCache
 	{
 	public:
@@ -119,18 +131,16 @@ namespace Metasound::Frontend
 		virtual const int32* FindDependencyIndex(const Metasound::Frontend::FNodeRegistryKey& InClassKey) const override;
 		virtual const int32* FindDependencyIndex(const FGuid& InClassID) const override;
 
-
-		virtual IDocumentGraphEdgeCache& GetEdgeCache() override { return *EdgeCache; }
-		virtual const IDocumentGraphEdgeCache& GetEdgeCache() const override { return *EdgeCache; }
-		virtual IDocumentGraphNodeCache& GetNodeCache() override { return *NodeCache; }
-		virtual const IDocumentGraphNodeCache& GetNodeCache() const override { return *NodeCache; }
-
-	protected:
-		virtual void OnDependencyAdded(int32 InNewIndex) override;
-		virtual void OnRemovingDependency(int32 IndexBeingRemoved) override;
 		virtual const FMetasoundFrontendDocument& GetDocument() const override;
+		virtual const IDocumentGraphEdgeCache& GetEdgeCache() const override;
+		virtual const IDocumentGraphNodeCache& GetNodeCache() const override;
+		virtual const IDocumentGraphInterfaceCache& GetInterfaceCache() const override;
+
+		void Init(FDocumentModifyDelegates& OutDelegates);
 
 	private:
+		void OnDependencyAdded(int32 InNewIndex);
+		void OnRemovingDependency(int32 IndexBeingRemoved);
 
 		// Cache of dependency (Class) ID to corresponding class dependency index
 		TSortedMap<FGuid, int32> IDToIndex;
@@ -138,11 +148,10 @@ namespace Metasound::Frontend
 		// Cache of version data to corresponding class dependency index
 		TSortedMap<FNodeRegistryKey, int32> KeyToIndex;
 
-		TUniquePtr<FDocumentGraphEdgeCache> EdgeCache;
-		TUniquePtr<FDocumentGraphNodeCache> NodeCache;
+		TSharedPtr<FDocumentGraphEdgeCache> EdgeCache;
+		TSharedPtr<FDocumentGraphNodeCache> NodeCache;
+		TSharedPtr<FDocumentGraphInterfaceCache> InterfaceCache;
 
 		const FMetasoundFrontendDocument* Document = nullptr;
-
-		friend struct ::FMetaSoundFrontendDocumentBuilder;
 	};
 } // namespace Metasound::Frontend
diff --git a/Engine/Plugins/Runtime/Metasound/Source/MetasoundFrontend/Private/MetasoundFrontendDocumentModifyDelegates.cpp b/Engine/Plugins/Runtime/Metasound/Source/MetasoundFrontend/Private/MetasoundFrontendDocumentModifyDelegates.cpp
new file mode 100644
--- /dev/null
+++ b/Engine/Plugins/Runtime/Metasound/Source/MetasoundFrontend/Private/MetasoundFrontendDocumentModifyDelegates.cpp
@@ -0,0 +1,8 @@
+// Copyright Epic Games, Inc. All Rights Reserved.
+
+#include "MetasoundFrontendDocumentModifyDelegates.h"
+
+
+namespace Metasound::Frontend
+{
+} // namespace Metasound::Frontend
diff --git a/Engine/Plugins/Runtime/Metasound/Source/MetasoundFrontend/Public/MetasoundDocumentInterface.h b/Engine/Plugins/Runtime/Metasound/Source/MetasoundFrontend/Public/MetasoundDocumentInterface.h
--- a/Engine/Plugins/Runtime/Metasound/Source/MetasoundFrontend/Public/MetasoundDocumentInterface.h
+++ b/Engine/Plugins/Runtime/Metasound/Source/MetasoundFrontend/Public/MetasoundDocumentInterface.h
@@ -3,6 +3,8 @@
 
 #include "UObject/Interface.h"
 #include "MetasoundFrontendDocument.h"
+#include "MetasoundFrontendDocumentModifyDelegates.h"
+#include "Templates/Function.h"
 
 #include "MetasoundDocumentInterface.generated.h"
 
@@ -32,10 +34,26 @@ public:
 	// Returns the parent class registered with the MetaSound UObject registry.
 	virtual const UClass& GetBaseMetaSoundUClass() const = 0;
 
-	// Returns the parent class registered with the MetaSound UObject registry.
-
 private:
 	virtual FMetasoundFrontendDocument& GetDocument() = 0;
 
 	friend struct FMetaSoundFrontendDocumentBuilder;
-};
\ No newline at end of file
+};
+
+namespace Metasound::Frontend
+{
+	class METASOUNDFRONTEND_API IMetaSoundDocumentBuilderRegistry
+	{
+	public:
+		// Returns delegates used to mutate any internal builder cached state or notify listeners of external system that
+		// has mutated a given document class.  Exists primarily for backward compatibility with the DocumentController system,
+		// and is not recommended for use outside of MetaSound plugin as it may be deprecated in the future (best practice is to
+		// mutate all documents using MetasoundDocumentBuilders).
+		virtual const FDocumentModifyDelegates* FindModifyDelegates(const FMetasoundFrontendClassName& InClassName) const = 0;
+
+		static IMetaSoundDocumentBuilderRegistry& GetChecked();
+
+	protected:
+		static void Set(TUniqueFunction<IMetaSoundDocumentBuilderRegistry&()>&& InGetInstance);
+	};
+} // namespace Metasound::Frontend
diff --git a/Engine/Plugins/Runtime/Metasound/Source/MetasoundFrontend/Public/MetasoundFrontendDocument.h b/Engine/Plugins/Runtime/Metasound/Source/MetasoundFrontend/Public/MetasoundFrontendDocument.h
--- a/Engine/Plugins/Runtime/Metasound/Source/MetasoundFrontend/Public/MetasoundFrontendDocument.h
+++ b/Engine/Plugins/Runtime/Metasound/Source/MetasoundFrontend/Public/MetasoundFrontendDocument.h
@@ -1,5 +1,4 @@
 // Copyright Epic Games, Inc. All Rights Reserved.
-
 #pragma once
 
 #include "Algo/Transform.h"
@@ -664,11 +663,9 @@ struct METASOUNDFRONTEND_API FMetasoundFrontendGraph
 	TArray<FMetasoundFrontendVariable> Variables;
 
 #if WITH_EDITORONLY_DATA
-
 	// Style of graph display.
 	UPROPERTY()
 	FMetasoundFrontendGraphStyle Style;
-
 #endif // WITH_EDITORONLY_DATA
 };
 
@@ -1192,7 +1189,7 @@ struct METASOUNDFRONTEND_API FMetasoundFrontendInterface : public FMetasoundFron
 
 
 // Name of a Metasound class
-USTRUCT()
+USTRUCT(BlueprintType)
 struct METASOUNDFRONTEND_API FMetasoundFrontendClassName
 {
 	GENERATED_BODY()
diff --git a/Engine/Plugins/Runtime/Metasound/Source/MetasoundFrontend/Public/MetasoundFrontendDocumentBuilder.h b/Engine/Plugins/Runtime/Metasound/Source/MetasoundFrontend/Public/MetasoundFrontendDocumentBuilder.h
--- a/Engine/Plugins/Runtime/Metasound/Source/MetasoundFrontend/Public/MetasoundFrontendDocumentBuilder.h
+++ b/Engine/Plugins/Runtime/Metasound/Source/MetasoundFrontend/Public/MetasoundFrontendDocumentBuilder.h
@@ -6,6 +6,7 @@
 #include "MetasoundFrontendDocumentCacheInterface.h"
 #include "MetasoundDocumentInterface.h"
 #include "MetasoundFrontendDocument.h"
+#include "MetasoundFrontendDocumentModifyDelegates.h"
 #include "MetasoundFrontendRegistries.h"
 #include "MetasoundFrontendTransform.h"
 #include "MetasoundVertex.h"
@@ -21,6 +22,8 @@ class FMetasoundAssetBase;
 
 namespace Metasound::Frontend
 {
+	using FFinalizeNodeFunctionRef = TFunctionRef<void(FMetasoundFrontendNode&, const Metasound::Frontend::FNodeRegistryKey&)>;
+
 	struct METASOUNDFRONTEND_API FNamedEdge
 	{
 		const FGuid OutputNodeID;
@@ -52,6 +55,7 @@ namespace Metasound::Frontend
 	struct METASOUNDFRONTEND_API FModifyInterfaceOptions
 	{
 		FModifyInterfaceOptions(const TArray<FMetasoundFrontendInterface>& InInterfacesToRemove, const TArray<FMetasoundFrontendInterface>& InInterfacesToAdd);
+		FModifyInterfaceOptions(TArray<FMetasoundFrontendInterface>&& InInterfacesToRemove, TArray<FMetasoundFrontendInterface>&& InInterfacesToAdd);
 		FModifyInterfaceOptions(const TArray<FMetasoundFrontendVersion>& InInterfaceVersionsToRemove, const TArray<FMetasoundFrontendVersion>& InInterfaceVersionsToAdd);
 
 		TArray<FMetasoundFrontendInterface> InterfacesToRemove;
@@ -114,7 +118,7 @@ private:
 // Builder used to support dynamically generating MetaSound documents at runtime. Builder contains caches that speed up
 // common search and modification operations on a given document, which may result in slower performance on construction,
 // but faster manipulation of its managed document.  The builder's managed copy of a document is expected to not be modified
-// by any external system to avoid cache to not become stale.
+// by any external system to avoid cache becoming stale.
 USTRUCT()
 struct METASOUNDFRONTEND_API FMetaSoundFrontendDocumentBuilder
 {
@@ -122,22 +126,12 @@ struct METASOUNDFRONTEND_API FMetaSoundFrontendDocumentBuilder
 
 public:
 	// Exists only to make UObject reflection happy.  Should never be
-	// used directly as builder class should be specified on construction.
-	FMetaSoundFrontendDocumentBuilder() = default;
-	FMetaSoundFrontendDocumentBuilder(FMetaSoundFrontendDocumentBuilder&& InBuilder);
+	// used directly as builder interface (and optionally delegates) should be specified on construction.
+	FMetaSoundFrontendDocumentBuilder();
 	FMetaSoundFrontendDocumentBuilder(TScriptInterface<IMetaSoundDocumentInterface> InDocumentInterface);
+	FMetaSoundFrontendDocumentBuilder(TScriptInterface<IMetaSoundDocumentInterface> InDocumentInterface, TSharedRef<Metasound::Frontend::FDocumentModifyDelegates> InDocumentDelegates);
 
-	// Copying actively used builders is highly discouraged (as only one builder should be acting on a document at any given time
-	// to avoid internal cache state to become out-of-date) but must contractually be implemented in this case as a USTRUCT (builders
-	// are USTRUCTS to keep document dependencies from being GC'ed).  To protect against this kind of "back door" building from another
-	// builder, copying builders also creates a new document & builder owner.
-	FMetaSoundFrontendDocumentBuilder(const FMetaSoundFrontendDocumentBuilder& InBuilder);
-	FMetaSoundFrontendDocumentBuilder& operator=(const FMetaSoundFrontendDocumentBuilder& InRHS);
-
-	FMetaSoundFrontendDocumentBuilder& operator=(FMetaSoundFrontendDocumentBuilder&& InRHS);
-
-	const FMetasoundFrontendClass* AddGraphDependency(const FMetasoundFrontendGraphClass& InClass);
-	const FMetasoundFrontendClass* AddNativeDependency(const FMetasoundFrontendClassMetadata& InClassMetadata);
+	const FMetasoundFrontendClass* AddDependency(const FMetasoundFrontendClass& InClass);
 	const FMetasoundFrontendEdge* AddEdge(FMetasoundFrontendEdge&& InNewEdge);
 	bool AddNamedEdges(const TSet<Metasound::Frontend::FNamedEdge>& ConnectionsToMake, TArray<const FMetasoundFrontendEdge*>* OutNewEdges = nullptr);
 	bool AddEdgesByNodeClassInterfaceBindings(const FGuid& InFromNodeID, const FGuid& InToNodeID);
@@ -147,12 +141,12 @@ public:
 	const FMetasoundFrontendNode* AddGraphOutput(const FMetasoundFrontendClassOutput& InClassOutput);
 	bool AddInterface(FName InterfaceName);
 
-	const FMetasoundFrontendNode* AddGraphNode(const FMetasoundFrontendGraphClass& InGraphClass, FGuid InNodeID = FGuid::NewGuid());
+	const FMetasoundFrontendNode* AddGraphNode(const FMetasoundFrontendGraphClass& InClass, FGuid InNodeID = FGuid::NewGuid());
+	const FMetasoundFrontendNode* AddNodeByClassName(const FMetasoundFrontendClassName& InClassName, int32 InMajorVersion, FGuid InNodeID = FGuid::NewGuid());
 
 	bool CanAddEdge(const FMetasoundFrontendEdge& InEdge) const;
 
 	void ClearGraph();
-
 	bool ContainsEdge(const FMetasoundFrontendEdge& InEdge) const;
 	bool ContainsNode(const FGuid& InNodeID) const;
 
@@ -183,7 +177,10 @@ public:
 	TArray<const FMetasoundFrontendVertex*> FindNodeOutputs(const FGuid& InNodeID, FName TypeName = FName()) const;
 
 	const FMetasoundFrontendDocument& GetDocument() const;
+	const Metasound::Frontend::FDocumentModifyDelegates& GetDocumentDelegates() const;
 	const IMetaSoundDocumentInterface& GetDocumentInterface() const;
+	const FMetasoundFrontendLiteral* GetNodeInputClassDefault(const FGuid& InNodeID, const FGuid& InVertexID) const;
+	const FMetasoundFrontendLiteral* GetNodeInputDefault(const FGuid& InNodeID, const FGuid& InVertexID) const;
 
 	static void InitGraphClassMetadata(FMetasoundFrontendClassMetadata& InOutMetadata, bool bResetVersion = false);
 	void InitDocument();
@@ -210,6 +207,7 @@ public:
 	bool RemoveInterface(FName InName);
 	bool RemoveNamedEdges(const TSet<Metasound::Frontend::FNamedEdge>& InNamedEdgesToRemove, TArray<FMetasoundFrontendEdge>* OutRemovedEdges = nullptr);
 	bool RemoveNode(const FGuid& InNodeID);
+	bool RemoveNodeInputDefault(const FGuid& InNodeID, const FGuid& InVertexID);
 	bool RemoveUnusedDependencies();
 
 #if WITH_EDITOR
@@ -225,15 +223,10 @@ public:
 private:
 	using FFinalizeNodeFunctionRef = TFunctionRef<void(FMetasoundFrontendNode&, const Metasound::Frontend::FNodeRegistryKey&)>;
 
-	UPROPERTY(Transient)
-	TScriptInterface<IMetaSoundDocumentInterface> DocumentInterface;
-
-	TUniquePtr<Metasound::Frontend::IDocumentCache> DocumentCache;
-
 	const FTopLevelAssetPath GetBuilderClassPath() const;
 	FMetasoundFrontendDocument& GetDocument();
 
-	FMetasoundFrontendNode* AddNodeInternal(const FMetasoundFrontendClassMetadata& InClassMetadata, FFinalizeNodeFunctionRef FinalizeNode, FGuid InNodeID = FGuid::NewGuid());
+	FMetasoundFrontendNode* AddNodeInternal(const FMetasoundFrontendClassMetadata& InClassMetadata, Metasound::Frontend::FFinalizeNodeFunctionRef FinalizeNode, FGuid InNodeID = FGuid::NewGuid());
 
 	const FMetasoundFrontendDocument* FindOrLoadNodeClassDocument(const FGuid& InNodeID) const;
 	const FMetasoundFrontendClassInput* FindNodeInputClassInput(const FGuid& InNodeID, const FGuid& InVertexID) const;
@@ -241,5 +234,16 @@ private:
 
 	void ReloadCache();
 
-	bool SetInputInheritsDefault(FName InName, bool bInputInheritsDefault);
+	bool SetGraphInputInheritsDefault(FName InName, bool bInputInheritsDefault);
+
+	UPROPERTY(Transient)
+	TScriptInterface<IMetaSoundDocumentInterface> DocumentInterface;
+
+	TSharedPtr<Metasound::Frontend::IDocumentCache> DocumentCache;
+
+	// SharedRef to struct of delegates fired when mutating document.  These are shared to safely support copying
+	// FrontendDocumentBuilders.  Copying is typically dissuaded as you can still have disparate builders that are
+	// not operating on the same cache; using a shared builder registered within the MetaSoundBuilderSubsystem
+	// in MetaSoundEngine is recommended.
+	TSharedRef<Metasound::Frontend::FDocumentModifyDelegates> DocumentDelegates;
 };
diff --git a/Engine/Plugins/Runtime/Metasound/Source/MetasoundFrontend/Public/MetasoundFrontendDocumentCacheInterface.h b/Engine/Plugins/Runtime/Metasound/Source/MetasoundFrontend/Public/MetasoundFrontendDocumentCacheInterface.h
--- a/Engine/Plugins/Runtime/Metasound/Source/MetasoundFrontend/Public/MetasoundFrontendDocumentCacheInterface.h
+++ b/Engine/Plugins/Runtime/Metasound/Source/MetasoundFrontend/Public/MetasoundFrontendDocumentCacheInterface.h
@@ -3,6 +3,8 @@
 
 #include "Containers/Array.h"
 #include "HAL/Platform.h"
+#include "MetasoundFrontendDocumentModifyDelegates.h"
+#include "Templates/SharedPointer.h"
 #include "UObject/Interface.h"
 
 
@@ -22,7 +24,7 @@ namespace Metasound::Frontend
 	using FNodeRegistryKey = FString;
 
 	/** Interface for querying cached document nodes. */
-	class IDocumentGraphNodeCache
+	class IDocumentGraphNodeCache : public TSharedFromThis<IDocumentGraphNodeCache>
 	{
 	public:
 		virtual ~IDocumentGraphNodeCache() = default;
@@ -30,14 +32,10 @@ namespace Metasound::Frontend
 		virtual bool ContainsNode(const FGuid& InNodeID) const = 0;
 		virtual bool ContainsNodesOfClassID(const FGuid& InClassID) const = 0;
 
-		virtual const int32* FindInputNodeIndex(FName InputName) const = 0;
-		virtual const int32* FindOutputNodeIndex(FName OutputName) const = 0;
 		virtual const int32* FindNodeIndex(const FGuid& InNodeID) const = 0;
 		virtual TArray<const FMetasoundFrontendNode*> FindNodesOfClassID(const FGuid& InClassID) const = 0;
 
 		virtual const FMetasoundFrontendNode* FindNode(const FGuid& InNodeID) const = 0;
-		virtual const FMetasoundFrontendNode* FindInputNode(FName InputName) const = 0;
-		virtual const FMetasoundFrontendNode* FindOutputNode(FName OutputName) const = 0;
 
 		virtual TArray<const FMetasoundFrontendVertex*> FindNodeInputs(const FGuid& InNodeID, FName TypeName = FName()) const = 0;
 		virtual TArray<const FMetasoundFrontendVertex*> FindNodeOutputs(const FGuid& InNodeID, FName TypeName = FName()) const = 0;
@@ -46,18 +44,10 @@ namespace Metasound::Frontend
 		virtual const FMetasoundFrontendVertex* FindInputVertex(const FGuid& InNodeID, FName InVertexName) const = 0;
 		virtual const FMetasoundFrontendVertex* FindOutputVertex(const FGuid& InNodeID, const FGuid& InVertexID) const = 0;
 		virtual const FMetasoundFrontendVertex* FindOutputVertex(const FGuid& InNodeID, FName InVertexName) const = 0;
-
-		virtual const FMetasoundFrontendDocument& GetDocument() const = 0;
-
-	protected:
-		virtual void OnNodeAdded(int32 NewItemIndex) = 0;
-		virtual void OnRemovingNode(int32 ItemIndexBeingRemoved) = 0;
-
-		friend struct ::FMetaSoundFrontendDocumentBuilder;
 	};
 
 	/** Interface for querying cached document edges. */
-	class IDocumentGraphEdgeCache
+	class IDocumentGraphEdgeCache : public TSharedFromThis<IDocumentGraphEdgeCache>
 	{
 	public:
 		virtual ~IDocumentGraphEdgeCache() = default;
@@ -70,18 +60,20 @@ namespace Metasound::Frontend
 		virtual TArray<const FMetasoundFrontendEdge*> FindEdges(const FGuid& InNodeID, const FGuid& InVertexID) const = 0;
 		virtual const int32* FindEdgeIndexToNodeInput(const FGuid& InNodeID, const FGuid& InVertexID) const = 0;
 		virtual const TArray<int32>* FindEdgeIndicesFromNodeOutput(const FGuid& InNodeID, const FGuid& InVertexID) const = 0;
+	};
 
-		virtual const FMetasoundFrontendDocument& GetDocument() const = 0;
-
-	protected:
-		virtual void OnEdgeAdded(int32 NewItemIndex) = 0;
-		virtual void OnRemovingEdge(int32 ItemIndexBeingRemoved) = 0;
+	/** Interface for querying cached document interface. */
+	class IDocumentGraphInterfaceCache : public TSharedFromThis<IDocumentGraphInterfaceCache>
+	{
+	public:
+		virtual ~IDocumentGraphInterfaceCache() = default;
 
-		friend struct ::FMetaSoundFrontendDocumentBuilder;
+		virtual const FMetasoundFrontendClassInput* FindInput(FName InputName) const = 0;
+		virtual const FMetasoundFrontendClassOutput* FindOutput(FName OutputName) const = 0;
 	};
 
 	/** Interface for querying cached document dependencies. */
-	class IDocumentCache
+	class IDocumentCache : public TSharedFromThis<IDocumentCache>
 	{
 	public:
 		virtual ~IDocumentCache() = default;
@@ -93,18 +85,11 @@ namespace Metasound::Frontend
 		virtual const int32* FindDependencyIndex(const Metasound::Frontend::FNodeRegistryKey& InClassKey) const = 0;
 		virtual const int32* FindDependencyIndex(const FGuid& InClassID) const = 0;
 
-		virtual IDocumentGraphNodeCache& GetNodeCache() = 0;
-		virtual const IDocumentGraphNodeCache& GetNodeCache() const = 0;
-		virtual IDocumentGraphEdgeCache& GetEdgeCache() = 0;
-		virtual const IDocumentGraphEdgeCache& GetEdgeCache() const = 0;
-
 		virtual const FMetasoundFrontendDocument& GetDocument() const = 0;
 
-	protected:
-		virtual void OnDependencyAdded(int32 NewItemIndex) = 0;
-		virtual void OnRemovingDependency(int32 ItemIndexBeingRemoved) = 0;
-
-		friend struct ::FMetaSoundFrontendDocumentBuilder;
+		virtual const IDocumentGraphNodeCache& GetNodeCache() const = 0;
+		virtual const IDocumentGraphEdgeCache& GetEdgeCache() const = 0;
+		virtual const IDocumentGraphInterfaceCache& GetInterfaceCache() const = 0;
 	};
 
 } // namespace Metasound::Frontend
diff --git a/Engine/Plugins/Runtime/Metasound/Source/MetasoundFrontend/Public/MetasoundFrontendDocumentModifyDelegates.h b/Engine/Plugins/Runtime/Metasound/Source/MetasoundFrontend/Public/MetasoundFrontendDocumentModifyDelegates.h
new file mode 100644
--- /dev/null
+++ b/Engine/Plugins/Runtime/Metasound/Source/MetasoundFrontend/Public/MetasoundFrontendDocumentModifyDelegates.h
@@ -0,0 +1,52 @@
+// Copyright Epic Games, Inc. All Rights Reserved.
+#pragma once
+
+#include "Delegates/DelegateCombinations.h"
+#include "MetasoundFrontendDocument.h"
+#include "Templates/SharedPointer.h"
+
+
+DECLARE_MULTICAST_DELEGATE_OneParam(FOnMetaSoundFrontendDocumentMutateArray, int32 /* Index */);
+DECLARE_MULTICAST_DELEGATE_OneParam(FOnMetaSoundFrontendDocumentMutateInterfaceArray, const FMetasoundFrontendInterface& /* Interface */);
+DECLARE_MULTICAST_DELEGATE_ThreeParams(FOnMetaSoundFrontendDocumentMutateNodeInputLiteralArray, int32 /* NodeIndex */, int32 /* VertexIndex */, int32 /* LiteralIndex */);
+
+
+namespace Metasound::Frontend
+{
+	struct METASOUNDFRONTEND_API FInterfaceModifyDelegates
+	{
+		FOnMetaSoundFrontendDocumentMutateInterfaceArray OnInterfaceAdded;
+		FOnMetaSoundFrontendDocumentMutateInterfaceArray OnRemovingInterface;
+
+		FOnMetaSoundFrontendDocumentMutateArray OnInputAdded;
+		FOnMetaSoundFrontendDocumentMutateArray OnRemovingInput;
+
+		FOnMetaSoundFrontendDocumentMutateArray OnOutputAdded;
+		FOnMetaSoundFrontendDocumentMutateArray OnRemovingOutput;
+	};
+
+	struct METASOUNDFRONTEND_API FNodeModifyDelegates
+	{
+		FOnMetaSoundFrontendDocumentMutateArray OnNodeAdded;
+		FOnMetaSoundFrontendDocumentMutateArray OnRemovingNode;
+
+		FOnMetaSoundFrontendDocumentMutateNodeInputLiteralArray OnNodeInputLiteralSet;
+		FOnMetaSoundFrontendDocumentMutateNodeInputLiteralArray OnRemovingNodeInputLiteral;
+	};
+
+	struct METASOUNDFRONTEND_API FEdgeModifyDelegates
+	{
+		FOnMetaSoundFrontendDocumentMutateArray OnEdgeAdded;
+		FOnMetaSoundFrontendDocumentMutateArray OnRemovingEdge;
+	};
+
+	struct METASOUNDFRONTEND_API FDocumentModifyDelegates : TSharedFromThis<FDocumentModifyDelegates>
+	{
+		FOnMetaSoundFrontendDocumentMutateArray OnDependencyAdded;
+		FOnMetaSoundFrontendDocumentMutateArray OnRemovingDependency;
+
+		FInterfaceModifyDelegates InterfaceDelegates;
+		FNodeModifyDelegates NodeDelegates;
+		FEdgeModifyDelegates EdgeDelegates;
+	};
+} // namespace Metasound::Frontend
