MetaSound deterministic id generation
- Deterministic generation of ids for various FGuids used in MetaSound documents so that Autoupdate can be moved from runtime to cook in the future while maintaining cook determinism
- cvar for enabling in editor au.MetaSound.EnableDeterministicIDGenerationInEditorCVar

diff --git a/Engine/Plugins/Runtime/Metasound/Source/MetasoundFrontend/Private/MetasoundAssetBase.cpp b/Engine/Plugins/Runtime/Metasound/Source/MetasoundFrontend/Private/MetasoundAssetBase.cpp
--- a/Engine/Plugins/Runtime/Metasound/Source/MetasoundFrontend/Private/MetasoundAssetBase.cpp
+++ b/Engine/Plugins/Runtime/Metasound/Source/MetasoundFrontend/Private/MetasoundAssetBase.cpp
@@ -19,6 +19,7 @@
 #include "MetasoundFrontendController.h"
 #include "MetasoundFrontendDocument.h"
 #include "MetasoundFrontendDocumentBuilder.h"
+#include "MetasoundFrontendDocumentIdGenerator.h"
 #include "MetasoundFrontendDocumentVersioning.h"
 #include "MetasoundFrontendGraph.h"
 #include "MetasoundFrontendNodeTemplateRegistry.h"
@@ -338,6 +339,10 @@ bool FMetasoundAssetBase::VersionAsset()
 	using namespace Metasound;
 	using namespace Metasound::Frontend;
 	METASOUND_TRACE_CPUPROFILER_EVENT_SCOPE(MetaSoundAssetBase::VersionAsset);
+#if WITH_EDITOR
+	const bool bIsDeterministic = MetaSoundEnableDeterministicIDGenerationInEditorCVar != 0 && !IsRunningCookCommandlet();
+	FDocumentIDGenerator::FScopeDeterminism DeterminismScope(bIsDeterministic);
+#endif // WITH_EDITOR
 
 	FName AssetName;
 	FString AssetPath;
@@ -833,6 +838,10 @@ void FMetasoundAssetBase::RegisterAssetDependencies(const Metasound::Frontend::F
 bool FMetasoundAssetBase::AutoUpdate(bool bInLogWarningsOnDroppedConnection)
 {
 	using namespace Metasound::Frontend;
+#if WITH_EDITOR
+	const bool bIsDeterministic = MetaSoundEnableDeterministicIDGenerationInEditorCVar != 0 && !IsRunningCookCommandlet();
+	FDocumentIDGenerator::FScopeDeterminism DeterminismScope(bIsDeterministic);
+#endif // WITH_EDITOR
 
 	FString OwningAssetName = GetOwningAssetName();
 	const bool bAutoUpdated = FAutoUpdateRootGraph(MoveTemp(OwningAssetName), bInLogWarningsOnDroppedConnection).Transform(GetDocumentHandle());
diff --git a/Engine/Plugins/Runtime/Metasound/Source/MetasoundFrontend/Private/MetasoundFrontendDocument.cpp b/Engine/Plugins/Runtime/Metasound/Source/MetasoundFrontend/Private/MetasoundFrontendDocument.cpp
--- a/Engine/Plugins/Runtime/Metasound/Source/MetasoundFrontend/Private/MetasoundFrontendDocument.cpp
+++ b/Engine/Plugins/Runtime/Metasound/Source/MetasoundFrontend/Private/MetasoundFrontendDocument.cpp
@@ -7,8 +7,9 @@
 #include "Algo/Transform.h"
 #include "IAudioParameterInterfaceRegistry.h"
 #include "Logging/LogMacros.h"
-#include "MetasoundFrontendDocumentVersioning.h"
 #include "MetasoundFrontend.h"
+#include "MetasoundFrontendDocumentIdGenerator.h"
+#include "MetasoundFrontendDocumentVersioning.h"
 #include "MetasoundFrontendRegistries.h"
 #include "MetasoundLog.h"
 #include "MetasoundParameterTransmitter.h"
@@ -424,9 +425,13 @@ FMetasoundFrontendClassInterface FMetasoundFrontendClassInterface::GenerateClass
 
 			ClassInput.Name = InputVertex.VertexName;
 			ClassInput.TypeName = InputVertex.DataTypeName;
-			ClassInput.VertexID = FGuid::NewGuid();
 			ClassInput.AccessType = Metasound::DocumentPrivate::CoreVertexAccessTypeToFrontendVertexAccessType(InputVertex.AccessType);
-
+#if WITH_EDITOR
+			const bool bIsDeterministic = MetaSoundEnableDeterministicIDGenerationInEditorCVar != 0 && !IsRunningCookCommandlet();
+			ClassInput.VertexID = bIsDeterministic ? FClassIDGenerator::Get().CreateInputID(ClassInput) : FGuid::NewGuid();
+#else 
+			ClassInput.VertexID = FGuid::NewGuid();
+#endif // WITH_EDITOR
 
 #if WITH_EDITOR
 			const FDataVertexMetadata& VertexMetadata = InputVertex.Metadata;
@@ -479,12 +484,15 @@ FMetasoundFrontendClassInterface FMetasoundFrontendClassInterface::GenerateClass
 		{
 			FMetasoundFrontendClassOutput ClassOutput;
 
-
 			ClassOutput.Name = OutputVertex.VertexName;
 			ClassOutput.TypeName = OutputVertex.DataTypeName;
-			ClassOutput.VertexID = FGuid::NewGuid();
 			ClassOutput.AccessType = Metasound::DocumentPrivate::CoreVertexAccessTypeToFrontendVertexAccessType(OutputVertex.AccessType);
-
+#if WITH_EDITOR
+			const bool bIsDeterministic = MetaSoundEnableDeterministicIDGenerationInEditorCVar != 0 && !IsRunningCookCommandlet();
+			ClassOutput.VertexID = bIsDeterministic ? FClassIDGenerator::Get().CreateOutputID(ClassOutput) : FGuid::NewGuid();
+#else 
+			ClassOutput.VertexID = FGuid::NewGuid();
+#endif // WITH_EDITOR
 #if WITH_EDITOR
 			const FDataVertexMetadata& VertexMetadata = OutputVertex.Metadata;
 
@@ -745,10 +753,17 @@ FMetasoundFrontendClassInput::FMetasoundFrontendClassInput(const FMetasoundFront
 
 FMetasoundFrontendClassInput::FMetasoundFrontendClassInput(const Audio::FParameterInterface::FInput& InInput)
 {
+	using namespace Metasound::Frontend;
+
 	Name = InInput.InitValue.ParamName;
 	DefaultLiteral = FMetasoundFrontendLiteral(InInput.InitValue);
 	TypeName = Metasound::DocumentPrivate::ResolveMemberDataType(InInput.DataType, InInput.InitValue.ParamType);
+#if WITH_EDITOR
+	const bool bIsDeterministic = MetaSoundEnableDeterministicIDGenerationInEditorCVar != 0 && !IsRunningCookCommandlet();
+	VertexID = bIsDeterministic ? FClassIDGenerator::Get().CreateInputID(InInput) : FGuid::NewGuid();
+#else 
 	VertexID = FGuid::NewGuid();
+#endif // WITH_EDITOR
 
 #if WITH_EDITOR
 	// Interfaces should never serialize text to avoid desync between
@@ -772,9 +787,16 @@ FMetasoundFrontendClassVariable::FMetasoundFrontendClassVariable(const FMetasoun
 
 FMetasoundFrontendClassOutput::FMetasoundFrontendClassOutput(const Audio::FParameterInterface::FOutput& Output)
 {
+	using namespace Metasound::Frontend;
+
 	Name = Output.ParamName;
 	TypeName = Metasound::DocumentPrivate::ResolveMemberDataType(Output.DataType, Output.ParamType);
+#if WITH_EDITOR
+	const bool bIsDeterministic = MetaSoundEnableDeterministicIDGenerationInEditorCVar != 0 && !IsRunningCookCommandlet();
+	VertexID = bIsDeterministic ? FClassIDGenerator::Get().CreateOutputID(Output) : FGuid::NewGuid();
+#else 
 	VertexID = FGuid::NewGuid();
+#endif // WITH_EDITOR
 
 #if WITH_EDITOR
 	// Interfaces should never serialize text to avoid desync between
diff --git a/Engine/Plugins/Runtime/Metasound/Source/MetasoundFrontend/Private/MetasoundFrontendDocumentController.cpp b/Engine/Plugins/Runtime/Metasound/Source/MetasoundFrontend/Private/MetasoundFrontendDocumentController.cpp
--- a/Engine/Plugins/Runtime/Metasound/Source/MetasoundFrontend/Private/MetasoundFrontendDocumentController.cpp
+++ b/Engine/Plugins/Runtime/Metasound/Source/MetasoundFrontend/Private/MetasoundFrontendDocumentController.cpp
@@ -5,6 +5,7 @@
 #include "Algo/ForEach.h"
 #include "HAL/FileManager.h"
 #include "MetasoundFrontendGraphController.h"
+#include "MetasoundFrontendDocumentIdGenerator.h"
 #include "MetasoundFrontendInvalidController.h"
 #include "MetasoundJsonBackend.h"
 #include "StructSerializer.h"
@@ -292,7 +293,7 @@ namespace Metasound
 
 					return NewClassPtr;
 				};
-
+				
 				if (FMetasoundFrontendClass* MetasoundClass = ClassPtr.Get())
 				{
 					// External node classes must match version to return shared definition.
@@ -305,7 +306,8 @@ namespace Metasound
 						FMetasoundFrontendClass NewClass = GenerateClass(InKey);
 						if (NewClass.Metadata.GetVersion().Major != MetasoundClass->Metadata.GetVersion().Major)
 						{
-							return AddClass(MoveTemp(NewClass), FGuid::NewGuid());
+							const FGuid ClassId = FDocumentIDGenerator::Get().CreateClassID(*Document);
+							return AddClass(MoveTemp(NewClass), ClassId);
 						}
 					}
 
@@ -321,7 +323,8 @@ namespace Metasound
 				}
 
 				FMetasoundFrontendClass NewClass = GenerateClass(InKey);
-				return AddClass(MoveTemp(NewClass), FGuid::NewGuid());
+				const FGuid ClassId = FDocumentIDGenerator::Get().CreateClassID(*Document);
+				return AddClass(MoveTemp(NewClass), ClassId);
 			}
 
 			return FConstClassAccessPtr();
@@ -509,6 +512,16 @@ namespace Metasound
 			return IGraphController::GetInvalidHandle();
 		}
 
+		FDocumentAccessPtr FDocumentController::GetDocumentPtr()
+		{
+			return DocumentPtr;
+		}
+
+		const FDocumentAccessPtr FDocumentController::GetDocumentPtr() const
+		{
+			return DocumentPtr;
+		}
+
 		TArray<FGraphHandle> FDocumentController::GetSubgraphHandles() 
 		{
 			TArray<FGraphHandle> Subgraphs;
diff --git a/Engine/Plugins/Runtime/Metasound/Source/MetasoundFrontend/Private/MetasoundFrontendDocumentController.h b/Engine/Plugins/Runtime/Metasound/Source/MetasoundFrontend/Private/MetasoundFrontendDocumentController.h
--- a/Engine/Plugins/Runtime/Metasound/Source/MetasoundFrontend/Private/MetasoundFrontendDocumentController.h
+++ b/Engine/Plugins/Runtime/Metasound/Source/MetasoundFrontend/Private/MetasoundFrontendDocumentController.h
@@ -69,6 +69,9 @@ namespace Metasound
 			FGraphHandle GetRootGraph() override;
 			FConstGraphHandle GetRootGraph() const override;
 
+			FDocumentAccessPtr GetDocumentPtr() override;
+			const FDocumentAccessPtr GetDocumentPtr() const override;
+
 			/** Returns an array of all subgraphs for this document. */
 			TArray<FGraphHandle> GetSubgraphHandles() override;
 
diff --git a/Engine/Plugins/Runtime/Metasound/Source/MetasoundFrontend/Private/MetasoundFrontendDocumentIdGenerator.cpp b/Engine/Plugins/Runtime/Metasound/Source/MetasoundFrontend/Private/MetasoundFrontendDocumentIdGenerator.cpp
new file mode 100644
--- /dev/null
+++ b/Engine/Plugins/Runtime/Metasound/Source/MetasoundFrontend/Private/MetasoundFrontendDocumentIdGenerator.cpp
@@ -0,0 +1,128 @@
+// Copyright Epic Games, Inc. All Rights Reserved.
+
+#include "MetasoundFrontendDocumentIdGenerator.h"
+
+namespace Metasound
+{
+	namespace Frontend
+	{
+		int32 MetaSoundEnableDeterministicIDGenerationInEditorCVar = 0;
+		FAutoConsoleVariableRef CVarMetaSoundEnableDeterministicIDGenerationInEditor(
+			TEXT("au.MetaSound.EnableDeterministicIDGenerationInEditorCVar"),
+			MetaSoundEnableDeterministicIDGenerationInEditorCVar,
+			TEXT("Enable deterministic ID generation in MetaSounds for AutoUpdate while in editor (requires editor restart for effect) \n")
+			TEXT("0: Disabled (default), !0: Enabled"),
+			ECVF_Default);
+
+		FDocumentIDGenerator::FScopeDeterminism::FScopeDeterminism(bool bInIsDeterministic)
+		{
+			FDocumentIDGenerator& IDGen = Get();
+			bOriginalValue = IDGen.GetDeterminism();
+			IDGen.SetDeterminism(bInIsDeterministic);
+		}
+
+		bool FDocumentIDGenerator::FScopeDeterminism::GetDeterminism() const
+		{
+			return Get().GetDeterminism();
+		}
+
+		FDocumentIDGenerator::FScopeDeterminism::~FScopeDeterminism()
+		{
+			Get().SetDeterminism(bOriginalValue);
+		}
+
+		void FDocumentIDGenerator::SetDeterminism(bool bInIsDeterministic)
+		{
+			bIsDeterministic = bInIsDeterministic;
+		}
+
+		bool FDocumentIDGenerator::GetDeterminism() const
+		{
+			return bIsDeterministic;
+		}
+		
+		FDocumentIDGenerator& FDocumentIDGenerator::Get()
+		{
+			thread_local FDocumentIDGenerator IDGenerator;
+			return IDGenerator;
+		}
+
+		FGuid FDocumentIDGenerator::CreateNodeID(const FMetasoundFrontendDocument& InDocument) const
+		{
+			return CreateIDFromDocument(InDocument);
+		}
+
+		FGuid FDocumentIDGenerator::CreateVertexID(const FMetasoundFrontendDocument& InDocument) const
+		{
+			return CreateIDFromDocument(InDocument);
+		}
+
+		FGuid FDocumentIDGenerator::CreateClassID(const FMetasoundFrontendDocument& InDocument) const
+		{
+			return CreateIDFromDocument(InDocument);
+		}
+
+		FGuid FDocumentIDGenerator::CreateIDFromDocument(const FMetasoundFrontendDocument& InDocument) const
+		{
+			if (bIsDeterministic)
+			{
+				const uint32 Value = InDocument.GetNextIdCounter();
+				const FGuid CounterGuid = FGuid(Value << 6, Value << 4, Value << 2, Value);
+				return FGuid::Combine(InDocument.RootGraph.ID, CounterGuid);
+			}
+			else
+			{
+				return FGuid::NewGuid();
+			}
+		}
+
+		FClassIDGenerator& FClassIDGenerator::Get()
+		{
+			static FClassIDGenerator IDGenerator;
+			return IDGenerator;
+		}
+
+		FGuid FClassIDGenerator::CreateInputID(const FMetasoundFrontendClassInput& Input) const
+		{
+			const FGuid ClassInputNamespaceGuid = FGuid("{149FEB6E-B9F9-47A6-AD4F-B78655F6EBE8}");
+			FString NameToHash = FString::Printf(TEXT("ClassInput.%s.%s.%s"), *Input.Name.ToString(), *Input.TypeName.ToString(), LexToString(Input.AccessType));
+
+			return CreateNamespacedIDFromString(ClassInputNamespaceGuid, NameToHash);
+		}
+
+		FGuid FClassIDGenerator::CreateInputID(const Audio::FParameterInterface::FInput& Input) const
+		{
+			const FGuid ParameterInterfaceInputNamespaceGuid = FGuid("{D9E893C0-92B3-4CB4-8306-4525ABEACADD}");
+			FString NameToHash = FString::Printf(TEXT("ParameterInterfaceInput.%s.%s"), *Input.InitValue.ParamName.ToString(), *Input.DataType.ToString());
+
+			return CreateNamespacedIDFromString(ParameterInterfaceInputNamespaceGuid, NameToHash);
+		}
+
+		FGuid FClassIDGenerator::CreateOutputID(const FMetasoundFrontendClassOutput& Output) const
+		{
+			const FGuid ClassOutputNamespaceGuid = FGuid("{C7B3ED2C-4407-4B2A-9144-7F1108387EBB}");
+			FString NameToHash = FString::Printf(TEXT("ClassOutput.%s.%s.%s"), *Output.Name.ToString(), *Output.TypeName.ToString(), LexToString(Output.AccessType));
+
+			return CreateNamespacedIDFromString(ClassOutputNamespaceGuid, NameToHash);
+		}
+
+		FGuid FClassIDGenerator::CreateOutputID(const Audio::FParameterInterface::FOutput& Output) const
+		{
+			const FGuid ParameterInterfaceOutputNamespaceGuid = FGuid("{6F41342A-2436-4462-81A0-8517887BB729}");
+			FString NameToHash = FString::Printf(TEXT("ParameterInterfaceOutput.%s.%s"), *Output.ParamName.ToString(), *Output.DataType.ToString());
+
+			return CreateNamespacedIDFromString(ParameterInterfaceOutputNamespaceGuid, NameToHash);
+		}
+
+		FGuid FClassIDGenerator::CreateNamespacedIDFromString(const FGuid NamespaceGuid, FString StringToHash) const
+		{
+			FSHA1 Hasher;
+			Hasher.Update(reinterpret_cast<const uint8*>(&NamespaceGuid), sizeof(FGuid));
+			Hasher.UpdateWithString(*StringToHash, StringToHash.Len());
+			FSHAHash HashValue = Hasher.Finalize();
+			const uint32* HashUint32 = reinterpret_cast<const uint32*>(HashValue.Hash);
+
+			return FGuid(HashUint32[0], HashUint32[1], HashUint32[2], HashUint32[3]);
+		}
+	}
+}
diff --git a/Engine/Plugins/Runtime/Metasound/Source/MetasoundFrontend/Private/MetasoundFrontendGraphController.cpp b/Engine/Plugins/Runtime/Metasound/Source/MetasoundFrontend/Private/MetasoundFrontendGraphController.cpp
--- a/Engine/Plugins/Runtime/Metasound/Source/MetasoundFrontend/Private/MetasoundFrontendGraphController.cpp
+++ b/Engine/Plugins/Runtime/Metasound/Source/MetasoundFrontend/Private/MetasoundFrontendGraphController.cpp
@@ -6,6 +6,7 @@
 #include "Internationalization/Text.h"
 #include "MetasoundFrontendDocument.h"
 #include "MetasoundFrontendDocumentAccessPtr.h"
+#include "MetasoundFrontendDocumentIdGenerator.h"
 #include "MetasoundFrontendGraph.h"
 #include "MetasoundFrontendNodeController.h"
 #include "MetasoundFrontendNodeTemplateRegistry.h"
@@ -974,14 +975,15 @@ namespace Metasound
 							FMetasoundFrontendNode& Node = GraphClass->Graph.Nodes.Emplace_GetRef(*InputClass);
 							Node.Name = NewName;
 
-							if (InClassInput.NodeID.IsValid())
+							FGuid NodeID = InClassInput.NodeID;
+							if (!NodeID.IsValid())
 							{
-								Node.UpdateID(InClassInput.NodeID);
-							}
-							else
-							{
-								Node.UpdateID(FGuid::NewGuid());
+								FDocumentAccessPtr DocumentPtr = OwningDocument->GetDocumentPtr();
+								const FMetasoundFrontendDocument* Document = DocumentPtr.Get();
+								check(Document);
+								NodeID = FDocumentIDGenerator::Get().CreateNodeID(*Document);
 							}
+							Node.UpdateID(NodeID);
 
 							// Set name on related vertices of input node
 							auto IsVertexWithTypeName = [&](FMetasoundFrontendVertex& Vertex) { return Vertex.TypeName == InClassInput.TypeName; };
@@ -1011,7 +1013,10 @@ namespace Metasound
 							{
 								// Create a new guid if there wasn't a valid guid attached
 								// to input.
-								NewInput.VertexID = FGuid::NewGuid();
+								FDocumentAccessPtr DocumentPtr = OwningDocument->GetDocumentPtr();
+								const FMetasoundFrontendDocument* Document = DocumentPtr.Get();
+								check(Document);
+								NewInput.VertexID = FDocumentIDGenerator::Get().CreateVertexID(*Document);
 							}
 
 #if WITH_EDITOR
@@ -1114,14 +1119,15 @@ namespace Metasound
 							FMetasoundFrontendNode& Node = GraphClass->Graph.Nodes.Add_GetRef(*OutputClass);
 							Node.Name = NewName;
 
-							if (InClassOutput.NodeID.IsValid())
-							{
-								Node.UpdateID(InClassOutput.NodeID);
-							}
-							else
+							FGuid NodeID = InClassOutput.NodeID;
+							if (!NodeID.IsValid())
 							{
-								Node.UpdateID(FGuid::NewGuid());
+								FDocumentAccessPtr DocumentPtr = OwningDocument->GetDocumentPtr();
+								const FMetasoundFrontendDocument* Document = DocumentPtr.Get();
+								check(Document);
+								NodeID = FDocumentIDGenerator::Get().CreateNodeID(*Document);
 							}
+							Node.UpdateID(NodeID);
 
 							// Set vertex name on output node
 							auto IsVertexWithTypeName = [&](FMetasoundFrontendVertex& Vertex) { return Vertex.TypeName == InClassOutput.TypeName; };
@@ -1152,7 +1158,10 @@ namespace Metasound
 							{
 								// Create a new guid if there wasn't a valid guid attached
 								// to output.
-								NewOutput.VertexID = FGuid::NewGuid();
+								FDocumentAccessPtr DocumentPtr = OwningDocument->GetDocumentPtr();
+								const FMetasoundFrontendDocument* Document = DocumentPtr.Get();
+								check(Document);
+								NewOutput.VertexID = FDocumentIDGenerator::Get().CreateVertexID(*Document);
 							}
 
 #if WITH_EDITOR
diff --git a/Engine/Plugins/Runtime/Metasound/Source/MetasoundFrontend/Private/MetasoundFrontendInvalidController.cpp b/Engine/Plugins/Runtime/Metasound/Source/MetasoundFrontend/Private/MetasoundFrontendInvalidController.cpp
--- a/Engine/Plugins/Runtime/Metasound/Source/MetasoundFrontend/Private/MetasoundFrontendInvalidController.cpp
+++ b/Engine/Plugins/Runtime/Metasound/Source/MetasoundFrontend/Private/MetasoundFrontendInvalidController.cpp
@@ -178,5 +178,17 @@ namespace Metasound
 		{
 			return IGraphController::GetInvalidHandle();
 		}
+
+		FDocumentAccessPtr FInvalidDocumentController::GetDocumentPtr()
+		{
+			static FDocumentAccessPtr Invalid;
+			return Invalid;
+		}
+
+		const FDocumentAccessPtr FInvalidDocumentController::GetDocumentPtr() const
+		{
+			static const FDocumentAccessPtr Invalid;
+			return Invalid;
+		}
 	}
 }
diff --git a/Engine/Plugins/Runtime/Metasound/Source/MetasoundFrontend/Private/MetasoundFrontendInvalidController.h b/Engine/Plugins/Runtime/Metasound/Source/MetasoundFrontend/Private/MetasoundFrontendInvalidController.h
--- a/Engine/Plugins/Runtime/Metasound/Source/MetasoundFrontend/Private/MetasoundFrontendInvalidController.h
+++ b/Engine/Plugins/Runtime/Metasound/Source/MetasoundFrontend/Private/MetasoundFrontendInvalidController.h
@@ -493,6 +493,9 @@ namespace Metasound
 
 				virtual TSharedRef<IGraphController> GetRootGraph() override;
 				virtual TSharedRef<const IGraphController> GetRootGraph() const override;
+				virtual FDocumentAccessPtr GetDocumentPtr() override;
+				virtual const FDocumentAccessPtr GetDocumentPtr() const override;
+
 				virtual bool ExportToJSONAsset(const FString& InAbsolutePath) const override { return false; }
 				virtual FString ExportToJSON() const override { return FString(); }
 
diff --git a/Engine/Plugins/Runtime/Metasound/Source/MetasoundFrontend/Private/MetasoundFrontendTransform.cpp b/Engine/Plugins/Runtime/Metasound/Source/MetasoundFrontend/Private/MetasoundFrontendTransform.cpp
--- a/Engine/Plugins/Runtime/Metasound/Source/MetasoundFrontend/Private/MetasoundFrontendTransform.cpp
+++ b/Engine/Plugins/Runtime/Metasound/Source/MetasoundFrontend/Private/MetasoundFrontendTransform.cpp
@@ -8,6 +8,7 @@
 #include "MetasoundAssetBase.h"
 #include "MetasoundFrontendDocument.h"
 #include "MetasoundFrontendDocumentController.h"
+#include "MetasoundFrontendDocumentIdGenerator.h"
 #include "MetasoundFrontendRegistries.h"
 #include "MetasoundFrontendSearchEngine.h"
 #include "MetasoundLog.h"
@@ -733,10 +734,12 @@ namespace Metasound
 			{
 				*InPresetNodeID = PresetNodeHandle->GetID();
 			}, EMetasoundFrontendClassType::External); 
-
+			
 			if (!PresetNodeID.IsValid())
 			{
-				PresetNodeID = FGuid::NewGuid();
+				// This ID was originally being set to FGuid::NewGuid. 
+				// If you were reliant on that ID, please resave the asset so it is serialized with a valid ID
+				PresetNodeID = InDocument->GetRootGraph()->GetClassID();
 			}
 
 			// Clear the root graph so it can be rebuilt.
@@ -857,7 +860,7 @@ namespace Metasound
 
 			Style.SortDefaults(NodeHandles, DocumentTransform::GetNodeDisplayNameProjection());
 
-			// Set input node location
+			// Set output node location
 			for (const FNodeHandle& OutputNode : NodeHandles)
 			{
 				FMetasoundFrontendNodeStyle NodeStyle;
@@ -891,8 +894,11 @@ namespace Metasound
 					ClassInput.Metadata.SetDescription(InputNode->GetDescription());
 					ClassInput.Metadata.SetDisplayName(Input->GetMetadata().GetDisplayName());
 #endif // WITH_EDITOR
-
-					ClassInput.VertexID = FGuid::NewGuid();
+					FConstDocumentHandle DocumentHandle = InPresetGraph->GetOwningDocument();
+					FDocumentAccessPtr DocumentPtr = DocumentHandle->GetDocumentPtr();
+					const FMetasoundFrontendDocument* Document = DocumentPtr.Get();
+					check(Document);
+					ClassInput.VertexID = FDocumentIDGenerator::Get().CreateVertexID(*Document);;
 
 					if (const FMetasoundFrontendClassInput* ExistingClassInput = InPresetGraph->FindClassInputWithName(NodeName).Get())
 					{
@@ -968,7 +974,11 @@ namespace Metasound
 					ClassOutput.Metadata.SetDisplayName(Output->GetMetadata().GetDisplayName());
 #endif // WITH_EDITOR
 
-					ClassOutput.VertexID = FGuid::NewGuid();
+					FConstDocumentHandle DocumentHandle = InPresetGraph->GetOwningDocument();
+					FDocumentAccessPtr DocumentPtr = DocumentHandle->GetDocumentPtr();
+					const FMetasoundFrontendDocument* Document = DocumentPtr.Get();
+					check(Document);
+					ClassOutput.VertexID = FDocumentIDGenerator::Get().CreateVertexID(*Document);
 
 					if (const FMetasoundFrontendClassOutput* ExistingClassOutput = InPresetGraph->FindClassOutputWithName(NodeName).Get())
 					{
diff --git a/Engine/Plugins/Runtime/Metasound/Source/MetasoundFrontend/Public/MetasoundFrontendController.h b/Engine/Plugins/Runtime/Metasound/Source/MetasoundFrontend/Public/MetasoundFrontendController.h
--- a/Engine/Plugins/Runtime/Metasound/Source/MetasoundFrontend/Public/MetasoundFrontendController.h
+++ b/Engine/Plugins/Runtime/Metasound/Source/MetasoundFrontend/Public/MetasoundFrontendController.h
@@ -1118,6 +1118,10 @@ namespace Metasound
 			/** Returns the root graph of this document. */
 			virtual FConstGraphHandle GetRootGraph() const = 0;
 
+			/** Returns the document if set. */
+			virtual FDocumentAccessPtr GetDocumentPtr() = 0;
+			virtual const FDocumentAccessPtr GetDocumentPtr() const = 0;
+
 			/** Exports the document to a json file at the provided path.
 			 *
 			 * @return True on success, false on failure. 
diff --git a/Engine/Plugins/Runtime/Metasound/Source/MetasoundFrontend/Public/MetasoundFrontendDocument.h b/Engine/Plugins/Runtime/Metasound/Source/MetasoundFrontend/Public/MetasoundFrontendDocument.h
--- a/Engine/Plugins/Runtime/Metasound/Source/MetasoundFrontend/Public/MetasoundFrontendDocument.h
+++ b/Engine/Plugins/Runtime/Metasound/Source/MetasoundFrontend/Public/MetasoundFrontendDocument.h
@@ -995,7 +995,7 @@ public:
 	TArray<FMetasoundFrontendClassEnvironmentVariable> Environment;
 
 private:
-	UPROPERTY()
+	UPROPERTY(Transient)
 	FGuid ChangeID;
 
 public:
@@ -1310,7 +1310,7 @@ private:
 
 	// ID used to identify if any of the above have been modified,
 	// to determine if the parent class should be auto-updated.
-	UPROPERTY()
+	UPROPERTY(Transient)
 	FGuid ChangeID;
 
 public:
@@ -1587,6 +1587,11 @@ public:
 	UPROPERTY()
 	TArray<FMetasoundFrontendClass> Dependencies;
 
+	uint32 GetNextIdCounter() const
+	{
+		return IdCounter++;
+	}
+
 private:
 	UPROPERTY(meta = (DeprecatedProperty, DeprecationMessage = "5.0 - ArchetypeVersion has been migrated to InterfaceVersions array."))
 	FMetasoundFrontendVersion ArchetypeVersion;
@@ -1594,6 +1599,9 @@ private:
 	UPROPERTY(meta = (DeprecatedProperty, DeprecationMessage = "5.0 - InterfaceVersions has been migrated to Interfaces set."))
 	TArray<FMetasoundFrontendVersion> InterfaceVersions;
 
+	// Used for generating deterministic IDs per document
+	mutable uint32 IdCounter = 1;
+
 public:
 	// Data migration for 5.0 Early Access data. ArchetypeVersion/InterfaceVersions properties can be removed post 5.0 release
 	// and this fix-up can be removed post 5.0 release.
diff --git a/Engine/Plugins/Runtime/Metasound/Source/MetasoundFrontend/Public/MetasoundFrontendDocumentIdGenerator.h b/Engine/Plugins/Runtime/Metasound/Source/MetasoundFrontend/Public/MetasoundFrontendDocumentIdGenerator.h
new file mode 100644
--- /dev/null
+++ b/Engine/Plugins/Runtime/Metasound/Source/MetasoundFrontend/Public/MetasoundFrontendDocumentIdGenerator.h
@@ -0,0 +1,92 @@
+// Copyright Epic Games, Inc. All Rights Reserved.
+
+#pragma once
+#include "MetasoundFrontendDocument.h"
+#include "HAL/IConsoleManager.h"
+#include "Misc/Guid.h"
+
+namespace Metasound
+{
+	namespace Frontend
+	{
+		extern METASOUNDFRONTEND_API int32 MetaSoundEnableDeterministicIDGenerationInEditorCVar;
+
+		/*** 
+		 * For generating IDs using a given document. 
+		 * USAGE:
+		 *
+		 * If you want everything within the calling scope to be deterministic, use
+		 * the scope determinism lock like you would a `FScopeLock`
+		 *
+		 * {
+		 * 		constexpr bool bIsDeterministic = true;
+		 * 		FDocumentIDGenerator::FScopeDeterminism DeterminismScope(bIsDeterministic);
+		 *
+		 *      // Anything called in this scope will use a deterministic ID generator.
+		 * 		// Once the `DeterminismScope` variable is destroyed, it will return to
+		 * 		// whatever the prior behavior was
+		 * 		MetaSoundAsset->UpdateOrWhatever();
+		 * }
+		 *
+		 * void FMetaSoundAsset::UpdateOrWhatever()
+		 * {
+		 *      // Anytime you call IDGen::Create*ID(..), it will obey whatever the outside scope set it to be.
+		 *      AddNewNodeOrWhatever(...)
+		 * }
+		 *
+		 * void FMetaSoundAsset::AddNodeOrWhatever(...)
+		 * {
+		 * 		FDocumentIDGenerator& IDGen = FDocumentIDGenerator::Get();
+		 *
+		 * 		FGuid NewNodeID = IDGen::CreateNewNodeID();
+		 * 		...
+		 * }
+		 *
+		 */
+		class FDocumentIDGenerator
+		{
+		public:
+			class METASOUNDFRONTEND_API FScopeDeterminism final
+			{
+			public:
+				FScopeDeterminism(bool bInIsDeterministic);
+				bool GetDeterminism() const;
+				~FScopeDeterminism();
+
+			private:
+				bool bOriginalValue = false;
+			};
+
+			static FDocumentIDGenerator& Get();
+
+			FGuid CreateNodeID(const FMetasoundFrontendDocument& InDocument) const;
+			FGuid CreateVertexID(const FMetasoundFrontendDocument& InDocument) const;
+			FGuid CreateClassID(const FMetasoundFrontendDocument& InDocument) const;
+
+			FGuid CreateIDFromDocument(const FMetasoundFrontendDocument& InDocument) const;
+
+		private:
+			FDocumentIDGenerator() = default;
+
+			friend class FScopeDeterminism;
+			void SetDeterminism(bool bInIsDeterministic);
+			bool GetDeterminism() const;
+
+			bool bIsDeterministic = false;
+		};
+
+		// For generating IDs that are not derived from a given document. 
+		class FClassIDGenerator
+		{
+		public: 
+			static FClassIDGenerator& Get();
+
+			FGuid CreateInputID(const FMetasoundFrontendClassInput& Input) const;
+			FGuid CreateInputID(const Audio::FParameterInterface::FInput& Input) const;
+			FGuid CreateOutputID(const FMetasoundFrontendClassOutput& Output) const;
+			FGuid CreateOutputID(const Audio::FParameterInterface::FOutput& Output) const;
+
+			FGuid CreateNamespacedIDFromString(const FGuid NamespaceGuid, FString StringToHash) const;
+		};
+	}
+}
