- Fix-up MetaSound dynamic builder engine tests
- Add missing move of transactor to generator init params

#rnx

diff --git a/Engine/Plugins/Runtime/Metasound/Source/MetasoundEngine/Private/MetasoundSource.cpp b/Engine/Plugins/Runtime/Metasound/Source/MetasoundEngine/Private/MetasoundSource.cpp
--- a/Engine/Plugins/Runtime/Metasound/Source/MetasoundEngine/Private/MetasoundSource.cpp
+++ b/Engine/Plugins/Runtime/Metasound/Source/MetasoundEngine/Private/MetasoundSource.cpp
@@ -605,7 +605,8 @@ ISoundGeneratorPtr UMetaSoundSource::CreateSoundGenerator(const FSoundGeneratorI
 		GetOutputAudioChannelOrder(),
 		MoveTemp(InDefaultParameters),
 		bBuildSynchronous,
-		DataChannel
+		DataChannel,
+		DynamicTransactor
 	};
 
 	TSharedPtr<FMetasoundGenerator> Generator = MakeShared<FMetasoundGenerator>(MoveTemp(InitParams));
diff --git a/Engine/Plugins/Runtime/Metasound/Source/MetasoundEngineTest/Private/EngineTestMetasoundSource.cpp b/Engine/Plugins/Runtime/Metasound/Source/MetasoundEngineTest/Private/EngineTestMetasoundSource.cpp
--- a/Engine/Plugins/Runtime/Metasound/Source/MetasoundEngineTest/Private/EngineTestMetasoundSource.cpp
+++ b/Engine/Plugins/Runtime/Metasound/Source/MetasoundEngineTest/Private/EngineTestMetasoundSource.cpp
@@ -7,13 +7,14 @@
 #include "Interfaces/MetasoundOutputFormatInterfaces.h"
 #include "Interfaces/MetasoundFrontendSourceInterface.h"
 #include "MetasoundBuilderSubsystem.h"
+#include "MetasoundDataReference.h"
 #include "MetasoundFrontendController.h"
 #include "MetasoundFrontendSearchEngine.h"
 #include "MetasoundSource.h"
 #include "Misc/AutomationTest.h"
 #include "Misc/Paths.h"
 #include "Tests/AutomationCommon.h"
-#include "../../MetasoundGraphCore/Public/MetasoundDataReference.h"
+
 
 #if WITH_DEV_AUTOMATION_TESTS
 
@@ -459,13 +460,14 @@ bool FAudioMetasoundSourceBuilderLiveUpdateNode::RunTest(const FString& Paramete
 			ADD_LATENT_AUTOMATION_COMMAND(FMetaSoundSourceBuilderAuditionLatentCommand(&Builder, AudioComponent, bEnableLiveUpdate));
 			ADD_LATENT_AUTOMATION_COMMAND(FEngineWaitLatentCommand(2.f));
 
-			// Disconnects graph audio output from existing sinosc output
+			// Disconnect graph audio output from existing sinosc output and connect to added triosc
 			ADD_LATENT_AUTOMATION_COMMAND(FMetaSoundSourceBuilderDisconnectInputLatentCommand(*this, &Builder, MonoOutNodeInput));
-
 			ADD_LATENT_AUTOMATION_COMMAND(FMetaSoundSourceBuilderCreateAndConnectTriGeneratorNodeLatentCommand(*this, &Builder, MonoOutNodeInput));
 
-			// TODO: Once we have the builder hooked up to the dynamic graph backend, remove this and voila!
-			ADD_LATENT_AUTOMATION_COMMAND(FMetaSoundSourceBuilderAuditionLatentCommand(&Builder, AudioComponent, bEnableLiveUpdate));
+			FName DataTypeName;
+			FMetasoundFrontendLiteral NewValue;
+			NewValue.Set(FMetasoundFrontendLiteral::FDefault{ });
+			ADD_LATENT_AUTOMATION_COMMAND(FMetaSoundSourceBuilderSetLiteralLatentCommand(*this, &Builder, MonoOutNodeInput, NewValue));
 
 			ADD_LATENT_AUTOMATION_COMMAND(FEngineWaitLatentCommand(2.f));
 			ADD_LATENT_AUTOMATION_COMMAND(FAudioComponentStopLatentCommand(AudioComponent));
@@ -510,24 +512,21 @@ bool FAudioMetasoundSourceBuilderLiveUpdateLiteral::RunTest(const FString& Param
 			ADD_LATENT_AUTOMATION_COMMAND(FMetaSoundSourceBuilderAuditionLatentCommand(&Builder, AudioComponent, bEnableLiveUpdate));
 			ADD_LATENT_AUTOMATION_COMMAND(FEngineWaitLatentCommand(0.25f));
 
-			// TODO: Once we have the builder hooked up to the dynamic graph backend, remove additional calls to
-			// FMetaSoundSourceBuilderAuditionLatentCommand. The dynamic graph should automatically make changes audible.
-
-			// Disconnects graph audio output from existing sinosc output
+			// Disconnects freq input node output from sinosc freq input. Initially was set to 220Hz above, and node's default is 440Hz,
+			// resulting in an octive pitch up.
 			ADD_LATENT_AUTOMATION_COMMAND(FMetaSoundSourceBuilderDisconnectInputLatentCommand(*this, &Builder, GenNodeFreqInput));
-			ADD_LATENT_AUTOMATION_COMMAND(FMetaSoundSourceBuilderAuditionLatentCommand(&Builder, AudioComponent, bEnableLiveUpdate));
 			ADD_LATENT_AUTOMATION_COMMAND(FEngineWaitLatentCommand(0.25f));
 
+			// Sets literal value on the sinosc freq input to 880Hz, pitching an octive yet again from previous.
 			FName DataTypeName;
 			FMetasoundFrontendLiteral NewValue = UMetaSoundBuilderSubsystem::GetChecked().CreateFloatMetaSoundLiteral(880.f, DataTypeName);
 			AddErrorIfFalse(DataTypeName == Metasound::GetMetasoundDataTypeName<float>(),
 				"Setting MetaSound Float literal returns non-float DataTypeName.");
 			ADD_LATENT_AUTOMATION_COMMAND(FMetaSoundSourceBuilderSetLiteralLatentCommand(*this, &Builder, GenNodeFreqInput, NewValue));
-			ADD_LATENT_AUTOMATION_COMMAND(FMetaSoundSourceBuilderAuditionLatentCommand(&Builder, AudioComponent, bEnableLiveUpdate));
 			ADD_LATENT_AUTOMATION_COMMAND(FEngineWaitLatentCommand(0.25f));
 
+			// Removes the literal value on the sinosc freq input set to 880Hz, reverting back to the class literal of 440Hz.
 			ADD_LATENT_AUTOMATION_COMMAND(FMetaSoundSourceBuilderRemoveNodeDefaultLiteralLatentCommand(*this, &Builder, GenNodeFreqInput));
-			ADD_LATENT_AUTOMATION_COMMAND(FMetaSoundSourceBuilderAuditionLatentCommand(&Builder, AudioComponent, bEnableLiveUpdate));
 			ADD_LATENT_AUTOMATION_COMMAND(FEngineWaitLatentCommand(0.25f));
 
 			ADD_LATENT_AUTOMATION_COMMAND(FAudioComponentStopLatentCommand(AudioComponent));
