[Metasound Bind] Fixup for bind issues caught by new automated tests.

diff --git a/Engine/Plugins/Runtime/AudioModulation/Source/AudioModulation/Private/MetaSoundNodes/GetModulatorValueNode.cpp b/Engine/Plugins/Runtime/AudioModulation/Source/AudioModulation/Private/MetaSoundNodes/GetModulatorValueNode.cpp
--- a/Engine/Plugins/Runtime/AudioModulation/Source/AudioModulation/Private/MetaSoundNodes/GetModulatorValueNode.cpp
+++ b/Engine/Plugins/Runtime/AudioModulation/Source/AudioModulation/Private/MetaSoundNodes/GetModulatorValueNode.cpp
@@ -93,25 +93,35 @@ namespace AudioModulation
 
 		virtual ~FGetModulatorValueNodeOperator() = default;
 
-		virtual Metasound::FDataReferenceCollection GetInputs() const override
+		virtual void BindInputs(Metasound::FInputVertexInterfaceData& InOutVertexData) override
 		{
 			using namespace Metasound;
-
-			FDataReferenceCollection Inputs;
-			Inputs.AddDataReadReference("Modulator", Modulator);
-			Inputs.AddDataReadReference("Normalized", Normalized);
-
-			return Inputs;
+			
+			InOutVertexData.BindReadVertex("Modulator", Modulator);
+			InOutVertexData.BindReadVertex("Normalized", Normalized);
 		}
 
-		virtual Metasound::FDataReferenceCollection GetOutputs() const override
+		virtual void BindOutputs(Metasound::FOutputVertexInterfaceData& InOutVertexData) override
 		{
 			using namespace Metasound;
+			
+			InOutVertexData.BindReadVertex("Out", OutValue);
+		}
 
-			FDataReferenceCollection Outputs;
-			Outputs.AddDataReadReference("Out", FFloatReadRef(OutValue));
+		virtual Metasound::FDataReferenceCollection GetInputs() const override
+		{
+			// This should never be called. Bind(...) is called instead. This method
+			// exists as a stop-gap until the API can be deprecated and removed.
+			checkNoEntry();
+			return {};
+		}
 
-			return Outputs;
+		virtual Metasound::FDataReferenceCollection GetOutputs() const override
+		{
+			// This should never be called. Bind(...) is called instead. This method
+			// exists as a stop-gap until the API can be deprecated and removed.
+			checkNoEntry();
+			return {};
 		}
 
 		void Execute()
diff --git a/Engine/Plugins/Runtime/AudioModulation/Source/AudioModulation/Private/MetaSoundNodes/MixModulatorsNode.cpp b/Engine/Plugins/Runtime/AudioModulation/Source/AudioModulation/Private/MetaSoundNodes/MixModulatorsNode.cpp
--- a/Engine/Plugins/Runtime/AudioModulation/Source/AudioModulation/Private/MetaSoundNodes/MixModulatorsNode.cpp
+++ b/Engine/Plugins/Runtime/AudioModulation/Source/AudioModulation/Private/MetaSoundNodes/MixModulatorsNode.cpp
@@ -105,27 +105,37 @@ namespace AudioModulation
 
 		virtual ~FMixModulatorsNodeOperator() = default;
 
-		virtual Metasound::FDataReferenceCollection GetInputs() const override
+		virtual void BindInputs(Metasound::FInputVertexInterfaceData& InOutVertexData) override
 		{
 			using namespace Metasound;
-
-			FDataReferenceCollection Inputs;
-			Inputs.AddDataReadReference("In1", Modulator1);
-			Inputs.AddDataReadReference("In2", Modulator2);
-			Inputs.AddDataReadReference("MixParameter", Parameter);
-			Inputs.AddDataReadReference("Normalized", Normalized);
-
-			return Inputs;
+			
+			InOutVertexData.BindReadVertex("In1", Modulator1);
+			InOutVertexData.BindReadVertex("In2", Modulator2);
+			InOutVertexData.BindReadVertex("MixParameter", Parameter);
+			InOutVertexData.BindReadVertex("Normalized", Normalized);
 		}
 
-		virtual Metasound::FDataReferenceCollection GetOutputs() const override
+		virtual void BindOutputs(Metasound::FOutputVertexInterfaceData& InOutVertexData) override
 		{
 			using namespace Metasound;
+			
+			InOutVertexData.BindReadVertex("Out", OutValue);
+		}
 
-			FDataReferenceCollection Outputs;
-			Outputs.AddDataReadReference("Out", FFloatReadRef(OutValue));
+		virtual Metasound::FDataReferenceCollection GetInputs() const override
+		{
+			// This should never be called. Bind(...) is called instead. This method
+			// exists as a stop-gap until the API can be deprecated and removed.
+			checkNoEntry();
+			return {};
+		}
 
-			return Outputs;
+		virtual Metasound::FDataReferenceCollection GetOutputs() const override
+		{
+			// This should never be called. Bind(...) is called instead. This method
+			// exists as a stop-gap until the API can be deprecated and removed.
+			checkNoEntry();
+			return {};
 		}
 
 		void Execute()
diff --git a/Engine/Plugins/Runtime/Metasound/Source/MetasoundEngine/Private/MetasoundAudioBusReaderNode.cpp b/Engine/Plugins/Runtime/Metasound/Source/MetasoundEngine/Private/MetasoundAudioBusReaderNode.cpp
--- a/Engine/Plugins/Runtime/Metasound/Source/MetasoundEngine/Private/MetasoundAudioBusReaderNode.cpp
+++ b/Engine/Plugins/Runtime/Metasound/Source/MetasoundEngine/Private/MetasoundAudioBusReaderNode.cpp
@@ -116,27 +116,35 @@ namespace Metasound
 			Reset(InParams);
 		}
 
-		virtual FDataReferenceCollection GetInputs() const override
+		virtual void BindInputs(FInputVertexInterfaceData& InOutVertexData) override
 		{
 			using namespace AudioBusReaderNode;
-
-			FDataReferenceCollection InputDataReferences;
-			InputDataReferences.AddDataReadReference(METASOUND_GET_PARAM_NAME(InParamAudioBusInput), FAudioBusAssetReadRef(AudioBusAsset));
-			return InputDataReferences;
+			InOutVertexData.BindReadVertex(METASOUND_GET_PARAM_NAME(InParamAudioBusInput), AudioBusAsset);
 		}
 
-		virtual FDataReferenceCollection GetOutputs() const override
+		virtual void BindOutputs(FOutputVertexInterfaceData& InOutVertexData) override
 		{
 			using namespace AudioBusReaderNode;
-
-			FDataReferenceCollection OutputDataReferences;
-
 			for (int32 ChannelIndex = 0; ChannelIndex < NumChannels; ++ChannelIndex)
 			{
-				OutputDataReferences.AddDataReadReference(METASOUND_GET_PARAM_NAME_WITH_INDEX(OutParamAudio, ChannelIndex), FAudioBufferWriteRef(AudioOutputs[ChannelIndex]));
+				InOutVertexData.BindReadVertex(METASOUND_GET_PARAM_NAME_WITH_INDEX(OutParamAudio, ChannelIndex), AudioOutputs[ChannelIndex]);
 			}
-			
-			return OutputDataReferences;			
+		}
+
+		virtual FDataReferenceCollection GetInputs() const override
+		{
+			// This should never be called. Bind(...) is called instead. This method
+			// exists as a stop-gap until the API can be deprecated and removed.
+			checkNoEntry();
+			return {};
+		}
+
+		virtual FDataReferenceCollection GetOutputs() const override
+		{
+			// This should never be called. Bind(...) is called instead. This method
+			// exists as a stop-gap until the API can be deprecated and removed.
+			checkNoEntry();
+			return {};
 		}
 
 		void Execute()
diff --git a/Engine/Plugins/Runtime/Metasound/Source/MetasoundEngine/Private/MetasoundAudioBusWriterNode.cpp b/Engine/Plugins/Runtime/Metasound/Source/MetasoundEngine/Private/MetasoundAudioBusWriterNode.cpp
--- a/Engine/Plugins/Runtime/Metasound/Source/MetasoundEngine/Private/MetasoundAudioBusWriterNode.cpp
+++ b/Engine/Plugins/Runtime/Metasound/Source/MetasoundEngine/Private/MetasoundAudioBusWriterNode.cpp
@@ -167,28 +167,37 @@ namespace Metasound
 			bFirstBlock = true;
 		}
 
-		virtual FDataReferenceCollection GetInputs() const override
+
+		virtual void BindInputs(FInputVertexInterfaceData& InOutVertexData) override
 		{
 			using namespace AudioBusWriterNode;
 
-			FDataReferenceCollection InputDataReferences;
-
-			InputDataReferences.AddDataReadReference(METASOUND_GET_PARAM_NAME(InParamAudioBusOutput), AudioBusAsset);
+			InOutVertexData.BindReadVertex(METASOUND_GET_PARAM_NAME(InParamAudioBusOutput), AudioBusAsset);
 
 			for (int32 ChannelIndex = 0; ChannelIndex < NumChannels; ++ChannelIndex)
 			{
-				InputDataReferences.AddDataReadReference(METASOUND_GET_PARAM_NAME_WITH_INDEX(InParamAudio, ChannelIndex), FAudioBufferReadRef(AudioInputs[ChannelIndex]));
+				InOutVertexData.BindReadVertex(METASOUND_GET_PARAM_NAME_WITH_INDEX(InParamAudio, ChannelIndex), AudioInputs[ChannelIndex]);
 			}
+		}
 
-			return InputDataReferences;
+		virtual void BindOutputs(FOutputVertexInterfaceData& InOutVertexData) override
+		{
 		}
 
-		virtual FDataReferenceCollection GetOutputs() const override
+		virtual FDataReferenceCollection GetInputs() const override
 		{
-			using namespace AudioBusWriterNode;
+			// This should never be called. Bind(...) is called instead. This method
+			// exists as a stop-gap until the API can be deprecated and removed.
+			checkNoEntry();
+			return {};
+		}
 
-			FDataReferenceCollection OutputDataReferences;
-			return OutputDataReferences;
+		virtual FDataReferenceCollection GetOutputs() const override
+		{
+			// This should never be called. Bind(...) is called instead. This method
+			// exists as a stop-gap until the API can be deprecated and removed.
+			checkNoEntry();
+			return {};
 		}
 
 		void Execute()
diff --git a/Engine/Plugins/Runtime/Metasound/Source/MetasoundEngine/Private/MetasoundWaveInfoNode.cpp b/Engine/Plugins/Runtime/Metasound/Source/MetasoundEngine/Private/MetasoundWaveInfoNode.cpp
--- a/Engine/Plugins/Runtime/Metasound/Source/MetasoundEngine/Private/MetasoundWaveInfoNode.cpp
+++ b/Engine/Plugins/Runtime/Metasound/Source/MetasoundEngine/Private/MetasoundWaveInfoNode.cpp
@@ -47,6 +47,8 @@ namespace Metasound
 		static const FNodeClassMetadata& GetNodeInfo();
 		static FVertexInterface DeclareVertexInterface();
 		static TUniquePtr<IOperator> CreateOperator(const FCreateOperatorParams& InParams, FBuildErrorArray& OutErrors);
+		virtual void BindInputs(FInputVertexInterfaceData& InOutVertexData) override;
+		virtual void BindOutputs(FOutputVertexInterfaceData& InOutVertexData) override;
 		virtual FDataReferenceCollection GetInputs() const override;
 		virtual FDataReferenceCollection GetOutputs() const override;
 		void Execute();
@@ -128,23 +130,35 @@ namespace Metasound
 		return MakeUnique < FWaveInfoNodeOperator >(InParams.OperatorSettings, WaveAssetIn);
 	}
 
-	FDataReferenceCollection FWaveInfoNodeOperator::GetInputs() const
+	void FWaveInfoNodeOperator::BindInputs(FInputVertexInterfaceData& InOutVertexData)
 	{
 		using namespace WaveInfoNodeParameterNames;
 		FDataReferenceCollection InputDataReferences;
-		InputDataReferences.AddDataReadReference(METASOUND_GET_PARAM_NAME(ParamWaveAsset), FWaveAssetReadRef(WaveAsset));
-
-		return InputDataReferences;
+		InOutVertexData.BindReadVertex(METASOUND_GET_PARAM_NAME(ParamWaveAsset), WaveAsset);	
 	}
 
-	FDataReferenceCollection FWaveInfoNodeOperator::GetOutputs() const
+	void FWaveInfoNodeOperator::BindOutputs(FOutputVertexInterfaceData& InOutVertexData)
 	{
 		// expose read access to our output buffer for other processors in the graph
 		using namespace WaveInfoNodeParameterNames;
 		FDataReferenceCollection OutputDataReferences;
-		OutputDataReferences.AddDataReadReference(METASOUND_GET_PARAM_NAME(ParamDurationSeconds), DurationSeconds);
+		InOutVertexData.BindReadVertex(METASOUND_GET_PARAM_NAME(ParamDurationSeconds), DurationSeconds);
+	}
+
+	FDataReferenceCollection FWaveInfoNodeOperator::GetInputs() const
+	{
+		// This should never be called. Bind(...) is called instead. This method
+		// exists as a stop-gap until the API can be deprecated and removed.
+		checkNoEntry();
+		return {};
+	}
 
-		return OutputDataReferences;
+	FDataReferenceCollection FWaveInfoNodeOperator::GetOutputs() const
+	{
+		// This should never be called. Bind(...) is called instead. This method
+		// exists as a stop-gap until the API can be deprecated and removed.
+		checkNoEntry();
+		return {};
 	}
 
 	void FWaveInfoNodeOperator::Execute()
diff --git a/Engine/Plugins/Runtime/Metasound/Source/MetasoundEngine/Private/MetasoundWavePlayerNode.cpp b/Engine/Plugins/Runtime/Metasound/Source/MetasoundEngine/Private/MetasoundWavePlayerNode.cpp
--- a/Engine/Plugins/Runtime/Metasound/Source/MetasoundEngine/Private/MetasoundWavePlayerNode.cpp
+++ b/Engine/Plugins/Runtime/Metasound/Source/MetasoundEngine/Private/MetasoundWavePlayerNode.cpp
@@ -356,37 +356,36 @@ namespace Metasound
 			}
 		}
 
-		virtual FDataReferenceCollection GetInputs() const override
+		virtual void BindInputs(FInputVertexInterfaceData& InOutVertexData) override
 		{
 			using namespace WavePlayerVertexNames;
 
 			FDataReferenceCollection InputDataReferences;
-			InputDataReferences.AddDataReadReference(METASOUND_GET_PARAM_NAME(InputTriggerPlay), PlayTrigger);
-			InputDataReferences.AddDataReadReference(METASOUND_GET_PARAM_NAME(InputTriggerStop), StopTrigger);
-			InputDataReferences.AddDataReadReference(METASOUND_GET_PARAM_NAME(InputWaveAsset), WaveAsset);
-			InputDataReferences.AddDataReadReference(METASOUND_GET_PARAM_NAME(InputStartTime), StartTime);
-			InputDataReferences.AddDataReadReference(METASOUND_GET_PARAM_NAME(InputPitchShift), PitchShift);
-			InputDataReferences.AddDataReadReference(METASOUND_GET_PARAM_NAME(InputLoop), bLoop);
-			InputDataReferences.AddDataReadReference(METASOUND_GET_PARAM_NAME(InputLoopStart), LoopStartTime);
-			InputDataReferences.AddDataReadReference(METASOUND_GET_PARAM_NAME(InputLoopDuration), LoopDuration);
-			return InputDataReferences;
+			InOutVertexData.BindReadVertex(METASOUND_GET_PARAM_NAME(InputTriggerPlay), PlayTrigger);
+			InOutVertexData.BindReadVertex(METASOUND_GET_PARAM_NAME(InputTriggerStop), StopTrigger);
+			InOutVertexData.BindReadVertex(METASOUND_GET_PARAM_NAME(InputWaveAsset), WaveAsset);
+			InOutVertexData.BindReadVertex(METASOUND_GET_PARAM_NAME(InputStartTime), StartTime);
+			InOutVertexData.BindReadVertex(METASOUND_GET_PARAM_NAME(InputPitchShift), PitchShift);
+			InOutVertexData.BindReadVertex(METASOUND_GET_PARAM_NAME(InputLoop), bLoop);
+			InOutVertexData.BindReadVertex(METASOUND_GET_PARAM_NAME(InputLoopStart), LoopStartTime);
+			InOutVertexData.BindReadVertex(METASOUND_GET_PARAM_NAME(InputLoopDuration), LoopDuration);
 		}
 
-		virtual FDataReferenceCollection GetOutputs() const override
+		virtual void BindOutputs(FOutputVertexInterfaceData& InOutVertexData) override
 		{
 			using namespace WavePlayerVertexNames;
 
 			FDataReferenceCollection OutputDataReferences;
-			OutputDataReferences.AddDataReadReference(METASOUND_GET_PARAM_NAME(OutputTriggerOnPlay), PlayTrigger);
-			OutputDataReferences.AddDataReadReference(METASOUND_GET_PARAM_NAME(OutputTriggerOnDone), TriggerOnDone);
-			OutputDataReferences.AddDataReadReference(METASOUND_GET_PARAM_NAME(OutputTriggerOnNearlyDone), TriggerOnNearlyDone);
-			OutputDataReferences.AddDataReadReference(METASOUND_GET_PARAM_NAME(OutputTriggerOnLooped), TriggerOnLooped);
-			OutputDataReferences.AddDataReadReference(METASOUND_GET_PARAM_NAME(OutputTriggerOnCuePoint), TriggerOnCuePoint);
-			OutputDataReferences.AddDataReadReference(METASOUND_GET_PARAM_NAME(OutputCuePointID), CuePointID);
-			OutputDataReferences.AddDataReadReference(METASOUND_GET_PARAM_NAME(OutputCuePointLabel), CuePointLabel);
-			OutputDataReferences.AddDataReadReference(METASOUND_GET_PARAM_NAME(OutputLoopRatio), LoopPercent);
-			OutputDataReferences.AddDataReadReference(METASOUND_GET_PARAM_NAME(OutputPlaybackLocation), PlaybackLocation);
-			OutputDataReferences.AddDataReadReference(METASOUND_GET_PARAM_NAME(OutputPlaybackTime), PlaybackTime);
+			InOutVertexData.BindReadVertex(METASOUND_GET_PARAM_NAME(OutputTriggerOnPlay), PlayTrigger);
+			InOutVertexData.BindReadVertex(METASOUND_GET_PARAM_NAME(OutputTriggerOnDone), TriggerOnDone);
+			InOutVertexData.BindReadVertex(METASOUND_GET_PARAM_NAME(OutputTriggerOnNearlyDone), TriggerOnNearlyDone);
+			InOutVertexData.BindReadVertex(METASOUND_GET_PARAM_NAME(OutputTriggerOnLooped), TriggerOnLooped);
+			InOutVertexData.BindReadVertex(METASOUND_GET_PARAM_NAME(OutputTriggerOnCuePoint), TriggerOnCuePoint);
+			InOutVertexData.BindReadVertex(METASOUND_GET_PARAM_NAME(OutputCuePointID), CuePointID);
+			InOutVertexData.BindReadVertex(METASOUND_GET_PARAM_NAME(OutputCuePointLabel), CuePointLabel);
+			InOutVertexData.BindReadVertex(METASOUND_GET_PARAM_NAME(OutputLoopRatio), LoopPercent);
+			InOutVertexData.BindReadVertex(METASOUND_GET_PARAM_NAME(OutputPlaybackLocation), PlaybackLocation);
+			InOutVertexData.BindReadVertex(METASOUND_GET_PARAM_NAME(OutputPlaybackTime), PlaybackTime);
 
 			check(OutputAudioBuffers.Num() == OutputAudioBufferVertexNames.Num());
 
@@ -394,8 +393,22 @@ namespace Metasound
 			{
 				OutputDataReferences.AddDataReadReference(OutputAudioBufferVertexNames[i], OutputAudioBuffers[i]);
 			}
-			
-			return OutputDataReferences;
+		}
+
+		virtual FDataReferenceCollection GetInputs() const override
+		{
+			// This should never be called. Bind(...) is called instead. This method
+			// exists as a stop-gap until the API can be deprecated and removed.
+			checkNoEntry();
+			return {};
+		}
+
+		virtual FDataReferenceCollection GetOutputs() const override
+		{
+			// This should never be called. Bind(...) is called instead. This method
+			// exists as a stop-gap until the API can be deprecated and removed.
+			checkNoEntry();
+			return {};
 		}
 
 
diff --git a/Engine/Plugins/Runtime/Metasound/Source/MetasoundEngine/Private/MetasoundWaveTableEnvelopeNode.cpp b/Engine/Plugins/Runtime/Metasound/Source/MetasoundEngine/Private/MetasoundWaveTableEnvelopeNode.cpp
--- a/Engine/Plugins/Runtime/Metasound/Source/MetasoundEngine/Private/MetasoundWaveTableEnvelopeNode.cpp
+++ b/Engine/Plugins/Runtime/Metasound/Source/MetasoundEngine/Private/MetasoundWaveTableEnvelopeNode.cpp
@@ -139,27 +139,38 @@ namespace Metasound
 
 		virtual ~FMetasoundWaveTableEnvelopeNodeOperator() = default;
 
-		virtual FDataReferenceCollection GetInputs() const override
+		virtual void BindInputs(FInputVertexInterfaceData& InOutVertexData) override
 		{
-			FDataReferenceCollection Inputs;
-			Inputs.AddDataReadReference("WaveTable", WaveTableReadRef);
-			Inputs.AddDataReadReference("Play", PlayReadRef);
-			Inputs.AddDataReadReference("Stop", StopReadRef);
-			Inputs.AddDataReadReference("Pause", PauseReadRef);
-			Inputs.AddDataReadReference("Duration", DurationReadRef);
-			Inputs.AddDataReadReference("Mode", ModeReadRef);
+			
+			InOutVertexData.BindReadVertex("WaveTable", WaveTableReadRef);
+			InOutVertexData.BindReadVertex("Play", PlayReadRef);
+			InOutVertexData.BindReadVertex("Stop", StopReadRef);
+			InOutVertexData.BindReadVertex("Pause", PauseReadRef);
+			InOutVertexData.BindReadVertex("Duration", DurationReadRef);
+			InOutVertexData.BindReadVertex("Mode", ModeReadRef);
+			InOutVertexData.BindReadVertex("Interpolation", InterpModeReadRef);
+		}
 
+		virtual void BindOutputs(FOutputVertexInterfaceData& InOutVertexData) override
+		{
+			InOutVertexData.BindReadVertex("OnFinished", OnFinishedWriteRef);
+			InOutVertexData.BindReadVertex("Out", OutWriteRef);
+		}
 
-			return Inputs;
+		virtual FDataReferenceCollection GetInputs() const override
+		{
+			// This should never be called. Bind(...) is called instead. This method
+			// exists as a stop-gap until the API can be deprecated and removed.
+			checkNoEntry();
+			return {};
 		}
 
 		virtual FDataReferenceCollection GetOutputs() const override
 		{
-			FDataReferenceCollection Outputs;
-			Outputs.AddDataReadReference("OnFinished", TDataReadReference<FTrigger>(OnFinishedWriteRef));
-			Outputs.AddDataReadReference("Out", TDataReadReference<float>(OutWriteRef));
-
-			return Outputs;
+			// This should never be called. Bind(...) is called instead. This method
+			// exists as a stop-gap until the API can be deprecated and removed.
+			checkNoEntry();
+			return {};
 		}
 
 		float GetStopValue(WaveTable::FWaveTableSampler::ESingleSampleMode InSampleMode) const
diff --git a/Engine/Plugins/Runtime/Metasound/Source/MetasoundEngine/Private/MetasoundWaveTableGetNode.cpp b/Engine/Plugins/Runtime/Metasound/Source/MetasoundEngine/Private/MetasoundWaveTableGetNode.cpp
--- a/Engine/Plugins/Runtime/Metasound/Source/MetasoundEngine/Private/MetasoundWaveTableGetNode.cpp
+++ b/Engine/Plugins/Runtime/Metasound/Source/MetasoundEngine/Private/MetasoundWaveTableGetNode.cpp
@@ -94,25 +94,33 @@ namespace Metasound
 
 		virtual ~FMetasoundWaveTableGetNodeOperator() = default;
 
-		virtual FDataReferenceCollection GetInputs() const override
+		virtual void BindInputs(FInputVertexInterfaceData& InOutVertexData) override
 		{
 			using namespace Metasound;
-
-			FDataReferenceCollection Inputs;
-			Inputs.AddDataReadReference("WaveTableBank", WaveTableBankReadRef);
-			Inputs.AddDataReadReference("TableIndex", TableIndexReadRef);
-
-			return Inputs;
+			InOutVertexData.BindReadVertex("WaveTableBank", WaveTableBankReadRef);
+			InOutVertexData.BindReadVertex("TableIndex", TableIndexReadRef);
 		}
 
-		virtual FDataReferenceCollection GetOutputs() const override
+		virtual void BindOutputs(FOutputVertexInterfaceData& InOutVertexData) override
 		{
 			using namespace Metasound;
+			InOutVertexData.BindReadVertex("Out", OutTable);
+		}
 
-			FDataReferenceCollection Outputs;
-			Outputs.AddDataReadReference("Out", FWaveTableWriteRef(OutTable));
+		virtual FDataReferenceCollection GetInputs() const override
+		{
+			// This should never be called. Bind(...) is called instead. This method
+			// exists as a stop-gap until the API can be deprecated and removed.
+			checkNoEntry();
+			return {};
+		}
 
-			return Outputs;
+		virtual FDataReferenceCollection GetOutputs() const override
+		{
+			// This should never be called. Bind(...) is called instead. This method
+			// exists as a stop-gap until the API can be deprecated and removed.
+			checkNoEntry();
+			return {};
 		}
 
 	private:
diff --git a/Engine/Plugins/Runtime/Metasound/Source/MetasoundEngine/Private/MetasoundWaveTableOscillatorNode.cpp b/Engine/Plugins/Runtime/Metasound/Source/MetasoundEngine/Private/MetasoundWaveTableOscillatorNode.cpp
--- a/Engine/Plugins/Runtime/Metasound/Source/MetasoundEngine/Private/MetasoundWaveTableOscillatorNode.cpp
+++ b/Engine/Plugins/Runtime/Metasound/Source/MetasoundEngine/Private/MetasoundWaveTableOscillatorNode.cpp
@@ -127,28 +127,39 @@ namespace Metasound
 
 		virtual ~FMetasoundWaveTableOscillatorNodeOperator() = default;
 
-		virtual FDataReferenceCollection GetInputs() const override
+		virtual void BindInputs(FInputVertexInterfaceData& InOutVertexData) override
 		{
-			FDataReferenceCollection Inputs;
-
-			Inputs.AddDataReadReference("WaveTable", WaveTableReadRef);
-			Inputs.AddDataReadReference("Sync", SyncReadRef);
-			Inputs.AddDataReadReference("Freq", FreqReadRef);
+			InOutVertexData.BindReadVertex("WaveTable", WaveTableReadRef);
+			InOutVertexData.BindReadVertex("Play", PlayReadRef);
+			InOutVertexData.BindReadVertex("Stop", StopReadRef);
+			InOutVertexData.BindReadVertex("Sync", SyncReadRef);
+			InOutVertexData.BindReadVertex("Freq", FreqReadRef);
 
 			if (PhaseModReadRef.IsSet())
 			{
-				Inputs.AddDataReadReference("PhaseMod", *PhaseModReadRef);
+				InOutVertexData.BindReadVertex("PhaseMod", *PhaseModReadRef);
 			}
+		}
 
-			return Inputs;
+		virtual void BindOutputs(FOutputVertexInterfaceData& InOutVertexData) override
+		{
+			InOutVertexData.BindReadVertex("Out", TDataReadReference<FAudioBuffer>(OutBufferWriteRef));
 		}
 
-		virtual FDataReferenceCollection GetOutputs() const override
+		virtual FDataReferenceCollection GetInputs() const override
 		{
-			FDataReferenceCollection Outputs;
-			Outputs.AddDataReadReference("Out", TDataReadReference<FAudioBuffer>(OutBufferWriteRef));
+			// This should never be called. Bind(...) is called instead. This method
+			// exists as a stop-gap until the API can be deprecated and removed.
+			checkNoEntry();
+			return {};
+		}
 
-			return Outputs;
+		virtual FDataReferenceCollection GetOutputs() const override
+		{
+			// This should never be called. Bind(...) is called instead. This method
+			// exists as a stop-gap until the API can be deprecated and removed.
+			checkNoEntry();
+			return {};
 		}
 
 		void Execute()
diff --git a/Engine/Plugins/Runtime/Metasound/Source/MetasoundEngine/Private/MetasoundWaveTablePlayerNode.cpp b/Engine/Plugins/Runtime/Metasound/Source/MetasoundEngine/Private/MetasoundWaveTablePlayerNode.cpp
--- a/Engine/Plugins/Runtime/Metasound/Source/MetasoundEngine/Private/MetasoundWaveTablePlayerNode.cpp
+++ b/Engine/Plugins/Runtime/Metasound/Source/MetasoundEngine/Private/MetasoundWaveTablePlayerNode.cpp
@@ -159,32 +159,43 @@ namespace Metasound
 		}
 
 		virtual ~FMetaSoundWaveTablePlayerNodeOperator() = default;
-
-		virtual FDataReferenceCollection GetInputs() const override
+		
+		virtual void BindInputs(FInputVertexInterfaceData& InOutVertexData) override
 		{
-			FDataReferenceCollection Inputs;
-
-			Inputs.AddDataReadReference("WaveTableBank", WaveTableBankReadRef);
-			Inputs.AddDataReadReference("Index", IndexReadRef);
-			Inputs.AddDataReadReference("Sync", SyncReadRef);
-			Inputs.AddDataReadReference("PitchShift", PitchShiftReadRef);
-			Inputs.AddDataReadReference("Loop", LoopReadRef);
+			InOutVertexData.BindReadVertex("WaveTableBank", WaveTableBankReadRef);
+			InOutVertexData.BindReadVertex("Index", IndexReadRef);
+			InOutVertexData.BindReadVertex("Play", PlayReadRef);
+			InOutVertexData.BindReadVertex("Stop", StopReadRef);
+			InOutVertexData.BindReadVertex("Sync", SyncReadRef);
+			InOutVertexData.BindReadVertex("PitchShift", PitchShiftReadRef);
+			InOutVertexData.BindReadVertex("Loop", LoopReadRef);
 
 			if (PhaseModReadRef.IsSet())
 			{
-				Inputs.AddDataReadReference("PhaseMod", *PhaseModReadRef);
+				InOutVertexData.BindReadVertex("PhaseMod", *PhaseModReadRef);
 			}
+		}
+
+		virtual void BindOutputs(FOutputVertexInterfaceData& InOutVertexData) override
+		{
+			InOutVertexData.BindReadVertex("MonoOut", TDataReadReference<FAudioBuffer>(OutBufferWriteRef));
+			InOutVertexData.BindReadVertex("OnFinished", TDataReadReference<FTrigger>(OutOnFinishedRef));
+		}
 
-			return Inputs;
+		virtual FDataReferenceCollection GetInputs() const override
+		{
+			// This should never be called. Bind(...) is called instead. This method
+			// exists as a stop-gap until the API can be deprecated and removed.
+			checkNoEntry();
+			return {};
 		}
 
 		virtual FDataReferenceCollection GetOutputs() const override
 		{
-			FDataReferenceCollection Outputs;
-			Outputs.AddDataReadReference("MonoOut", TDataReadReference<FAudioBuffer>(OutBufferWriteRef));
-			Outputs.AddDataReadReference("OnFinished", TDataReadReference<FTrigger>(OutOnFinishedRef));
-
-			return Outputs;
+			// This should never be called. Bind(...) is called instead. This method
+			// exists as a stop-gap until the API can be deprecated and removed.
+			checkNoEntry();
+			return {};
 		}
 
 		void Execute()
diff --git a/Engine/Plugins/Runtime/Metasound/Source/MetasoundFrontend/Public/MetasoundArrayRandomNode.h b/Engine/Plugins/Runtime/Metasound/Source/MetasoundFrontend/Public/MetasoundArrayRandomNode.h
--- a/Engine/Plugins/Runtime/Metasound/Source/MetasoundFrontend/Public/MetasoundArrayRandomNode.h
+++ b/Engine/Plugins/Runtime/Metasound/Source/MetasoundFrontend/Public/MetasoundArrayRandomNode.h
@@ -159,7 +159,7 @@ namespace Metasound
 			FInt32ReadRef InNoRepeatOrder = InParams.InputDataReferences.GetDataReadReferenceOrConstructWithVertexDefault<int32>(Inputs, METASOUND_GET_PARAM_NAME(InputNoRepeatOrder), InParams.OperatorSettings);
 			FBoolReadRef bInEnableSharedState = InParams.InputDataReferences.GetDataReadReferenceOrConstructWithVertexDefault<bool>(Inputs, METASOUND_GET_PARAM_NAME(InputEnableSharedState), InParams.OperatorSettings);
 
-			return MakeUnique<TArrayRandomGetOperator<ArrayType>>(InParams, InTriggerNext, InTriggerReset, InInputArray, InInputWeightsArray, InSeedValue, InNoRepeatOrder, *bInEnableSharedState);
+			return MakeUnique<TArrayRandomGetOperator<ArrayType>>(InParams, InTriggerNext, InTriggerReset, InInputArray, InInputWeightsArray, InSeedValue, InNoRepeatOrder, bInEnableSharedState);
 		}
 
 		TArrayRandomGetOperator(
@@ -170,7 +170,7 @@ namespace Metasound
 			const TDataReadReference<WeightsArrayType>& InInputWeightsArray,
 			const FInt32ReadRef& InSeedValue,
 			const FInt32ReadRef& InNoRepeatOrder,
-			bool bInEnableSharedState)
+			const FBoolReadRef& bInEnableSharedState)
 			: TriggerNext(InTriggerNext)
 			, TriggerReset(InTriggerReset)
 			, InputArray(InInputArray)
@@ -191,41 +191,40 @@ namespace Metasound
 
 		virtual ~TArrayRandomGetOperator() = default;
 
-		virtual FDataReferenceCollection GetInputs() const override
+		virtual void BindInputs(FInputVertexInterfaceData& InOutVertexData) override
 		{
-			FDataReferenceCollection Inputs;
-			// This function is no longer used
-			checkNoEntry();
-			return Inputs;
+			using namespace ArrayNodeRandomGetVertexNames;
+			InOutVertexData.BindReadVertex(METASOUND_GET_PARAM_NAME(InputTriggerNextValue), TriggerNext);
+			InOutVertexData.BindReadVertex(METASOUND_GET_PARAM_NAME(InputTriggerResetSeed), TriggerReset);
+			InOutVertexData.BindReadVertex(METASOUND_GET_PARAM_NAME(InputRandomArray), InputArray);
+			InOutVertexData.BindReadVertex(METASOUND_GET_PARAM_NAME(InputWeights), InputWeightsArray);
+			InOutVertexData.BindReadVertex(METASOUND_GET_PARAM_NAME(InputSeed), SeedValue);
+			InOutVertexData.BindReadVertex(METASOUND_GET_PARAM_NAME(InputNoRepeatOrder), NoRepeatOrder);
+			InOutVertexData.BindReadVertex(METASOUND_GET_PARAM_NAME(InputEnableSharedState), bEnableSharedState);
 		}
 
-		virtual FDataReferenceCollection GetOutputs() const override
+		virtual void BindOutputs(FOutputVertexInterfaceData& InOutVertexData) override
 		{
-			FDataReferenceCollection Outputs;
-			// This function is no longer used
-			checkNoEntry();
-			return Outputs;
+			using namespace ArrayNodeRandomGetVertexNames;
+			InOutVertexData.BindReadVertex(METASOUND_GET_PARAM_NAME(OutputTriggerOnNext), TriggerOnNext);
+			InOutVertexData.BindReadVertex(METASOUND_GET_PARAM_NAME(OutputTriggerOnReset), TriggerOnReset);
+			InOutVertexData.BindReadVertex(METASOUND_GET_PARAM_NAME(ShuffleOutputValue), OutValue);
 		}
 
-		virtual void Bind(FVertexInterfaceData& InOutVertexData) const override
+		virtual FDataReferenceCollection GetInputs() const override
 		{
-			using namespace ArrayNodeRandomGetVertexNames;
-
-			FInputVertexInterfaceData& Inputs = InOutVertexData.GetInputs();
-
-			Inputs.BindReadVertex(METASOUND_GET_PARAM_NAME(InputTriggerNextValue), TriggerNext);
-			Inputs.BindReadVertex(METASOUND_GET_PARAM_NAME(InputTriggerResetSeed), TriggerReset);
-			Inputs.BindReadVertex(METASOUND_GET_PARAM_NAME(InputRandomArray), InputArray);
-			Inputs.BindReadVertex(METASOUND_GET_PARAM_NAME(InputWeights), InputWeightsArray);
-			Inputs.BindReadVertex(METASOUND_GET_PARAM_NAME(InputSeed), SeedValue);
-			Inputs.BindReadVertex(METASOUND_GET_PARAM_NAME(InputNoRepeatOrder), NoRepeatOrder);
-			Inputs.SetValue(METASOUND_GET_PARAM_NAME(InputEnableSharedState), bEnableSharedState);
-
-			FOutputVertexInterfaceData& Outputs = InOutVertexData.GetOutputs();
+			// This should never be called. Bind(...) is called instead. This method
+			// exists as a stop-gap until the API can be deprecated and removed.
+			checkNoEntry();
+			return {};
+		}
 
-			Outputs.BindReadVertex(METASOUND_GET_PARAM_NAME(OutputTriggerOnNext), TriggerOnNext);
-			Outputs.BindReadVertex(METASOUND_GET_PARAM_NAME(OutputTriggerOnReset), TriggerOnReset);
-			Outputs.BindReadVertex(METASOUND_GET_PARAM_NAME(ShuffleOutputValue), OutValue);
+		virtual FDataReferenceCollection GetOutputs() const override
+		{
+			// This should never be called. Bind(...) is called instead. This method
+			// exists as a stop-gap until the API can be deprecated and removed.
+			checkNoEntry();
+			return {};
 		}
 
 		void Reset(const IOperator::FResetParams& InParams)
@@ -258,7 +257,7 @@ namespace Metasound
 
 		bool UseSharedState() const
 		{
-			return bEnableSharedState && bSharedStateInitialized;
+			return *bEnableSharedState && bSharedStateInitialized;
 		}
 
 		void Execute()
@@ -396,7 +395,7 @@ namespace Metasound
 			bSharedStateInitialized = false;
 			if (InArraySize > 0)
 			{
-				if (bEnableSharedState)
+				if (*bEnableSharedState)
 				{
 					// Get the environment variable for the unique ID of the sound
 					check(SharedStateUniqueId.IsValid());
@@ -454,7 +453,7 @@ namespace Metasound
 		FGuid SharedStateUniqueId;
 		int32 PrevArraySize = 0;
 		bool bIsPreviewSound = false;
-		bool bEnableSharedState = false;
+		FBoolReadRef bEnableSharedState;
 		bool bSharedStateInitialized = false;
 	};
 
diff --git a/Engine/Plugins/Runtime/Metasound/Source/MetasoundFrontend/Public/MetasoundArrayShuffleNode.h b/Engine/Plugins/Runtime/Metasound/Source/MetasoundFrontend/Public/MetasoundArrayShuffleNode.h
--- a/Engine/Plugins/Runtime/Metasound/Source/MetasoundFrontend/Public/MetasoundArrayShuffleNode.h
+++ b/Engine/Plugins/Runtime/Metasound/Source/MetasoundFrontend/Public/MetasoundArrayShuffleNode.h
@@ -225,33 +225,43 @@ namespace Metasound
 
 		virtual ~TArrayShuffleOperator() = default;
 
-		virtual FDataReferenceCollection GetInputs() const override
+		virtual void BindInputs(FInputVertexInterfaceData& InOutVertexData) override
 		{
 			using namespace ArrayNodeShuffleVertexNames;
 
-			FDataReferenceCollection Inputs;
-			Inputs.AddDataReadReference(METASOUND_GET_PARAM_NAME(InputTriggerNext), TriggerNext);
-			Inputs.AddDataReadReference(METASOUND_GET_PARAM_NAME(InputTriggerShuffle), TriggerShuffle);
-			Inputs.AddDataReadReference(METASOUND_GET_PARAM_NAME(InputTriggerReset), TriggerReset);
-			Inputs.AddDataReadReference(METASOUND_GET_PARAM_NAME(InputShuffleArray), InputArray);
-			Inputs.AddDataReadReference(METASOUND_GET_PARAM_NAME(InputShuffleSeed), SeedValue);
-			Inputs.AddDataReadReference(METASOUND_GET_PARAM_NAME(InputAutoShuffle), bAutoShuffle);
-			Inputs.AddDataReadReference(METASOUND_GET_PARAM_NAME(InputShuffleEnableSharedState), bEnableSharedState);
-
-			return Inputs;
+			InOutVertexData.BindReadVertex(METASOUND_GET_PARAM_NAME(InputTriggerNext), TriggerNext);
+			InOutVertexData.BindReadVertex(METASOUND_GET_PARAM_NAME(InputTriggerShuffle), TriggerShuffle);
+			InOutVertexData.BindReadVertex(METASOUND_GET_PARAM_NAME(InputTriggerReset), TriggerReset);
+			InOutVertexData.BindReadVertex(METASOUND_GET_PARAM_NAME(InputShuffleArray), InputArray);
+			InOutVertexData.BindReadVertex(METASOUND_GET_PARAM_NAME(InputShuffleSeed), SeedValue);
+			InOutVertexData.BindReadVertex(METASOUND_GET_PARAM_NAME(InputAutoShuffle), bAutoShuffle);
+			InOutVertexData.BindReadVertex(METASOUND_GET_PARAM_NAME(InputShuffleEnableSharedState), bEnableSharedState);
 		}
 
-		virtual FDataReferenceCollection GetOutputs() const override
+		virtual void BindOutputs(FOutputVertexInterfaceData& InOutVertexData) override
 		{
 			using namespace ArrayNodeShuffleVertexNames;
 
-			FDataReferenceCollection Outputs;
-			Outputs.AddDataReadReference(METASOUND_GET_PARAM_NAME(OutputTriggerOnNext), TriggerNext);
-			Outputs.AddDataReadReference(METASOUND_GET_PARAM_NAME(OutputTriggerOnShuffle), TriggerOnShuffle);
-			Outputs.AddDataReadReference(METASOUND_GET_PARAM_NAME(OutputTriggerOnResetSeed), TriggerReset);
-			Outputs.AddDataReadReference(METASOUND_GET_PARAM_NAME(OutputShuffledValue), OutValue);
+			InOutVertexData.BindReadVertex(METASOUND_GET_PARAM_NAME(OutputTriggerOnNext), TriggerNext);
+			InOutVertexData.BindReadVertex(METASOUND_GET_PARAM_NAME(OutputTriggerOnShuffle), TriggerOnShuffle);
+			InOutVertexData.BindReadVertex(METASOUND_GET_PARAM_NAME(OutputTriggerOnResetSeed), TriggerReset);
+			InOutVertexData.BindReadVertex(METASOUND_GET_PARAM_NAME(OutputShuffledValue), OutValue);
+		}
+
+		virtual FDataReferenceCollection GetInputs() const override
+		{
+			// This should never be called. Bind(...) is called instead. This method
+			// exists as a stop-gap until the API can be deprecated and removed.
+			checkNoEntry();
+			return {};
+		}
 
-			return Outputs;
+		virtual FDataReferenceCollection GetOutputs() const override
+		{
+			// This should never be called. Bind(...) is called instead. This method
+			// exists as a stop-gap until the API can be deprecated and removed.
+			checkNoEntry();
+			return {};
 		}
 
 		void Reset(const IOperator::FResetParams& InParams)
diff --git a/Engine/Plugins/Runtime/Metasound/Source/MetasoundFrontend/Public/MetasoundReceiveNode.h b/Engine/Plugins/Runtime/Metasound/Source/MetasoundFrontend/Public/MetasoundReceiveNode.h
--- a/Engine/Plugins/Runtime/Metasound/Source/MetasoundFrontend/Public/MetasoundReceiveNode.h
+++ b/Engine/Plugins/Runtime/Metasound/Source/MetasoundFrontend/Public/MetasoundReceiveNode.h
@@ -116,25 +116,33 @@ namespace Metasound
 					ResetReceiverAndCleanupChannel();
 				}
 
-				virtual FDataReferenceCollection GetInputs() const override
+				virtual void BindInputs(FInputVertexInterfaceData& InOutVertexData) override
 				{
 					using namespace ReceiveNodeInfo; 
+					InOutVertexData.BindReadVertex<TDataType>(METASOUND_GET_PARAM_NAME(DefaultDataInput), DefaultData);
+					InOutVertexData.BindReadVertex<FSendAddress>(METASOUND_GET_PARAM_NAME(AddressInput), SendAddress);
+				}
 
-					FDataReferenceCollection Inputs;
-
-					Inputs.AddDataReadReference<TDataType>(METASOUND_GET_PARAM_NAME(DefaultDataInput), DefaultData);
-					Inputs.AddDataReadReference<FSendAddress>(METASOUND_GET_PARAM_NAME(AddressInput), SendAddress);
+				virtual void BindOutputs(FOutputVertexInterfaceData& InOutVertexData) override
+				{
+					using namespace ReceiveNodeInfo;
+					InOutVertexData.BindReadVertex<TDataType>(METASOUND_GET_PARAM_NAME(Output), TDataReadReference<TDataType>(OutputData));
+				}
 
-					return Inputs;
+				virtual FDataReferenceCollection GetInputs() const override
+				{
+					// This should never be called. Bind(...) is called instead. This method
+					// exists as a stop-gap until the API can be deprecated and removed.
+					checkNoEntry();
+					return {};
 				}
 
 				virtual FDataReferenceCollection GetOutputs() const override
 				{
-					using namespace ReceiveNodeInfo;
-
-					FDataReferenceCollection Outputs;
-					Outputs.AddDataReadReference<TDataType>(METASOUND_GET_PARAM_NAME(Output), TDataReadReference<TDataType>(OutputData));
-					return Outputs;
+					// This should never be called. Bind(...) is called instead. This method
+					// exists as a stop-gap until the API can be deprecated and removed.
+					checkNoEntry();
+					return {};
 				}
 
 				void Execute()
diff --git a/Engine/Plugins/Runtime/Metasound/Source/MetasoundFrontend/Public/MetasoundSendNode.h b/Engine/Plugins/Runtime/Metasound/Source/MetasoundFrontend/Public/MetasoundSendNode.h
--- a/Engine/Plugins/Runtime/Metasound/Source/MetasoundFrontend/Public/MetasoundSendNode.h
+++ b/Engine/Plugins/Runtime/Metasound/Source/MetasoundFrontend/Public/MetasoundSendNode.h
@@ -105,18 +105,30 @@ namespace Metasound
 					ResetSenderAndCleanupChannel();
 				}
 
-				virtual FDataReferenceCollection GetInputs() const override
+				virtual void BindInputs(FInputVertexInterfaceData& InOutVertexData) override
 				{
 					using namespace SendVertexNames; 
+					InOutVertexData.BindReadVertex<FSendAddress>(METASOUND_GET_PARAM_NAME(AddressInput), SendAddress);
+					InOutVertexData.BindReadVertex<TDataType>(GetSendInputName(), TDataReadReference<TDataType>(InputData));
+				}
+
+				virtual void BindOutputs(FOutputVertexInterfaceData& InOutVertexData) override
+				{
+				}
 
-					FDataReferenceCollection Inputs;
-					Inputs.AddDataReadReference<FSendAddress>(METASOUND_GET_PARAM_NAME(AddressInput), SendAddress);
-					Inputs.AddDataReadReference<TDataType>(GetSendInputName(), TDataReadReference<TDataType>(InputData));
-					return Inputs;
+				virtual FDataReferenceCollection GetInputs() const override
+				{
+					// This should never be called. Bind(...) is called instead. This method
+					// exists as a stop-gap until the API can be deprecated and removed.
+					checkNoEntry();
+					return {};
 				}
 
 				virtual FDataReferenceCollection GetOutputs() const override
 				{
+					// This should never be called. Bind(...) is called instead. This method
+					// exists as a stop-gap until the API can be deprecated and removed.
+					checkNoEntry();
 					return {};
 				}
 
diff --git a/Engine/Plugins/Runtime/Metasound/Source/MetasoundStandardNodes/Private/MetasoundPerlinNoiseNode.cpp b/Engine/Plugins/Runtime/Metasound/Source/MetasoundStandardNodes/Private/MetasoundPerlinNoiseNode.cpp
--- a/Engine/Plugins/Runtime/Metasound/Source/MetasoundStandardNodes/Private/MetasoundPerlinNoiseNode.cpp
+++ b/Engine/Plugins/Runtime/Metasound/Source/MetasoundStandardNodes/Private/MetasoundPerlinNoiseNode.cpp
@@ -161,6 +161,8 @@ namespace Metasound
 			InOutVertexData.BindReadVertex(METASOUND_GET_PARAM_NAME(XPin), Pins.X);
 			InOutVertexData.SetValue(METASOUND_GET_PARAM_NAME(OffsetPin), Pins.Offset);
 			InOutVertexData.BindReadVertex(METASOUND_GET_PARAM_NAME(OctavesPin), Pins.Octaves);
+			InOutVertexData.BindReadVertex(METASOUND_GET_PARAM_NAME(MinValuePin), Pins.MinValue);
+			InOutVertexData.BindReadVertex(METASOUND_GET_PARAM_NAME(MaxValuePin), Pins.MaxValue);
 		}
 
 		virtual void BindOutputs(FOutputVertexInterfaceData& InOutVertexData) override
diff --git a/Engine/Plugins/Runtime/Metasound/Source/MetasoundStandardNodes/Private/MetasoundStereoDelayNode.cpp b/Engine/Plugins/Runtime/Metasound/Source/MetasoundStandardNodes/Private/MetasoundStereoDelayNode.cpp
--- a/Engine/Plugins/Runtime/Metasound/Source/MetasoundStandardNodes/Private/MetasoundStereoDelayNode.cpp
+++ b/Engine/Plugins/Runtime/Metasound/Source/MetasoundStandardNodes/Private/MetasoundStereoDelayNode.cpp
@@ -150,22 +150,22 @@ namespace Metasound
 	{
 		using namespace StereoDelay;
 
-		InOutVertexData.BindReadVertex(METASOUND_GET_PARAM_NAME(InAudioInputLeft), FAudioBufferReadRef(LeftAudioInput));
-		InOutVertexData.BindReadVertex(METASOUND_GET_PARAM_NAME(InAudioInputRight), FAudioBufferReadRef(RightAudioInput));
-		InOutVertexData.BindReadVertex(METASOUND_GET_PARAM_NAME(InDelayMode), FStereoDelayModeReadRef(StereoDelayMode));
-		InOutVertexData.BindReadVertex(METASOUND_GET_PARAM_NAME(InDelayTime), FTimeReadRef(DelayTime));
-		InOutVertexData.BindReadVertex(METASOUND_GET_PARAM_NAME(InDelayRatio), FFloatReadRef(DelayRatio));
-		InOutVertexData.BindReadVertex(METASOUND_GET_PARAM_NAME(InDryLevel), FFloatReadRef(DryLevel));
-		InOutVertexData.BindReadVertex(METASOUND_GET_PARAM_NAME(InWetLevel), FFloatReadRef(WetLevel));
-		InOutVertexData.BindReadVertex(METASOUND_GET_PARAM_NAME(InFeedbackAmount), FFloatReadRef(Feedback));
+		InOutVertexData.BindReadVertex(METASOUND_GET_PARAM_NAME(InAudioInputLeft), LeftAudioInput);
+		InOutVertexData.BindReadVertex(METASOUND_GET_PARAM_NAME(InAudioInputRight), RightAudioInput);
+		InOutVertexData.BindReadVertex(METASOUND_GET_PARAM_NAME(InDelayMode), StereoDelayMode);
+		InOutVertexData.BindReadVertex(METASOUND_GET_PARAM_NAME(InDelayTime), DelayTime);
+		InOutVertexData.BindReadVertex(METASOUND_GET_PARAM_NAME(InDelayRatio), DelayRatio);
+		InOutVertexData.BindReadVertex(METASOUND_GET_PARAM_NAME(InDryLevel), DryLevel);
+		InOutVertexData.BindReadVertex(METASOUND_GET_PARAM_NAME(InWetLevel), WetLevel);
+		InOutVertexData.BindReadVertex(METASOUND_GET_PARAM_NAME(InFeedbackAmount), Feedback);
 	}
 
 	void FStereoDelayOperator::BindOutputs(FOutputVertexInterfaceData& InOutVertexData)
 	{
 		using namespace StereoDelay;
 
-		InOutVertexData.BindReadVertex(METASOUND_GET_PARAM_NAME(OutAudioLeft), FAudioBufferReadRef(LeftAudioOutput));
-		InOutVertexData.BindReadVertex(METASOUND_GET_PARAM_NAME(OutAudioRight), FAudioBufferReadRef(RightAudioOutput));
+		InOutVertexData.BindReadVertex(METASOUND_GET_PARAM_NAME(OutAudioLeft), LeftAudioOutput);
+		InOutVertexData.BindReadVertex(METASOUND_GET_PARAM_NAME(OutAudioRight), RightAudioOutput);
 	}
 
 	FDataReferenceCollection FStereoDelayOperator::GetInputs() const
diff --git a/Engine/Plugins/Runtime/Metasound/Source/MetasoundStandardNodes/Private/MetasoundWaveWriterNode.cpp b/Engine/Plugins/Runtime/Metasound/Source/MetasoundStandardNodes/Private/MetasoundWaveWriterNode.cpp
--- a/Engine/Plugins/Runtime/Metasound/Source/MetasoundStandardNodes/Private/MetasoundWaveWriterNode.cpp
+++ b/Engine/Plugins/Runtime/Metasound/Source/MetasoundStandardNodes/Private/MetasoundWaveWriterNode.cpp
@@ -226,11 +226,11 @@ namespace Metasound
 		// Theoretical limit of .WAV files.
 		static_assert(NumInputChannels > 0 && NumInputChannels <= 65535, "Num Channels > 0 and <= 65535");
 
-		TWaveWriterOperator(const FOperatorSettings& InSettings, TArray<FAudioBufferReadRef>&& InAudioBuffers, FBoolReadRef&& InEnabled, const TSharedPtr<FNumberedFileCache, ESPMode::ThreadSafe>& InNumberedFileCache, const FString& InFilenamePrefix)
+		TWaveWriterOperator(const FOperatorSettings& InSettings, TArray<FAudioBufferReadRef>&& InAudioBuffers, FBoolReadRef&& InEnabled, const TSharedPtr<FNumberedFileCache, ESPMode::ThreadSafe>& InNumberedFileCache, FStringReadRef&& InFilenamePrefix)
 			: AudioInputs{ MoveTemp(InAudioBuffers) }
 			, Enabled{ MoveTemp(InEnabled) }
 			, NumberedFileCacheSP{ InNumberedFileCache }
-			, FileNamePrefix{ InFilenamePrefix }
+			, FileNamePrefix{ MoveTemp(InFilenamePrefix) }
 			, SampleRate{ InSettings.GetSampleRate() }
 		{
 			check(AudioInputs.Num() == NumInputChannels);
@@ -254,6 +254,10 @@ namespace Metasound
 			{
 				InOutVertexData.BindReadVertex(GetAudioInputName(i), AudioInputs[i]);
 			}
+
+			using namespace WaveWriterVertexNames;
+			InOutVertexData.BindReadVertex(METASOUND_GET_PARAM_NAME(InEnabledPin), Enabled);
+			InOutVertexData.BindReadVertex(METASOUND_GET_PARAM_NAME(InFilenamePrefixPin), FileNamePrefix);
 		}
 
 		virtual void BindOutputs(FOutputVertexInterfaceData&) override
@@ -498,7 +502,7 @@ namespace Metasound
 		FBoolReadRef Enabled;
 		TUniquePtr<FWaveWriter> Writer;
 		TSharedPtr<FNumberedFileCache, ESPMode::ThreadSafe> NumberedFileCacheSP;
-		FString FileNamePrefix;
+		FStringReadRef FileNamePrefix;
 		float SampleRate = 0.f;
 		bool bIsEnabled = false;
 	};
@@ -513,8 +517,6 @@ namespace Metasound
 		const FOperatorSettings& Settings = InParams.OperatorSettings;
 		const FInputVertexInterface& InputInterface = DeclareVertexInterface().GetInputInterface();
 
-		FStringReadRef FilenamePrefix = InputCol.GetDataReadReferenceOrConstructWithVertexDefault<FString>(InputInterface, METASOUND_GET_PARAM_NAME(InFilenamePrefixPin), Settings);
-
 		int32 NumConnectedAudioPins = 0;
 		TArray<FAudioBufferReadRef> InputBuffers;
 		for (int32 i = 0; i < NumInputChannels; ++i)
@@ -532,7 +534,7 @@ namespace Metasound
 				MoveTemp(InputBuffers),
 				InputCol.GetDataReadReferenceOrConstructWithVertexDefault<bool>(InputInterface, METASOUND_GET_PARAM_NAME(InEnabledPin), Settings),
 				GetNameCache(),
-				*FilenamePrefix
+				InputCol.GetDataReadReferenceOrConstructWithVertexDefault<FString>(InputInterface, METASOUND_GET_PARAM_NAME(InFilenamePrefixPin), Settings)
 			);
 		}
 
