[Metasound Node] Guard against zero-duration grain spawning.
(NANs / bad memory access).

diff --git a/Engine/Plugins/Runtime/Metasound/Source/MetasoundStandardNodes/Private/MetasoundGrainDelayNode.cpp b/Engine/Plugins/Runtime/Metasound/Source/MetasoundStandardNodes/Private/MetasoundGrainDelayNode.cpp
--- a/Engine/Plugins/Runtime/Metasound/Source/MetasoundStandardNodes/Private/MetasoundGrainDelayNode.cpp
+++ b/Engine/Plugins/Runtime/Metasound/Source/MetasoundStandardNodes/Private/MetasoundGrainDelayNode.cpp
@@ -342,17 +342,24 @@ namespace Metasound
 			// execute the next spawn trigger
 			if (CurrAudioFrame == NextSpawnFrame)
 			{
+				++SpawnTrigIndex;
+
+				const float Delta = *GrainDurationDelta;
+				const float NewDuration = GrainDelayProcessor.GetGrainDurationClamped(*GrainDuration + FMath::RandRange(-0.5f * Delta, 0.5f * Delta));
+				if(FMath::IsNearlyZero(NewDuration))
+				{
+					// avoid spawning nodes of zero duration
+					continue;
+				}
+
 				const float NewDelay = GrainDelayProcessor.GetGrainDelayClamped(*GrainDelay + FMath::RandRange(-0.5f * (*GrainDelayDelta), 0.5f * (*GrainDelayDelta)));
-				const float NewDuration = GrainDelayProcessor.GetGrainDurationClamped(*GrainDuration + FMath::RandRange(-0.5f * (*GrainDurationDelta), 0.5f * (*GrainDurationDelta)));
 
 				// Clamp the pitch shift delta before determining the random range
 				const float ClampedPitchShiftDelta = GrainDelayProcessor.GetGrainPitchShiftClamped(*GrainPitchShiftDelta);
 				const float NewPitchShiftRatioOffset = GrainDelayProcessor.GetGrainPitchShiftFrameRatio(FMath::RandRange(-0.5f * ClampedPitchShiftDelta, 0.5f * ClampedPitchShiftDelta));
-				
+
 				// Spawn a new grain. Note the grain manager handles spawning logic and rate limiting.
 				GrainDelayProcessor.SpawnGrain(NewDelay, NewDuration, NewPitchShiftRatioOffset);
-				
-				++SpawnTrigIndex;
 			}
 		}
 	}
diff --git a/Engine/Source/Runtime/SignalProcessing/Private/GrainDelay.cpp b/Engine/Source/Runtime/SignalProcessing/Private/GrainDelay.cpp
--- a/Engine/Source/Runtime/SignalProcessing/Private/GrainDelay.cpp
+++ b/Engine/Source/Runtime/SignalProcessing/Private/GrainDelay.cpp
@@ -144,6 +144,8 @@ namespace Audio
 				FGrain& Grain = GrainPool[GrainId];
 				
 				// Calculate the grain envelope
+				// should have never been added to the pool
+				check(!FMath::IsNearlyZero(Grain.DurationFrames));
 				const float Fraction = FMath::Min(Grain.NumFramesRendered / Grain.DurationFrames, 1.0f);
 				const float GrainVolume = Grain::GetValue(GrainEnvelope, Fraction);
 
