Forbid naming MetaSound members as that of existing namespace members on the given document

diff --git a/Engine/Plugins/Runtime/Metasound/Source/MetasoundEditor/Private/MetasoundEditorGraph.cpp b/Engine/Plugins/Runtime/Metasound/Source/MetasoundEditor/Private/MetasoundEditorGraph.cpp
--- a/Engine/Plugins/Runtime/Metasound/Source/MetasoundEditor/Private/MetasoundEditorGraph.cpp
+++ b/Engine/Plugins/Runtime/Metasound/Source/MetasoundEditor/Private/MetasoundEditorGraph.cpp
@@ -609,6 +609,10 @@ bool UMetasoundEditorGraphVertex::CanRename(const FText& InNewText, FText& OutEr
 	}
 
 	const FName NewName(*NewNameString);
+	FName Namespace;
+	FName ParameterName;
+	Audio::FParameterPath::SplitName(NewName, Namespace, ParameterName);
+
 	bool bIsNameValid = true;
 	FConstNodeHandle NodeHandle = GetConstNodeHandle();
 	FConstGraphHandle GraphHandle = NodeHandle->GetOwningGraph();
@@ -616,11 +620,18 @@ bool UMetasoundEditorGraphVertex::CanRename(const FText& InNewText, FText& OutEr
 	{
 		if (NodeID != NodeToCompare->GetID())
 		{
-			if (NewName == NodeToCompare->GetNodeName())
+			const FName& OtherName = NodeToCompare->GetNodeName();
+			if (NewName == OtherName)
 			{
 				bIsNameValid = false;
 				OutError = FText::Format(LOCTEXT("GraphVertexRenameInvalid_NameTaken", "{0} is already in use"), InNewText);
 			}
+
+			if (Namespace == OtherName)
+			{
+				bIsNameValid = false;
+				OutError = FText::Format(LOCTEXT("GraphVertexRenameInvalid_NamespaceTaken", "{0} cannot share name of existing namespace"), InNewText);
+			}
 		}
 	}, GetClassType());
 
@@ -1113,6 +1124,9 @@ bool UMetasoundEditorGraphVariable::CanRename(const FText& InNewText, FText& Out
 	}
 
 	const FName NewName(*NewNameString);
+	FName Namespace;
+	FName ParameterName;
+	Audio::FParameterPath::SplitName(NewName, Namespace, ParameterName);
 
 	FConstVariableHandle VariableHandle = GetConstVariableHandle();
 	TArray<FConstVariableHandle> Variables = VariableHandle->GetOwningGraph()->GetVariables();
@@ -1120,11 +1134,18 @@ bool UMetasoundEditorGraphVariable::CanRename(const FText& InNewText, FText& Out
 	{
 		if (VariableID != OtherVariable->GetID())
 		{
-			if (NewName == OtherVariable->GetName())
+			const FName& OtherName = OtherVariable->GetName();
+			if (NewName == OtherName)
 			{
 				OutError = FText::Format(LOCTEXT("GraphVariableRenameInvalid_NameTaken", "{0} is already in use"), InNewText);
 				return false;
 			}
+
+			if (Namespace == OtherName)
+			{
+				OutError = FText::Format(LOCTEXT("GraphVariableRenameInvalid_NamespaceTaken", "{0} cannot share name of existing namespace"), InNewText);
+				return false;
+			}
 		}
 	}
 
