Add missing BPReadWrite mark-up to MetaSoundClassName so nodes can be added via builder in BP & LexToString for better error logging of MetaSound types to be used in subsequent CL for BP debugging. Minor formatting fix-ups.
#rnx
#rnx

diff --git a/Engine/Plugins/Runtime/Metasound/Source/MetasoundEngine/Private/Interfaces/MetasoundOutputFormatInterfaces.cpp b/Engine/Plugins/Runtime/Metasound/Source/MetasoundEngine/Private/Interfaces/MetasoundOutputFormatInterfaces.cpp
--- a/Engine/Plugins/Runtime/Metasound/Source/MetasoundEngine/Private/Interfaces/MetasoundOutputFormatInterfaces.cpp
+++ b/Engine/Plugins/Runtime/Metasound/Source/MetasoundEngine/Private/Interfaces/MetasoundOutputFormatInterfaces.cpp
@@ -20,9 +20,10 @@ namespace Metasound::Engine
 	{
 		Audio::FParameterInterface::FOutput GetFrontLeftOutput(const FName& InVertexName)
 		{
-			return {
+			return
+			{
 				LOCTEXT("OutputFormatSurroundInterface_GeneratedFrontLeftDisplayName", "Out Front Left"),
-				LOCTEXT("OutputFormatSurroundInterface_GeneratedFrontLeftDescription", "The resulting front left channel output audio from this source."),
+				LOCTEXT("OutputFormatSurroundInterface_GeneratedFrontLeftDescription", "The resulting front left channel output audio."),
 				GetMetasoundDataTypeName<FAudioBuffer>(),
 				InVertexName,
 				FText::GetEmpty(), // RequiredText
@@ -33,9 +34,10 @@ namespace Metasound::Engine
 
 		Audio::FParameterInterface::FOutput GetFrontRightOutput(const FName& InVertexName)
 		{
-			return {
+			return
+			{
 				LOCTEXT("OutputFormatSurroundInterface_GeneratedFrontRightDisplayName", "Out Front Right"),
-				LOCTEXT("OutputFormatSurroundInterface_GeneratedFrontRightDescription", "The resulting front right channel output audio from this source."),
+				LOCTEXT("OutputFormatSurroundInterface_GeneratedFrontRightDescription", "The resulting front right channel output audio."),
 				GetMetasoundDataTypeName<FAudioBuffer>(),
 				InVertexName,
 				FText::GetEmpty(), // RequiredText
@@ -46,9 +48,10 @@ namespace Metasound::Engine
 
 		Audio::FParameterInterface::FOutput GetFrontCenterOutput(const FName& InVertexName)
 		{
-			return {
+			return
+			{
 				LOCTEXT("OutputFormatSurroundInterface_GeneratedCenterDisplayName", "Out Front Center"),
-				LOCTEXT("OutputFormatSurroundInterface_GeneratedCenterDescription", "The resulting front center channel output audio from this source."),
+				LOCTEXT("OutputFormatSurroundInterface_GeneratedCenterDescription", "The resulting front center channel output audio."),
 				GetMetasoundDataTypeName<FAudioBuffer>(),
 				InVertexName,
 				FText::GetEmpty(), // RequiredText
@@ -59,9 +62,10 @@ namespace Metasound::Engine
 
 		Audio::FParameterInterface::FOutput GetLowFrequencyOutput(const FName& InVertexName)
 		{
-			return {
+			return
+			{
 				LOCTEXT("OutputFormatSurroundInterface_GeneratedLowFrequencyDisplayName", "Out Low Frequency"),
-				LOCTEXT("OutputFormatSurroundInterface_GeneratedLowFrequencyDescription", "The resulting low frequency channel output audio from this source."),
+				LOCTEXT("OutputFormatSurroundInterface_GeneratedLowFrequencyDescription", "The resulting low frequency channel output audio."),
 				GetMetasoundDataTypeName<FAudioBuffer>(),
 				InVertexName,
 				FText::GetEmpty(), // RequiredText
@@ -72,9 +76,10 @@ namespace Metasound::Engine
 
 		Audio::FParameterInterface::FOutput GetSideLeftOutput(const FName& InVertexName)
 		{
-			return {
+			return
+			{
 				LOCTEXT("OutputFormatSurroundInterface_GeneratedSideLeftDisplayName", "Out Side Left"),
-				LOCTEXT("OutputFormatSurroundInterface_GeneratedSideLeftDescription", "The resulting side left channel output audio from this source."),
+				LOCTEXT("OutputFormatSurroundInterface_GeneratedSideLeftDescription", "The resulting side left channel output audio."),
 				GetMetasoundDataTypeName<FAudioBuffer>(),
 				InVertexName,
 				FText::GetEmpty(), // RequiredText
@@ -85,9 +90,10 @@ namespace Metasound::Engine
 
 		Audio::FParameterInterface::FOutput GetSideRightOutput(const FName& InVertexName)
 		{
-			return {
+			return
+			{
 				LOCTEXT("OutputFormatSurroundInterface_GeneratedSideRightDisplayName", "Out Side Right"),
-				LOCTEXT("OutputFormatSurroundInterface_GeneratedSideRightDescription", "The resulting side right channel output audio from this source."),
+				LOCTEXT("OutputFormatSurroundInterface_GeneratedSideRightDescription", "The resulting side right channel output audio."),
 				GetMetasoundDataTypeName<FAudioBuffer>(),
 				InVertexName,
 				FText::GetEmpty(), // RequiredText
@@ -98,9 +104,10 @@ namespace Metasound::Engine
 
 		Audio::FParameterInterface::FOutput GetBackLeftOutput(const FName& InVertexName)
 		{
-			return {
+			return
+			{
 				LOCTEXT("OutputFormatSurroundInterface_GeneratedBackLeftDisplayName", "Out Back Left"),
-				LOCTEXT("OutputFormatSurroundInterface_GeneratedBackLeftDescription", "The resulting back left channel output audio from this source."),
+				LOCTEXT("OutputFormatSurroundInterface_GeneratedBackLeftDescription", "The resulting back left channel output audio."),
 				GetMetasoundDataTypeName<FAudioBuffer>(),
 				InVertexName,
 				FText::GetEmpty(), // RequiredText
@@ -111,9 +118,10 @@ namespace Metasound::Engine
 
 		Audio::FParameterInterface::FOutput GetBackRightOutput(const FName& InVertexName)
 		{
-			return {
+			return
+			{
 				LOCTEXT("OutputFormatSurroundInterface_GeneratedBackRightDisplayName", "Out Back Right"),
-				LOCTEXT("OutputFormatSurroundInterface_GeneratedBackRightDescription", "The resulting back right channel output audio from this source."),
+				LOCTEXT("OutputFormatSurroundInterface_GeneratedBackRightDescription", "The resulting back right channel output audio."),
 				GetMetasoundDataTypeName<FAudioBuffer>(),
 				InVertexName,
 				FText::GetEmpty(), // RequiredText
@@ -167,7 +175,7 @@ namespace Metasound::Engine
 					{
 						{
 							LOCTEXT("GeneratedAudioDisplayName", "Out Mono"),
-							LOCTEXT("GeneratedAudioDescription", "The resulting mono output from this source."),
+							LOCTEXT("GeneratedAudioDescription", "The resulting mono output."),
 							GetMetasoundDataTypeName<FAudioBuffer>(),
 							Outputs::MonoOut,
 							FText::GetEmpty(), // RequiredText
@@ -218,7 +226,7 @@ namespace Metasound::Engine
 					{
 						{
 							LOCTEXT("OutputFormatStereoInterface_GeneratedLeftDisplayName", "Out Left"),
-							LOCTEXT("OutputFormatStereoInterface_GeneratedLeftDescription", "The resulting left channel output audio from this source."),
+							LOCTEXT("OutputFormatStereoInterface_GeneratedLeftDescription", "The resulting left channel output audio."),
 							GetMetasoundDataTypeName<FAudioBuffer>(),
 							Outputs::LeftOut,
 							FText::GetEmpty(), // RequiredText
@@ -227,7 +235,7 @@ namespace Metasound::Engine
 						},
 						{
 							LOCTEXT("OutputFormatStereoInterface_GeneratedRightDisplayName", "Out Right"),
-							LOCTEXT("OutputFormatStereoInterface_GeneratedRightDescription", "The resulting right channel output audio from this source."),
+							LOCTEXT("OutputFormatStereoInterface_GeneratedRightDescription", "The resulting right channel output audio."),
 							GetMetasoundDataTypeName<FAudioBuffer>(),
 							Outputs::RightOut,
 							FText::GetEmpty(), // RequiredText
diff --git a/Engine/Plugins/Runtime/Metasound/Source/MetasoundEngine/Public/Interfaces/MetasoundOutputFormatInterfaces.h b/Engine/Plugins/Runtime/Metasound/Source/MetasoundEngine/Public/Interfaces/MetasoundOutputFormatInterfaces.h
--- a/Engine/Plugins/Runtime/Metasound/Source/MetasoundEngine/Public/Interfaces/MetasoundOutputFormatInterfaces.h
+++ b/Engine/Plugins/Runtime/Metasound/Source/MetasoundEngine/Public/Interfaces/MetasoundOutputFormatInterfaces.h
@@ -25,107 +25,104 @@ enum class EMetaSoundOutputAudioFormat : uint8
 };
 
 
-namespace Metasound
+namespace Metasound::Engine
 {
-	namespace Engine
+	namespace OutputFormatMonoInterface
 	{
-		namespace OutputFormatMonoInterface
+		namespace Outputs
 		{
-			namespace Outputs
-			{
-				METASOUNDENGINE_API const extern FName MonoOut;
-			} // namespace Outputs
+			METASOUNDENGINE_API const extern FName MonoOut;
+		} // namespace Outputs
 
-			METASOUNDENGINE_API const FMetasoundFrontendVersion& GetVersion();
+		METASOUNDENGINE_API const FMetasoundFrontendVersion& GetVersion();
 
-			UE_DEPRECATED(5.3, "Interfaces are no longer registered with the UClass as interfaces can support multiple UClasses")
-			METASOUNDENGINE_API Audio::FParameterInterfacePtr CreateInterface(const UClass& InClass);
-			METASOUNDENGINE_API Audio::FParameterInterfacePtr CreateInterface();
-		} // namespace OutputFormatMonoInterface
+		UE_DEPRECATED(5.3, "Interfaces are no longer registered with the UClass as interfaces can support multiple UClasses")
+		METASOUNDENGINE_API Audio::FParameterInterfacePtr CreateInterface(const UClass& InClass);
+		METASOUNDENGINE_API Audio::FParameterInterfacePtr CreateInterface();
+	} // namespace OutputFormatMonoInterface
 
-		namespace OutputFormatStereoInterface
+	namespace OutputFormatStereoInterface
+	{
+		namespace Outputs
 		{
-			namespace Outputs
-			{
-				METASOUNDENGINE_API const extern FName LeftOut;
-				METASOUNDENGINE_API const extern FName RightOut;
-			} // namespace Outputs
+			METASOUNDENGINE_API const extern FName LeftOut;
+			METASOUNDENGINE_API const extern FName RightOut;
+		} // namespace Outputs
 
-			METASOUNDENGINE_API const FMetasoundFrontendVersion& GetVersion();
+		METASOUNDENGINE_API const FMetasoundFrontendVersion& GetVersion();
 
-			UE_DEPRECATED(5.3, "Interfaces are no longer registered with the UClass as interfaces can support multiple UClasses")
-			METASOUNDENGINE_API Audio::FParameterInterfacePtr CreateInterface(const UClass& InClass);
-			METASOUNDENGINE_API Audio::FParameterInterfacePtr CreateInterface();
-		} // namespace OutputFormatStereoInterface
+		UE_DEPRECATED(5.3, "Interfaces are no longer registered with the UClass as interfaces can support multiple UClasses")
+		METASOUNDENGINE_API Audio::FParameterInterfacePtr CreateInterface(const UClass& InClass);
+		METASOUNDENGINE_API Audio::FParameterInterfacePtr CreateInterface();
+	} // namespace OutputFormatStereoInterface
 
-		namespace OutputFormatQuadInterface
-		{
-			namespace Outputs
-			{
-				METASOUNDENGINE_API const extern FName FrontLeftOut;
-				METASOUNDENGINE_API const extern FName FrontRightOut;
-				METASOUNDENGINE_API const extern FName SideLeftOut;
-				METASOUNDENGINE_API const extern FName SideRightOut;
-			} // namespace Outputs
-
-			METASOUNDENGINE_API const FMetasoundFrontendVersion& GetVersion();
-
-			UE_DEPRECATED(5.3, "Interfaces are no longer registered with the UClass as interfaces can support multiple UClasses")
-			METASOUNDENGINE_API Audio::FParameterInterfacePtr CreateInterface(const UClass& InClass);
-			METASOUNDENGINE_API Audio::FParameterInterfacePtr CreateInterface();
-		} // namespace OutputFormatQuadInterface
-
-		namespace OutputFormatFiveDotOneInterface
+	namespace OutputFormatQuadInterface
+	{
+		namespace Outputs
 		{
-			namespace Outputs
-			{
-				METASOUNDENGINE_API const extern FName FrontLeftOut;
-				METASOUNDENGINE_API const extern FName FrontRightOut;
-				METASOUNDENGINE_API const extern FName FrontCenterOut;
-				METASOUNDENGINE_API const extern FName LowFrequencyOut;
-				METASOUNDENGINE_API const extern FName SideLeftOut;
-				METASOUNDENGINE_API const extern FName SideRightOut;
-			} // namespace Outputs
-
-			METASOUNDENGINE_API const FMetasoundFrontendVersion& GetVersion();
-
-			UE_DEPRECATED(5.3, "Interfaces are no longer registered with the UClass as interfaces can support multiple UClasses")
-			METASOUNDENGINE_API Audio::FParameterInterfacePtr CreateInterface(const UClass& InClass);
-			METASOUNDENGINE_API Audio::FParameterInterfacePtr CreateInterface();
-		} // namespace OutputFormatFiveDotOneInterface
-
-		namespace OutputFormatSevenDotOneInterface
+			METASOUNDENGINE_API const extern FName FrontLeftOut;
+			METASOUNDENGINE_API const extern FName FrontRightOut;
+			METASOUNDENGINE_API const extern FName SideLeftOut;
+			METASOUNDENGINE_API const extern FName SideRightOut;
+		} // namespace Outputs
+
+		METASOUNDENGINE_API const FMetasoundFrontendVersion& GetVersion();
+
+		UE_DEPRECATED(5.3, "Interfaces are no longer registered with the UClass as interfaces can support multiple UClasses")
+		METASOUNDENGINE_API Audio::FParameterInterfacePtr CreateInterface(const UClass& InClass);
+		METASOUNDENGINE_API Audio::FParameterInterfacePtr CreateInterface();
+	} // namespace OutputFormatQuadInterface
+
+	namespace OutputFormatFiveDotOneInterface
+	{
+		namespace Outputs
 		{
-			namespace Outputs
-			{
-				METASOUNDENGINE_API const extern FName FrontLeftOut;
-				METASOUNDENGINE_API const extern FName FrontRightOut;
-				METASOUNDENGINE_API const extern FName FrontCenterOut;
-				METASOUNDENGINE_API const extern FName LowFrequencyOut;
-				METASOUNDENGINE_API const extern FName SideLeftOut;
-				METASOUNDENGINE_API const extern FName SideRightOut;
-				METASOUNDENGINE_API const extern FName BackLeftOut;
-				METASOUNDENGINE_API const extern FName BackRightOut;
-			} // namespace Outputs
-
-			METASOUNDENGINE_API const FMetasoundFrontendVersion& GetVersion();
-
-			UE_DEPRECATED(5.3, "Interfaces are no longer registered with the UClass as interfaces can support multiple UClasses")
-			METASOUNDENGINE_API Audio::FParameterInterfacePtr CreateInterface(const UClass& InClass);
-			METASOUNDENGINE_API Audio::FParameterInterfacePtr CreateInterface();
-		} // namespace OutputFormatSevenDotOneInterface
-
-		// Contains information on output audio formats
-		struct METASOUNDENGINE_API FOutputAudioFormatInfo
+			METASOUNDENGINE_API const extern FName FrontLeftOut;
+			METASOUNDENGINE_API const extern FName FrontRightOut;
+			METASOUNDENGINE_API const extern FName FrontCenterOut;
+			METASOUNDENGINE_API const extern FName LowFrequencyOut;
+			METASOUNDENGINE_API const extern FName SideLeftOut;
+			METASOUNDENGINE_API const extern FName SideRightOut;
+		} // namespace Outputs
+
+		METASOUNDENGINE_API const FMetasoundFrontendVersion& GetVersion();
+
+		UE_DEPRECATED(5.3, "Interfaces are no longer registered with the UClass as interfaces can support multiple UClasses")
+		METASOUNDENGINE_API Audio::FParameterInterfacePtr CreateInterface(const UClass& InClass);
+		METASOUNDENGINE_API Audio::FParameterInterfacePtr CreateInterface();
+	} // namespace OutputFormatFiveDotOneInterface
+
+	namespace OutputFormatSevenDotOneInterface
+	{
+		namespace Outputs
 		{
-			FMetasoundFrontendVersion InterfaceVersion;
-			TArray<Metasound::FVertexName> OutputVertexChannelOrder;
-		};
+			METASOUNDENGINE_API const extern FName FrontLeftOut;
+			METASOUNDENGINE_API const extern FName FrontRightOut;
+			METASOUNDENGINE_API const extern FName FrontCenterOut;
+			METASOUNDENGINE_API const extern FName LowFrequencyOut;
+			METASOUNDENGINE_API const extern FName SideLeftOut;
+			METASOUNDENGINE_API const extern FName SideRightOut;
+			METASOUNDENGINE_API const extern FName BackLeftOut;
+			METASOUNDENGINE_API const extern FName BackRightOut;
+		} // namespace Outputs
+
+		METASOUNDENGINE_API const FMetasoundFrontendVersion& GetVersion();
+
+		UE_DEPRECATED(5.3, "Interfaces are no longer registered with the UClass as interfaces can support multiple UClasses")
+		METASOUNDENGINE_API Audio::FParameterInterfacePtr CreateInterface(const UClass& InClass);
+		METASOUNDENGINE_API Audio::FParameterInterfacePtr CreateInterface();
+	} // namespace OutputFormatSevenDotOneInterface
+
+	// Contains information on output audio formats
+	struct METASOUNDENGINE_API FOutputAudioFormatInfo
+	{
+		FMetasoundFrontendVersion InterfaceVersion;
+		TArray<Metasound::FVertexName> OutputVertexChannelOrder;
+	};
 
-		using FOutputAudioFormatInfoMap = TMap<EMetaSoundOutputAudioFormat, FOutputAudioFormatInfo>;
-		using FOutputAudioFormatInfoPair = TPair<EMetaSoundOutputAudioFormat, FOutputAudioFormatInfo>;
+	using FOutputAudioFormatInfoMap = TMap<EMetaSoundOutputAudioFormat, FOutputAudioFormatInfo>;
+	using FOutputAudioFormatInfoPair = TPair<EMetaSoundOutputAudioFormat, FOutputAudioFormatInfo>;
 
-		// Return a map containing all the supported output audio formats for a MetaSound.
-		METASOUNDENGINE_API const FOutputAudioFormatInfoMap& GetOutputAudioFormatInfo();
-	} // namespace Engine
-} // namespace Metasound
+	// Return a map containing all the supported output audio formats for a MetaSound.
+	METASOUNDENGINE_API const FOutputAudioFormatInfoMap& GetOutputAudioFormatInfo();
+} // namespace Metasound::Engine
diff --git a/Engine/Plugins/Runtime/Metasound/Source/MetasoundFrontend/Private/MetasoundFrontendDocument.cpp b/Engine/Plugins/Runtime/Metasound/Source/MetasoundFrontend/Private/MetasoundFrontendDocument.cpp
--- a/Engine/Plugins/Runtime/Metasound/Source/MetasoundFrontend/Private/MetasoundFrontendDocument.cpp
+++ b/Engine/Plugins/Runtime/Metasound/Source/MetasoundFrontend/Private/MetasoundFrontendDocument.cpp
@@ -838,6 +838,38 @@ FMetasoundFrontendDocument::FMetasoundFrontendDocument()
 	ArchetypeVersion = FMetasoundFrontendVersion::GetInvalid();
 }
 
+const TCHAR* LexToString(EMetasoundFrontendClassType InClassType)
+{
+	switch (InClassType)
+	{
+		case EMetasoundFrontendClassType::External:
+			return TEXT("External");
+		case EMetasoundFrontendClassType::Graph:
+			return TEXT("Graph");
+		case EMetasoundFrontendClassType::Input:
+			return TEXT("Input");
+		case EMetasoundFrontendClassType::Output:
+			return TEXT("Output");
+		case EMetasoundFrontendClassType::Literal:
+			return TEXT("Literal");
+		case EMetasoundFrontendClassType::Variable:
+			return TEXT("Variable");
+		case EMetasoundFrontendClassType::VariableDeferredAccessor:
+			return TEXT("Variable (Deferred Accessor)");
+		case EMetasoundFrontendClassType::VariableAccessor:
+			return TEXT("Variable (Accessor)");
+		case EMetasoundFrontendClassType::VariableMutator:
+			return TEXT("Variable (Mutator)");
+		case EMetasoundFrontendClassType::Template:
+			return TEXT("Template");
+		case EMetasoundFrontendClassType::Invalid:
+			return TEXT("Invalid");
+		default:
+			static_assert(static_cast<int32>(EMetasoundFrontendClassType::Invalid) == 10, "Possible missed EMetasoundFrontendClassType case coverage");
+			return nullptr;
+	}
+}
+
 const TCHAR* LexToString(EMetasoundFrontendVertexAccessType InVertexAccess)
 {
 	switch (InVertexAccess)
diff --git a/Engine/Plugins/Runtime/Metasound/Source/MetasoundFrontend/Public/MetasoundFrontendDocument.h b/Engine/Plugins/Runtime/Metasound/Source/MetasoundFrontend/Public/MetasoundFrontendDocument.h
--- a/Engine/Plugins/Runtime/Metasound/Source/MetasoundFrontend/Public/MetasoundFrontendDocument.h
+++ b/Engine/Plugins/Runtime/Metasound/Source/MetasoundFrontend/Public/MetasoundFrontendDocument.h
@@ -1201,15 +1201,15 @@ struct METASOUNDFRONTEND_API FMetasoundFrontendClassName
 	FMetasoundFrontendClassName(const Metasound::FNodeClassName& InName);
 
 	// Namespace of class.
-	UPROPERTY(EditAnywhere, Category = General)
+	UPROPERTY(EditAnywhere, BlueprintReadWrite, Category = General)
 	FName Namespace;
 
 	// Name of class.
-	UPROPERTY(EditAnywhere, Category = General)
+	UPROPERTY(EditAnywhere, BlueprintReadWrite, Category = General)
 	FName Name;
 
 	// Variant of class. The Variant is used to describe an equivalent class which performs the same operation but on differing types.
-	UPROPERTY(EditAnywhere, Category = General)
+	UPROPERTY(EditAnywhere, BlueprintReadWrite, AdvancedDisplay, Category = General)
 	FName Variant;
 
 	// Returns a full name of the class.
@@ -1625,5 +1625,5 @@ public:
 	}
 };
 
-
+METASOUNDFRONTEND_API const TCHAR* LexToString(EMetasoundFrontendClassType InClassType);
 METASOUNDFRONTEND_API const TCHAR* LexToString(EMetasoundFrontendVertexAccessType InVertexAccess);
