[Metasounds] Follow up fixes for new BindInputs/ BindOutputs API.

diff --git a/Engine/Plugins/Runtime/Metasound/Source/MetasoundEngine/Private/MetasoundWavePlayerNode.cpp b/Engine/Plugins/Runtime/Metasound/Source/MetasoundEngine/Private/MetasoundWavePlayerNode.cpp
--- a/Engine/Plugins/Runtime/Metasound/Source/MetasoundEngine/Private/MetasoundWavePlayerNode.cpp
+++ b/Engine/Plugins/Runtime/Metasound/Source/MetasoundEngine/Private/MetasoundWavePlayerNode.cpp
@@ -359,8 +359,7 @@ namespace Metasound
 		virtual void BindInputs(FInputVertexInterfaceData& InOutVertexData) override
 		{
 			using namespace WavePlayerVertexNames;
-
-			FDataReferenceCollection InputDataReferences;
+			
 			InOutVertexData.BindReadVertex(METASOUND_GET_PARAM_NAME(InputTriggerPlay), PlayTrigger);
 			InOutVertexData.BindReadVertex(METASOUND_GET_PARAM_NAME(InputTriggerStop), StopTrigger);
 			InOutVertexData.BindReadVertex(METASOUND_GET_PARAM_NAME(InputWaveAsset), WaveAsset);
@@ -375,7 +374,6 @@ namespace Metasound
 		{
 			using namespace WavePlayerVertexNames;
 
-			FDataReferenceCollection OutputDataReferences;
 			InOutVertexData.BindReadVertex(METASOUND_GET_PARAM_NAME(OutputTriggerOnPlay), PlayTrigger);
 			InOutVertexData.BindReadVertex(METASOUND_GET_PARAM_NAME(OutputTriggerOnDone), TriggerOnDone);
 			InOutVertexData.BindReadVertex(METASOUND_GET_PARAM_NAME(OutputTriggerOnNearlyDone), TriggerOnNearlyDone);
@@ -391,7 +389,7 @@ namespace Metasound
 
 			for (int32 i = 0; i < OutputAudioBuffers.Num(); i++)
 			{
-				OutputDataReferences.AddDataReadReference(OutputAudioBufferVertexNames[i], OutputAudioBuffers[i]);
+				InOutVertexData.BindReadVertex(OutputAudioBufferVertexNames[i], OutputAudioBuffers[i]);
 			}
 		}
 
