Fix for cache corruption when removing inputs/outputs using document builder
#rnx

diff --git a/Engine/Plugins/Runtime/Metasound/Source/MetasoundFrontend/Private/MetasoundFrontendDocumentCache.cpp b/Engine/Plugins/Runtime/Metasound/Source/MetasoundFrontend/Private/MetasoundFrontendDocumentCache.cpp
--- a/Engine/Plugins/Runtime/Metasound/Source/MetasoundFrontend/Private/MetasoundFrontendDocumentCache.cpp
+++ b/Engine/Plugins/Runtime/Metasound/Source/MetasoundFrontend/Private/MetasoundFrontendDocumentCache.cpp
@@ -2,6 +2,7 @@
 
 #include "MetasoundFrontendDocumentCache.h"
 
+#include "Algo/ForEach.h"
 #include "MetasoundFrontendDocument.h"
 #include "MetasoundFrontendDocumentBuilder.h"
 #include "MetasoundFrontendRegistries.h"
@@ -204,6 +205,7 @@ namespace Metasound::Frontend
 		const FMetasoundFrontendGraphClass& GraphClass = Document.RootGraph;
 		const FMetasoundFrontendClassInput& Input = GraphClass.Interface.Inputs[IndexBeingRemoved];
 		InputNameToIndex.Remove(Input.Name);
+		Algo::ForEach(InputNameToIndex, [](TPair<FName, int32>& Entry) { --Entry.Value; });
 	}
 
 	void FDocumentGraphInterfaceCache::OnRemovingOutput(int32 IndexBeingRemoved)
@@ -212,6 +214,7 @@ namespace Metasound::Frontend
 		const FMetasoundFrontendGraphClass& GraphClass = Document.RootGraph;
 		const FMetasoundFrontendClassOutput& Output = GraphClass.Interface.Outputs[IndexBeingRemoved];
 		OutputNameToIndex.Remove(Output.Name);
+		Algo::ForEach(OutputNameToIndex, [](TPair<FName, int32>& Entry) { --Entry.Value; });
 	}
 
 	FDocumentGraphNodeCache::FDocumentGraphNodeCache(TSharedRef<IDocumentCache> ParentCache)
