Forbid naming MetaSound members as that of existing namespace members on the given document
#rnx

diff --git a/Engine/Plugins/Runtime/Metasound/Source/MetasoundEditor/Private/MetasoundEditorGraph.cpp b/Engine/Plugins/Runtime/Metasound/Source/MetasoundEditor/Private/MetasoundEditorGraph.cpp
--- a/Engine/Plugins/Runtime/Metasound/Source/MetasoundEditor/Private/MetasoundEditorGraph.cpp
+++ b/Engine/Plugins/Runtime/Metasound/Source/MetasoundEditor/Private/MetasoundEditorGraph.cpp
@@ -620,17 +620,26 @@ bool UMetasoundEditorGraphVertex::CanRename(const FText& InNewText, FText& OutEr
 	{
 		if (NodeID != NodeToCompare->GetID())
 		{
-			const FName& OtherName = NodeToCompare->GetNodeName();
+			FName OtherName = NodeToCompare->GetNodeName();
 			if (NewName == OtherName)
 			{
 				bIsNameValid = false;
 				OutError = FText::Format(LOCTEXT("GraphVertexRenameInvalid_NameTaken", "{0} is already in use"), InNewText);
 			}
-
-			if (Namespace == OtherName)
+			else if (Namespace == OtherName)
 			{
 				bIsNameValid = false;
-				OutError = FText::Format(LOCTEXT("GraphVertexRenameInvalid_NamespaceTaken", "{0} cannot share name of existing namespace"), InNewText);
+				OutError = FText::Format(LOCTEXT("GraphVertexRenameInvalid_NamespaceTaken", "Namespace of '{0}' cannot be the same as an existing member's name"), InNewText);
+			}
+			else
+			{
+				FName OtherNamespace;
+				Audio::FParameterPath::SplitName(OtherName, OtherNamespace, OtherName);
+				if (OtherNamespace == NewName)
+				{
+					bIsNameValid = false;
+					OutError = FText::Format(LOCTEXT("GraphVertexRenameInvalid_NamespaceTaken2", "Name of '{0}' cannot be the same as an existing member's namespace"), InNewText);
+				}
 			}
 		}
 	}, GetClassType());
@@ -1134,7 +1143,7 @@ bool UMetasoundEditorGraphVariable::CanRename(const FText& InNewText, FText& Out
 	{
 		if (VariableID != OtherVariable->GetID())
 		{
-			const FName& OtherName = OtherVariable->GetName();
+			FName OtherName = OtherVariable->GetName();
 			if (NewName == OtherName)
 			{
 				OutError = FText::Format(LOCTEXT("GraphVariableRenameInvalid_NameTaken", "{0} is already in use"), InNewText);
@@ -1143,7 +1152,15 @@ bool UMetasoundEditorGraphVariable::CanRename(const FText& InNewText, FText& Out
 
 			if (Namespace == OtherName)
 			{
-				OutError = FText::Format(LOCTEXT("GraphVariableRenameInvalid_NamespaceTaken", "{0} cannot share name of existing namespace"), InNewText);
+				OutError = FText::Format(LOCTEXT("GraphVariableRenameInvalid_NamespaceTaken", "Namespace of '{0}' cannot be the same as an existing variable's name"), InNewText);
+				return false;
+			}
+
+			FName OtherNamespace;
+			Audio::FParameterPath::SplitName(OtherName, OtherNamespace, OtherName);
+			if (OtherNamespace == NewName)
+			{
+				OutError = FText::Format(LOCTEXT("GraphVariableRenameInvalid_NamespaceTaken2", "Name of '{0}' cannot be the same as an existing variable's namespace"), InNewText);
 				return false;
 			}
 		}
